Learning Agents: Create initial experimental version

diff --git a/Engine/Plugins/Experimental/LearningAgents/Source/Learning/Private/LearningArrayMap.cpp b/Engine/Plugins/Experimental/LearningAgents/Source/Learning/Private/LearningArrayMap.cpp
new file mode 100644
--- /dev/null
+++ b/Engine/Plugins/Experimental/LearningAgents/Source/Learning/Private/LearningArrayMap.cpp
@@ -0,0 +1,246 @@
+// Copyright Epic Games, Inc. All Rights Reserved.
+
+#include "LearningArrayMap.h"
+
+namespace UE::Learning
+{
+	FArrayMap::FArrayMap() {}
+	FArrayMap::~FArrayMap() { Empty(); }
+
+	void FArrayMap::Empty()
+	{
+		Handles.Empty();
+
+		const int32 VariableNum = Arrays.Num();
+
+		for (int32 VariableIdx = 0; VariableIdx < VariableNum; VariableIdx++)
+		{
+			if (Flags[VariableIdx] & FlagConstructed)
+			{
+				Destructors[VariableIdx](
+					Arrays[VariableIdx].Data,
+					DimNums[VariableIdx],
+					Arrays[VariableIdx].Shape);
+			}
+
+			if (Flags[VariableIdx] & FlagAllocated)
+			{
+				FMemory::Free(Arrays[VariableIdx].Data);
+			}
+		}
+
+		Arrays.Empty();
+		Destructors.Empty();
+		DimNums.Empty();
+		Flags.Empty();
+		TypeIds.Empty();
+		TypeNames.Empty();
+	}
+
+	FArrayMapHandle FArrayMap::Lookup(const FArrayMapKey Key) const
+	{
+		return Handles[Key];
+	}
+
+	const FArrayMapHandle* FArrayMap::Find(const FArrayMapKey Key) const
+	{
+		return Handles.Find(Key);
+	}
+
+	int16 FArrayMap::TypeId(const FArrayMapKey Key) const
+	{
+		return TypeIds[Lookup(Key).Index];
+	}
+
+	int16 FArrayMap::TypeId(const FArrayMapHandle Variable) const
+	{
+		return TypeIds[Variable.Index];
+	}
+
+	const TCHAR* FArrayMap::TypeName(const FArrayMapKey Key) const
+	{
+		return TypeNames[Lookup(Key).Index];
+	}
+
+	const TCHAR* FArrayMap::TypeName(const FArrayMapHandle Variable) const
+	{
+		return TypeNames[Variable.Index];
+	}
+
+	uint8 FArrayMap::DimNum(const FArrayMapKey Key) const
+	{
+		return DimNums[Lookup(Key).Index];
+	}
+
+	uint8 FArrayMap::DimNum(const FArrayMapHandle Handle) const
+	{
+		return DimNums[Handle.Index];
+	}
+
+	bool FArrayMap::Contains(const FArrayMapKey Key) const
+	{
+		return Handles.Contains(Key);
+	}
+
+	void FArrayMap::Link(const FArrayMapKey Src, const FArrayMapKey Dst)
+	{
+		const FArrayMapHandle* SrcVariable = Handles.Find(Src);
+		const FArrayMapHandle* DstVariable = Handles.Find(Dst);
+
+		UE_LEARNING_CHECKF(SrcVariable, TEXT("Does not contain an array called { \"%s\", \"%s\" }"),
+			*Src.Namespace.ToString(), *Src.Variable.ToString());
+
+		UE_LEARNING_CHECKF(DstVariable, TEXT("Does not contain an array called { \"%s\", \"%s\" }"),
+			*Dst.Namespace.ToString(), *Dst.Variable.ToString());
+
+		const int16 DstIndex = DstVariable->Index;
+		const int16 SrcIndex = SrcVariable->Index;
+
+		UE_LEARNING_CHECKF(!(Flags[DstIndex] & FlagLinked), TEXT("Destination array is already linked!"));
+
+		UE_LEARNING_CHECKF(TypeIds[SrcIndex] == TypeIds[DstIndex],
+			TEXT("Source array { \"%s\", \"%s\" } is of type %s, while destination array { \"%s\", \"%s\" } is of type %s"),
+			*Src.Namespace.ToString(), *Src.Variable.ToString(), TypeNames[SrcIndex],
+			*Dst.Namespace.ToString(), *Dst.Variable.ToString(), TypeNames[DstIndex]);
+
+		UE_LEARNING_CHECKF(DimNums[SrcIndex] == DimNums[DstIndex],
+			TEXT("Source array { \"%s\", \"%s\" } has %i dimensions, while destination array { \"%s\", \"%s\" } has %i dimensions"),
+			*Src.Namespace.ToString(), *Src.Variable.ToString(), DimNums[SrcIndex],
+			*Dst.Namespace.ToString(), *Dst.Variable.ToString(), DimNums[DstIndex]);
+
+		const uint8 DimNum = DimNums[SrcIndex];
+
+		for (uint8 DimIdx = 0; DimIdx < DimNum; DimIdx++)
+		{
+			UE_LEARNING_CHECKF(Arrays[SrcIndex].Shape[DimIdx] == Arrays[DstIndex].Shape[DimIdx],
+				TEXT("Array Shapes don't match on dimension %i of %i ({ \"%s\", \"%s\" }: %i, { \"%s\", \"%s\" }: %i)"),
+				DimIdx + 1, DimNum,
+				*Src.Namespace.ToString(), *Src.Variable.ToString(), Arrays[SrcIndex].Shape[DimIdx],
+				*Dst.Namespace.ToString(), *Dst.Variable.ToString(), Arrays[DstIndex].Shape[DimIdx]);
+		}
+
+		// Free the Destination array data
+		if (Flags[DstIndex] & FlagConstructed)
+		{
+			Destructors[DstIndex](
+				Arrays[DstIndex].Data,
+				DimNums[DstIndex],
+				Arrays[DstIndex].Shape);
+		}
+
+		if (Flags[DstIndex] & FlagAllocated)
+		{
+			FMemory::Free(Arrays[DstIndex].Data);
+		}
+
+		// Re-direct existing Links
+		for (int16 ArrayIdx = 0; ArrayIdx < Arrays.Num(); ArrayIdx++)
+		{
+			if (Flags[ArrayIdx] == FlagLinked && Arrays[ArrayIdx].Data == Arrays[DstIndex].Data)
+			{
+				Arrays[ArrayIdx].Data = Arrays[SrcIndex].Data;
+			}
+		}
+
+		// Insert the new Link
+		Arrays[DstIndex].Data = Arrays[SrcIndex].Data;
+		Flags[DstIndex] = FlagLinked;
+	}
+
+	void FArrayMap::Link(const FArrayMapHandle Src, const FArrayMapHandle Dst)
+	{
+		const int16 SrcIndex = Src.Index;
+		const int16 DstIndex = Dst.Index;
+
+		UE_LEARNING_CHECKF(!(Flags[DstIndex] & FlagLinked), TEXT("Destination array is already linked!"));
+
+		UE_LEARNING_CHECKF(TypeIds[SrcIndex] == TypeIds[DstIndex],
+			TEXT("Source array is of type %s, while destination array is of type %s"),
+			TypeNames[SrcIndex], TypeNames[DstIndex]);
+
+		UE_LEARNING_CHECKF(DimNums[SrcIndex] == DimNums[DstIndex],
+			TEXT("Source array has %i dimensions, while destination array has %i dimensions"),
+			DimNums[SrcIndex], DimNums[DstIndex]);
+
+		const uint8 DimNum = DimNums[SrcIndex];
+
+		for (uint8 DimIdx = 0; DimIdx < DimNum; DimIdx++)
+		{
+			UE_LEARNING_CHECKF(Arrays[SrcIndex].Shape[DimIdx] == Arrays[DstIndex].Shape[DimIdx],
+				TEXT("Array Shapes don't match on dimension %i of %i (%i vs %i)"),
+				DimIdx + 1, DimNum, Arrays[SrcIndex].Shape[DimIdx], Arrays[DstIndex].Shape[DimIdx]);
+		}
+
+		// Free the Destination array data
+		if (Flags[DstIndex] & FlagConstructed)
+		{
+			Destructors[DstIndex](
+				Arrays[DstIndex].Data,
+				DimNums[DstIndex],
+				Arrays[DstIndex].Shape);
+		}
+
+		if (Flags[DstIndex] & FlagAllocated)
+		{
+			FMemory::Free(Arrays[DstIndex].Data);
+		}
+
+		// Re-direct existing Links
+		for (int16 ArrayIdx = 0; ArrayIdx < Arrays.Num(); ArrayIdx++)
+		{
+			if (Flags[ArrayIdx] == FlagLinked && Arrays[ArrayIdx].Data == Arrays[DstIndex].Data)
+			{
+				Arrays[ArrayIdx].Data = Arrays[SrcIndex].Data;
+			}
+		}
+
+		// Insert the new Link
+		Arrays[DstIndex].Data = Arrays[SrcIndex].Data;
+		Flags[DstIndex] = FlagLinked;
+	}
+
+	bool FArrayMap::HasLink(const FArrayMapKey Dst) const
+	{
+		const FArrayMapHandle* DstVariable = Handles.Find(Dst);
+
+		UE_LEARNING_CHECKF(DstVariable, TEXT("Does not contain an array called { \"%s\", \"%s\" }"),
+			*Dst.Namespace.ToString(), *Dst.Variable.ToString());
+
+		return HasLink(*DstVariable);
+	}
+
+	bool FArrayMap::HasLink(const FArrayMapHandle Dst) const
+	{
+		return Flags[Dst.Index] & FlagLinked;
+	}
+
+	bool FArrayMap::IsLinkedTo(const FArrayMapKey Src) const
+	{
+		const FArrayMapHandle* SrcVariable = Handles.Find(Src);
+
+		UE_LEARNING_CHECKF(SrcVariable, TEXT("Does not contain an array called { \"%s\", \"%s\" }"),
+			*Src.Namespace.ToString(), *Src.Variable.ToString());
+
+		return IsLinkedTo(*SrcVariable);
+	}
+
+	bool FArrayMap::IsLinkedTo(const FArrayMapHandle Src) const
+	{
+		const int32 VariableNum = Arrays.Num();
+
+		for (int32 VariableIdx = 0; VariableIdx < VariableNum; VariableIdx++)
+		{
+			if (Arrays[VariableIdx].Data == Arrays[Src.Index].Data && (Flags[VariableIdx] & FlagLinked))
+			{
+				return true;
+			}
+		}
+
+		return false;
+	}
+
+	const FArrayMapKey* FArrayMap::FindKey(FArrayMapHandle Handle) const
+	{
+		return Handles.FindKey(Handle);
+	}
+}
diff --git a/Engine/Plugins/Experimental/LearningAgents/Source/Learning/Private/LearningCompletion.cpp b/Engine/Plugins/Experimental/LearningAgents/Source/Learning/Private/LearningCompletion.cpp
new file mode 100644
--- /dev/null
+++ b/Engine/Plugins/Experimental/LearningAgents/Source/Learning/Private/LearningCompletion.cpp
@@ -0,0 +1,94 @@
+// Copyright Epic Games, Inc. All Rights Reserved.
+
+#include "LearningCompletion.h"
+
+namespace UE::Learning
+{
+	void FResetInstanceBuffer::Resize(const int32 InMaxInstanceNum)
+	{
+		MaxInstanceNum = InMaxInstanceNum;
+		ResetInstanceNum = 0;
+		ResetInstances.SetNumUninitialized({ MaxInstanceNum });
+	}
+
+	void FResetInstanceBuffer::SetResetInstances(const FIndexSet Instances)
+	{
+		ResetInstanceNum = 0;
+
+		for (const int32 InstanceIdx : Instances)
+		{
+			ResetInstances[ResetInstanceNum] = InstanceIdx;
+			ResetInstanceNum++;
+		}
+
+		ResetInstancesSet = ResetInstances.Slice(0, ResetInstanceNum);
+		ResetInstancesSet.TryMakeSlice();
+	}
+
+	void FResetInstanceBuffer::SetResetInstancesFromCompletions(const TLearningArrayView<1, const ECompletionMode> Completions, const FIndexSet Instances)
+	{
+		ResetInstanceNum = 0;
+
+		for (const int32 InstanceIdx : Instances)
+		{
+			if (Completions[InstanceIdx] != ECompletionMode::Running)
+			{
+				ResetInstances[ResetInstanceNum] = InstanceIdx;
+				ResetInstanceNum++;
+			}
+		}
+
+		ResetInstancesSet = ResetInstances.Slice(0, ResetInstanceNum);
+		ResetInstancesSet.TryMakeSlice();
+	}
+
+	const int32 FResetInstanceBuffer::GetResetInstanceNum() const { return ResetInstanceNum; }
+
+	const FIndexSet FResetInstanceBuffer::GetResetInstances() const { return ResetInstancesSet; }
+
+	namespace Completion
+	{
+		const TCHAR* CompletionModeString(const ECompletionMode Completion)
+		{
+			switch (Completion)
+			{
+			case ECompletionMode::Running: return TEXT("Running");
+			case ECompletionMode::Truncated: return TEXT("Truncated");
+			case ECompletionMode::Terminated: return TEXT("Terminated");
+			default: UE_LEARNING_NOT_IMPLEMENTED(); return TEXT("Unknown");
+			}
+		}
+
+		ECompletionMode Or(const ECompletionMode Lhs, const ECompletionMode Rhs)
+		{
+			if ((Lhs == ECompletionMode::Running && Rhs != ECompletionMode::Running) ||
+				(Lhs == ECompletionMode::Truncated && Rhs == ECompletionMode::Terminated))
+			{
+				return Rhs;
+			}
+			else
+			{
+				return Lhs;
+			}
+		}
+
+		void EvaluateEndOfEpisodeCompletions(
+			TLearningArrayView<1, ECompletionMode> InOutCompletions,
+			const TLearningArrayView<1, const int32> EpisodeStepNums,
+			const int32 EpisodeMaxStepNum,
+			const ECompletionMode EpisodeEndCompletionMode,
+			const FIndexSet Instances)
+		{
+			UE_LEARNING_TRACE_CPUPROFILER_EVENT_SCOPE(Learning::Completion::EvaluateEndOfEpisodeCompletions);
+
+			for (const int32 InstanceIdx : Instances)
+			{
+				if (InOutCompletions[InstanceIdx] == ECompletionMode::Running && EpisodeStepNums[InstanceIdx] == EpisodeMaxStepNum)
+				{
+					InOutCompletions[InstanceIdx] = EpisodeEndCompletionMode;
+				}
+			}
+		}
+	}
+}
+
diff --git a/Engine/Plugins/Experimental/LearningAgents/Source/Learning/Private/LearningCompletionObject.cpp b/Engine/Plugins/Experimental/LearningAgents/Source/Learning/Private/LearningCompletionObject.cpp
new file mode 100644
--- /dev/null
+++ b/Engine/Plugins/Experimental/LearningAgents/Source/Learning/Private/LearningCompletionObject.cpp
@@ -0,0 +1,276 @@
+// Copyright Epic Games, Inc. All Rights Reserved.
+
+#include "LearningCompletionObject.h"
+#include "LearningLog.h"
+
+namespace UE::Learning
+{
+	FCompletionObject::FCompletionObject(
+		const FName& InIdentifier,
+		const TSharedRef<FArrayMap>& InInstanceData,
+		const int32 InMaxInstanceNum,
+		const ECompletionMode InCompletionMode)
+		: InstanceData(InInstanceData)
+		, CompletionMode(InCompletionMode)
+	{
+		CompletionHandle = InstanceData->Add<1, ECompletionMode>({ InIdentifier, TEXT("Completion") }, { InMaxInstanceNum }, ECompletionMode::Running);
+	}
+
+	TLearningArrayView<1, ECompletionMode> FCompletionObject::CompletionBuffer()
+	{
+		return InstanceData->View(CompletionHandle);
+	}
+
+	//------------------------------------------------------------------
+
+	FAnyCompletion::FAnyCompletion(
+		const FName& InIdentifier,
+		const TLearningArrayView<1, const TSharedRef<FCompletionObject>> InCompletions,
+		const TSharedRef<FArrayMap>& InInstanceData,
+		const int32 InMaxInstanceNum)
+		: FCompletionObject(InIdentifier, InInstanceData, InMaxInstanceNum, ECompletionMode::Terminated)
+		, Completions(InCompletions) {}
+
+	void FAnyCompletion::Evaluate(const FIndexSet Instances)
+	{
+		UE_LEARNING_TRACE_CPUPROFILER_EVENT_SCOPE(FAnyCompletion::Evaluate);
+
+		TLearningArrayView<1, ECompletionMode> Completion = InstanceData->View(CompletionHandle);
+
+		const int32 CompletionNum = Completions.Num();
+
+		for (int32 CompletionIdx = 0; CompletionIdx < CompletionNum; CompletionIdx++)
+		{
+			Completions[CompletionIdx]->Evaluate(Instances);
+		}
+
+		Array::Set(Completion, ECompletionMode::Running, Instances);
+
+		for (int32 CompletionIdx = 0; CompletionIdx < CompletionNum; CompletionIdx++)
+		{
+			const TLearningArrayView<1, const ECompletionMode> Input = InstanceData->ConstView(Completions[CompletionIdx]->CompletionHandle);
+
+			for (const int32 InstanceIdx : Instances)
+			{
+				Completion[InstanceIdx] = Completion::Or(Completion[InstanceIdx], Input[InstanceIdx]);
+			}
+		}
+	}
+
+	//------------------------------------------------------------------
+
+	FConditionalCompletion::FConditionalCompletion(
+		const FName& InIdentifier,
+		const TSharedRef<FArrayMap>& InInstanceData,
+		const int32 InMaxInstanceNum,
+		const ECompletionMode InCompletionMode)
+		: FCompletionObject(InIdentifier, InInstanceData, InMaxInstanceNum, InCompletionMode)
+	{
+		ConditionHandle = InstanceData->Add<1, bool>({ InIdentifier, TEXT("Condition") }, { InMaxInstanceNum }, 0.0f);
+	}
+
+	void FConditionalCompletion::Evaluate(const FIndexSet Instances)
+	{
+		UE_LEARNING_TRACE_CPUPROFILER_EVENT_SCOPE(FConditionalCompletion::Evaluate);
+
+		const TLearningArrayView<1, const bool> Condition = InstanceData->ConstView(ConditionHandle);
+		TLearningArrayView<1, ECompletionMode> Completion = InstanceData->View(CompletionHandle);
+
+		for (const int32 InstanceIdx : Instances)
+		{
+			Completion[InstanceIdx] = Condition[InstanceIdx] ? CompletionMode : ECompletionMode::Running;
+		}
+	}
+
+	//------------------------------------------------------------------
+
+	FScalarPositionDifferenceCompletion::FScalarPositionDifferenceCompletion(
+		const FName& InIdentifier,
+		const TSharedRef<FArrayMap>& InInstanceData,
+		const int32 InMaxInstanceNum,
+		const float InThreshold,
+		const ECompletionMode InCompletionMode)
+		: FCompletionObject(InIdentifier, InInstanceData, InMaxInstanceNum, InCompletionMode)
+	{
+		Position0Handle = InstanceData->Add<1, float>({ InIdentifier, TEXT("Position0") }, { InMaxInstanceNum }, 0.0f);
+		Position1Handle = InstanceData->Add<1, float>({ InIdentifier, TEXT("Position1") }, { InMaxInstanceNum }, 0.0f);
+		ThresholdHandle = InstanceData->Add<1, float>({ InIdentifier, TEXT("Threshold") }, { InMaxInstanceNum }, InThreshold);
+	}
+
+	void FScalarPositionDifferenceCompletion::Evaluate(const FIndexSet Instances)
+	{
+		UE_LEARNING_TRACE_CPUPROFILER_EVENT_SCOPE(FScalarPositionDifferenceCompletion::Evaluate);
+
+		const TLearningArrayView<1, const float> Position0 = InstanceData->ConstView(Position0Handle);
+		const TLearningArrayView<1, const float> Position1 = InstanceData->ConstView(Position1Handle);
+		const TLearningArrayView<1, const float> Threshold = InstanceData->ConstView(ThresholdHandle);
+		TLearningArrayView<1, ECompletionMode> Completion = InstanceData->View(CompletionHandle);
+
+		for (const int32 InstanceIdx : Instances)
+		{
+			if (FMath::Abs(Position0[InstanceIdx] - Position1[InstanceIdx]) > Threshold[InstanceIdx])
+			{
+				Completion[InstanceIdx] = CompletionMode;
+			}
+			else
+			{
+				Completion[InstanceIdx] = ECompletionMode::Running;
+			}
+		}
+	}
+
+	FPlanarPositionDifferenceCompletion::FPlanarPositionDifferenceCompletion(
+		const FName& InIdentifier,
+		const TSharedRef<FArrayMap>& InInstanceData,
+		const int32 InMaxInstanceNum,
+		const float InThreshold,
+		const ECompletionMode InCompletionMode,
+		const FVector InAxis0,
+		const FVector InAxis1)
+		: FCompletionObject(InIdentifier, InInstanceData, InMaxInstanceNum, InCompletionMode)
+		, Axis0(InAxis0)
+		, Axis1(InAxis1)
+	{
+		Position0Handle = InstanceData->Add<1, FVector>({ InIdentifier, TEXT("Position0") }, { InMaxInstanceNum }, FVector::ZeroVector);
+		Position1Handle = InstanceData->Add<1, FVector>({ InIdentifier, TEXT("Position1") }, { InMaxInstanceNum }, FVector::ZeroVector);
+		ThresholdHandle = InstanceData->Add<1, float>({ InIdentifier, TEXT("Threshold") }, { InMaxInstanceNum }, InThreshold);
+	}
+
+	void FPlanarPositionDifferenceCompletion::Evaluate(const FIndexSet Instances)
+	{
+		UE_LEARNING_TRACE_CPUPROFILER_EVENT_SCOPE(FPlanarPositionDifferenceCompletion::Evaluate);
+
+		const TLearningArrayView<1, const FVector> Position0 = InstanceData->ConstView(Position0Handle);
+		const TLearningArrayView<1, const FVector> Position1 = InstanceData->ConstView(Position1Handle);
+		const TLearningArrayView<1, const float> Threshold = InstanceData->ConstView(ThresholdHandle);
+		TLearningArrayView<1, ECompletionMode> Completion = InstanceData->View(CompletionHandle);
+
+		for (const int32 InstanceIdx : Instances)
+		{
+			const FVector ProjectedPosition0 = FVector(Axis0.Dot(Position0[InstanceIdx]), Axis1.Dot(Position0[InstanceIdx]), 0.0f);
+			const FVector ProjectedPosition1 = FVector(Axis0.Dot(Position1[InstanceIdx]), Axis1.Dot(Position1[InstanceIdx]), 0.0f);
+
+			if (FVector::Distance(ProjectedPosition0, ProjectedPosition1) > Threshold[InstanceIdx])
+			{
+				Completion[InstanceIdx] = CompletionMode;
+			}
+			else
+			{
+				Completion[InstanceIdx] = ECompletionMode::Running;
+			}
+		}
+	}
+
+	//------------------------------------------------------------------
+
+	FScalarVelocityDifferenceCompletion::FScalarVelocityDifferenceCompletion(
+		const FName& InIdentifier,
+		const TSharedRef<FArrayMap>& InInstanceData,
+		const int32 InMaxInstanceNum,
+		const float InThreshold,
+		const ECompletionMode InCompletionMode)
+		: FCompletionObject(InIdentifier, InInstanceData, InMaxInstanceNum, InCompletionMode)
+	{
+		Velocity0Handle = InstanceData->Add<1, float>({ InIdentifier, TEXT("Velocity0") }, { InMaxInstanceNum }, 0.0f);
+		Velocity1Handle = InstanceData->Add<1, float>({ InIdentifier, TEXT("Velocity1") }, { InMaxInstanceNum }, 0.0f);
+		ThresholdHandle = InstanceData->Add<1, float>({ InIdentifier, TEXT("Threshold") }, { InMaxInstanceNum }, InThreshold);
+	}
+
+	void FScalarVelocityDifferenceCompletion::Evaluate(const FIndexSet Instances)
+	{
+		UE_LEARNING_TRACE_CPUPROFILER_EVENT_SCOPE(FScalarVelocityDifferenceCompletion::Evaluate);
+
+		const TLearningArrayView<1, const float> Velocity0 = InstanceData->ConstView(Velocity0Handle);
+		const TLearningArrayView<1, const float> Velocity1 = InstanceData->ConstView(Velocity1Handle);
+		const TLearningArrayView<1, const float> Threshold = InstanceData->ConstView(ThresholdHandle);
+		TLearningArrayView<1, ECompletionMode> Completion = InstanceData->View(CompletionHandle);
+
+		for (const int32 InstanceIdx : Instances)
+		{
+			if (FMath::Abs(Velocity0[InstanceIdx] - Velocity1[InstanceIdx]) > Threshold[InstanceIdx])
+			{
+				Completion[InstanceIdx] = CompletionMode;
+			}
+			else
+			{
+				Completion[InstanceIdx] = ECompletionMode::Running;
+			}
+		}
+	}
+
+	//------------------------------------------------------------------
+
+	FScalarRotationDifferenceCompletion::FScalarRotationDifferenceCompletion(
+		const FName& InIdentifier,
+		const TSharedRef<FArrayMap>& InInstanceData,
+		const int32 InMaxInstanceNum,
+		const float InThreshold,
+		const ECompletionMode InCompletionMode)
+		: FCompletionObject(InIdentifier, InInstanceData, InMaxInstanceNum, InCompletionMode)
+	{
+		Rotation0Handle = InstanceData->Add<1, float>({ InIdentifier, TEXT("Rotation0") }, { InMaxInstanceNum }, 0.0f);
+		Rotation1Handle = InstanceData->Add<1, float>({ InIdentifier, TEXT("Rotation1") }, { InMaxInstanceNum }, 0.0f);
+		ThresholdHandle = InstanceData->Add<1, float>({ InIdentifier, TEXT("Threshold") }, { InMaxInstanceNum }, InThreshold);
+	}
+
+	void FScalarRotationDifferenceCompletion::Evaluate(const FIndexSet Instances)
+	{
+		UE_LEARNING_TRACE_CPUPROFILER_EVENT_SCOPE(FScalarRotationDifferenceCompletion::Evaluate);
+
+		const TLearningArrayView<1, const float> Rotation0 = InstanceData->ConstView(Rotation0Handle);
+		const TLearningArrayView<1, const float> Rotation1 = InstanceData->ConstView(Rotation1Handle);
+		const TLearningArrayView<1, const float> Threshold = InstanceData->ConstView(ThresholdHandle);
+		TLearningArrayView<1, ECompletionMode> Completion = InstanceData->View(CompletionHandle);
+
+		for (const int32 InstanceIdx : Instances)
+		{
+			if (FMath::Abs(FMath::FindDeltaAngleRadians(Rotation1[InstanceIdx], Rotation0[InstanceIdx])) > Threshold[InstanceIdx])
+			{
+				Completion[InstanceIdx] = CompletionMode;
+			}
+			else
+			{
+				Completion[InstanceIdx] = ECompletionMode::Running;
+			}
+		}
+	}
+
+	//------------------------------------------------------------------
+
+	FScalarAngularVelocityDifferenceCompletion::FScalarAngularVelocityDifferenceCompletion(
+		const FName& InIdentifier,
+		const TSharedRef<FArrayMap>& InInstanceData,
+		const int32 InMaxInstanceNum,
+		const float InThreshold,
+		const ECompletionMode InCompletionMode)
+		: FCompletionObject(InIdentifier, InInstanceData, InMaxInstanceNum, InCompletionMode)
+	{
+		AngularVelocity0Handle = InstanceData->Add<1, float>({ InIdentifier, TEXT("AngularVelocity0") }, { InMaxInstanceNum }, 0.0f);
+		AngularVelocity1Handle = InstanceData->Add<1, float>({ InIdentifier, TEXT("AngularVelocity1") }, { InMaxInstanceNum }, 0.0f);
+		ThresholdHandle = InstanceData->Add<1, float>({ InIdentifier, TEXT("Threshold") }, { InMaxInstanceNum }, InThreshold);
+	}
+
+	void FScalarAngularVelocityDifferenceCompletion::Evaluate(const FIndexSet Instances)
+	{
+		UE_LEARNING_TRACE_CPUPROFILER_EVENT_SCOPE(FScalarAngularVelocityDifferenceCompletion::Evaluate);
+
+		const TLearningArrayView<1, const float> AngularVelocity0 = InstanceData->ConstView(AngularVelocity0Handle);
+		const TLearningArrayView<1, const float> AngularVelocity1 = InstanceData->ConstView(AngularVelocity1Handle);
+		const TLearningArrayView<1, const float> Threshold = InstanceData->ConstView(ThresholdHandle);
+		TLearningArrayView<1, ECompletionMode> Completion = InstanceData->View(CompletionHandle);
+
+		for (const int32 InstanceIdx : Instances)
+		{
+			if (FMath::Abs(AngularVelocity0[InstanceIdx] - AngularVelocity1[InstanceIdx]) > Threshold[InstanceIdx])
+			{
+				Completion[InstanceIdx] = CompletionMode;
+			}
+			else
+			{
+				Completion[InstanceIdx] = ECompletionMode::Running;
+			}
+		}
+	}
+
+	//------------------------------------------------------------------
+}
