Learning Agents: Create initial experimental version

diff --git a/Engine/Plugins/Experimental/LearningAgents/Source/Learning/Private/LearningFeatureObject.cpp b/Engine/Plugins/Experimental/LearningAgents/Source/Learning/Private/LearningFeatureObject.cpp
new file mode 100644
--- /dev/null
+++ b/Engine/Plugins/Experimental/LearningAgents/Source/Learning/Private/LearningFeatureObject.cpp
@@ -0,0 +1,2298 @@
+// Copyright Epic Games, Inc. All Rights Reserved.
+
+#include "LearningFeatureObject.h"
+#include "LearningLog.h"
+
+#if UE_LEARNING_ISPC
+#include "Learning.ispc.generated.h"
+#endif
+
+namespace UE::Learning
+{
+	namespace Feature::Encode
+	{
+		void Float(
+			float& Output,
+			const float Value,
+			const float Scale,
+			const float Epsilon)
+		{
+			UE_LEARNING_ARRAY_VALUE_CHECK(Scale > 0.0f);
+
+			Output = Value / FMath::Max(Scale, Epsilon);
+		}
+
+		void Time(
+			float& Output,
+			const float Time,
+			const float RelativeTime,
+			const float Scale,
+			const float Epsilon)
+		{
+			UE_LEARNING_ARRAY_VALUE_CHECK(Scale > 0.0f);
+
+			Output = (Time - RelativeTime) / FMath::Max(Scale, Epsilon);
+		}
+
+
+		void PlanarDirection(
+			TLearningArrayView<1, float> Output,
+			const FVector Direction,
+			const FQuat RelativeRotation,
+			const float Scale,
+			const FVector Axis0,
+			const FVector Axis1,
+			const float Epsilon)
+		{
+			UE_LEARNING_ARRAY_SHAPE_CHECK(Output.Num() == 2);
+			UE_LEARNING_ARRAY_VALUE_CHECK(Scale > 0.0f);
+
+			const FVector LocalDirection = RelativeRotation.UnrotateVector(Direction);
+			const FVector PlanarDirection = FVector(LocalDirection.Dot(Axis0), LocalDirection.Dot(Axis1), 0.0f).GetSafeNormal(Epsilon, FVector::ForwardVector);
+			Output[0] = PlanarDirection.X / FMath::Max(Scale, Epsilon);
+			Output[1] = PlanarDirection.Y / FMath::Max(Scale, Epsilon);
+		}
+
+		void Direction(
+			TLearningArrayView<1, float> Output,
+			const FVector Direction,
+			const FQuat RelativeRotation,
+			const float Scale,
+			const float Epsilon)
+		{
+			UE_LEARNING_ARRAY_SHAPE_CHECK(Output.Num() == 3);
+			UE_LEARNING_ARRAY_VALUE_CHECK(Scale > 0.0f);
+
+			const FVector LocalDirection = RelativeRotation.UnrotateVector(Direction);
+			Output[0] = LocalDirection.X / FMath::Max(Scale, Epsilon);
+			Output[1] = LocalDirection.Y / FMath::Max(Scale, Epsilon);
+			Output[2] = LocalDirection.Z / FMath::Max(Scale, Epsilon);
+		}
+
+
+		void ScalarPosition(
+			float& Output,
+			const float Position,
+			const float RelativePosition,
+			const float Scale,
+			const float Epsilon)
+		{
+			UE_LEARNING_ARRAY_VALUE_CHECK(Scale > 0.0f);
+
+			Output = (Position - RelativePosition) / FMath::Max(Scale, Epsilon);
+		}
+
+		void PlanarPosition(
+			TLearningArrayView<1, float> Output,
+			const FVector Position,
+			const FVector RelativePosition,
+			const FQuat RelativeRotation,
+			const float Scale,
+			const FVector Axis0,
+			const FVector Axis1,
+			const float Epsilon)
+		{
+			UE_LEARNING_ARRAY_SHAPE_CHECK(Output.Num() == 2);
+			UE_LEARNING_ARRAY_VALUE_CHECK(Scale > 0.0f);
+
+			const FVector LocalPosition = RelativeRotation.UnrotateVector(Position - RelativePosition);
+			Output[0] = LocalPosition.Dot(Axis0) / FMath::Max(Scale, Epsilon);
+			Output[1] = LocalPosition.Dot(Axis1) / FMath::Max(Scale, Epsilon);
+		}
+
+		void Position(
+			TLearningArrayView<1, float> Output,
+			const FVector Position,
+			const FVector RelativePosition,
+			const FQuat RelativeRotation,
+			const float Scale,
+			const float Epsilon)
+		{
+			UE_LEARNING_ARRAY_SHAPE_CHECK(Output.Num() == 3);
+			UE_LEARNING_ARRAY_VALUE_CHECK(Scale > 0.0f);
+
+			const FVector LocalPosition = RelativeRotation.UnrotateVector(Position - RelativePosition);
+			Output[0] = LocalPosition.X / FMath::Max(Scale, Epsilon);
+			Output[1] = LocalPosition.Y / FMath::Max(Scale, Epsilon);
+			Output[2] = LocalPosition.Z / FMath::Max(Scale, Epsilon);
+		}
+
+
+		void ScalarVelocity(
+			float& Output,
+			const float Velocity,
+			const float Scale,
+			const float Epsilon)
+		{
+			UE_LEARNING_ARRAY_VALUE_CHECK(Scale > 0.0f);
+
+			Output = Velocity / FMath::Max(Scale, Epsilon);
+		}
+
+		void PlanarVelocity(
+			TLearningArrayView<1, float> Output,
+			const FVector Velocity,
+			const FQuat RelativeRotation,
+			const float Scale,
+			const FVector Axis0,
+			const FVector Axis1,
+			const float Epsilon)
+		{
+			UE_LEARNING_ARRAY_SHAPE_CHECK(Output.Num() == 2);
+			UE_LEARNING_ARRAY_VALUE_CHECK(Scale > 0.0f);
+
+			const FVector LocalVelocity = RelativeRotation.UnrotateVector(Velocity);
+			Output[0] = LocalVelocity.Dot(Axis0) / FMath::Max(Scale, Epsilon);
+			Output[1] = LocalVelocity.Dot(Axis1) / FMath::Max(Scale, Epsilon);
+		}
+
+		void Velocity(
+			TLearningArrayView<1, float> Output,
+			const FVector Velocity,
+			const FQuat RelativeRotation,
+			const float Scale,
+			const float Epsilon)
+		{
+			UE_LEARNING_ARRAY_SHAPE_CHECK(Output.Num() == 3);
+			UE_LEARNING_ARRAY_VALUE_CHECK(Scale > 0.0f);
+
+			const FVector LocalVelocity = RelativeRotation.UnrotateVector(Velocity);
+			Output[0] = LocalVelocity.X / FMath::Max(Scale, Epsilon);
+			Output[1] = LocalVelocity.Y / FMath::Max(Scale, Epsilon);
+			Output[2] = LocalVelocity.Z / FMath::Max(Scale, Epsilon);
+		}
+
+
+		void ScalarAcceleration(
+			float& Output,
+			const float Acceleration,
+			const float Scale,
+			const float Epsilon)
+		{
+			UE_LEARNING_ARRAY_VALUE_CHECK(Scale > 0.0f);
+
+			Output = Acceleration  / FMath::Max(Scale, Epsilon);
+		}
+
+		void PlanarAcceleration(
+			TLearningArrayView<1, float> Output,
+			const FVector Acceleration,
+			const FQuat RelativeRotation,
+			const float Scale,
+			const FVector Axis0,
+			const FVector Axis1,
+			const float Epsilon)
+		{
+			UE_LEARNING_ARRAY_SHAPE_CHECK(Output.Num() == 2);
+			UE_LEARNING_ARRAY_VALUE_CHECK(Scale > 0.0f);
+
+			const FVector LocalAcceleration = RelativeRotation.UnrotateVector(Acceleration);
+			Output[0] = LocalAcceleration.Dot(Axis0) / FMath::Max(Scale, Epsilon);
+			Output[1] = LocalAcceleration.Dot(Axis1) / FMath::Max(Scale, Epsilon);
+		}
+
+		void Acceleration(
+			TLearningArrayView<1, float> Output,
+			const FVector Acceleration,
+			const FQuat RelativeRotation,
+			const float Scale,
+			const float Epsilon)
+		{
+			UE_LEARNING_ARRAY_SHAPE_CHECK(Output.Num() == 3);
+			UE_LEARNING_ARRAY_VALUE_CHECK(Scale > 0.0f);
+
+			const FVector LocalAcceleration = RelativeRotation.UnrotateVector(Acceleration);
+			Output[0] = LocalAcceleration.X / FMath::Max(Scale, Epsilon);
+			Output[1] = LocalAcceleration.Y / FMath::Max(Scale, Epsilon);
+			Output[2] = LocalAcceleration.Z / FMath::Max(Scale, Epsilon);
+		}
+
+
+		void Angle(
+			TLearningArrayView<1, float> Output,
+			const float Angle,
+			const float RelativeAngle,
+			const float Scale,
+			const float Epsilon)
+		{
+			UE_LEARNING_ARRAY_SHAPE_CHECK(Output.Num() == 2);
+			UE_LEARNING_ARRAY_VALUE_CHECK(Scale > 0.0f);
+
+			const float LocalAngle = FMath::FindDeltaAngleRadians(RelativeAngle, Angle);
+			Output[0] = FMath::Sin(LocalAngle) / FMath::Max(Scale, Epsilon);
+			Output[1] = FMath::Cos(LocalAngle) / FMath::Max(Scale, Epsilon);
+		}
+
+		void PlanarRotation(
+			TLearningArrayView<1, float> Output,
+			const FQuat Rotation,
+			const FQuat RelativeRotation,
+			const float Scale,
+			const FVector Axis,
+			const float Epsilon)
+		{
+			UE_LEARNING_ARRAY_SHAPE_CHECK(Output.Num() == 2);
+			UE_LEARNING_ARRAY_VALUE_CHECK(Scale > 0.0f);
+
+			const FQuat LocalRotation = RelativeRotation.Inverse() * Rotation;
+			const float LocalAngle = LocalRotation.GetTwistAngle(Axis);
+			Output[0] = FMath::Sin(LocalAngle) / FMath::Max(Scale, Epsilon);
+			Output[1] = FMath::Cos(LocalAngle) / FMath::Max(Scale, Epsilon);
+		}
+
+		void Rotation(
+			TLearningArrayView<1, float> Output,
+			const FQuat Rotation,
+			const FQuat RelativeRotation,
+			const float Scale,
+			const float Epsilon)
+		{
+			UE_LEARNING_ARRAY_SHAPE_CHECK(Output.Num() == 6);
+			UE_LEARNING_ARRAY_VALUE_CHECK(Scale > 0.0f);
+
+			const FQuat LocalRotation = RelativeRotation.Inverse() * Rotation;
+			const FVector LocalRotationForward = LocalRotation.GetForwardVector();
+			const FVector LocalRotationRight = LocalRotation.GetRightVector();
+
+			Output[0] = LocalRotationForward.X / FMath::Max(Scale, Epsilon);
+			Output[1] = LocalRotationForward.Y / FMath::Max(Scale, Epsilon);
+			Output[2] = LocalRotationForward.Z / FMath::Max(Scale, Epsilon);
+			Output[3] = LocalRotationRight.X / FMath::Max(Scale, Epsilon);
+			Output[4] = LocalRotationRight.Y / FMath::Max(Scale, Epsilon);
+			Output[5] = LocalRotationRight.Z / FMath::Max(Scale, Epsilon);
+		}
+
+		void RotationVector(
+			TLearningArrayView<1, float> Output,
+			const FVector RotationVector,
+			const float Scale,
+			const float Epsilon)
+		{
+			UE_LEARNING_ARRAY_SHAPE_CHECK(Output.Num() == 3);
+			UE_LEARNING_ARRAY_VALUE_CHECK(Scale > 0.0f);
+
+			Output[0] = RotationVector.X / FMath::Max(Scale, Epsilon);
+			Output[1] = RotationVector.Y / FMath::Max(Scale, Epsilon);
+			Output[2] = RotationVector.Z / FMath::Max(Scale, Epsilon);
+		}
+
+
+		void ScalarAngularVelocity(
+			float& Output,
+			const float AngularVelocity,
+			const float Scale,
+			const float Epsilon)
+		{
+			UE_LEARNING_ARRAY_VALUE_CHECK(Scale > 0.0f);
+
+			Output = AngularVelocity / FMath::Max(Scale, Epsilon);
+		}
+
+		void PlanarAngularVelocity(
+			float& Output,
+			const FVector AngularVelocity,
+			const FQuat RelativeRotation,
+			const float Scale,
+			const FVector Axis,
+			const float Epsilon)
+		{
+			UE_LEARNING_ARRAY_VALUE_CHECK(Scale > 0.0f);
+
+			Output = RelativeRotation.UnrotateVector(AngularVelocity).Dot(Axis) / FMath::Max(Scale, Epsilon);
+		}
+
+		void AngularVelocity(
+			TLearningArrayView<1, float> Output,
+			const FVector AngularVelocity,
+			const FQuat RelativeRotation,
+			const float Scale,
+			const float Epsilon)
+		{
+			UE_LEARNING_ARRAY_SHAPE_CHECK(Output.Num() == 3);
+			UE_LEARNING_ARRAY_VALUE_CHECK(Scale > 0.0f);
+
+			const FVector LocalAngularVelocity = RelativeRotation.UnrotateVector(AngularVelocity);
+			Output[0] = LocalAngularVelocity.X / FMath::Max(Scale, Epsilon);
+			Output[1] = LocalAngularVelocity.Y / FMath::Max(Scale, Epsilon);
+			Output[2] = LocalAngularVelocity.Z / FMath::Max(Scale, Epsilon);
+		}
+
+
+		void ScalarAngularAcceleration(
+			float& Output,
+			const float AngularAcceleration,
+			const float Scale,
+			const float Epsilon)
+		{
+			UE_LEARNING_ARRAY_VALUE_CHECK(Scale > 0.0f);
+
+			Output = AngularAcceleration / FMath::Max(Scale, Epsilon);
+		}
+
+		void PlanarAngularAcceleration(
+			float& Output,
+			const FVector AngularAcceleration,
+			const FQuat RelativeRotation,
+			const float Scale,
+			const FVector Axis,
+			const float Epsilon)
+		{
+			UE_LEARNING_ARRAY_VALUE_CHECK(Scale > 0.0f);
+
+			Output = RelativeRotation.UnrotateVector(AngularAcceleration).Dot(Axis) / FMath::Max(Scale, Epsilon);
+		}
+
+		void AngularAcceleration(
+			TLearningArrayView<1, float> Output,
+			const FVector AngularAcceleration,
+			const FQuat RelativeRotation,
+			const float Scale,
+			const float Epsilon)
+		{
+			UE_LEARNING_ARRAY_SHAPE_CHECK(Output.Num() == 3);
+			UE_LEARNING_ARRAY_VALUE_CHECK(Scale > 0.0f);
+
+			const FVector LocalAngularAcceleration = RelativeRotation.UnrotateVector(AngularAcceleration);
+			Output[0] = LocalAngularAcceleration.X / FMath::Max(Scale, Epsilon);
+			Output[1] = LocalAngularAcceleration.Y / FMath::Max(Scale, Epsilon);
+			Output[2] = LocalAngularAcceleration.Z / FMath::Max(Scale, Epsilon);
+		}
+
+
+		void Scale(
+			float& Output,
+			const float InScale,
+			const float RelativeScale,
+			const float Scale,
+			const float Epsilon)
+		{
+			UE_LEARNING_ARRAY_VALUE_CHECK(InScale > 0.0f);
+			UE_LEARNING_ARRAY_VALUE_CHECK(RelativeScale > 0.0f);
+			UE_LEARNING_ARRAY_VALUE_CHECK(Scale > 0.0f);
+
+			Output = FMath::Loge(FMath::Max(InScale, Epsilon) / FMath::Max(RelativeScale, Epsilon)) / FMath::Max(Scale, Epsilon);
+		}
+	}
+
+	namespace Feature::Decode
+	{
+
+		void Float(
+			float& OutValue,
+			const float Input,
+			const float Scale)
+		{
+			OutValue = Input * Scale;
+		}
+
+		void Time(
+			float& OutTime,
+			const float Input,
+			const float RelativeTime,
+			const float Scale)
+		{
+			OutTime = Input * Scale + RelativeTime;
+		}
+
+
+		void PlanarDirection(
+			FVector& OutPlanarDirection,
+			const TLearningArrayView<1, const float> Input,
+			const FQuat RelativeRotation,
+			const float Scale,
+			const FVector Axis0,
+			const FVector Axis1,
+			const float Epsilon)
+		{
+			UE_LEARNING_ARRAY_SHAPE_CHECK(Input.Num() == 2);
+
+			const FVector LocalDirection = (
+				Axis0 * Input[0] * Scale +
+				Axis1 * Input[1] * Scale).GetSafeNormal(Epsilon, Axis0);
+			OutPlanarDirection = RelativeRotation.RotateVector(LocalDirection);
+		}
+
+		void Direction(
+			FVector& OutDirection,
+			const TLearningArrayView<1, const float> Input,
+			const FQuat RelativeRotation,
+			const float Scale,
+			const float Epsilon)
+		{
+			UE_LEARNING_ARRAY_SHAPE_CHECK(Input.Num() == 3);
+
+			const FVector LocalDirection = (FVector(Input[0], Input[1], Input[2]) * Scale).GetSafeNormal(Epsilon, FVector::ForwardVector);
+			OutDirection = RelativeRotation.RotateVector(LocalDirection);
+		}
+
+
+		void ScalarPosition(
+			float& OutScalarPosition,
+			const float Input,
+			const float RelativePosition,
+			const float Scale)
+		{
+			OutScalarPosition = Input * Scale + RelativePosition;
+		}
+
+		void PlanarPosition(
+			FVector& OutPlanarPosition,
+			const TLearningArrayView<1, const float> Input,
+			const FVector RelativePosition,
+			const FQuat RelativeRotation,
+			const float Scale,
+			const FVector Axis0,
+			const FVector Axis1)
+		{
+			UE_LEARNING_ARRAY_SHAPE_CHECK(Input.Num() == 2);
+
+			const FVector LocalPosition =
+				Axis0 * Input[0] * Scale +
+				Axis1 * Input[1] * Scale;
+			OutPlanarPosition = RelativeRotation.RotateVector(LocalPosition) + RelativePosition;
+		}
+
+		void Position(
+			FVector& OutPosition,
+			const TLearningArrayView<1, const float> Input,
+			const FVector RelativePosition,
+			const FQuat RelativeRotation,
+			const float Scale)
+		{
+			UE_LEARNING_ARRAY_SHAPE_CHECK(Input.Num() == 3);
+
+			const FVector LocalPosition = FVector(Input[0], Input[1], Input[2]) * Scale;
+			OutPosition = RelativeRotation.RotateVector(LocalPosition) + RelativePosition;
+		}
+
+
+		void ScalarVelocity(
+			float& OutScalarVelocity,
+			const float Input,
+			const float Scale)
+		{
+			OutScalarVelocity = Input * Scale;
+		}
+
+		void PlanarVelocity(
+			FVector& OutPlanarVelocity,
+			const TLearningArrayView<1, const float> Input,
+			const FQuat RelativeRotation,
+			const float Scale,
+			const FVector Axis0,
+			const FVector Axis1)
+		{
+			UE_LEARNING_ARRAY_SHAPE_CHECK(Input.Num() == 2);
+
+			const FVector LocalVelocity =
+				Axis0 * Input[0] * Scale +
+				Axis1 * Input[1] * Scale;
+			OutPlanarVelocity = RelativeRotation.RotateVector(LocalVelocity);
+		}
+
+		void Velocity(
+			FVector& OutVelocity,
+			const TLearningArrayView<1, const float> Input,
+			const FQuat RelativeRotation,
+			const float Scale)
+		{
+			UE_LEARNING_ARRAY_SHAPE_CHECK(Input.Num() == 3);
+
+			const FVector LocalVelocity = FVector(Input[0], Input[1], Input[2]) * Scale;
+			OutVelocity = RelativeRotation.RotateVector(LocalVelocity);
+		}
+
+
+		void ScalarAcceleration(
+			float& OutScalarAcceleration,
+			const float Input,
+			const float Scale)
+		{
+			OutScalarAcceleration = Input * Scale;
+		}
+
+		void PlanarAcceleration(
+			FVector& OutPlanarAcceleration,
+			const TLearningArrayView<1, const float> Input,
+			const FQuat RelativeRotation,
+			const float Scale,
+			const FVector Axis0,
+			const FVector Axis1)
+		{
+			UE_LEARNING_ARRAY_SHAPE_CHECK(Input.Num() == 2);
+
+			const FVector LocalAcceleration =
+				Axis0 * Input[0] * Scale +
+				Axis1 * Input[1] * Scale;
+			OutPlanarAcceleration = RelativeRotation.RotateVector(LocalAcceleration);
+		}
+
+		void Acceleration(
+			FVector& OutAcceleration,
+			const TLearningArrayView<1, const float> Input,
+			const FQuat RelativeRotation,
+			const float Scale)
+		{
+			UE_LEARNING_ARRAY_SHAPE_CHECK(Input.Num() == 3);
+
+			const FVector LocalAcceleration = FVector(Input[0], Input[1], Input[2]) * Scale;
+			OutAcceleration = RelativeRotation.RotateVector(LocalAcceleration);
+		}
+
+		void Angle(
+			float& OutAngle,
+			const TLearningArrayView<1, const float> Input,
+			const float RelativeAngle,
+			const float Scale)
+		{
+			UE_LEARNING_ARRAY_SHAPE_CHECK(Input.Num() == 2);
+
+			OutAngle = FMath::Atan2(Input[0] * Scale, Input[1] * Scale) + RelativeAngle;
+		}
+
+		void PlanarRotation(
+			FQuat& OutPlanarRotation,
+			const TLearningArrayView<1, const float> Input,
+			const FQuat RelativeRotation,
+			const float Scale,
+			const FVector Axis)
+		{
+			UE_LEARNING_ARRAY_SHAPE_CHECK(Input.Num() == 2);
+
+			const float LocalAngle = FMath::Atan2(Input[0] * Scale, Input[1] * Scale);
+			OutPlanarRotation = RelativeRotation * FQuat::MakeFromRotationVector(LocalAngle * Axis);
+		}
+
+		void Rotation(
+			FQuat& OutRotation,
+			const TLearningArrayView<1, const float> Input,
+			const FQuat RelativeRotation,
+			const float Scale,
+			const float Epsilon)
+		{
+			UE_LEARNING_ARRAY_SHAPE_CHECK(Input.Num() == 6);
+
+			const FVector InputAxisForward = FVector(Input[0], Input[1], Input[2]) * Scale;
+			const FVector InputAxisRight = FVector(Input[3], Input[4], Input[5]) * Scale;
+			const FVector AxisUp = InputAxisForward.Cross(InputAxisRight).GetSafeNormal(Epsilon, FVector::UpVector);
+			const FVector AxisRight = AxisUp.Cross(InputAxisForward).GetSafeNormal(Epsilon, FVector::RightVector);
+			const FVector AxisForward = InputAxisForward.GetSafeNormal(Epsilon, FVector::ForwardVector);
+
+			FMatrix RotationMatrix = FMatrix::Identity;
+			RotationMatrix.SetAxis(0, AxisForward);
+			RotationMatrix.SetAxis(1, AxisRight);
+			RotationMatrix.SetAxis(2, AxisUp);
+
+			OutRotation = RelativeRotation * RotationMatrix.ToQuat();
+		}
+
+		void RotationVector(
+			FVector& OutRotationVector,
+			const TLearningArrayView<1, const float> Input,
+			const float Scale)
+		{
+			UE_LEARNING_ARRAY_SHAPE_CHECK(Input.Num() == 3);
+
+			OutRotationVector = FVector(Input[0], Input[1], Input[2]) * Scale;
+		}
+
+
+		void ScalarAngularVelocity(
+			float& OutScalarVelocity,
+			const float Input,
+			const float Scale)
+		{
+			OutScalarVelocity = Input * Scale;
+		}
+
+		void PlanarAngularVelocity(
+			FVector& OutPlanarAngularVelocity,
+			const float Input,
+			const FQuat RelativeRotation,
+			const float Scale,
+			const FVector Axis)
+		{
+			OutPlanarAngularVelocity = RelativeRotation.RotateVector(Axis * Input * Scale);
+		}
+
+		void AngularVelocity(
+			FVector& OutAngularVelocity,
+			const TLearningArrayView<1, const float> Input,
+			const FQuat RelativeRotation,
+			const float Scale)
+		{
+			UE_LEARNING_ARRAY_SHAPE_CHECK(Input.Num() == 3);
+
+			OutAngularVelocity = RelativeRotation.RotateVector(FVector(Input[0], Input[1], Input[2]) * Scale);
+		}
+
+
+		void ScalarAngularAcceleration(
+			float& OutScalarAngularAcceleration,
+			const float Input,
+			const float Scale)
+		{
+			OutScalarAngularAcceleration = Input * Scale;
+		}
+
+		void PlanarAngularAcceleration(
+			FVector& OutPlanarAngularAcceleration,
+			const float Input,
+			const FQuat RelativeRotation,
+			const float Scale,
+			const FVector Axis)
+		{
+			OutPlanarAngularAcceleration = RelativeRotation.RotateVector(Axis * Input * Scale);
+		}
+
+		void AngularAcceleration(
+			FVector& OutAngularAcceleration,
+			const TLearningArrayView<1, const float> Input,
+			const FQuat RelativeRotation,
+			const float Scale)
+		{
+			UE_LEARNING_ARRAY_SHAPE_CHECK(Input.Num() == 3);
+
+			const FVector LocalAngularAcceleration = FVector(Input[0], Input[1], Input[2]) * Scale;
+			OutAngularAcceleration = RelativeRotation.RotateVector(LocalAngularAcceleration);
+		}
+
+
+		void Scale(
+			float& OutScale,
+			const float Input,
+			const float RelativeScale,
+			const float Scale,
+			const float Epsilon)
+		{
+			UE_LEARNING_ARRAY_VALUE_CHECK(RelativeScale > 0.0f);
+
+			OutScale = FMath::Exp(Input * FMath::Max(RelativeScale, Epsilon)) * Scale;
+		}
+	}
+
+	//------------------------------------------------------------------
+
+	FFeatureObject::FFeatureObject(
+		const FName& InIdentifier,
+		const TSharedRef<FArrayMap>& InInstanceData,
+		const int32 InMaxInstanceNum,
+		const int32 InFeatureDimNum,
+		const float InScale)
+		: InstanceData(InInstanceData)
+		, FeatureDimNum(InFeatureDimNum)
+		, Scale(InScale)
+	{
+		FeatureHandle = InstanceData->Add<2, float>({ InIdentifier, TEXT("Feature") }, { InMaxInstanceNum, InFeatureDimNum }, 0.0f);
+	}
+
+	int32 FFeatureObject::DimNum() const
+	{
+		return FeatureDimNum;
+	}
+
+	TLearningArrayView<2, float> FFeatureObject::FeatureBuffer()
+	{
+		return InstanceData->View(FeatureHandle);
+	}
+
+	static inline int32 FeatureTotalDimension(const TLearningArrayView<1, const TSharedRef<FFeatureObject>> InFeatures)
+	{
+		const int32 FeatureNum = InFeatures.Num();
+
+		int32 Total = 0;
+		for (int32 FeatureIdx = 0; FeatureIdx < FeatureNum; FeatureIdx++)
+		{
+			Total += InFeatures[FeatureIdx]->FeatureDimNum;
+		}
+
+		return Total;
+	}
+
+	//------------------------------------------------------------------
+
+	FConcatenateFeature::FConcatenateFeature(
+		const FName& InIdentifier,
+		const TLearningArrayView<1, const TSharedRef<FFeatureObject>> InFeatures,
+		const TSharedRef<FArrayMap>& InInstanceData,
+		const int32 InMaxInstanceNum,
+		const float InScale)
+		: FFeatureObject(InIdentifier, InInstanceData, InMaxInstanceNum, FeatureTotalDimension(InFeatures), InScale)
+		, Features(InFeatures)
+	{
+		const int32 FeatureNum = Features.Num();
+
+		Offsets.SetNumUninitialized({ FeatureNum });
+		Sizes.SetNumUninitialized({ FeatureNum });
+
+		int32 Total = 0;
+		for (int32 FeatureIdx = 0; FeatureIdx < FeatureNum; FeatureIdx++)
+		{
+			Offsets[FeatureIdx] = Total;
+			Sizes[FeatureIdx] = Features[FeatureIdx]->FeatureDimNum;
+			Total += Features[FeatureIdx]->FeatureDimNum;
+		}
+
+		UE_LEARNING_CHECK(Total == FeatureTotalDimension(InFeatures));
+	}
+
+	bool FConcatenateFeature::IsEncodable() const
+	{
+		const int32 FeatureNum = Features.Num();
+
+		for (int32 FeatureIdx = 0; FeatureIdx < FeatureNum; FeatureIdx++)
+		{
+			if (!Features[FeatureIdx]->IsEncodable()) { return false; }
+		}
+		return true;
+	}
+
+	bool FConcatenateFeature::IsDecodable() const
+	{
+		const int32 FeatureNum = Features.Num();
+
+		for (int32 FeatureIdx = 0; FeatureIdx < FeatureNum; FeatureIdx++)
+		{
+			if (!Features[FeatureIdx]->IsDecodable()) { return false; }
+		}
+		return true;
+	}
+
+	void FConcatenateFeature::Encode(const FIndexSet Instances)
+	{
+		UE_LEARNING_TRACE_CPUPROFILER_EVENT_SCOPE(FConcatenateFeature::Encode);
+
+		TLearningArrayView<2, float> Feature = InstanceData->View(FeatureHandle);
+
+		const int32 FeatureNum = Features.Num();
+
+		for (int32 FeatureIdx = 0; FeatureIdx < FeatureNum; FeatureIdx++)
+		{
+			Features[FeatureIdx]->Encode(Instances);
+
+			{
+				UE_LEARNING_TRACE_CPUPROFILER_EVENT_SCOPE(FConcatenateFeature::Encode::Combine);
+
+				const TLearningArrayView<2, const float> FeatureInput = InstanceData->ConstView(Features[FeatureIdx]->FeatureHandle);
+				const int32 FeatureOffset = Offsets[FeatureIdx];
+				const int32 FeatureSize = Sizes[FeatureIdx];
+
+				UE_LEARNING_CHECK(FeatureInput.Num<1>() == FeatureSize);
+
+				if (UE_LEARNING_ISPC && Instances.IsSlice())
+				{
+#if UE_LEARNING_ISPC
+					ispc::LearningCombineFeature(
+						Feature.Slice(Instances.GetSliceStart(), Instances.GetSliceNum()).GetData(),
+						FeatureInput.Slice(Instances.GetSliceStart(), Instances.GetSliceNum()).GetData(),
+						FeatureOffset,
+						FeatureSize,
+						Feature.Num<1>(),
+						Instances.GetSliceNum(),
+						Scale);
+#endif
+				}
+				else
+				{
+					for (const int32 InstanceIdx : Instances)
+					{
+						for (int32 DimIdx = 0; DimIdx < FeatureSize; DimIdx++)
+						{
+							Feature[InstanceIdx][FeatureOffset + DimIdx] = Scale * FeatureInput[InstanceIdx][DimIdx];
+						}
+					}
+				}
+			}
+		}
+
+		Array::Check(Feature, Instances);
+	}
+
+	void FConcatenateFeature::Decode(const FIndexSet Instances)
+	{
+		UE_LEARNING_TRACE_CPUPROFILER_EVENT_SCOPE(FConcatenateFeature::Decode);
+
+		const TLearningArrayView<2, const float> Feature = InstanceData->ConstView(FeatureHandle);
+		Array::Check(Feature, Instances);
+
+		const int32 FeatureNum = Features.Num();
+
+		for (int32 FeatureIdx = 0; FeatureIdx < FeatureNum; FeatureIdx++)
+		{
+			{
+				UE_LEARNING_TRACE_CPUPROFILER_EVENT_SCOPE(FConcatenateFeature::Decode::Separate);
+
+				TLearningArrayView<2, float> FeatureOutput = InstanceData->View(Features[FeatureIdx]->FeatureHandle);
+				const int32 FeatureOffset = Offsets[FeatureIdx];
+				const int32 FeatureSize = Sizes[FeatureIdx];
+
+				UE_LEARNING_CHECK(FeatureOutput.Num<1>() == FeatureSize);
+
+				if (UE_LEARNING_ISPC && Instances.IsSlice())
+				{
+#if UE_LEARNING_ISPC
+					ispc::LearningSeparateFeature(
+						FeatureOutput.Slice(Instances.GetSliceStart(), Instances.GetSliceNum()).GetData(),
+						Feature.Slice(Instances.GetSliceStart(), Instances.GetSliceNum()).GetData(),
+						FeatureOffset,
+						FeatureSize,
+						Feature.Num<1>(),
+						Instances.GetSliceNum(),
+						Scale,
+						UE_SMALL_NUMBER);
+#endif
+				}
+				else
+				{
+					for (const int32 InstanceIdx : Instances)
+					{
+						for (int32 DimIdx = 0; DimIdx < FeatureSize; DimIdx++)
+						{
+							FeatureOutput[InstanceIdx][DimIdx] = Feature[InstanceIdx][FeatureOffset + DimIdx] / FMath::Max(Scale, UE_SMALL_NUMBER);
+						}
+					}
+				}
+			}
+
+			Features[FeatureIdx]->Decode(Instances);
+		}
+	}
+
+	//------------------------------------------------------------------
+
+	FFloatFeature::FFloatFeature(
+		const FName& InIdentifier,
+		const TSharedRef<FArrayMap>& InInstanceData,
+		const int32 InMaxInstanceNum,
+		const int32 InFloatNum,
+		const float InScale)
+		: FFeatureObject(InIdentifier, InInstanceData, InMaxInstanceNum, InFloatNum, InScale)
+	{
+		ValueHandle = InstanceData->Add<2, float>({ InIdentifier, TEXT("Value") }, { InMaxInstanceNum, InFloatNum }, 0.0f);
+	}
+
+	void FFloatFeature::Encode(const FIndexSet Instances)
+	{
+		UE_LEARNING_TRACE_CPUPROFILER_EVENT_SCOPE(FScalarFeature::Encode);
+
+		const TLearningArrayView<2, const float> Value = InstanceData->ConstView(ValueHandle);
+		TLearningArrayView<2, float> Feature = InstanceData->View(FeatureHandle);
+
+		const int32 FloatNum = Value.Num<1>();
+
+		for (const int32 InstanceIdx : Instances)
+		{
+			for (int32 FloatIdx = 0; FloatIdx < FloatNum; FloatIdx++)
+			{
+				Feature::Encode::Float(
+					Feature[InstanceIdx][FloatIdx],
+					Value[InstanceIdx][FloatIdx],
+					Scale);
+			}
+		}
+	}
+
+	void FFloatFeature::Decode(const FIndexSet Instances)
+	{
+		UE_LEARNING_TRACE_CPUPROFILER_EVENT_SCOPE(FScalarFeature::Decode);
+
+		TLearningArrayView<2, float> Value = InstanceData->View(ValueHandle);
+		const TLearningArrayView<2, const float> Feature = InstanceData->ConstView(FeatureHandle);
+
+		const int32 FloatNum = Value.Num<1>();
+
+		for (const int32 InstanceIdx : Instances)
+		{
+			for (int32 FloatIdx = 0; FloatIdx < FloatNum; FloatIdx++)
+			{
+				Feature::Decode::Float(
+					Value[InstanceIdx][FloatIdx],
+					Feature[InstanceIdx][FloatIdx],
+					Scale);
+			}
+		}
+	}
+
+	//------------------------------------------------------------------
+
+	FTimeFeature::FTimeFeature(
+		const FName& InIdentifier,
+		const TSharedRef<FArrayMap>& InInstanceData,
+		const int32 InMaxInstanceNum,
+		const int32 InTimeNum,
+		const float InScale)
+		: FFeatureObject(InIdentifier, InInstanceData, InMaxInstanceNum, InTimeNum, InScale)
+	{
+		TimeHandle = InstanceData->Add<2, float>({ InIdentifier, TEXT("Time") }, { InMaxInstanceNum, InTimeNum }, 0.0f);
+		RelativeTimeHandle = InstanceData->Add<1, float>({ InIdentifier, TEXT("RelativeTime") }, { InMaxInstanceNum }, 0.0f);
+	}
+
+	void FTimeFeature::Encode(const FIndexSet Instances)
+	{
+		UE_LEARNING_TRACE_CPUPROFILER_EVENT_SCOPE(FTimeFeature::Encode);
+
+		const TLearningArrayView<2, const float> Time = InstanceData->ConstView(TimeHandle);
+		const TLearningArrayView<1, const float> RelativeTime = InstanceData->ConstView(RelativeTimeHandle);
+		TLearningArrayView<2, float> Feature = InstanceData->View(FeatureHandle);
+
+		const int32 TimeNum = Time.Num<1>();
+		
+		for (const int32 InstanceIdx : Instances)
+		{
+			for (int32 TimeIdx = 0; TimeIdx < TimeNum; TimeIdx++)
+			{
+				Feature::Encode::Time(
+					Feature[InstanceIdx][TimeIdx],
+					Time[InstanceIdx][TimeIdx],
+					RelativeTime[InstanceIdx],
+					Scale);
+			}
+		}
+	}
+
+	void FTimeFeature::Decode(const FIndexSet Instances)
+	{
+		UE_LEARNING_TRACE_CPUPROFILER_EVENT_SCOPE(FTimeFeature::Decode);
+
+		TLearningArrayView<2, float> Time = InstanceData->View(TimeHandle);
+		const TLearningArrayView<1, const float> RelativeTime = InstanceData->ConstView(RelativeTimeHandle);
+		const TLearningArrayView<2, const float> Feature = InstanceData->ConstView(FeatureHandle);
+
+		const int32 TimeNum = Time.Num<1>();
+
+		for (const int32 InstanceIdx : Instances)
+		{
+			for (int32 TimeIdx = 0; TimeIdx < TimeNum; TimeIdx++)
+			{
+				Feature::Decode::Time(
+					Time[InstanceIdx][TimeIdx],
+					Feature[InstanceIdx][TimeIdx],
+					RelativeTime[InstanceIdx],
+					Scale);
+			}
+		}
+	}
+
+	//------------------------------------------------------------------
+
+	FPlanarDirectionFeature::FPlanarDirectionFeature(
+		const FName& InIdentifier,
+		const TSharedRef<FArrayMap>& InInstanceData,
+		const int32 InMaxInstanceNum,
+		const int32 InDirectionNum,
+		const float InScale,
+		const FVector InAxis0,
+		const FVector InAxis1)
+		: FFeatureObject(InIdentifier, InInstanceData, InMaxInstanceNum, InDirectionNum * 2, InScale)
+		, Axis0(InAxis0)
+		, Axis1(InAxis1)
+	{
+		DirectionHandle = InstanceData->Add<2, FVector>({ InIdentifier, TEXT("Direction") }, { InMaxInstanceNum, InDirectionNum }, FVector::ZeroVector);
+		RelativeRotationHandle = InstanceData->Add<1, FQuat>({ InIdentifier, TEXT("RelativeRotation") }, { InMaxInstanceNum }, FQuat::Identity);
+	}
+
+	void FPlanarDirectionFeature::Encode(const FIndexSet Instances)
+	{
+		UE_LEARNING_TRACE_CPUPROFILER_EVENT_SCOPE(FPlanarDirectionFeature::Encode);
+
+		const TLearningArrayView<2, const FVector> Direction = InstanceData->ConstView(DirectionHandle);
+		const TLearningArrayView<1, const FQuat> RelativeRotation = InstanceData->ConstView(RelativeRotationHandle);
+		TLearningArrayView<2, float> Feature = InstanceData->View(FeatureHandle);
+
+		const int32 DirectionNum = Direction.Num<1>();
+
+		for (const int32 InstanceIdx : Instances)
+		{
+			for (int32 DirectionIdx = 0; DirectionIdx < DirectionNum; DirectionIdx++)
+			{
+				Feature::Encode::PlanarDirection(
+					Feature[InstanceIdx].Slice(DirectionIdx * 2, 2),
+					Direction[InstanceIdx][DirectionIdx],
+					RelativeRotation[InstanceIdx],
+					Scale,
+					Axis0,
+					Axis1);
+			}
+		}
+	}
+
+	void FPlanarDirectionFeature::Decode(const FIndexSet Instances)
+	{
+		UE_LEARNING_TRACE_CPUPROFILER_EVENT_SCOPE(FPlanarDirectionFeature::Decode);
+
+		TLearningArrayView<2, FVector> Direction = InstanceData->View(DirectionHandle);
+		const TLearningArrayView<1, const FQuat> RelativeRotation = InstanceData->ConstView(RelativeRotationHandle);
+		const TLearningArrayView<2, const float> Feature = InstanceData->ConstView(FeatureHandle);
+
+		const int32 DirectionNum = Direction.Num<1>();
+
+		for (const int32 InstanceIdx : Instances)
+		{
+			for (int32 DirectionIdx = 0; DirectionIdx < DirectionNum; DirectionIdx++)
+			{
+				Feature::Decode::PlanarDirection(
+					Direction[InstanceIdx][DirectionIdx],
+					Feature[InstanceIdx].Slice(DirectionIdx * 2, 2),
+					RelativeRotation[InstanceIdx],
+					Scale,
+					Axis0,
+					Axis1);
+			}
+		}
+	}
+
+	FDirectionFeature::FDirectionFeature(
+		const FName& InIdentifier,
+		const TSharedRef<FArrayMap>& InInstanceData,
+		const int32 InMaxInstanceNum,
+		const int32 InDirectionNum,
+		const float InScale)
+		: FFeatureObject(InIdentifier, InInstanceData, InMaxInstanceNum, InDirectionNum * 3, InScale)
+	{
+		DirectionHandle = InstanceData->Add<2, FVector>({ InIdentifier, TEXT("Direction") }, { InMaxInstanceNum }, FVector::ZeroVector);
+		RelativeRotationHandle = InstanceData->Add<1, FQuat>({ InIdentifier, TEXT("RelativeRotation") }, { InMaxInstanceNum }, FQuat::Identity);
+	}
+
+	void FDirectionFeature::Encode(const FIndexSet Instances)
+	{
+		UE_LEARNING_TRACE_CPUPROFILER_EVENT_SCOPE(FDirectionFeature::Encode);
+
+		const TLearningArrayView<2, const FVector> Direction = InstanceData->ConstView(DirectionHandle);
+		const TLearningArrayView<1, const FQuat> RelativeRotation = InstanceData->ConstView(RelativeRotationHandle);
+		TLearningArrayView<2, float> Feature = InstanceData->View(FeatureHandle);
+
+		const int32 DirectionNum = Direction.Num<1>();
+
+		for (const int32 InstanceIdx : Instances)
+		{
+			for (int32 DirectionIdx = 0; DirectionIdx < DirectionNum; DirectionIdx++)
+			{
+				Feature::Encode::Direction(
+					Feature[InstanceIdx].Slice(DirectionIdx * 3, 3),
+					Direction[InstanceIdx][DirectionIdx],
+					RelativeRotation[InstanceIdx],
+					Scale);
+			}
+		}
+	}
+
+	void FDirectionFeature::Decode(const FIndexSet Instances)
+	{
+		UE_LEARNING_TRACE_CPUPROFILER_EVENT_SCOPE(FDirectionFeature::Decode);
+
+		TLearningArrayView<2, FVector> Direction = InstanceData->View(DirectionHandle);
+		const TLearningArrayView<1, const FQuat> RelativeRotation = InstanceData->ConstView(RelativeRotationHandle);
+		const TLearningArrayView<2, const float> Feature = InstanceData->ConstView(FeatureHandle);
+
+		const int32 DirectionNum = Direction.Num<1>();
+
+		for (const int32 InstanceIdx : Instances)
+		{
+			for (int32 DirectionIdx = 0; DirectionIdx < DirectionNum; DirectionIdx++)
+			{
+				Feature::Decode::Direction(
+					Direction[InstanceIdx][DirectionIdx],
+					Feature[InstanceIdx].Slice(DirectionIdx * 3, 3),
+					RelativeRotation[InstanceIdx],
+					Scale);
+			}
+		}
+	}
+
+	//------------------------------------------------------------------
+
+	FScalarPositionFeature::FScalarPositionFeature(
+		const FName& InIdentifier,
+		const TSharedRef<FArrayMap>& InInstanceData,
+		const int32 InMaxInstanceNum,
+		const int32 InPositionNum,
+		const float InScale)
+		: FFeatureObject(InIdentifier, InInstanceData, InMaxInstanceNum, InPositionNum, InScale)
+	{
+		PositionHandle = InstanceData->Add<2, float>({ InIdentifier, TEXT("Position") }, { InMaxInstanceNum, InPositionNum }, 0.0f);
+		RelativePositionHandle = InstanceData->Add<1, float>({ InIdentifier, TEXT("RelativePosition") }, { InMaxInstanceNum }, 0.0f);
+	}
+
+	void FScalarPositionFeature::Encode(const FIndexSet Instances)
+	{
+		UE_LEARNING_TRACE_CPUPROFILER_EVENT_SCOPE(FScalarPositionFeature::Encode);
+
+		const TLearningArrayView<2, const float> Position = InstanceData->ConstView(PositionHandle);
+		const TLearningArrayView<1, const float> RelativePosition = InstanceData->ConstView(RelativePositionHandle);
+		TLearningArrayView<2, float> Feature = InstanceData->View(FeatureHandle);
+
+		const int32 PositionNum = Position.Num<1>();
+
+		for (const int32 InstanceIdx : Instances)
+		{
+			for (int32 PositionIdx = 0; PositionIdx < PositionNum; PositionIdx++)
+			{
+				Feature::Encode::ScalarPosition(
+					Feature[InstanceIdx][PositionIdx],
+					Position[InstanceIdx][PositionIdx],
+					RelativePosition[InstanceIdx],
+					Scale);
+			}
+		}
+	}
+
+	void FScalarPositionFeature::Decode(const FIndexSet Instances)
+	{
+		UE_LEARNING_TRACE_CPUPROFILER_EVENT_SCOPE(FScalarPositionFeature::Decode);
+
+		TLearningArrayView<2, float> Position = InstanceData->View(PositionHandle);
+		const TLearningArrayView<1, const float> RelativePosition = InstanceData->ConstView(RelativePositionHandle);
+		const TLearningArrayView<2, const float> Feature = InstanceData->ConstView(FeatureHandle);
+
+		const int32 PositionNum = Position.Num<1>();
+
+		for (const int32 InstanceIdx : Instances)
+		{
+			for (int32 PositionIdx = 0; PositionIdx < PositionNum; PositionIdx++)
+			{
+				Feature::Decode::ScalarPosition(
+					Position[InstanceIdx][PositionIdx],
+					Feature[InstanceIdx][PositionIdx],
+					RelativePosition[InstanceIdx],
+					Scale);
+			}
+		}
+	}
+
+	FPlanarPositionFeature::FPlanarPositionFeature(
+		const FName& InIdentifier,
+		const TSharedRef<FArrayMap>& InInstanceData,
+		const int32 InMaxInstanceNum,
+		const int32 InPositionNum,
+		const float InScale,
+		const FVector InAxis0,
+		const FVector InAxis1)
+		: FFeatureObject(InIdentifier, InInstanceData, InMaxInstanceNum, InPositionNum * 2, InScale)
+		, Axis0(InAxis0)
+		, Axis1(InAxis1)
+	{
+		PositionHandle = InstanceData->Add<2, FVector>({ InIdentifier, TEXT("Position") }, { InMaxInstanceNum, InPositionNum }, FVector::ZeroVector);
+		RelativePositionHandle = InstanceData->Add<1, FVector>({ InIdentifier, TEXT("RelativePosition") }, { InMaxInstanceNum }, FVector::ZeroVector);
+		RelativeRotationHandle = InstanceData->Add<1, FQuat>({ InIdentifier, TEXT("RelativeRotation") }, { InMaxInstanceNum }, FQuat::Identity);
+	}
+
+	void FPlanarPositionFeature::Encode(const FIndexSet Instances)
+	{
+		UE_LEARNING_TRACE_CPUPROFILER_EVENT_SCOPE(FPlanarPositionFeature::Encode);
+
+		const TLearningArrayView<2, const FVector> Position = InstanceData->ConstView(PositionHandle);
+		const TLearningArrayView<1, const FVector> RelativePosition = InstanceData->ConstView(RelativePositionHandle);
+		const TLearningArrayView<1, const FQuat> RelativeRotation = InstanceData->ConstView(RelativeRotationHandle);
+		TLearningArrayView<2, float> Feature = InstanceData->View(FeatureHandle);
+
+		const int32 PositionNum = Position.Num<1>();
+
+		for (const int32 InstanceIdx : Instances)
+		{
+			for (int32 PositionIdx = 0; PositionIdx < PositionNum; PositionIdx++)
+			{
+				Feature::Encode::PlanarPosition(
+					Feature[InstanceIdx].Slice(PositionIdx * 2, 2),
+					Position[InstanceIdx][PositionIdx],
+					RelativePosition[InstanceIdx],
+					RelativeRotation[InstanceIdx],
+					Scale,
+					Axis0,
+					Axis1);
+			}
+		}
+	}
+
+	void FPlanarPositionFeature::Decode(const FIndexSet Instances)
+	{
+		UE_LEARNING_TRACE_CPUPROFILER_EVENT_SCOPE(FPlanarPositionFeature::Decode);
+
+		TLearningArrayView<2, FVector> Position = InstanceData->View(PositionHandle);
+		const TLearningArrayView<1, const FVector> RelativePosition = InstanceData->ConstView(RelativePositionHandle);
+		const TLearningArrayView<1, const FQuat> RelativeRotation = InstanceData->ConstView(RelativeRotationHandle);
+		const TLearningArrayView<2, const float> Feature = InstanceData->ConstView(FeatureHandle);
+
+		const int32 PositionNum = Position.Num<1>();
+
+		for (const int32 InstanceIdx : Instances)
+		{
+			for (int32 PositionIdx = 0; PositionIdx < PositionNum; PositionIdx++)
+			{
+				Feature::Decode::PlanarPosition(
+					Position[InstanceIdx][PositionIdx],
+					Feature[InstanceIdx].Slice(PositionIdx * 2, 2),
+					RelativePosition[InstanceIdx],
+					RelativeRotation[InstanceIdx],
+					Scale,
+					Axis0,
+					Axis1);
+			}
+		}
+	}
+
+	FPositionFeature::FPositionFeature(
+		const FName& InIdentifier,
+		const TSharedRef<FArrayMap>& InInstanceData,
+		const int32 InMaxInstanceNum,
+		const int32 InPositionNum,
+		const float InScale)
+		: FFeatureObject(InIdentifier, InInstanceData, InMaxInstanceNum, InPositionNum * 3, InScale)
+	{
+		PositionHandle = InstanceData->Add<2, FVector>({ InIdentifier, TEXT("Position") }, { InMaxInstanceNum, InPositionNum }, FVector::ZeroVector);
+		RelativePositionHandle = InstanceData->Add<1, FVector>({ InIdentifier, TEXT("RelativePosition") }, { InMaxInstanceNum }, FVector::ZeroVector);
+		RelativeRotationHandle = InstanceData->Add<1, FQuat>({ InIdentifier, TEXT("RelativeRotation") }, { InMaxInstanceNum }, FQuat::Identity);
+	}
+
+	void FPositionFeature::Encode(const FIndexSet Instances)
+	{
+		UE_LEARNING_TRACE_CPUPROFILER_EVENT_SCOPE(FPositionFeature::Encode);
+
+		const TLearningArrayView<2, const FVector> Position = InstanceData->ConstView(PositionHandle);
+		const TLearningArrayView<1, const FVector> RelativePosition = InstanceData->ConstView(RelativePositionHandle);
+		const TLearningArrayView<1, const FQuat> RelativeRotation = InstanceData->ConstView(RelativeRotationHandle);
+		TLearningArrayView<2, float> Feature = InstanceData->View(FeatureHandle);
+
+		const int32 PositionNum = Position.Num<1>();
+
+		for (const int32 InstanceIdx : Instances)
+		{
+			for (int32 PositionIdx = 0; PositionIdx < PositionNum; PositionIdx++)
+			{
+				Feature::Encode::Position(
+					Feature[InstanceIdx].Slice(PositionIdx * 3, 3),
+					Position[InstanceIdx][PositionIdx],
+					RelativePosition[InstanceIdx],
+					RelativeRotation[InstanceIdx],
+					Scale);
+			}
+		}
+	}
+
+	void FPositionFeature::Decode(const FIndexSet Instances)
+	{
+		UE_LEARNING_TRACE_CPUPROFILER_EVENT_SCOPE(FPositionFeature::Decode);
+
+		TLearningArrayView<2, FVector> Position = InstanceData->View(PositionHandle);
+		const TLearningArrayView<1, const FVector> RelativePosition = InstanceData->ConstView(RelativePositionHandle);
+		const TLearningArrayView<1, const FQuat> RelativeRotation = InstanceData->ConstView(RelativeRotationHandle);
+		const TLearningArrayView<2, const float> Feature = InstanceData->ConstView(FeatureHandle);
+
+		const int32 PositionNum = Position.Num<1>();
+
+		for (const int32 InstanceIdx : Instances)
+		{
+			for (int32 PositionIdx = 0; PositionIdx < PositionNum; PositionIdx++)
+			{
+				Feature::Decode::Position(
+					Position[InstanceIdx][PositionIdx],
+					Feature[InstanceIdx].Slice(PositionIdx * 3, 3),
+					RelativePosition[InstanceIdx],
+					RelativeRotation[InstanceIdx],
+					Scale);
+			}
+		}
+	}
+
+	//------------------------------------------------------------------
+
+	FScalarVelocityFeature::FScalarVelocityFeature(
+		const FName& InIdentifier,
+		const TSharedRef<FArrayMap>& InInstanceData,
+		const int32 InMaxInstanceNum,
+		const int32 InVelocityNum,
+		const float InScale)
+		: FFeatureObject(InIdentifier, InInstanceData, InMaxInstanceNum, InVelocityNum, InScale)
+	{
+		VelocityHandle = InstanceData->Add<2, float>({ InIdentifier, TEXT("Velocity") }, { InMaxInstanceNum, InVelocityNum }, 0.0f);
+	}
+
+	void FScalarVelocityFeature::Encode(const FIndexSet Instances)
+	{
+		UE_LEARNING_TRACE_CPUPROFILER_EVENT_SCOPE(FScalarVelocityFeature::Encode);
+
+		const TLearningArrayView<2, const float> Velocity = InstanceData->ConstView(VelocityHandle);
+		TLearningArrayView<2, float> Feature = InstanceData->View(FeatureHandle);
+
+		const int32 VelocityNum = Velocity.Num<1>();
+
+		for (const int32 InstanceIdx : Instances)
+		{
+			for (int32 VelocityIdx = 0; VelocityIdx < VelocityNum; VelocityIdx++)
+			{
+				Feature::Encode::ScalarVelocity(
+					Feature[InstanceIdx][VelocityIdx],
+					Velocity[InstanceIdx][VelocityIdx],
+					Scale);
+			}
+		}
+	}
+
+	void FScalarVelocityFeature::Decode(const FIndexSet Instances)
+	{
+		UE_LEARNING_TRACE_CPUPROFILER_EVENT_SCOPE(FScalarVelocityFeature::Decode);
+
+		TLearningArrayView<2, float> Velocity = InstanceData->View(VelocityHandle);
+		const TLearningArrayView<2, const float> Feature = InstanceData->ConstView(FeatureHandle);
+
+		const int32 VelocityNum = Velocity.Num<1>();
+
+		for (const int32 InstanceIdx : Instances)
+		{
+			for (int32 VelocityIdx = 0; VelocityIdx < VelocityNum; VelocityIdx++)
+			{
+				Feature::Decode::ScalarVelocity(
+					Velocity[InstanceIdx][VelocityIdx],
+					Feature[InstanceIdx][VelocityIdx],
+					Scale);
+			}
+		}
+	}
+
+	FPlanarVelocityFeature::FPlanarVelocityFeature(
+		const FName& InIdentifier,
+		const TSharedRef<FArrayMap>& InInstanceData,
+		const int32 InMaxInstanceNum,
+		const int32 InVelocityNum,
+		const float InScale,
+		const FVector InAxis0,
+		const FVector InAxis1)
+		: FFeatureObject(InIdentifier, InInstanceData, InMaxInstanceNum, InVelocityNum * 2, InScale)
+		, Axis0(InAxis0)
+		, Axis1(InAxis1)
+	{
+		VelocityHandle = InstanceData->Add<2, FVector>({ InIdentifier, TEXT("Velocity") }, { InMaxInstanceNum, InVelocityNum }, FVector::ZeroVector);
+		RelativeRotationHandle = InstanceData->Add<1, FQuat>({ InIdentifier, TEXT("RelativeRotation") }, { InMaxInstanceNum }, FQuat::Identity);
+	}
+
+	void FPlanarVelocityFeature::Encode(const FIndexSet Instances)
+	{
+		UE_LEARNING_TRACE_CPUPROFILER_EVENT_SCOPE(FPlanarVelocityFeature::Encode);
+
+		const TLearningArrayView<2, const FVector> Velocity = InstanceData->ConstView(VelocityHandle);
+		const TLearningArrayView<1, const FQuat> RelativeRotation = InstanceData->ConstView(RelativeRotationHandle);
+		TLearningArrayView<2, float> Feature = InstanceData->View(FeatureHandle);
+
+		const int32 VelocityNum = Velocity.Num<1>();
+
+		for (const int32 InstanceIdx : Instances)
+		{
+			for (int32 VelocityIdx = 0; VelocityIdx < VelocityNum; VelocityIdx++)
+			{
+				Feature::Encode::PlanarVelocity(
+					Feature[InstanceIdx].Slice(VelocityIdx * 2, 2),
+					Velocity[InstanceIdx][VelocityIdx],
+					RelativeRotation[InstanceIdx],
+					Scale,
+					Axis0,
+					Axis1);
+			}
+		}
+	}
+
+	void FPlanarVelocityFeature::Decode(const FIndexSet Instances)
+	{
+		UE_LEARNING_TRACE_CPUPROFILER_EVENT_SCOPE(FPlanarVelocityFeature::Decode);
+
+		TLearningArrayView<2, FVector> Velocity = InstanceData->View(VelocityHandle);
+		const TLearningArrayView<1, const FQuat> RelativeRotation = InstanceData->ConstView(RelativeRotationHandle);
+		const TLearningArrayView<2, const float> Feature = InstanceData->ConstView(FeatureHandle);
+
+		const int32 VelocityNum = Velocity.Num<1>();
+
+		for (const int32 InstanceIdx : Instances)
+		{
+			for (int32 VelocityIdx = 0; VelocityIdx < VelocityNum; VelocityIdx++)
+			{
+				Feature::Decode::PlanarVelocity(
+					Velocity[InstanceIdx][VelocityIdx],
+					Feature[InstanceIdx].Slice(VelocityIdx * 2, 2),
+					RelativeRotation[InstanceIdx],
+					Scale,
+					Axis0,
+					Axis1);
+			}
+		}
+	}
+
+	FVelocityFeature::FVelocityFeature(
+		const FName& InIdentifier,
+		const TSharedRef<FArrayMap>& InInstanceData,
+		const int32 InMaxInstanceNum,
+		const int32 InVelocityNum,
+		const float InScale)
+		: FFeatureObject(InIdentifier, InInstanceData, InMaxInstanceNum, InVelocityNum * 3, InScale)
+	{
+		VelocityHandle = InstanceData->Add<2, FVector>({ InIdentifier, TEXT("Velocity") }, { InMaxInstanceNum, InVelocityNum }, FVector::ZeroVector);
+		RelativeRotationHandle = InstanceData->Add<1, FQuat>({ InIdentifier, TEXT("RelativeRotation") }, { InMaxInstanceNum }, FQuat::Identity);
+	}
+
+	void FVelocityFeature::Encode(const FIndexSet Instances)
+	{
+		UE_LEARNING_TRACE_CPUPROFILER_EVENT_SCOPE(FVelocityFeature::Encode);
+
+		const TLearningArrayView<2, const FVector> Velocity = InstanceData->ConstView(VelocityHandle);
+		const TLearningArrayView<1, const FQuat> RelativeRotation = InstanceData->ConstView(RelativeRotationHandle);
+		TLearningArrayView<2, float> Feature = InstanceData->View(FeatureHandle);
+
+		const int32 VelocityNum = Velocity.Num<1>();
+
+		for (const int32 InstanceIdx : Instances)
+		{
+			for (int32 VelocityIdx = 0; VelocityIdx < VelocityNum; VelocityIdx++)
+			{
+				Feature::Encode::Velocity(
+					Feature[InstanceIdx].Slice(VelocityIdx * 3, 3),
+					Velocity[InstanceIdx][VelocityIdx],
+					RelativeRotation[InstanceIdx],
+					Scale);
+			}
+		}
+	}
+
+	void FVelocityFeature::Decode(const FIndexSet Instances)
+	{
+		UE_LEARNING_TRACE_CPUPROFILER_EVENT_SCOPE(FVelocityFeature::Decode);
+
+		TLearningArrayView<2, FVector> Velocity = InstanceData->View(VelocityHandle);
+		const TLearningArrayView<1, const FQuat> RelativeRotation = InstanceData->ConstView(RelativeRotationHandle);
+		const TLearningArrayView<2, const float> Feature = InstanceData->ConstView(FeatureHandle);
+
+		const int32 VelocityNum = Velocity.Num<1>();
+
+		for (const int32 InstanceIdx : Instances)
+		{
+			for (int32 VelocityIdx = 0; VelocityIdx < VelocityNum; VelocityIdx++)
+			{
+				Feature::Decode::Velocity(
+					Velocity[InstanceIdx][VelocityIdx],
+					Feature[InstanceIdx].Slice(VelocityIdx * 3, 3),
+					RelativeRotation[InstanceIdx],
+					Scale);
+			}
+		}
+	}
+
+	//------------------------------------------------------------------
+
+	FScalarAccelerationFeature::FScalarAccelerationFeature(
+		const FName& InIdentifier,
+		const TSharedRef<FArrayMap>& InInstanceData,
+		const int32 InMaxInstanceNum,
+		const int32 InAccelerationNum,
+		const float InScale)
+		: FFeatureObject(InIdentifier, InInstanceData, InMaxInstanceNum, InAccelerationNum, InScale)
+	{
+		AccelerationHandle = InstanceData->Add<2, float>({ InIdentifier, TEXT("Acceleration") }, { InMaxInstanceNum, InAccelerationNum }, 0.0f);
+	}
+
+	void FScalarAccelerationFeature::Encode(const FIndexSet Instances)
+	{
+		UE_LEARNING_TRACE_CPUPROFILER_EVENT_SCOPE(FScalarAccelerationFeature::Encode);
+
+		const TLearningArrayView<2, const float> Acceleration = InstanceData->ConstView(AccelerationHandle);
+		TLearningArrayView<2, float> Feature = InstanceData->View(FeatureHandle);
+
+		const int32 AccelerationNum = Acceleration.Num<1>();
+
+		for (const int32 InstanceIdx : Instances)
+		{
+			for (int32 AccelerationIdx = 0; AccelerationIdx < AccelerationNum; AccelerationIdx++)
+			{
+				Feature::Encode::ScalarAcceleration(
+					Feature[InstanceIdx][AccelerationIdx],
+					Acceleration[InstanceIdx][AccelerationIdx],
+					Scale);
+			}
+		}
+	}
+
+	void FScalarAccelerationFeature::Decode(const FIndexSet Instances)
+	{
+		UE_LEARNING_TRACE_CPUPROFILER_EVENT_SCOPE(FScalarAccelerationFeature::Decode);
+
+		TLearningArrayView<2, float> Acceleration = InstanceData->View(AccelerationHandle);
+		const TLearningArrayView<2, const float> Feature = InstanceData->ConstView(FeatureHandle);
+
+		const int32 AccelerationNum = Acceleration.Num<1>();
+
+		for (const int32 InstanceIdx : Instances)
+		{
+			for (int32 AccelerationIdx = 0; AccelerationIdx < AccelerationNum; AccelerationIdx++)
+			{
+				Feature::Decode::ScalarAcceleration(
+					Acceleration[InstanceIdx][AccelerationIdx],
+					Feature[InstanceIdx][AccelerationIdx],
+					Scale);
+			}
+		}
+	}
+
+	FPlanarAccelerationFeature::FPlanarAccelerationFeature(
+		const FName& InIdentifier,
+		const TSharedRef<FArrayMap>& InInstanceData,
+		const int32 InMaxInstanceNum,
+		const int32 InAccelerationNum,
+		const float InScale,
+		const FVector InAxis0,
+		const FVector InAxis1)
+		: FFeatureObject(InIdentifier, InInstanceData, InMaxInstanceNum, InAccelerationNum * 2, InScale)
+		, Axis0(InAxis0)
+		, Axis1(InAxis1)
+	{
+		AccelerationHandle = InstanceData->Add<2, FVector>({ InIdentifier, TEXT("Acceleration") }, { InMaxInstanceNum, InAccelerationNum }, FVector::ZeroVector);
+		RelativeRotationHandle = InstanceData->Add<1, FQuat>({ InIdentifier, TEXT("RelativeRotation") }, { InMaxInstanceNum }, FQuat::Identity);
+	}
+
+	void FPlanarAccelerationFeature::Encode(const FIndexSet Instances)
+	{
+		UE_LEARNING_TRACE_CPUPROFILER_EVENT_SCOPE(FPlanarAccelerationFeature::Encode);
+
+		const TLearningArrayView<2, const FVector> Acceleration = InstanceData->ConstView(AccelerationHandle);
+		const TLearningArrayView<1, const FQuat> RelativeRotation = InstanceData->ConstView(RelativeRotationHandle);
+		TLearningArrayView<2, float> Feature = InstanceData->View(FeatureHandle);
+
+		const int32 AccelerationNum = Acceleration.Num<1>();
+
+		for (const int32 InstanceIdx : Instances)
+		{
+			for (int32 AccelerationIdx = 0; AccelerationIdx < AccelerationNum; AccelerationIdx++)
+			{
+				Feature::Encode::PlanarAcceleration(
+					Feature[InstanceIdx].Slice(AccelerationIdx * 2, 2),
+					Acceleration[InstanceIdx][AccelerationIdx],
+					RelativeRotation[InstanceIdx],
+					Scale,
+					Axis0,
+					Axis1);
+			}
+		}
+	}
+
+	void FPlanarAccelerationFeature::Decode(const FIndexSet Instances)
+	{
+		UE_LEARNING_TRACE_CPUPROFILER_EVENT_SCOPE(FPlanarAccelerationFeature::Decode);
+
+		TLearningArrayView<2, FVector> Acceleration = InstanceData->View(AccelerationHandle);
+		const TLearningArrayView<1, const FQuat> RelativeRotation = InstanceData->ConstView(RelativeRotationHandle);
+		const TLearningArrayView<2, const float> Feature = InstanceData->ConstView(FeatureHandle);
+
+		const int32 AccelerationNum = Acceleration.Num<1>();
+
+		for (const int32 InstanceIdx : Instances)
+		{
+			for (int32 AccelerationIdx = 0; AccelerationIdx < AccelerationNum; AccelerationIdx++)
+			{
+				Feature::Decode::PlanarAcceleration(
+					Acceleration[InstanceIdx][AccelerationIdx],
+					Feature[InstanceIdx].Slice(AccelerationIdx * 2, 2),
+					RelativeRotation[InstanceIdx],
+					Scale,
+					Axis0,
+					Axis1);
+			}
+		}
+	}
+
+	FAccelerationFeature::FAccelerationFeature(
+		const FName& InIdentifier,
+		const TSharedRef<FArrayMap>& InInstanceData,
+		const int32 InMaxInstanceNum,
+		const int32 InAccelerationNum,
+		const float InScale)
+		: FFeatureObject(InIdentifier, InInstanceData, InMaxInstanceNum, InAccelerationNum * 3, InScale)
+	{
+		AccelerationHandle = InstanceData->Add<2, FVector>({ InIdentifier, TEXT("Acceleration") }, { InMaxInstanceNum, InAccelerationNum }, FVector::ZeroVector);
+		RelativeRotationHandle = InstanceData->Add<1, FQuat>({ InIdentifier, TEXT("RelativeRotation") }, { InMaxInstanceNum }, FQuat::Identity);
+	}
+
+	void FAccelerationFeature::Encode(const FIndexSet Instances)
+	{
+		UE_LEARNING_TRACE_CPUPROFILER_EVENT_SCOPE(FAccelerationFeature::Encode);
+
+		const TLearningArrayView<2, const FVector> Acceleration = InstanceData->ConstView(AccelerationHandle);
+		const TLearningArrayView<1, const FQuat> RelativeRotation = InstanceData->ConstView(RelativeRotationHandle);
+		TLearningArrayView<2, float> Feature = InstanceData->View(FeatureHandle);
+
+		const int32 AccelerationNum = Acceleration.Num<1>();
+
+		for (const int32 InstanceIdx : Instances)
+		{
+			for (int32 AccelerationIdx = 0; AccelerationIdx < AccelerationNum; AccelerationIdx++)
+			{
+				Feature::Encode::Acceleration(
+					Feature[InstanceIdx].Slice(AccelerationIdx * 3, 3),
+					Acceleration[InstanceIdx][AccelerationIdx],
+					RelativeRotation[InstanceIdx],
+					Scale);
+			}
+		}
+	}
+
+	void FAccelerationFeature::Decode(const FIndexSet Instances)
+	{
+		UE_LEARNING_TRACE_CPUPROFILER_EVENT_SCOPE(FAccelerationFeature::Decode);
+
+		TLearningArrayView<2, FVector> Acceleration = InstanceData->View(AccelerationHandle);
+		const TLearningArrayView<1, const FQuat> RelativeRotation = InstanceData->ConstView(RelativeRotationHandle);
+		const TLearningArrayView<2, const float> Feature = InstanceData->ConstView(FeatureHandle);
+
+		const int32 AccelerationNum = Acceleration.Num<1>();
+
+		for (const int32 InstanceIdx : Instances)
+		{
+			for (int32 AccelerationIdx = 0; AccelerationIdx < AccelerationNum; AccelerationIdx++)
+			{
+				Feature::Decode::Acceleration(
+					Acceleration[InstanceIdx][AccelerationIdx],
+					Feature[InstanceIdx].Slice(AccelerationIdx * 3, 3),
+					RelativeRotation[InstanceIdx],
+					Scale);
+			}
+		}
+	}
+
+	//------------------------------------------------------------------
+
+	FAngleFeature::FAngleFeature(
+		const FName& InIdentifier,
+		const TSharedRef<FArrayMap>& InInstanceData,
+		const int32 InMaxInstanceNum,
+		const int32 InAngleNum,
+		const float InScale)
+		: FFeatureObject(InIdentifier, InInstanceData, InMaxInstanceNum, InAngleNum * 2, InScale)
+	{
+		AngleHandle = InstanceData->Add<2, float>({ InIdentifier, TEXT("Angle") }, { InMaxInstanceNum, InAngleNum }, 0.0f);
+		RelativeAngleHandle = InstanceData->Add<1, float>({ InIdentifier, TEXT("RelativeAngle") }, { InMaxInstanceNum }, 0.0f);
+	}
+
+	void FAngleFeature::Encode(const FIndexSet Instances)
+	{
+		UE_LEARNING_TRACE_CPUPROFILER_EVENT_SCOPE(FAngleFeature::Encode);
+
+		const TLearningArrayView<2, const float> Angle = InstanceData->ConstView(AngleHandle);
+		const TLearningArrayView<1, const float> RelativeAngle = InstanceData->ConstView(RelativeAngleHandle);
+		TLearningArrayView<2, float> Feature = InstanceData->View(FeatureHandle);
+
+		const int32 AngleNum = Angle.Num<1>();
+
+		for (const int32 InstanceIdx : Instances)
+		{
+			for (int32 AngleIdx = 0; AngleIdx < AngleNum; AngleIdx++)
+			{
+				Feature::Encode::Angle(
+					Feature[InstanceIdx].Slice(AngleIdx * 2, 2),
+					Angle[InstanceIdx][AngleIdx],
+					RelativeAngle[InstanceIdx],
+					Scale);
+			}
+		}
+	}
+
+	void FAngleFeature::Decode(const FIndexSet Instances)
+	{
+		UE_LEARNING_TRACE_CPUPROFILER_EVENT_SCOPE(FAngleFeature::Decode);
+
+		TLearningArrayView<2, float> Angle = InstanceData->View(AngleHandle);
+		const TLearningArrayView<1, const float> RelativeAngle = InstanceData->ConstView(RelativeAngleHandle);
+		const TLearningArrayView<2, const float> Feature = InstanceData->ConstView(FeatureHandle);
+
+		const int32 AngleNum = Angle.Num<1>();
+
+		for (const int32 InstanceIdx : Instances)
+		{
+			for (int32 AngleIdx = 0; AngleIdx < AngleNum; AngleIdx++)
+			{
+				Feature::Decode::Angle(
+					Angle[InstanceIdx][AngleIdx],
+					Feature[InstanceIdx].Slice(AngleIdx * 2, 2),
+					RelativeAngle[InstanceIdx],
+					Scale);
+			}
+		}
+	}
+
+	FPlanarRotationFeature::FPlanarRotationFeature(
+		const FName& InIdentifier,
+		const TSharedRef<FArrayMap>& InInstanceData,
+		const int32 InMaxInstanceNum,
+		const int32 InRotationNum,
+		const float InScale,
+		const FVector InAxis)
+		: FFeatureObject(InIdentifier, InInstanceData, InMaxInstanceNum, InRotationNum * 2, InScale)
+		, Axis(InAxis)
+	{
+		RotationHandle = InstanceData->Add<2, FQuat>({ InIdentifier, TEXT("Rotation") }, { InMaxInstanceNum, InRotationNum }, FQuat::Identity);
+		RelativeRotationHandle = InstanceData->Add<1, FQuat>({ InIdentifier, TEXT("RelativeRotation") }, { InMaxInstanceNum }, FQuat::Identity);
+	}
+
+	void FPlanarRotationFeature::Encode(const FIndexSet Instances)
+	{
+		UE_LEARNING_TRACE_CPUPROFILER_EVENT_SCOPE(FPlanarRotationFeature::Encode);
+
+		const TLearningArrayView<2, const FQuat> Rotation = InstanceData->ConstView(RotationHandle);
+		const TLearningArrayView<1, const FQuat> RelativeRotation = InstanceData->ConstView(RelativeRotationHandle);
+		TLearningArrayView<2, float> Feature = InstanceData->View(FeatureHandle);
+
+		const int32 RotationNum = Rotation.Num<1>();
+
+		for (const int32 InstanceIdx : Instances)
+		{
+			for (int32 RotationIdx = 0; RotationIdx < RotationNum; RotationIdx++)
+			{
+				Feature::Encode::PlanarRotation(
+					Feature[InstanceIdx].Slice(RotationIdx * 2, 2),
+					Rotation[InstanceIdx][RotationIdx],
+					RelativeRotation[InstanceIdx],
+					Scale,
+					Axis);
+			}
+		}
+	}
+
+	void FPlanarRotationFeature::Decode(const FIndexSet Instances)
+	{
+		UE_LEARNING_TRACE_CPUPROFILER_EVENT_SCOPE(FPlanarRotationFeature::Decode);
+
+		TLearningArrayView<2, FQuat> Rotation = InstanceData->View(RotationHandle);
+		const TLearningArrayView<1, const FQuat> RelativeRotation = InstanceData->ConstView(RelativeRotationHandle);
+		const TLearningArrayView<2, const float> Feature = InstanceData->ConstView(FeatureHandle);
+
+		const int32 RotationNum = Rotation.Num<1>();
+
+		for (const int32 InstanceIdx : Instances)
+		{
+			for (int32 RotationIdx = 0; RotationIdx < RotationNum; RotationIdx++)
+			{
+				Feature::Decode::PlanarRotation(
+					Rotation[InstanceIdx][RotationIdx],
+					Feature[InstanceIdx].Slice(RotationIdx * 2, 2),
+					RelativeRotation[InstanceIdx],
+					Scale,
+					Axis);
+			}
+		}
+	}
+
+	FRotationFeature::FRotationFeature(
+		const FName& InIdentifier,
+		const TSharedRef<FArrayMap>& InInstanceData,
+		const int32 InMaxInstanceNum,
+		const int32 InRotationNum,
+		const float InScale)
+		: FFeatureObject(InIdentifier, InInstanceData, InMaxInstanceNum, InRotationNum * 6, InScale)
+	{
+		RotationHandle = InstanceData->Add<2, FQuat>({ InIdentifier, TEXT("Rotation") }, { InMaxInstanceNum, InRotationNum }, FQuat::Identity);
+		RelativeRotationHandle = InstanceData->Add<1, FQuat>({ InIdentifier, TEXT("RelativeRotation") }, { InMaxInstanceNum }, FQuat::Identity);
+	}
+
+	void FRotationFeature::Encode(const FIndexSet Instances)
+	{
+		UE_LEARNING_TRACE_CPUPROFILER_EVENT_SCOPE(FRotationFeature::Encode);
+
+		const TLearningArrayView<2, const FQuat> Rotation = InstanceData->ConstView(RotationHandle);
+		const TLearningArrayView<1, const FQuat> RelativeRotation = InstanceData->ConstView(RelativeRotationHandle);
+		TLearningArrayView<2, float> Feature = InstanceData->View(FeatureHandle);
+
+		const int32 RotationNum = Rotation.Num<1>();
+
+		for (const int32 InstanceIdx : Instances)
+		{
+			for (int32 RotationIdx = 0; RotationIdx < RotationNum; RotationIdx++)
+			{
+				Feature::Encode::Rotation(
+					Feature[InstanceIdx].Slice(RotationIdx * 6, 6),
+					Rotation[InstanceIdx][RotationIdx],
+					RelativeRotation[InstanceIdx],
+					Scale);
+			}
+		}
+	}
+
+	void FRotationFeature::Decode(const FIndexSet Instances)
+	{
+		UE_LEARNING_TRACE_CPUPROFILER_EVENT_SCOPE(FRotationFeature::Decode);
+
+		TLearningArrayView<2, FQuat> Rotation = InstanceData->View(RotationHandle);
+		const TLearningArrayView<1, const FQuat> RelativeRotation = InstanceData->ConstView(RelativeRotationHandle);
+		const TLearningArrayView<2, const float> Feature = InstanceData->ConstView(FeatureHandle);
+
+		const int32 RotationNum = Rotation.Num<1>();
+
+		for (const int32 InstanceIdx : Instances)
+		{
+			for (int32 RotationIdx = 0; RotationIdx < RotationNum; RotationIdx++)
+			{
+				Feature::Decode::Rotation(
+					Rotation[InstanceIdx][RotationIdx],
+					Feature[InstanceIdx].Slice(RotationIdx * 6, 6),
+					RelativeRotation[InstanceIdx],
+					Scale);
+			}
+		}
+	}
+
+
+	FRotationVectorFeature::FRotationVectorFeature(
+		const FName& InIdentifier,
+		const TSharedRef<FArrayMap>& InInstanceData,
+		const int32 InMaxInstanceNum,
+		const int32 InRotationVectorNum,
+		const float InScale)
+		: FFeatureObject(InIdentifier, InInstanceData, InMaxInstanceNum, InRotationVectorNum * 3, InScale)
+	{
+		RotationVectorsHandle = InstanceData->Add<2, FVector>({ InIdentifier, TEXT("RotationVectors") }, { InMaxInstanceNum, InRotationVectorNum }, FVector::ZeroVector);
+	}
+
+	void FRotationVectorFeature::Encode(const FIndexSet Instances)
+	{
+		UE_LEARNING_TRACE_CPUPROFILER_EVENT_SCOPE(FRotationVectorArrayFeature::Encode);
+
+		const TLearningArrayView<2, const FVector> RotationVectors = InstanceData->ConstView(RotationVectorsHandle);
+		TLearningArrayView<2, float> Feature = InstanceData->View(FeatureHandle);
+
+		const int32 RotationVectorNum = RotationVectors.Num<1>();
+
+		for (const int32 InstanceIdx : Instances)
+		{
+			for (int32 RotationVectorIdx = 0; RotationVectorIdx < RotationVectorNum; RotationVectorIdx++)
+			{
+				Feature::Encode::RotationVector(
+					Feature[InstanceIdx].Slice(RotationVectorIdx * 3, 3),
+					RotationVectors[InstanceIdx][RotationVectorIdx],
+					Scale);
+			}
+		}
+	}
+
+	void FRotationVectorFeature::Decode(const FIndexSet Instances)
+	{
+		UE_LEARNING_TRACE_CPUPROFILER_EVENT_SCOPE(FRotationVectorFeature::Decode);
+
+		TLearningArrayView<2, FVector> RotationVectors = InstanceData->View(RotationVectorsHandle);
+		const TLearningArrayView<2, const float> Feature = InstanceData->ConstView(FeatureHandle);
+
+		const int32 RotationVectorNum = RotationVectors.Num<1>();
+
+		for (const int32 InstanceIdx : Instances)
+		{
+			for (int32 RotationVectorIdx = 0; RotationVectorIdx < RotationVectorNum; RotationVectorIdx++)
+			{
+				Feature::Decode::RotationVector(
+					RotationVectors[InstanceIdx][RotationVectorIdx],
+					Feature[InstanceIdx].Slice(RotationVectorIdx * 3, 3),
+					Scale);
+			}
+		}
+	}
+
+
+	//------------------------------------------------------------------
+
+	FScalarAngularVelocityFeature::FScalarAngularVelocityFeature(
+		const FName& InIdentifier,
+		const TSharedRef<FArrayMap>& InInstanceData,
+		const int32 InMaxInstanceNum,
+		const int32 InAngularVelocityNum,
+		const float InScale)
+		: FFeatureObject(InIdentifier, InInstanceData, InMaxInstanceNum, InAngularVelocityNum, InScale)
+	{
+		AngularVelocityHandle = InstanceData->Add<2, float>({ InIdentifier, TEXT("AngularVelocity") }, { InMaxInstanceNum, InAngularVelocityNum }, 0.0f);
+	}
+
+	void FScalarAngularVelocityFeature::Encode(const FIndexSet Instances)
+	{
+		UE_LEARNING_TRACE_CPUPROFILER_EVENT_SCOPE(FScalarAngularVelocityFeature::Encode);
+
+		const TLearningArrayView<2, const float> AngularVelocity = InstanceData->ConstView(AngularVelocityHandle);
+		TLearningArrayView<2, float> Feature = InstanceData->View(FeatureHandle);
+
+		const int32 AngularVelocityNum = AngularVelocity.Num<1>();
+
+		for (const int32 InstanceIdx : Instances)
+		{
+			for (int32 AngularVelocityIdx = 0; AngularVelocityIdx < AngularVelocityNum; AngularVelocityIdx++)
+			{
+				Feature::Encode::ScalarAngularVelocity(
+					Feature[InstanceIdx][AngularVelocityIdx],
+					AngularVelocity[InstanceIdx][AngularVelocityIdx],
+					Scale);
+			}
+		}
+	}
+
+	void FScalarAngularVelocityFeature::Decode(const FIndexSet Instances)
+	{
+		UE_LEARNING_TRACE_CPUPROFILER_EVENT_SCOPE(FScalarAngularVelocityFeature::Decode);
+
+		TLearningArrayView<2, float> AngularVelocity = InstanceData->View(AngularVelocityHandle);
+		const TLearningArrayView<2, const float> Feature = InstanceData->ConstView(FeatureHandle);
+
+		const int32 AngularVelocityNum = AngularVelocity.Num<1>();
+
+		for (const int32 InstanceIdx : Instances)
+		{
+			for (int32 AngularVelocityIdx = 0; AngularVelocityIdx < AngularVelocityNum; AngularVelocityIdx++)
+			{
+				Feature::Decode::ScalarAngularVelocity(
+					AngularVelocity[InstanceIdx][AngularVelocityIdx],
+					Feature[InstanceIdx][AngularVelocityIdx],
+					Scale);
+			}
+		}
+	}
+
+	FPlanarAngularVelocityFeature::FPlanarAngularVelocityFeature(
+		const FName& InIdentifier,
+		const TSharedRef<FArrayMap>& InInstanceData,
+		const int32 InMaxInstanceNum,
+		const int32 InAngularVelocityNum,
+		const float InScale,
+		const FVector InAxis)
+		: FFeatureObject(InIdentifier, InInstanceData, InMaxInstanceNum, InAngularVelocityNum, InScale)
+		, Axis(InAxis)
+	{
+		AngularVelocityHandle = InstanceData->Add<2, FVector>({ InIdentifier, TEXT("AngularVelocity") }, { InMaxInstanceNum }, FVector::ZeroVector);
+		RelativeRotationHandle = InstanceData->Add<1, FQuat>({ InIdentifier, TEXT("RelativeRotation") }, { InMaxInstanceNum }, FQuat::Identity);
+	}
+
+	void FPlanarAngularVelocityFeature::Encode(const FIndexSet Instances)
+	{
+		UE_LEARNING_TRACE_CPUPROFILER_EVENT_SCOPE(FPlanarAngularVelocityFeature::Encode);
+
+		const TLearningArrayView<2, const FVector> AngularVelocity = InstanceData->ConstView(AngularVelocityHandle);
+		const TLearningArrayView<1, const FQuat> RelativeRotation = InstanceData->ConstView(RelativeRotationHandle);
+		TLearningArrayView<2, float> Feature = InstanceData->View(FeatureHandle);
+
+		const int32 AngularVelocityNum = AngularVelocity.Num<1>();
+
+		for (const int32 InstanceIdx : Instances)
+		{
+			for (int32 AngularVelocityIdx = 0; AngularVelocityIdx < AngularVelocityNum; AngularVelocityIdx++)
+			{
+				Feature::Encode::PlanarAngularVelocity(
+					Feature[InstanceIdx][AngularVelocityIdx],
+					AngularVelocity[InstanceIdx][AngularVelocityIdx],
+					RelativeRotation[InstanceIdx],
+					Scale,
+					Axis);
+			}
+		}
+	}
+
+	void FPlanarAngularVelocityFeature::Decode(const FIndexSet Instances)
+	{
+		UE_LEARNING_TRACE_CPUPROFILER_EVENT_SCOPE(FPlanarAngularVelocityFeature::Decode);
+
+		TLearningArrayView<2, FVector> AngularVelocity = InstanceData->View(AngularVelocityHandle);
+		const TLearningArrayView<1, const FQuat> RelativeRotation = InstanceData->ConstView(RelativeRotationHandle);
+		const TLearningArrayView<2, const float> Feature = InstanceData->ConstView(FeatureHandle);
+
+		const int32 AngularVelocityNum = AngularVelocity.Num<1>();
+
+		for (const int32 InstanceIdx : Instances)
+		{
+			for (int32 AngularVelocityIdx = 0; AngularVelocityIdx < AngularVelocityNum; AngularVelocityIdx++)
+			{
+				Feature::Decode::PlanarAngularVelocity(
+					AngularVelocity[InstanceIdx][AngularVelocityIdx],
+					Feature[InstanceIdx][AngularVelocityIdx],
+					RelativeRotation[InstanceIdx],
+					Scale,
+					Axis);
+			}
+		}
+	}
+
+	FAngularVelocityFeature::FAngularVelocityFeature(
+		const FName& InIdentifier,
+		const TSharedRef<FArrayMap>& InInstanceData,
+		const int32 InMaxInstanceNum,
+		const int32 InAngularVelocityNum,
+		const float InScale)
+		: FFeatureObject(InIdentifier, InInstanceData, InMaxInstanceNum, InAngularVelocityNum * 3, InScale)
+	{
+		AngularVelocityHandle = InstanceData->Add<2, FVector>({ InIdentifier, TEXT("AngularVelocity") }, { InMaxInstanceNum, InAngularVelocityNum }, FVector::ZeroVector);
+		RelativeRotationHandle = InstanceData->Add<1, FQuat>({ InIdentifier, TEXT("RelativeRotation") }, { InMaxInstanceNum }, FQuat::Identity);
+	}
+
+	void FAngularVelocityFeature::Encode(const FIndexSet Instances)
+	{
+		UE_LEARNING_TRACE_CPUPROFILER_EVENT_SCOPE(FAngularVelocityFeature::Encode);
+
+		const TLearningArrayView<2, const FVector> AngularVelocity = InstanceData->ConstView(AngularVelocityHandle);
+		const TLearningArrayView<1, const FQuat> RelativeRotation = InstanceData->ConstView(RelativeRotationHandle);
+		TLearningArrayView<2, float> Feature = InstanceData->View(FeatureHandle);
+
+		const int32 AngularVelocityNum = AngularVelocity.Num<1>();
+
+		for (const int32 InstanceIdx : Instances)
+		{
+			for (int32 AngularVelocityIdx = 0; AngularVelocityIdx < AngularVelocityNum; AngularVelocityIdx++)
+			{
+				Feature::Encode::AngularVelocity(
+					Feature[InstanceIdx].Slice(AngularVelocityIdx * 3, 3),
+					AngularVelocity[InstanceIdx][AngularVelocityIdx],
+					RelativeRotation[InstanceIdx],
+					Scale);
+			}
+		}
+	}
+
+	void FAngularVelocityFeature::Decode(const FIndexSet Instances)
+	{
+		UE_LEARNING_TRACE_CPUPROFILER_EVENT_SCOPE(FAngularVelocityFeature::Decode);
+
+		TLearningArrayView<2, FVector> AngularVelocity = InstanceData->View(AngularVelocityHandle);
+		const TLearningArrayView<1, const FQuat> RelativeRotation = InstanceData->ConstView(RelativeRotationHandle);
+		const TLearningArrayView<2, const float> Feature = InstanceData->ConstView(FeatureHandle);
+
+		const int32 AngularVelocityNum = AngularVelocity.Num<1>();
+
+		for (const int32 InstanceIdx : Instances)
+		{
+			for (int32 AngularVelocityIdx = 0; AngularVelocityIdx < AngularVelocityNum; AngularVelocityIdx++)
+			{
+				Feature::Decode::AngularVelocity(
+					AngularVelocity[InstanceIdx][AngularVelocityIdx],
+					Feature[InstanceIdx].Slice(AngularVelocityIdx * 3, 3),
+					RelativeRotation[InstanceIdx],
+					Scale);
+			}
+		}
+	}
+
+	//------------------------------------------------------------------
+
+	FScalarAngularAccelerationFeature::FScalarAngularAccelerationFeature(
+		const FName& InIdentifier,
+		const TSharedRef<FArrayMap>& InInstanceData,
+		const int32 InMaxInstanceNum,
+		const int32 InAngularAccelerationNum,
+		const float InScale)
+		: FFeatureObject(InIdentifier, InInstanceData, InMaxInstanceNum, InAngularAccelerationNum, InScale)
+	{
+		AngularAccelerationHandle = InstanceData->Add<2, float>({ InIdentifier, TEXT("AngularAcceleration") }, { InMaxInstanceNum, InAngularAccelerationNum }, 0.0f);
+	}
+
+	void FScalarAngularAccelerationFeature::Encode(const FIndexSet Instances)
+	{
+		UE_LEARNING_TRACE_CPUPROFILER_EVENT_SCOPE(FScalarAngularAccelerationFeature::Encode);
+
+		const TLearningArrayView<2, const float> AngularAcceleration = InstanceData->ConstView(AngularAccelerationHandle);
+		TLearningArrayView<2, float> Feature = InstanceData->View(FeatureHandle);
+
+		const int32 AngularAccelerationNum = AngularAcceleration.Num<1>();
+
+		for (const int32 InstanceIdx : Instances)
+		{
+			for (int32 AngularAccelerationIdx = 0; AngularAccelerationIdx < AngularAccelerationNum; AngularAccelerationIdx++)
+			{
+				Feature::Encode::ScalarAngularAcceleration(
+					Feature[InstanceIdx][AngularAccelerationIdx],
+					AngularAcceleration[InstanceIdx][AngularAccelerationIdx],
+					Scale);
+			}
+		}
+	}
+
+	void FScalarAngularAccelerationFeature::Decode(const FIndexSet Instances)
+	{
+		UE_LEARNING_TRACE_CPUPROFILER_EVENT_SCOPE(FScalarAngularAccelerationFeature::Decode);
+
+		TLearningArrayView<2, float> AngularAcceleration = InstanceData->View(AngularAccelerationHandle);
+		const TLearningArrayView<2, const float> Feature = InstanceData->ConstView(FeatureHandle);
+
+		const int32 AngularAccelerationNum = AngularAcceleration.Num<1>();
+
+		for (const int32 InstanceIdx : Instances)
+		{
+			for (int32 AngularAccelerationIdx = 0; AngularAccelerationIdx < AngularAccelerationNum; AngularAccelerationIdx++)
+			{
+				Feature::Decode::ScalarAngularAcceleration(
+					AngularAcceleration[InstanceIdx][AngularAccelerationIdx],
+					Feature[InstanceIdx][AngularAccelerationIdx],
+					Scale);
+			}
+		}
+	}
+
+	FPlanarAngularAccelerationFeature::FPlanarAngularAccelerationFeature(
+		const FName& InIdentifier,
+		const TSharedRef<FArrayMap>& InInstanceData,
+		const int32 InMaxInstanceNum,
+		const int32 InAngularAccelerationNum,
+		const float InScale,
+		const FVector InAxis)
+		: FFeatureObject(InIdentifier, InInstanceData, InMaxInstanceNum, InAngularAccelerationNum, InScale)
+		, Axis(InAxis)
+	{
+		AngularAccelerationHandle = InstanceData->Add<2, FVector>({ InIdentifier, TEXT("AngularAcceleration") }, { InMaxInstanceNum, InAngularAccelerationNum }, FVector::ZeroVector);
+		RelativeRotationHandle = InstanceData->Add<1, FQuat>({ InIdentifier, TEXT("RelativeRotation") }, { InMaxInstanceNum }, FQuat::Identity);
+	}
+
+	void FPlanarAngularAccelerationFeature::Encode(const FIndexSet Instances)
+	{
+		UE_LEARNING_TRACE_CPUPROFILER_EVENT_SCOPE(FPlanarAngularAccelerationFeature::Encode);
+
+		const TLearningArrayView<2, const FVector> AngularAcceleration = InstanceData->ConstView(AngularAccelerationHandle);
+		const TLearningArrayView<1, const FQuat> RelativeRotation = InstanceData->ConstView(RelativeRotationHandle);
+		TLearningArrayView<2, float> Feature = InstanceData->View(FeatureHandle);
+
+		const int32 AngularAccelerationNum = AngularAcceleration.Num<1>();
+
+		for (const int32 InstanceIdx : Instances)
+		{
+			for (int32 AngularAccelerationIdx = 0; AngularAccelerationIdx < AngularAccelerationNum; AngularAccelerationIdx++)
+			{
+				Feature::Encode::PlanarAngularAcceleration(
+					Feature[InstanceIdx][AngularAccelerationIdx],
+					AngularAcceleration[InstanceIdx][AngularAccelerationIdx],
+					RelativeRotation[InstanceIdx],
+					Scale,
+					Axis);
+			}
+		}
+	}
+
+	void FPlanarAngularAccelerationFeature::Decode(const FIndexSet Instances)
+	{
+		UE_LEARNING_TRACE_CPUPROFILER_EVENT_SCOPE(FPlanarAngularAccelerationFeature::Decode);
+
+		TLearningArrayView<2, FVector> AngularAcceleration = InstanceData->View(AngularAccelerationHandle);
+		const TLearningArrayView<1, const FQuat> RelativeRotation = InstanceData->ConstView(RelativeRotationHandle);
+		const TLearningArrayView<2, const float> Feature = InstanceData->ConstView(FeatureHandle);
+
+		const int32 AngularAccelerationNum = AngularAcceleration.Num<1>();
+
+		for (const int32 InstanceIdx : Instances)
+		{
+			for (int32 AngularAccelerationIdx = 0; AngularAccelerationIdx < AngularAccelerationNum; AngularAccelerationIdx++)
+			{
+				Feature::Decode::PlanarAngularAcceleration(
+					AngularAcceleration[InstanceIdx][AngularAccelerationIdx],
+					Feature[InstanceIdx][AngularAccelerationIdx],
+					RelativeRotation[InstanceIdx],
+					Scale,
+					Axis);
+			}
+		}
+	}
+
+	FAngularAccelerationFeature::FAngularAccelerationFeature(
+		const FName& InIdentifier,
+		const TSharedRef<FArrayMap>& InInstanceData,
+		const int32 InMaxInstanceNum,
+		const int32 InAngularAccelerationNum,
+		const float InScale)
+		: FFeatureObject(InIdentifier, InInstanceData, InMaxInstanceNum, InAngularAccelerationNum * 3, InScale)
+	{
+		AngularAccelerationHandle = InstanceData->Add<2, FVector>({ InIdentifier, TEXT("AngularAcceleration") }, { InMaxInstanceNum, InAngularAccelerationNum }, FVector::ZeroVector);
+		RelativeRotationHandle = InstanceData->Add<1, FQuat>({ InIdentifier, TEXT("RelativeRotation") }, { InMaxInstanceNum }, FQuat::Identity);
+	}
+
+	void FAngularAccelerationFeature::Encode(const FIndexSet Instances)
+	{
+		UE_LEARNING_TRACE_CPUPROFILER_EVENT_SCOPE(FAngularAccelerationFeature::Encode);
+
+		const TLearningArrayView<2, const FVector> AngularAcceleration = InstanceData->ConstView(AngularAccelerationHandle);
+		const TLearningArrayView<1, const FQuat> RelativeRotation = InstanceData->ConstView(RelativeRotationHandle);
+		TLearningArrayView<2, float> Feature = InstanceData->View(FeatureHandle);
+
+		const int32 AngularAccelerationNum = AngularAcceleration.Num<1>();
+
+		for (const int32 InstanceIdx : Instances)
+		{
+			for (int32 AngularAccelerationIdx = 0; AngularAccelerationIdx < AngularAccelerationNum; AngularAccelerationIdx++)
+			{
+				Feature::Encode::AngularAcceleration(
+					Feature[InstanceIdx].Slice(AngularAccelerationIdx * 3, 3),
+					AngularAcceleration[InstanceIdx][AngularAccelerationIdx],
+					RelativeRotation[InstanceIdx],
+					Scale);
+			}
+		}
+	}
+
+	void FAngularAccelerationFeature::Decode(const FIndexSet Instances)
+	{
+		UE_LEARNING_TRACE_CPUPROFILER_EVENT_SCOPE(FAngularAccelerationFeature::Decode);
+
+		TLearningArrayView<2, FVector> AngularAcceleration = InstanceData->View(AngularAccelerationHandle);
+		const TLearningArrayView<1, const FQuat> RelativeRotation = InstanceData->ConstView(RelativeRotationHandle);
+		const TLearningArrayView<2, const float> Feature = InstanceData->ConstView(FeatureHandle);
+
+		const int32 AngularAccelerationNum = AngularAcceleration.Num<1>();
+
+		for (const int32 InstanceIdx : Instances)
+		{
+			for (int32 AngularAccelerationIdx = 0; AngularAccelerationIdx < AngularAccelerationNum; AngularAccelerationIdx++)
+			{
+				Feature::Decode::AngularAcceleration(
+					AngularAcceleration[InstanceIdx][AngularAccelerationIdx],
+					Feature[InstanceIdx].Slice(AngularAccelerationIdx * 3, 3),
+					RelativeRotation[InstanceIdx],
+					Scale);
+			}
+		}
+	}
+
+	//------------------------------------------------------------------
+
+	FScaleFeature::FScaleFeature(
+		const FName& InIdentifier,
+		const TSharedRef<FArrayMap>& InInstanceData,
+		const int32 InMaxInstanceNum,
+		const int32 InScaleNum,
+		const float InScale)
+		: FFeatureObject(InIdentifier, InInstanceData, InMaxInstanceNum, InScaleNum, InScale)
+	{
+		ScaleHandle = InstanceData->Add<2, float>({ InIdentifier, TEXT("Scale") }, { InMaxInstanceNum, InScaleNum }, 1.0f);
+		RelativeScaleHandle = InstanceData->Add<1, float>({ InIdentifier, TEXT("RelativeScale") }, { InMaxInstanceNum }, 1.0f);
+	}
+
+	void FScaleFeature::Encode(const FIndexSet Instances)
+	{
+		UE_LEARNING_TRACE_CPUPROFILER_EVENT_SCOPE(FScaleFeature::Encode);
+
+		const TLearningArrayView<2, const float> ScaleValue = InstanceData->ConstView(ScaleHandle);
+		const TLearningArrayView<1, const float> RelativeScale = InstanceData->ConstView(RelativeScaleHandle);
+		TLearningArrayView<2, float> Feature = InstanceData->View(FeatureHandle);
+
+		const int32 ScaleNum = ScaleValue.Num<1>();
+
+		for (const int32 InstanceIdx : Instances)
+		{
+			for (int32 ScaleIdx = 0; ScaleIdx < ScaleNum; ScaleIdx++)
+			{
+				Feature::Encode::Scale(
+					Feature[InstanceIdx][ScaleIdx],
+					ScaleValue[InstanceIdx][ScaleIdx],
+					RelativeScale[InstanceIdx],
+					Scale);
+			}
+		}
+	}
+
+	void FScaleFeature::Decode(const FIndexSet Instances)
+	{
+		UE_LEARNING_TRACE_CPUPROFILER_EVENT_SCOPE(FScaleFeature::Decode);
+
+		TLearningArrayView<2, float> ScaleValue = InstanceData->View(ScaleHandle);
+		const TLearningArrayView<1, const float> RelativeScale = InstanceData->ConstView(RelativeScaleHandle);
+		const TLearningArrayView<2, const float> Feature = InstanceData->ConstView(FeatureHandle);
+
+		const int32 ScaleNum = ScaleValue.Num<1>();
+
+		for (const int32 InstanceIdx : Instances)
+		{
+			for (int32 ScaleIdx = 0; ScaleIdx < ScaleNum; ScaleIdx++)
+			{
+				Feature::Decode::Scale(
+					ScaleValue[InstanceIdx][ScaleIdx],
+					Feature[InstanceIdx][ScaleIdx],
+					RelativeScale[InstanceIdx],
+					Scale);
+			}
+		}
+	}
+}
diff --git a/Engine/Plugins/Experimental/LearningAgents/Source/Learning/Private/LearningFunctionObject.cpp b/Engine/Plugins/Experimental/LearningAgents/Source/Learning/Private/LearningFunctionObject.cpp
new file mode 100644
--- /dev/null
+++ b/Engine/Plugins/Experimental/LearningAgents/Source/Learning/Private/LearningFunctionObject.cpp
@@ -0,0 +1,229 @@
+// Copyright Epic Games, Inc. All Rights Reserved.
+
+#include "LearningFunctionObject.h"
+
+#include "LearningRandom.h"
+
+namespace UE::Learning
+{
+	FFunctionObject::FFunctionObject(const TSharedRef<FArrayMap>& InInstanceData) : InstanceData(InInstanceData) {}
+
+	FSequentialFunction::FSequentialFunction(
+		const TLearningArrayView<1, const TSharedRef<FFunctionObject>> InFunctions,
+		const TSharedRef<FArrayMap>& InInstanceData)
+		: FFunctionObject(InInstanceData)
+		, Functions(InFunctions) {}
+
+	void FSequentialFunction::Evaluate(const FIndexSet Instances)
+	{
+		for (const TSharedRef<FFunctionObject>& Function : Functions)
+		{
+			Function->Evaluate(Instances);
+		}
+	}
+
+	FCopyVectorsFunction::FCopyVectorsFunction(
+		const FName& InIdentifier,
+		const TSharedRef<FArrayMap>& InInstanceData,
+		const int32 InMaxInstanceNum,
+		const int32 InDimensionsNum)
+		: FFunctionObject(InInstanceData)
+	{
+		InputHandle = InstanceData->Add<2, float>({ InIdentifier, TEXT("Seed") }, { InMaxInstanceNum, InDimensionsNum }, 0.0f);
+		OutputHandle = InstanceData->Add<2, float>({ InIdentifier, TEXT("Min") }, { InMaxInstanceNum, InDimensionsNum }, 0.0f);
+	}
+
+	void FCopyVectorsFunction::Evaluate(const FIndexSet Instances)
+	{
+		UE_LEARNING_TRACE_CPUPROFILER_EVENT_SCOPE(Learning::FCopyVectorsFunction::Evaluate);
+
+		const TLearningArrayView<2, const float> Input = InstanceData->ConstView(InputHandle);
+		TLearningArrayView<2, float> Output = InstanceData->View(OutputHandle);
+
+		Array::Copy(Output, Input, Instances);
+	}
+
+	FExtractRotationsFromTransformsFunction::FExtractRotationsFromTransformsFunction(
+		const FName& InIdentifier,
+		const TSharedRef<FArrayMap>& InInstanceData,
+		const int32 InMaxInstanceNum)
+		: FFunctionObject(InInstanceData)
+	{
+		TransformHandle = InstanceData->Add<1, FTransform>({ InIdentifier, TEXT("Transform") }, { InMaxInstanceNum }, FTransform::Identity);
+		RotationHandle = InstanceData->Add<1, FQuat>({ InIdentifier, TEXT("Rotation") }, { InMaxInstanceNum }, FQuat::Identity);
+	}
+
+	void FExtractRotationsFromTransformsFunction::Evaluate(const FIndexSet Instances)
+	{
+		UE_LEARNING_TRACE_CPUPROFILER_EVENT_SCOPE(Learning::FExtractRotationsFromTransformsFunction::Evaluate);
+
+		const TLearningArrayView<1, const FTransform> Transform = InstanceData->ConstView(TransformHandle);
+		TLearningArrayView<1, FQuat> Rotation = InstanceData->View(RotationHandle);
+
+		for (const int32 InstanceIdx : Instances)
+		{
+			Rotation[InstanceIdx] = Transform[InstanceIdx].GetRotation();
+		}
+	}
+
+	FExtractPositionsAndRotationsFromTransformsFunction::FExtractPositionsAndRotationsFromTransformsFunction(
+		const FName& InIdentifier,
+		const TSharedRef<FArrayMap>& InInstanceData,
+		const int32 InMaxInstanceNum)
+		: FFunctionObject(InInstanceData)
+	{
+		TransformHandle = InstanceData->Add<1, FTransform>({ InIdentifier, TEXT("Transform") }, { InMaxInstanceNum }, FTransform::Identity);
+		RotationHandle = InstanceData->Add<1, FQuat>({ InIdentifier, TEXT("Rotation") }, { InMaxInstanceNum }, FQuat::Identity);
+		PositionHandle = InstanceData->Add<1, FVector>({ InIdentifier, TEXT("Position") }, { InMaxInstanceNum }, FVector::ZeroVector);
+	}
+
+	void FExtractPositionsAndRotationsFromTransformsFunction::Evaluate(const FIndexSet Instances)
+	{
+		UE_LEARNING_TRACE_CPUPROFILER_EVENT_SCOPE(Learning::FExtractPositionsAndRotationsFromTransformsFunction::Evaluate);
+
+		const TLearningArrayView<1, const FTransform> Transform = InstanceData->ConstView(TransformHandle);
+		TLearningArrayView<1, FQuat> Rotation = InstanceData->View(RotationHandle);
+		TLearningArrayView<1, FVector> Position = InstanceData->View(PositionHandle);
+
+		for (const int32 InstanceIdx : Instances)
+		{
+			Rotation[InstanceIdx] = Transform[InstanceIdx].GetRotation();
+			Position[InstanceIdx] = Transform[InstanceIdx].GetTranslation();
+		}
+	}
+
+	FRandomUniformFunction::FRandomUniformFunction(
+		const FName& InIdentifier,
+		const TSharedRef<FArrayMap>& InInstanceData,
+		const int32 InMaxInstanceNum,
+		const uint32 InSeed,
+		const float InMin,
+		const float InMax)
+		: FFunctionObject(InInstanceData)
+	{
+		SeedHandle = InstanceData->Add<1, uint32>({ InIdentifier, TEXT("Seed") }, { InMaxInstanceNum });
+		MinHandle = InstanceData->Add<1, float>({ InIdentifier, TEXT("Min") }, { InMaxInstanceNum }, InMin);
+		MaxHandle = InstanceData->Add<1, float>({ InIdentifier, TEXT("Max") }, { InMaxInstanceNum }, InMax);
+		ValueHandle = InstanceData->Add<1, float>({ InIdentifier, TEXT("Value") }, { InMaxInstanceNum }, 0.0f);
+
+		Random::IntArray(InstanceData->View(SeedHandle), InSeed);
+	}
+
+	void FRandomUniformFunction::Evaluate(const FIndexSet Instances)
+	{
+		UE_LEARNING_TRACE_CPUPROFILER_EVENT_SCOPE(Learning::FRandomUniformFunction::Evaluate);
+
+		const TLearningArrayView<1, const float> Min = InstanceData->ConstView(MinHandle);
+		const TLearningArrayView<1, const float> Max = InstanceData->ConstView(MaxHandle);
+		TLearningArrayView<1, uint32> Seed = InstanceData->View(SeedHandle);
+		TLearningArrayView<1, float> Values = InstanceData->View(ValueHandle);
+
+		for (const int32 InstanceIdx : Instances)
+		{
+			Values[InstanceIdx] = Random::SampleUniform(
+				Seed[InstanceIdx],
+				Min[InstanceIdx],
+				Max[InstanceIdx]);
+		}
+	}
+
+	FRandomPlanarClippedGaussianFunction::FRandomPlanarClippedGaussianFunction(
+		const FName& InIdentifier,
+		const TSharedRef<FArrayMap>& InInstanceData,
+		const int32 InMaxInstanceNum,
+		const uint32 InSeed,
+		const float InMean,
+		const float InStd,
+		const float InClip)
+		: FFunctionObject(InInstanceData)
+	{
+		SeedHandle = InstanceData->Add<1, uint32>({ InIdentifier, TEXT("Seed") }, { InMaxInstanceNum });
+		MeanHandle = InstanceData->Add<1, float>({ InIdentifier, TEXT("Mean") }, { InMaxInstanceNum }, InMean);
+		StdHandle = InstanceData->Add<1, float>({ InIdentifier, TEXT("Std") }, { InMaxInstanceNum }, InStd);
+		ClipHandle = InstanceData->Add<1, float>({ InIdentifier, TEXT("Clip") }, { InMaxInstanceNum }, InClip);
+		ValueHandle = InstanceData->Add<1, FVector>({ InIdentifier, TEXT("Value") }, { InMaxInstanceNum }, FVector::ZeroVector);
+
+		Random::IntArray(InstanceData->View(SeedHandle), InSeed);
+	}
+
+	void FRandomPlanarClippedGaussianFunction::Evaluate(const FIndexSet Instances)
+	{
+		UE_LEARNING_TRACE_CPUPROFILER_EVENT_SCOPE(Learning::FRandomPlanarClippedGaussianFunction::Evaluate);
+
+		const TLearningArrayView<1, const float> Mean = InstanceData->ConstView(MeanHandle);
+		const TLearningArrayView<1, const float> Std = InstanceData->ConstView(StdHandle);
+		const TLearningArrayView<1, const float> Clip = InstanceData->ConstView(ClipHandle);
+		TLearningArrayView<1, uint32> Seed = InstanceData->View(SeedHandle);
+		TLearningArrayView<1, FVector> Values = InstanceData->View(ValueHandle);
+
+		for (const int32 InstanceIdx : Instances)
+		{
+			Values[InstanceIdx] = Random::SamplePlanarClippedGaussian(
+				Seed[InstanceIdx],
+				Mean[InstanceIdx],
+				Std[InstanceIdx],
+				Clip[InstanceIdx]);
+		}
+	}
+
+	FRandomPlanarDirectionFunction::FRandomPlanarDirectionFunction(
+		const FName& InIdentifier,
+		const TSharedRef<FArrayMap>& InInstanceData,
+		const uint32 InSeed,
+		const int32 InMaxInstanceNum)
+		: FFunctionObject(InInstanceData)
+	{
+		SeedHandle = InstanceData->Add<1, uint32>({ InIdentifier, TEXT("Seed") }, { InMaxInstanceNum });
+		DirectionHandle = InstanceData->Add<1, FVector>({ InIdentifier, TEXT("Direction") }, { InMaxInstanceNum }, FVector::ForwardVector);
+
+		Random::IntArray(InstanceData->View(SeedHandle), InSeed);
+	}
+
+	void FRandomPlanarDirectionFunction::Evaluate(const FIndexSet Instances)
+	{
+		UE_LEARNING_TRACE_CPUPROFILER_EVENT_SCOPE(Learning::FRandomPlanarDirectionFunction::Evaluate);
+
+		TLearningArrayView<1, uint32> Seed = InstanceData->View(SeedHandle);
+		TLearningArrayView<1, FVector> Direction = InstanceData->View(DirectionHandle);
+
+		for (const int32 InstanceIdx : Instances)
+		{
+			Direction[InstanceIdx] = Random::SamplePlanarDirection(Seed[InstanceIdx]);
+		}
+	}
+
+	FRandomPlanarDirectionVelocityFunction::FRandomPlanarDirectionVelocityFunction(
+		const FName& InIdentifier,
+		const TSharedRef<FArrayMap>& InInstanceData,
+		const int32 InMaxInstanceNum,
+		const uint32 InSeed,
+		const float InVelocityScale)
+		: FFunctionObject(InInstanceData)
+	{
+		SeedHandle = InstanceData->Add<1, uint32>({ InIdentifier, TEXT("Seed") }, { InMaxInstanceNum });
+		VelocityScaleHandle = InstanceData->Add<1, float>({ InIdentifier, TEXT("VelocityScale") }, { InMaxInstanceNum }, InVelocityScale);
+		DirectionHandle = InstanceData->Add<1, FVector>({ InIdentifier, TEXT("Direction") }, { InMaxInstanceNum }, FVector::ForwardVector);
+		VelocityHandle = InstanceData->Add<1, FVector>({ InIdentifier, TEXT("Velocity") }, { InMaxInstanceNum }, FVector::ZeroVector);
+
+		Random::IntArray(InstanceData->View(SeedHandle), InSeed);
+	}
+
+	void FRandomPlanarDirectionVelocityFunction::Evaluate(const FIndexSet Instances)
+	{
+		UE_LEARNING_TRACE_CPUPROFILER_EVENT_SCOPE(Learning::FRandomPlanarDirectionVelocityFunction::Evaluate);
+
+		const TLearningArrayView<1, const float> VelocityScale = InstanceData->ConstView(VelocityScaleHandle);
+		TLearningArrayView<1, uint32> Seed = InstanceData->View(SeedHandle);
+		TLearningArrayView<1, FVector> Direction = InstanceData->View(DirectionHandle);
+		TLearningArrayView<1, FVector> Velocity = InstanceData->View(VelocityHandle);
+
+		for (const int32 InstanceIdx : Instances)
+		{
+			const FVector RandomDirection = Random::SamplePlanarDirection(Seed[InstanceIdx]);
+			Direction[InstanceIdx] = RandomDirection;
+			Velocity[InstanceIdx] = VelocityScale[InstanceIdx] * RandomDirection;
+		}
+	}
+
+}
+
diff --git a/Engine/Plugins/Experimental/LearningAgents/Source/Learning/Private/LearningNeuralNetwork.cpp b/Engine/Plugins/Experimental/LearningAgents/Source/Learning/Private/LearningNeuralNetwork.cpp
new file mode 100644
--- /dev/null
+++ b/Engine/Plugins/Experimental/LearningAgents/Source/Learning/Private/LearningNeuralNetwork.cpp
@@ -0,0 +1,133 @@
+// Copyright Epic Games, Inc. All Rights Reserved.
+
+#include "LearningNeuralNetwork.h"
+
+namespace UE::Learning
+{
+	const TCHAR* ActivationFunctionString(const EActivationFunction ActivationFunction)
+	{
+		switch (ActivationFunction)
+		{
+		case EActivationFunction::ReLU: return TEXT("ReLU");
+		case EActivationFunction::ELU: return TEXT("ELU");
+		case EActivationFunction::TanH: return TEXT("TanH");
+		default: UE_LEARNING_NOT_IMPLEMENTED(); return TEXT("Unknown");
+		}
+	}
+
+	void FNeuralNetwork::Resize(
+		const int32 InputNum,
+		const int32 OutputNum,
+		const int32 HiddenNum,
+		const int32 LayerNum)
+	{
+		UE_LEARNING_TRACE_CPUPROFILER_EVENT_SCOPE(Learning::FNeuralNetwork::Resize);
+		UE_LEARNING_CHECKF(LayerNum >= 2, TEXT("At least two layers required (input and output layers)"));
+
+		Weights.SetNum(LayerNum);
+		Biases.SetNum(LayerNum);
+
+		Weights[0].SetNumUninitialized({ InputNum, HiddenNum });
+		Biases[0].SetNumUninitialized({ HiddenNum });
+
+		for (int32 LayerIdx = 0; LayerIdx < LayerNum - 2; LayerIdx++)
+		{
+			Weights[LayerIdx + 1].SetNumUninitialized({ HiddenNum, HiddenNum });
+			Biases[LayerIdx + 1].SetNumUninitialized({ HiddenNum });
+		}
+
+		Weights[LayerNum - 1].SetNumUninitialized({ HiddenNum, OutputNum });
+		Biases[LayerNum - 1].SetNumUninitialized({ OutputNum });
+
+		for (int32 LayerIdx = 0; LayerIdx < LayerNum; LayerIdx++)
+		{
+			Array::Zero(Weights[LayerIdx]);
+			Array::Zero(Biases[LayerIdx]);
+		}
+	}
+
+	int32 FNeuralNetwork::GetInputNum() const
+	{
+		return Weights[0].Num<0>();
+	}
+
+	int32 FNeuralNetwork::GetOutputNum() const
+	{
+		return Weights[Weights.Num() - 1].Num<1>();
+	}
+
+	int32 FNeuralNetwork::GetLayerNum() const
+	{
+		return Weights.Num();
+	}
+
+	int32 FNeuralNetwork::GetHiddenNum() const
+	{
+		return Weights[0].Num(1);
+	}
+
+	int32 FNeuralNetwork::GetTotalByteNum() const
+	{
+		const int32 LayerNum = Weights.Num();
+
+		int32 Total = 0;
+
+		for (int32 LayerIdx = 0; LayerIdx < LayerNum; LayerIdx++)
+		{
+			Total += Weights[LayerIdx].Num() * sizeof(float);
+			Total += Biases[LayerIdx].Num() * sizeof(float);
+		}
+
+		return Total;
+	}
+
+	void FNeuralNetwork::DeserializeFromBytes(const TLearningArrayView<1, const uint8> RawBytes)
+	{
+		UE_LEARNING_TRACE_CPUPROFILER_EVENT_SCOPE(Learning::FNeuralNetwork::DeserializeFromBytes);
+
+		UE_LEARNING_CHECK(RawBytes.Num() == GetTotalByteNum());
+
+		const int32 LayerNum = Weights.Num();
+
+		int32 Offset = 0;
+
+		for (int32 LayerIdx = 0; LayerIdx < LayerNum; LayerIdx++)
+		{
+			const int32 WeightByteNum = Weights[LayerIdx].Num() * sizeof(float);
+			FMemory::Memcpy((uint8*)Weights[LayerIdx].GetData(), &RawBytes[Offset], WeightByteNum);
+			Array::Check(Weights[LayerIdx]);
+			Offset += WeightByteNum;
+
+			const int32 BiasByteNum = Biases[LayerIdx].Num() * sizeof(float);
+			FMemory::Memcpy((uint8*)Biases[LayerIdx].GetData(), &RawBytes[Offset], BiasByteNum);
+			Array::Check(Biases[LayerIdx]);
+			Offset += BiasByteNum;
+		}
+
+		UE_LEARNING_CHECK(Offset == GetTotalByteNum());
+	}
+
+	void FNeuralNetwork::SerializeToBytes(TLearningArrayView<1, uint8> OutRawBytes) const
+	{
+		UE_LEARNING_TRACE_CPUPROFILER_EVENT_SCOPE(Learning::FNeuralNetwork::SerializeToBytes);
+
+		UE_LEARNING_CHECK(OutRawBytes.Num() == GetTotalByteNum());
+
+		const int32 LayerNum = Weights.Num();
+
+		int32 Offset = 0;
+
+		for (int32 LayerIdx = 0; LayerIdx < LayerNum; LayerIdx++)
+		{
+			const int32 WeightByteNum = Weights[LayerIdx].Num() * sizeof(float);
+			FMemory::Memcpy(&OutRawBytes[Offset], (uint8*)Weights[LayerIdx].GetData(), WeightByteNum);
+			Offset += WeightByteNum;
+
+			const int32 BiasByteNum = Biases[LayerIdx].Num() * sizeof(float);
+			FMemory::Memcpy(&OutRawBytes[Offset], (uint8*)Biases[LayerIdx].GetData(), BiasByteNum);
+			Offset += BiasByteNum;
+		}
+
+		UE_LEARNING_CHECK(Offset == GetTotalByteNum());
+	}
+}
