Learning Agents: Create initial experimental version

diff --git a/Engine/Plugins/Experimental/LearningAgents/Source/Learning/Private/LearningPolicyObject.cpp b/Engine/Plugins/Experimental/LearningAgents/Source/Learning/Private/LearningPolicyObject.cpp
new file mode 100644
--- /dev/null
+++ b/Engine/Plugins/Experimental/LearningAgents/Source/Learning/Private/LearningPolicyObject.cpp
@@ -0,0 +1,341 @@
+// Copyright Epic Games, Inc. All Rights Reserved.
+
+#include "LearningPolicyObject.h"
+
+#include "LearningNeuralNetwork.h"
+#include "LearningEigen.h"
+#include "LearningRandom.h"
+
+#if UE_LEARNING_ISPC
+#include "Learning.ispc.generated.h"
+#endif
+
+namespace UE::Learning
+{
+	namespace NeuralNetworkPolicy
+	{
+		static inline float ReLU(const float X)
+		{
+			return FMath::Max(X, 0.0f);
+		}
+
+		static inline float ELU(const float X)
+		{
+			return X > 0.0f ? X : FMath::InvExpApprox(-X) - 1.0f;
+		}
+
+		static inline float Sigmoid(const float X)
+		{
+			return 1.0f / (1.0f + FMath::InvExpApprox(X));
+		}
+
+		static inline float TanH(const float X)
+		{
+			return FMath::Tanh(X);
+		}
+
+		static inline void MatMulPlusBias(
+			TLearningArrayView<2, float> Output,
+			const TLearningArrayView<2, const float> Input,
+			const TLearningArrayView<2, const float> Weights,
+			const TLearningArrayView<1, const float> Biases,
+			const FIndexSet Instances)
+		{
+			UE_LEARNING_TRACE_CPUPROFILER_EVENT_SCOPE(Learning::NeuralNetworkPolicy::MatMulPlusBias);
+
+			const int32 RowNum = Weights.Num<0>();
+			const int32 ColNum = Weights.Num<1>();
+
+			if (UE_LEARNING_ISPC && Instances.IsSlice())
+			{
+#if UE_LEARNING_ISPC
+				ispc::LearningLayerMatMulPlusBias(
+					Output.Slice(Instances.GetSliceStart(), Instances.GetSliceNum()).GetData(),
+					Input.Slice(Instances.GetSliceStart(), Instances.GetSliceNum()).GetData(),
+					Weights.GetData(),
+					Biases.GetData(),
+					Instances.GetSliceNum(),
+					RowNum,
+					ColNum);
+#endif
+			}
+			else
+			{
+				OutEigenMatrix(Output).noalias() =
+					(InEigenMatrix(Weights).transpose() *
+						InEigenMatrix(Input).transpose()).transpose().rowwise() +
+					InEigenRowVector(Biases);
+			}
+		}
+
+		static inline void ActivationReLU(TLearningArrayView<2, float> InputOutput, const FIndexSet Instances)
+		{
+			UE_LEARNING_TRACE_CPUPROFILER_EVENT_SCOPE(Learning::NeuralNetworkPolicy::ActivationReLU);
+
+			const int32 HiddenNum = InputOutput.Num<1>();
+
+			if (UE_LEARNING_ISPC && Instances.IsSlice())
+			{
+#if UE_LEARNING_ISPC
+				ispc::LearningLayerReLU(
+					InputOutput.Slice(Instances.GetSliceStart(), Instances.GetSliceNum()).GetData(),
+					Instances.GetSliceNum(),
+					HiddenNum);
+#endif
+			}
+			else
+			{
+				for (const int32 InstanceIdx : Instances)
+				{
+					for (int32 HiddenIdx = 0; HiddenIdx < HiddenNum; HiddenIdx++)
+					{
+						InputOutput[InstanceIdx][HiddenIdx] = ReLU(InputOutput[InstanceIdx][HiddenIdx]);
+					}
+				}
+			}
+		}
+
+		static inline void ActivationELU(TLearningArrayView<2, float> InputOutput, const FIndexSet Instances)
+		{
+			UE_LEARNING_TRACE_CPUPROFILER_EVENT_SCOPE(Learning::NeuralNetworkPolicy::ActivationELU);
+
+			const int32 HiddenNum = InputOutput.Num<1>();
+
+			if (UE_LEARNING_ISPC && Instances.IsSlice())
+			{
+#if UE_LEARNING_ISPC
+				ispc::LearningLayerELU(
+					InputOutput.Slice(Instances.GetSliceStart(), Instances.GetSliceNum()).GetData(),
+					Instances.GetSliceNum(),
+					HiddenNum);
+#endif
+			}
+			else
+			{
+				for (const int32 InstanceIdx : Instances)
+				{
+					for (int32 HiddenIdx = 0; HiddenIdx < HiddenNum; HiddenIdx++)
+					{
+						InputOutput[InstanceIdx][HiddenIdx] = ELU(InputOutput[InstanceIdx][HiddenIdx]);
+					}
+				}
+			}
+		}
+
+		static inline void ActivationTanH(TLearningArrayView<2, float> InputOutput, const FIndexSet Instances)
+		{
+			UE_LEARNING_TRACE_CPUPROFILER_EVENT_SCOPE(Learning::NeuralNetworkPolicy::ActivationTanH);
+
+			const int32 HiddenNum = InputOutput.Num<1>();
+
+			if (UE_LEARNING_ISPC && Instances.IsSlice())
+			{
+#if UE_LEARNING_ISPC
+				ispc::LearningLayerTanH(
+					InputOutput.Slice(Instances.GetSliceStart(), Instances.GetSliceNum()).GetData(),
+					Instances.GetSliceNum(),
+					HiddenNum);
+#endif
+			}
+			else
+			{
+				for (const int32 InstanceIdx : Instances)
+				{
+					for (int32 HiddenIdx = 0; HiddenIdx < HiddenNum; HiddenIdx++)
+					{
+						InputOutput[InstanceIdx][HiddenIdx] = TanH(InputOutput[InstanceIdx][HiddenIdx]);
+					}
+				}
+			}
+		}
+
+		static inline void ActionNoise(
+			TLearningArrayView<2, float> Output,
+			const TLearningArrayView<2, float> Input,
+			TLearningArrayView<1, uint32> Seed,
+			const TLearningArrayView<1, const float> ActionNoiseScale,
+			const float LogActionNoiseMin,
+			const float LogActionNoiseMax,
+			const FIndexSet Instances)
+		{
+			UE_LEARNING_TRACE_CPUPROFILER_EVENT_SCOPE(Learning::NeuralNetworkPolicy::ActionNoise);
+
+			const int32 InputNum = Input.Num<1>();
+			const int32 OutputNum = Output.Num<1>();
+
+			if (UE_LEARNING_ISPC && Instances.IsSlice())
+			{
+#if UE_LEARNING_ISPC
+				ispc::LearningLayerActionNoise(
+					Output.Slice(Instances.GetSliceStart(), Instances.GetSliceNum()).GetData(),
+					Input.Slice(Instances.GetSliceStart(), Instances.GetSliceNum()).GetData(),
+					Seed.Slice(Instances.GetSliceStart(), Instances.GetSliceNum()).GetData(),
+					ActionNoiseScale.Slice(Instances.GetSliceStart(), Instances.GetSliceNum()).GetData(),
+					Instances.GetSliceNum(),
+					InputNum,
+					OutputNum,
+					LogActionNoiseMin,
+					LogActionNoiseMax);
+#endif
+			}
+			else
+			{
+				for (const int32 InstanceIdx : Instances)
+				{
+					for (int32 OutputIdx = 0; OutputIdx < OutputNum; OutputIdx++)
+					{
+						const float ActionMean = Input[InstanceIdx][OutputIdx];
+						const float LogActionStd = Input[InstanceIdx][OutputNum + OutputIdx];
+						const float ActionStd = ActionNoiseScale[InstanceIdx] *
+							FMath::Exp(Sigmoid(LogActionStd) * (LogActionNoiseMax - LogActionNoiseMin) + LogActionNoiseMin);
+
+						Output[InstanceIdx][OutputIdx] = Random::Gaussian(
+							Seed[InstanceIdx] ^ 0xab744615 ^ Random::Int(OutputIdx ^ 0xf8a88a27),
+							ActionMean,
+							ActionStd);
+					}
+				}
+			}
+
+			Random::ResampleStateArray(Seed, Instances);
+		}
+
+		static inline void EvaluateLayer(
+			TLearningArrayView<2, float> Outputs,
+			TArrayView<TLearningArray<2, float, TInlineAllocator<128>>> Activations,
+			const TLearningArrayView<2, const float> Inputs,
+			TLearningArrayView<1, uint32> Seed,
+			const TLearningArrayView<1, const float> ActionNoiseScale,
+			const int32 LayerIdx,
+			const int32 LayerNum,
+			const float LogActionNoiseMin,
+			const float LogActionNoiseMax,
+			const FNeuralNetwork& NeuralNetwork,
+			const FIndexSet Instances)
+		{
+			UE_LEARNING_TRACE_CPUPROFILER_EVENT_SCOPE(Learning::NeuralNetworkPolicy::EvaluateLayer);
+
+			TLearningArrayView<2, float> LayerOutput = Activations[LayerIdx - 0];
+			const TLearningArrayView<2, const float> LayerInput = LayerIdx == 0 ? Inputs : Activations[LayerIdx - 1];
+
+			// Apply Linear Transformation
+
+			MatMulPlusBias(
+				LayerOutput,
+				LayerInput,
+				NeuralNetwork.Weights[LayerIdx],
+				NeuralNetwork.Biases[LayerIdx],
+				Instances);
+
+			// Apply Noise on final layer otherwise apply activation in-place
+
+			if (LayerIdx == LayerNum - 1)
+			{
+				ActionNoise(
+					Outputs,
+					LayerOutput,
+					Seed,
+					ActionNoiseScale,
+					LogActionNoiseMin,
+					LogActionNoiseMax,
+					Instances);
+			}
+			else
+			{
+				switch (NeuralNetwork.ActivationFunction)
+				{
+				case EActivationFunction::ReLU: ActivationReLU(LayerOutput, Instances); break;
+				case EActivationFunction::ELU: ActivationELU(LayerOutput, Instances); break;
+				case EActivationFunction::TanH: ActivationTanH(LayerOutput, Instances); break;
+				default: UE_LEARNING_NOT_IMPLEMENTED();
+				}
+			}
+		}
+	}
+
+	FNeuralNetworkPolicyFunction::FNeuralNetworkPolicyFunction(
+		const FName& InIdentifier,
+		const TSharedRef<FArrayMap>& InInstanceData,
+		const int32 InMaxInstanceNum,
+		const TSharedRef<FNeuralNetwork>& InNeuralNetwork,
+		const uint32 InSeed,
+		const FNeuralNetworkPolicyFunctionSettings& InSettings)
+		: FFunctionObject(InInstanceData)
+		, NeuralNetwork(InNeuralNetwork)
+		, Settings(InSettings)
+	{
+		SeedHandle = InstanceData->Add<1, uint32>({ InIdentifier, TEXT("Seed") }, { InMaxInstanceNum });
+		InputHandle = InstanceData->Add<2, float>({ InIdentifier, TEXT("Input") }, { InMaxInstanceNum, NeuralNetwork->GetInputNum() }, 0.0f);
+		OutputHandle = InstanceData->Add<2, float>({ InIdentifier, TEXT("Output") }, { InMaxInstanceNum, NeuralNetwork->GetOutputNum() / 2 }, 0.0f);
+		ActionNoiseScaleHandle = InstanceData->Add<1, float>({ InIdentifier, TEXT("ActionNoiseScale") }, { InMaxInstanceNum }, Settings.ActionNoiseScale);
+
+		Random::IntArray(InstanceData->View(SeedHandle), InSeed);
+
+		// Allocate temporary storage for activations
+
+		const int32 LayerNum = InNeuralNetwork->GetLayerNum();
+
+		Activations.SetNum(LayerNum);
+		for (int32 LayerIdx = 0; LayerIdx < LayerNum; LayerIdx++)
+		{
+			Activations[LayerIdx].SetNumUninitialized({ InMaxInstanceNum, NeuralNetwork->Weights[LayerIdx].Num(1) });
+		}
+	}
+
+	void FNeuralNetworkPolicyFunction::Evaluate(const FIndexSet Instances)
+	{
+		UE_LEARNING_TRACE_CPUPROFILER_EVENT_SCOPE(Learning::FNeuralNetworkPolicyFunction::Evaluate);
+
+		const TLearningArrayView<2, const float> Inputs = InstanceData->ConstView(InputHandle);
+		const TLearningArrayView<1, const float> ActionNoiseScale = InstanceData->ConstView(ActionNoiseScaleHandle);
+		TLearningArrayView<2, float> Outputs = InstanceData->View(OutputHandle);
+		TLearningArrayView<1, uint32> Seed = InstanceData->View(SeedHandle);
+
+		UE_LEARNING_CHECK(NeuralNetwork->GetOutputNum() == 2 * Outputs.Num<1>());
+		UE_LEARNING_CHECK(NeuralNetwork->GetInputNum() == Inputs.Num<1>());
+
+		const int32 LayerNum = NeuralNetwork->GetLayerNum();
+
+		if (!ensureMsgf(LayerNum > 0, TEXT("Empty Neural Network used in Policy")))
+		{
+			Array::Zero(Outputs, Instances);
+			return;
+		}
+
+		UE_LEARNING_ARRAY_VALUE_CHECK(Settings.ActionNoiseMin >= 0.0f && Settings.ActionNoiseMax >= 0.0f);
+		const float LogActionNoiseMin = FMath::Loge(Settings.ActionNoiseMin + UE_KINDA_SMALL_NUMBER);
+		const float LogActionNoiseMax = FMath::Loge(Settings.ActionNoiseMax + UE_KINDA_SMALL_NUMBER);
+
+		// Compute Layers
+
+		for (int32 LayerIdx = 0; LayerIdx < LayerNum; LayerIdx++)
+		{
+			auto LayerEvaluationFunction = [&](const int32 SliceStart, const int32 SliceNum)
+			{
+				NeuralNetworkPolicy::EvaluateLayer(
+					Outputs,
+					Activations,
+					Inputs,
+					Seed,
+					ActionNoiseScale,
+					LayerIdx,
+					LayerNum,
+					LogActionNoiseMin,
+					LogActionNoiseMax,
+					*NeuralNetwork,
+					Instances.Slice(SliceStart, SliceNum));
+			};
+
+			if (Settings.bParallelEvaluation && Instances.Num() > Settings.MinParallelBatchSize)
+			{
+				SlicedParallelFor(Instances.Num(), Settings.MinParallelBatchSize, LayerEvaluationFunction);
+			}
+			else
+			{
+				LayerEvaluationFunction(0, Instances.Num());
+			}
+		}
+
+		Array::Check(Outputs, Instances);
+	}
+}
diff --git a/Engine/Plugins/Experimental/LearningAgents/Source/Learning/Private/LearningRecorderObject.cpp b/Engine/Plugins/Experimental/LearningAgents/Source/Learning/Private/LearningRecorderObject.cpp
new file mode 100644
--- /dev/null
+++ b/Engine/Plugins/Experimental/LearningAgents/Source/Learning/Private/LearningRecorderObject.cpp
@@ -0,0 +1,46 @@
+// Copyright Epic Games, Inc. All Rights Reserved.
+
+#include "LearningRecorderObject.h"
+
+namespace UE::Learning
+{
+	FRecorderObject::FRecorderObject(
+		const int32 InMaxInstanceNum,
+		const int32 InMaxStepNum,
+		const TSharedRef<FArrayMap>& InSrcArrayMap,
+		const TSharedRef<FArrayMap>& InLogArrayMap)
+		: MaxInstanceNum(InMaxInstanceNum)
+		, MaxStepNum(InMaxStepNum)
+		, SrcArrayMap(InSrcArrayMap)
+		, LogArrayMap(InLogArrayMap)
+	{
+		StepNums.SetNumUninitialized({ MaxInstanceNum });
+		Array::Zero(StepNums);
+	}
+
+	FRecorderObject::~FRecorderObject() { }
+
+	void FRecorderObject::Reset(const FIndexSet Instances)
+	{
+		Array::Zero(StepNums, Instances);
+	}
+
+	void FRecorderObject::Update(const FIndexSet Instances)
+	{
+		const int32 VariableNum = SrcHandles.Num();
+
+		for (int32 VariableIdx = 0; VariableIdx < VariableNum; VariableIdx++)
+		{
+			CopyFunctions[VariableIdx](
+				StepNums,
+				*SrcArrayMap,
+				*LogArrayMap,
+				SrcHandles[VariableIdx],
+				LogHandles[VariableIdx],
+				MaxStepNum,
+				Instances);
+		}
+	}
+
+}
+
diff --git a/Engine/Plugins/Experimental/LearningAgents/Source/Learning/Private/LearningRewardObject.cpp b/Engine/Plugins/Experimental/LearningAgents/Source/Learning/Private/LearningRewardObject.cpp
new file mode 100644
--- /dev/null
+++ b/Engine/Plugins/Experimental/LearningAgents/Source/Learning/Private/LearningRewardObject.cpp
@@ -0,0 +1,1014 @@
+// Copyright Epic Games, Inc. All Rights Reserved.
+
+#include "LearningRewardObject.h"
+#include "LearningLog.h"
+
+namespace UE::Learning
+{
+	namespace Reward
+	{
+		// Misc
+
+		static inline FVector VectorExp(const FVector X)
+		{
+			return FVector(FMath::Exp(X[0]), FMath::Exp(X[1]), FMath::Exp(X[2]));
+		}
+
+		static inline FVector VectorMax(const FVector X, const float Y)
+		{
+			return FVector(FMath::Max(X[0], Y), FMath::Max(X[1], Y), FMath::Max(X[2], Y));
+		}
+
+		static inline float SpringInertializeHalfLifeToDamping(const float HalfLife, const float Epsilon = UE_SMALL_NUMBER)
+		{
+			UE_LEARNING_ARRAY_VALUE_CHECK(HalfLife > 0.0f);
+
+			const float Ln2 = 0.69314718056f;
+			return (4.0f * Ln2) / FMath::Max(HalfLife, Epsilon);
+		}
+
+		static inline FVector SpringInertializeDecayIntersectionTime(const FVector Position, const FVector Velocity, const float HalfLife, const float Epsilon = UE_SMALL_NUMBER)
+		{
+			UE_LEARNING_ARRAY_VALUE_CHECK(HalfLife > 0.0f);
+
+			const float Y = SpringInertializeHalfLifeToDamping(HalfLife, Epsilon) / 2.0f;
+			return -Position / VectorMax(Velocity + Position * Y, Epsilon);
+		}
+
+		static inline FVector SpringInertializeDecayTotalDisplacement(const FVector PositionOffset, const FVector VelocityOffset, const float HalfLife, const float Epsilon = UE_SMALL_NUMBER)
+		{
+			const FVector X = PositionOffset;
+			const FVector V = VelocityOffset;
+			const float Y = SpringInertializeHalfLifeToDamping(HalfLife, Epsilon) / 2.0f;
+			const FVector T = SpringInertializeDecayIntersectionTime(X, V, HalfLife, Epsilon);
+
+			const FVector Integral0 = (X * Y * 2 + V) / (Y * Y);
+			const FVector IntegralT = (VectorExp(-Y * T) / (Y * Y)) * (T * V * Y + X * Y * (T * Y + 2) + V);
+
+			return FVector(
+				T[0] > 0.0f ? FMath::Abs(Integral0[0] - IntegralT[0]) + FMath::Abs(IntegralT[0]) : FMath::Abs(Integral0[0]),
+				T[1] > 0.0f ? FMath::Abs(Integral0[1] - IntegralT[1]) + FMath::Abs(IntegralT[1]) : FMath::Abs(Integral0[1]),
+				T[2] > 0.0f ? FMath::Abs(Integral0[2] - IntegralT[2]) + FMath::Abs(IntegralT[2]) : FMath::Abs(Integral0[2]));
+		}
+
+		// Helpers
+
+		float DistanceToReward(const float Distance)
+		{
+			return SquaredDistanceToReward(FMath::Square(Distance));
+		}
+
+		float DistanceToPenalty(const float Distance)
+		{
+			return -Distance;
+		}
+
+		float SquaredDistanceToReward(const float SquaredDistance)
+		{
+			return FMath::InvExpApprox(SquaredDistance);
+		}
+
+		float SquaredDistanceToPenalty(const float SquaredDistance)
+		{
+			return DistanceToPenalty(FMath::Sqrt(SquaredDistance));
+		}
+
+		// Rewards
+
+		float ScalarSimilarityReward(
+			const float Scalar0,
+			const float Scalar1,
+			const float Scale,
+			const float Threshold,
+			const float Weight,
+			const float Epsilon)
+		{
+			UE_LEARNING_ARRAY_VALUE_CHECK(Scale > 0.0f);
+
+			return Weight * DistanceToReward(FMath::Max(FMath::Abs(Scalar0 - Scalar1) - Threshold, 0.0f) / FMath::Max(Scale, Epsilon));
+		}
+
+		float ScalarPositionSimilarityReward(
+			const float Position0,
+			const float Position1,
+			const float Scale,
+			const float Threshold,
+			const float Weight,
+			const float Epsilon)
+		{
+			UE_LEARNING_ARRAY_VALUE_CHECK(Scale > 0.0f);
+
+			return Weight * DistanceToReward(FMath::Max(FMath::Abs(Position0 - Position1) - Threshold, 0.0f) / FMath::Max(Scale, Epsilon));
+		}
+
+		float PositionSimilarityReward(
+			const FVector Position0,
+			const FVector Position1,
+			const float Scale,
+			const float Threshold,
+			const float Weight,
+			const float Epsilon)
+		{
+			UE_LEARNING_ARRAY_VALUE_CHECK(Scale > 0.0f);
+
+			return Weight * DistanceToReward(FMath::Max(FVector::Distance(Position0, Position1) - Threshold, 0.0f) / FMath::Max(Scale, Epsilon));
+		}
+
+		float ScalarRotationSimilarityReward(
+			const float Rotation0,
+			const float Rotation1,
+			const float Scale,
+			const float Threshold,
+			const float Weight,
+			const float Epsilon)
+		{
+			UE_LEARNING_ARRAY_VALUE_CHECK(Scale > 0.0f);
+
+			return Weight * DistanceToReward(FMath::Max(
+				FMath::Abs(FMath::FindDeltaAngleRadians(Rotation1, Rotation0)) - Threshold, 0.0f) / FMath::Max(Scale, Epsilon));
+		}
+
+		float ScalarAngularVelocitySimilarityReward(
+			const float AngularVelocity0,
+			const float AngularVelocity1,
+			const float Scale,
+			const float Threshold,
+			const float Weight,
+			const float Epsilon)
+		{
+			UE_LEARNING_ARRAY_VALUE_CHECK(Scale > 0.0f);
+
+			return Weight * DistanceToReward(FMath::Max(FMath::Abs(AngularVelocity0 - AngularVelocity1) - Threshold, 0.0f) / FMath::Max(Scale, Epsilon));
+		}
+
+		LEARNING_API float ScalarVelocityReward(
+			const float Velocity,
+			const float Scale,
+			const float Weight,
+			const float Epsilon)
+		{
+			return Weight * (Velocity / FMath::Max(Scale, Epsilon));
+		}
+
+		LEARNING_API float LocalDirectionalVelocityReward(
+			const FVector Velocity,
+			const FQuat RelativeRotation,
+			const float Scale,
+			const float Weight,
+			const FVector Axis,
+			const float Epsilon)
+		{
+			UE_LEARNING_ARRAY_VALUE_CHECK(Scale > 0.0f);
+
+			return Weight * (RelativeRotation.UnrotateVector(Velocity).Dot(Axis) / FMath::Max(Scale, Epsilon));
+		}
+
+		// Penalties
+
+		float ScalarDifferencePenalty(
+			const float Scalar0,
+			const float Scalar1,
+			const float Scale,
+			const float Threshold,
+			const float Weight,
+			const float Epsilon)
+		{
+			UE_LEARNING_ARRAY_VALUE_CHECK(Scale > 0.0f);
+
+			return Weight * DistanceToPenalty(FMath::Max(FMath::Abs(Scalar0 - Scalar1) - Threshold, 0.0f) / FMath::Max(Scale, Epsilon));
+		}
+
+		float ScalarMagnitudePenalty(
+			const float Scalar,
+			const float Scale,
+			const float Threshold,
+			const float Weight,
+			const float Epsilon)
+		{
+			UE_LEARNING_ARRAY_VALUE_CHECK(Scale > 0.0f);
+
+			return Weight * DistanceToPenalty(FMath::Max(FMath::Abs(Scalar) - Threshold, 0.0f) / FMath::Max(Scale, Epsilon));
+		}
+
+		float DirectionAngleDifferencePenalty(
+			const FVector Direction0,
+			const FVector Direction1,
+			const float Scale,
+			const float Threshold,
+			const float Weight,
+			const float Epsilon)
+		{
+			UE_LEARNING_ARRAY_VALUE_CHECK(Scale > 0.0f);
+
+			return Weight * DistanceToPenalty(FMath::Max(FMath::Acos(Direction0.Dot(Direction1)) - Threshold, 0.0f) / FMath::Max(Scale, Epsilon));
+		}
+
+		float PlanarPositionDifferencePenalty(
+			const FVector Position0,
+			const FVector Position1,
+			const float Scale,
+			const float Threshold,
+			const float Weight,
+			const FVector Axis0,
+			const FVector Axis1,
+			const float Epsilon)
+		{
+			UE_LEARNING_ARRAY_VALUE_CHECK(Scale > 0.0f);
+
+			const FVector ProjectedPosition0 = FVector(Axis0.Dot(Position0), Axis1.Dot(Position0), 0.0f);
+			const FVector ProjectedPosition1 = FVector(Axis0.Dot(Position1), Axis1.Dot(Position1), 0.0f);
+
+			return Weight * DistanceToPenalty(FMath::Max(FVector::Distance(ProjectedPosition0, ProjectedPosition1) - Threshold, 0.0f) / FMath::Max(Scale, Epsilon));
+		}
+
+		float PositionDifferencePenalty(
+			const FVector Position0,
+			const FVector Position1,
+			const float Scale,
+			const float Threshold,
+			const float Weight,
+			const float Epsilon)
+		{
+			UE_LEARNING_ARRAY_VALUE_CHECK(Scale > 0.0f);
+
+			return Weight * DistanceToPenalty(FMath::Max(FVector::Dist(Position0, Position1) - Threshold, 0.0f) / FMath::Max(Scale, Epsilon));
+		}
+
+		float VelocityDifferencePenalty(
+			const FVector Velocity0,
+			const FVector Velocity1,
+			const float Scale,
+			const float Threshold,
+			const float Weight,
+			const float Epsilon)
+		{
+			UE_LEARNING_ARRAY_VALUE_CHECK(Scale > 0.0f);
+
+			return Weight * DistanceToPenalty(FMath::Max(FVector::Dist(Velocity0, Velocity1) - Threshold, 0.0f) / FMath::Max(Scale, Epsilon));
+		}
+
+		float SpringInertializationPenalty(
+			const FVector Position0,
+			const FVector Position1,
+			const FVector Velocity0,
+			const FVector Velocity1,
+			const float HalfLife,
+			const float Scale,
+			const float Threshold,
+			const float Weight,
+			const float Epsilon)
+		{
+			UE_LEARNING_ARRAY_VALUE_CHECK(Scale > 0.0f);
+			UE_LEARNING_ARRAY_VALUE_CHECK(HalfLife > 0.0f);
+
+			return Weight * DistanceToPenalty(FMath::Max(SpringInertializeDecayTotalDisplacement(
+				Position0 - Position1,
+				Velocity0 - Velocity1,
+				HalfLife,
+				Epsilon).Length() / FMath::Max(HalfLife, Epsilon) - Threshold, 0.0f) / FMath::Max(Scale, Epsilon));
+		}
+	}
+
+	//------------------------------------------------------------------
+
+	FRewardObject::FRewardObject(
+		const FName& InIdentifier,
+		const TSharedRef<FArrayMap>& InInstanceData,
+		const int32 InMaxInstanceNum)
+		: InstanceData(InInstanceData)
+	{
+		RewardHandle = InstanceData->Add<1, float>({ InIdentifier, TEXT("Reward") }, { InMaxInstanceNum }, 0.0f);
+	}
+
+	TLearningArrayView<1, float> FRewardObject::RewardBuffer()
+	{
+		return  InstanceData->View(RewardHandle);
+	}
+
+	//------------------------------------------------------------------
+
+	FSumReward::FSumReward(
+		const FName& InIdentifier,
+		const TLearningArrayView<1, const TSharedRef<FRewardObject>> InRewards,
+		const TSharedRef<FArrayMap>& InInstanceData,
+		const int32 InMaxInstanceNum)
+		: FRewardObject(InIdentifier, InInstanceData, InMaxInstanceNum)
+		, Rewards(InRewards) {}
+
+	void FSumReward::Evaluate(const FIndexSet Instances)
+	{
+		UE_LEARNING_TRACE_CPUPROFILER_EVENT_SCOPE(FSumReward::Evaluate);
+
+		TLearningArrayView<1, float> Reward = InstanceData->View(RewardHandle);
+
+		const int32 RewardNum = Rewards.Num();
+
+		for (int32 RewardIdx = 0; RewardIdx < RewardNum; RewardIdx++)
+		{
+			Rewards[RewardIdx]->Evaluate(Instances);
+		}
+
+		for (const int32 InstanceIdx : Instances)
+		{
+			Reward[InstanceIdx] = 0.0f;
+		}
+
+		for (int32 RewardIdx = 0; RewardIdx < RewardNum; RewardIdx++)
+		{
+			const TLearningArrayView<1, const float> Input = InstanceData->ConstView(Rewards[RewardIdx]->RewardHandle);
+
+			for (const int32 InstanceIdx : Instances)
+			{
+				Reward[InstanceIdx] += Input[InstanceIdx];
+			}
+		}
+	}
+
+	//------------------------------------------------------------------
+
+	FConstantReward::FConstantReward(
+		const FName& InIdentifier,
+		const TSharedRef<FArrayMap>& InInstanceData,
+		const int32 InMaxInstanceNum,
+		const float InValue)
+		: FRewardObject(InIdentifier, InInstanceData, InMaxInstanceNum)
+		, Value(InValue) {}
+
+	void FConstantReward::Evaluate(const FIndexSet Instances)
+	{
+		UE_LEARNING_TRACE_CPUPROFILER_EVENT_SCOPE(FConstantReward::Evaluate);
+
+		Array::Set(InstanceData->View(RewardHandle), Value, Instances);
+	}
+
+	FConstantPenalty::FConstantPenalty(
+		const FName& InIdentifier,
+		const TSharedRef<FArrayMap>& InInstanceData,
+		const int32 InMaxInstanceNum,
+		const float InValue)
+		: FRewardObject(InIdentifier, InInstanceData, InMaxInstanceNum)
+		, Value(InValue) {}
+
+	void FConstantPenalty::Evaluate(const FIndexSet Instances)
+	{
+		UE_LEARNING_TRACE_CPUPROFILER_EVENT_SCOPE(FConstantPenalty::Evaluate);
+
+		Array::Set(InstanceData->View(RewardHandle), -Value, Instances);
+	}
+
+	FConditionalConstantReward::FConditionalConstantReward(
+		const FName& InIdentifier,
+		const TSharedRef<FArrayMap>& InInstanceData,
+		const int32 InMaxInstanceNum,
+		const float InValue)
+		: FRewardObject(InIdentifier, InInstanceData, InMaxInstanceNum)
+		, Value(InValue)
+	{
+		ConditionHandle = InstanceData->Add<1, bool>({ InIdentifier, TEXT("Condition") }, { InMaxInstanceNum }, false);
+	}
+
+	void FConditionalConstantReward::Evaluate(const FIndexSet Instances)
+	{
+		UE_LEARNING_TRACE_CPUPROFILER_EVENT_SCOPE(FConditionalConstantReward::Evaluate);
+
+		const TLearningArrayView<1, const bool> Condition = InstanceData->ConstView(ConditionHandle);
+		TLearningArrayView<1, float> Reward = InstanceData->View(RewardHandle);
+
+		for (const int32 InstanceIdx : Instances)
+		{
+			Reward[InstanceIdx] = Condition[InstanceIdx] ? Value : 0.0f;
+		}
+	}
+
+	FFloatReward::FFloatReward(
+		const FName& InIdentifier,
+		const TSharedRef<FArrayMap>& InInstanceData,
+		const int32 InMaxInstanceNum,
+		const float InWeight)
+		: FRewardObject(InIdentifier, InInstanceData, InMaxInstanceNum)
+	{
+		ValueHandle = InstanceData->Add<1, float>({ InIdentifier, TEXT("Value") }, { InMaxInstanceNum }, 0.0f);
+		WeightHandle = InstanceData->Add<1, float>({ InIdentifier, TEXT("Weight") }, { InMaxInstanceNum }, InWeight);
+	}
+
+	void FFloatReward::Evaluate(const FIndexSet Instances)
+	{
+		UE_LEARNING_TRACE_CPUPROFILER_EVENT_SCOPE(FFloatReward::Evaluate);
+
+		const TLearningArrayView<1, const float> Value = InstanceData->ConstView(ValueHandle);
+		TLearningArrayView<1, float> Reward = InstanceData->View(RewardHandle);
+
+		Array::Copy(Reward, Value, Instances);
+	}
+
+	//------------------------------------------------------------------
+
+	FVectorAverageMagnitudePenalty::FVectorAverageMagnitudePenalty(
+		const FName& InIdentifier,
+		const TSharedRef<FArrayMap>& InInstanceData,
+		const int32 InMaxInstanceNum,
+		const int32 InDimNum,
+		const float InWeight,
+		const float InScale,
+		const float InThreshold)
+		: FRewardObject(InIdentifier, InInstanceData, InMaxInstanceNum)
+		, DimNum(InDimNum)
+	{
+		ValueHandle = InstanceData->Add<2, float>({ InIdentifier, TEXT("Value") }, { InMaxInstanceNum, DimNum }, 1.0f);
+
+		WeightHandle = InstanceData->Add<1, float>({ InIdentifier, TEXT("Weight") }, { InMaxInstanceNum }, InWeight);
+		ScaleHandle = InstanceData->Add<1, float>({ InIdentifier, TEXT("Scale") }, { InMaxInstanceNum }, InScale);
+		ThresholdHandle = InstanceData->Add<1, float>({ InIdentifier, TEXT("Threshold") }, { InMaxInstanceNum }, InThreshold);
+	}
+
+	void FVectorAverageMagnitudePenalty::Evaluate(const FIndexSet Instances)
+	{
+		UE_LEARNING_TRACE_CPUPROFILER_EVENT_SCOPE(FVectorAverageMagnitudePenalty::Evaluate);
+
+		const TLearningArrayView<2, const float> Value = InstanceData->ConstView(ValueHandle);
+
+		const TLearningArrayView<1, const float> Scale = InstanceData->ConstView(ScaleHandle);
+		const TLearningArrayView<1, const float> Threshold = InstanceData->ConstView(ThresholdHandle);
+		const TLearningArrayView<1, const float> Weight = InstanceData->ConstView(WeightHandle);
+
+		TLearningArrayView<1, float> Reward = InstanceData->View(RewardHandle);
+
+		for (const int32 InstanceIdx : Instances)
+		{
+			Reward[InstanceIdx] = 0.0f;
+
+			for (int32 DimIdx = 0; DimIdx < DimNum; DimIdx++)
+			{
+				Reward[InstanceIdx] += Reward::ScalarMagnitudePenalty(
+					Value[InstanceIdx][DimIdx],
+					Scale[InstanceIdx],
+					Threshold[InstanceIdx],
+					Weight[InstanceIdx]) / DimNum;
+			}
+		}
+	}
+
+	//------------------------------------------------------------------
+
+	FScalarPositionSimilarityReward::FScalarPositionSimilarityReward(
+		const FName& InIdentifier,
+		const TSharedRef<FArrayMap>& InInstanceData,
+		const int32 InMaxInstanceNum,
+		const float InWeight,
+		const float InScale,
+		const float InThreshold)
+		: FRewardObject(InIdentifier, InInstanceData, InMaxInstanceNum)
+	{
+		Position0Handle = InstanceData->Add<1, float>({ InIdentifier, TEXT("Position0") }, { InMaxInstanceNum }, 0.0f);
+		Position1Handle = InstanceData->Add<1, float>({ InIdentifier, TEXT("Position1") }, { InMaxInstanceNum }, 0.0f);
+
+		WeightHandle = InstanceData->Add<1, float>({ InIdentifier, TEXT("Weight") }, { InMaxInstanceNum }, InWeight);
+		ScaleHandle = InstanceData->Add<1, float>({ InIdentifier, TEXT("Scale") }, { InMaxInstanceNum }, InScale);
+		ThresholdHandle = InstanceData->Add<1, float>({ InIdentifier, TEXT("Threshold") }, { InMaxInstanceNum }, InThreshold);
+	}
+
+	void FScalarPositionSimilarityReward::Evaluate(const FIndexSet Instances)
+	{
+		UE_LEARNING_TRACE_CPUPROFILER_EVENT_SCOPE(FScalarPositionSimilarityReward::Evaluate);
+
+		const TLearningArrayView<1, const float> Position0 = InstanceData->ConstView(Position0Handle);
+		const TLearningArrayView<1, const float> Position1 = InstanceData->ConstView(Position1Handle);
+
+		const TLearningArrayView<1, const float> Scale = InstanceData->ConstView(ScaleHandle);
+		const TLearningArrayView<1, const float> Threshold = InstanceData->ConstView(ThresholdHandle);
+		const TLearningArrayView<1, const float> Weight = InstanceData->ConstView(WeightHandle);
+
+		TLearningArrayView<1, float> Reward = InstanceData->View(RewardHandle);
+
+		for (const int32 InstanceIdx : Instances)
+		{
+			Reward[InstanceIdx] = Reward::ScalarPositionSimilarityReward(
+				Position0[InstanceIdx],
+				Position1[InstanceIdx],
+				Scale[InstanceIdx],
+				Threshold[InstanceIdx],
+				Weight[InstanceIdx]);
+		}
+	}
+
+	//------------------------------------------------------------------
+
+	FScalarRotationSimilarityReward::FScalarRotationSimilarityReward(
+		const FName& InIdentifier,
+		const TSharedRef<FArrayMap>& InInstanceData,
+		const int32 InMaxInstanceNum,
+		const float InWeight,
+		const float InScale,
+		const float InThreshold)
+		: FRewardObject(InIdentifier, InInstanceData, InMaxInstanceNum)
+	{
+		Rotation0Handle = InstanceData->Add<1, float>({ InIdentifier, TEXT("Rotation0") }, { InMaxInstanceNum }, 0.0f);
+		Rotation1Handle = InstanceData->Add<1, float>({ InIdentifier, TEXT("Rotation1") }, { InMaxInstanceNum }, 0.0f);
+
+		WeightHandle = InstanceData->Add<1, float>({ InIdentifier, TEXT("Weight") }, { InMaxInstanceNum }, InWeight);
+		ScaleHandle = InstanceData->Add<1, float>({ InIdentifier, TEXT("Scale") }, { InMaxInstanceNum }, InScale);
+		ThresholdHandle = InstanceData->Add<1, float>({ InIdentifier, TEXT("Threshold") }, { InMaxInstanceNum }, InThreshold);
+	}
+
+	void FScalarRotationSimilarityReward::Evaluate(const FIndexSet Instances)
+	{
+		UE_LEARNING_TRACE_CPUPROFILER_EVENT_SCOPE(FScalarRotationSimilarityReward::Evaluate);
+
+		const TLearningArrayView<1, const float> Rotation0 = InstanceData->ConstView(Rotation0Handle);
+		const TLearningArrayView<1, const float> Rotation1 = InstanceData->ConstView(Rotation1Handle);
+
+		const TLearningArrayView<1, const float> Scale = InstanceData->ConstView(ScaleHandle);
+		const TLearningArrayView<1, const float> Threshold = InstanceData->ConstView(ThresholdHandle);
+		const TLearningArrayView<1, const float> Weight = InstanceData->ConstView(WeightHandle);
+
+		TLearningArrayView<1, float> Reward = InstanceData->View(RewardHandle);
+
+		for (const int32 InstanceIdx : Instances)
+		{
+			Reward[InstanceIdx] = Reward::ScalarRotationSimilarityReward(
+				Rotation0[InstanceIdx],
+				Rotation1[InstanceIdx],
+				Scale[InstanceIdx],
+				Threshold[InstanceIdx],
+				Weight[InstanceIdx]);
+		}
+	}
+
+	//------------------------------------------------------------------
+
+	FScalarAngularVelocitySimilarityReward::FScalarAngularVelocitySimilarityReward(
+		const FName& InIdentifier,
+		const TSharedRef<FArrayMap>& InInstanceData,
+		const int32 InMaxInstanceNum,
+		const float InWeight,
+		const float InScale,
+		const float InThreshold)
+		: FRewardObject(InIdentifier, InInstanceData, InMaxInstanceNum)
+	{
+		AngularVelocity0Handle = InstanceData->Add<1, float>({ InIdentifier, TEXT("AngularVelocity0") }, { InMaxInstanceNum }, 0.0f);
+		AngularVelocity1Handle = InstanceData->Add<1, float>({ InIdentifier, TEXT("AngularVelocity1") }, { InMaxInstanceNum }, 0.0f);
+
+		WeightHandle = InstanceData->Add<1, float>({ InIdentifier, TEXT("Weight") }, { InMaxInstanceNum }, InWeight);
+		ScaleHandle = InstanceData->Add<1, float>({ InIdentifier, TEXT("Scale") }, { InMaxInstanceNum }, InScale);
+		ThresholdHandle = InstanceData->Add<1, float>({ InIdentifier, TEXT("Threshold") }, { InMaxInstanceNum }, InThreshold);
+	}
+
+	void FScalarAngularVelocitySimilarityReward::Evaluate(const FIndexSet Instances)
+	{
+		UE_LEARNING_TRACE_CPUPROFILER_EVENT_SCOPE(FScalarAngularVelocitySimilarityReward::Evaluate);
+
+		const TLearningArrayView<1, const float> AngularVelocity0 = InstanceData->ConstView(AngularVelocity0Handle);
+		const TLearningArrayView<1, const float> AngularVelocity1 = InstanceData->ConstView(AngularVelocity1Handle);
+
+		const TLearningArrayView<1, const float> Scale = InstanceData->ConstView(ScaleHandle);
+		const TLearningArrayView<1, const float> Threshold = InstanceData->ConstView(ThresholdHandle);
+		const TLearningArrayView<1, const float> Weight = InstanceData->ConstView(WeightHandle);
+
+		TLearningArrayView<1, float> Reward = InstanceData->View(RewardHandle);
+
+		for (const int32 InstanceIdx : Instances)
+		{
+			Reward[InstanceIdx] = Reward::ScalarAngularVelocitySimilarityReward(
+				AngularVelocity0[InstanceIdx],
+				AngularVelocity1[InstanceIdx],
+				Scale[InstanceIdx],
+				Threshold[InstanceIdx],
+				Weight[InstanceIdx]);
+		}
+	}
+
+	//------------------------------------------------------------------
+
+	FScalarVelocityReward::FScalarVelocityReward(
+		const FName& InIdentifier,
+		const TSharedRef<FArrayMap>& InInstanceData,
+		const int32 InMaxInstanceNum,
+		const float InWeight,
+		const float InScale)
+		: FRewardObject(InIdentifier, InInstanceData, InMaxInstanceNum)
+	{
+		VelocityHandle = InstanceData->Add<1, float>({ InIdentifier, TEXT("Velocity") }, { InMaxInstanceNum }, 0.0f);
+
+		WeightHandle = InstanceData->Add<1, float>({ InIdentifier, TEXT("Weight") }, { InMaxInstanceNum }, InWeight);
+		ScaleHandle = InstanceData->Add<1, float>({ InIdentifier, TEXT("Scale") }, { InMaxInstanceNum }, InScale);
+	}
+
+	void FScalarVelocityReward::Evaluate(const FIndexSet Instances)
+	{
+		UE_LEARNING_TRACE_CPUPROFILER_EVENT_SCOPE(FPlanarVelocityMagnitudeReward::Evaluate);
+
+		const TLearningArrayView<1, const float> Velocity = InstanceData->ConstView(VelocityHandle);
+
+		const TLearningArrayView<1, const float> Scale = InstanceData->ConstView(ScaleHandle);
+		const TLearningArrayView<1, const float> Weight = InstanceData->ConstView(WeightHandle);
+
+		TLearningArrayView<1, float> Reward = InstanceData->View(RewardHandle);
+
+		for (const int32 InstanceIdx : Instances)
+		{
+			Reward[InstanceIdx] = Reward::ScalarVelocityReward(
+				Velocity[InstanceIdx],
+				Scale[InstanceIdx],
+				Weight[InstanceIdx]);
+		}
+	}
+
+	FLocalDirectionalVelocityReward::FLocalDirectionalVelocityReward(
+		const FName& InIdentifier,
+		const TSharedRef<FArrayMap>& InInstanceData,
+		const int32 InMaxInstanceNum,
+		const float InWeight,
+		const float InScale,
+		const FVector InAxis)
+		: FRewardObject(InIdentifier, InInstanceData, InMaxInstanceNum)
+		, Axis(InAxis)
+	{
+		VelocityHandle = InstanceData->Add<1, FVector>({ InIdentifier, TEXT("Velocity") }, { InMaxInstanceNum }, FVector::ZeroVector);
+		RelativeRotationHandle = InstanceData->Add<1, FQuat>({ InIdentifier, TEXT("RelativeRotation") }, { InMaxInstanceNum }, FQuat::Identity);
+
+		WeightHandle = InstanceData->Add<1, float>({ InIdentifier, TEXT("Weight") }, { InMaxInstanceNum }, InWeight);
+		ScaleHandle = InstanceData->Add<1, float>({ InIdentifier, TEXT("Scale") }, { InMaxInstanceNum }, InScale);
+	}
+
+	void FLocalDirectionalVelocityReward::Evaluate(const FIndexSet Instances)
+	{
+		UE_LEARNING_TRACE_CPUPROFILER_EVENT_SCOPE(FLocalDirectionalVelocityReward::Evaluate);
+
+		const TLearningArrayView<1, const FVector> Velocity = InstanceData->ConstView(VelocityHandle);
+		const TLearningArrayView<1, const FQuat> RelativeRotation = InstanceData->ConstView(RelativeRotationHandle);
+
+		const TLearningArrayView<1, const float> Scale = InstanceData->ConstView(ScaleHandle);
+		const TLearningArrayView<1, const float> Weight = InstanceData->ConstView(WeightHandle);
+
+		TLearningArrayView<1, float> Reward = InstanceData->View(RewardHandle);
+
+		for (const int32 InstanceIdx : Instances)
+		{
+			Reward[InstanceIdx] = Reward::LocalDirectionalVelocityReward(
+				Velocity[InstanceIdx],
+				RelativeRotation[InstanceIdx],
+				Scale[InstanceIdx],
+				Weight[InstanceIdx],
+				Axis);
+		}
+	}
+
+	//------------------------------------------------------------------
+
+	FPlanarPositionDifferencePenalty::FPlanarPositionDifferencePenalty(
+		const FName& InIdentifier,
+		const TSharedRef<FArrayMap>& InInstanceData,
+		const int32 InMaxInstanceNum,
+		const float InWeight,
+		const float InScale,
+		const float InThreshold,
+		const FVector InAxis0,
+		const FVector InAxis1)
+		: FRewardObject(InIdentifier, InInstanceData, InMaxInstanceNum)
+		, Axis0(InAxis0)
+		, Axis1(InAxis1)
+	{
+		Position0Handle = InstanceData->Add<1, FVector>({ InIdentifier, TEXT("Position0") }, { InMaxInstanceNum }, FVector::ZeroVector);
+		Position1Handle = InstanceData->Add<1, FVector>({ InIdentifier, TEXT("Position1") }, { InMaxInstanceNum }, FVector::ZeroVector);
+
+		WeightHandle = InstanceData->Add<1, float>({ InIdentifier, TEXT("Weight") }, { InMaxInstanceNum }, InWeight);
+		ScaleHandle = InstanceData->Add<1, float>({ InIdentifier, TEXT("Scale") }, { InMaxInstanceNum }, InScale);
+		ThresholdHandle = InstanceData->Add<1, float>({ InIdentifier, TEXT("Threshold") }, { InMaxInstanceNum }, InThreshold);
+	}
+
+	void FPlanarPositionDifferencePenalty::Evaluate(const FIndexSet Instances)
+	{
+		UE_LEARNING_TRACE_CPUPROFILER_EVENT_SCOPE(FPlanarPositionDifferencePenalty::Evaluate);
+
+		const TLearningArrayView<1, const FVector> Position0 = InstanceData->ConstView(Position0Handle);
+		const TLearningArrayView<1, const FVector> Position1 = InstanceData->ConstView(Position1Handle);
+
+		const TLearningArrayView<1, const float> Scale = InstanceData->ConstView(ScaleHandle);
+		const TLearningArrayView<1, const float> Threshold = InstanceData->ConstView(ThresholdHandle);
+		const TLearningArrayView<1, const float> Weight = InstanceData->ConstView(WeightHandle);
+
+		TLearningArrayView<1, float> Reward = InstanceData->View(RewardHandle);
+
+		for (const int32 InstanceIdx : Instances)
+		{
+			Reward[InstanceIdx] = Reward::PlanarPositionDifferencePenalty(
+				Position0[InstanceIdx],
+				Position1[InstanceIdx],
+				Scale[InstanceIdx],
+				Threshold[InstanceIdx],
+				Weight[InstanceIdx],
+				Axis0,
+				Axis1);
+		}
+	}
+
+	FPositionDifferencePenalty::FPositionDifferencePenalty(
+		const FName& InIdentifier,
+		const TSharedRef<FArrayMap>& InInstanceData,
+		const int32 InMaxInstanceNum,
+		const float InWeight,
+		const float InScale,
+		const float InThreshold)
+		: FRewardObject(InIdentifier, InInstanceData, InMaxInstanceNum)
+	{
+		Position0Handle = InstanceData->Add<1, FVector>({ InIdentifier, TEXT("Position0") }, { InMaxInstanceNum }, FVector::ZeroVector);
+		Position1Handle = InstanceData->Add<1, FVector>({ InIdentifier, TEXT("Position1") }, { InMaxInstanceNum }, FVector::ZeroVector);
+
+		WeightHandle = InstanceData->Add<1, float>({ InIdentifier, TEXT("Weight") }, { InMaxInstanceNum }, InWeight);
+		ScaleHandle = InstanceData->Add<1, float>({ InIdentifier, TEXT("Scale") }, { InMaxInstanceNum }, InScale);
+		ThresholdHandle = InstanceData->Add<1, float>({ InIdentifier, TEXT("Threshold") }, { InMaxInstanceNum }, InThreshold);
+	}
+
+	void FPositionDifferencePenalty::Evaluate(const FIndexSet Instances)
+	{
+		UE_LEARNING_TRACE_CPUPROFILER_EVENT_SCOPE(FPositionDifferencePenalty::Evaluate);
+
+		const TLearningArrayView<1, const FVector> Position0 = InstanceData->ConstView(Position0Handle);
+		const TLearningArrayView<1, const FVector> Position1 = InstanceData->ConstView(Position1Handle);
+
+		const TLearningArrayView<1, const float> Scale = InstanceData->ConstView(ScaleHandle);
+		const TLearningArrayView<1, const float> Threshold = InstanceData->ConstView(ThresholdHandle);
+		const TLearningArrayView<1, const float> Weight = InstanceData->ConstView(WeightHandle);
+
+		TLearningArrayView<1, float> Reward = InstanceData->View(RewardHandle);
+
+		for (const int32 InstanceIdx : Instances)
+		{
+			Reward[InstanceIdx] = Reward::PositionDifferencePenalty(
+				Position0[InstanceIdx],
+				Position1[InstanceIdx],
+				Scale[InstanceIdx],
+				Threshold[InstanceIdx],
+				Weight[InstanceIdx]);
+		}
+	}
+
+
+	//------------------------------------------------------------------
+
+	FVelocityDifferencePenalty::FVelocityDifferencePenalty(
+		const FName& InIdentifier,
+		const TSharedRef<FArrayMap>& InInstanceData,
+		const int32 InMaxInstanceNum,
+		const float InWeight,
+		const float InScale,
+		const float InThreshold)
+		: FRewardObject(InIdentifier, InInstanceData, InMaxInstanceNum)
+	{
+		Velocity0Handle = InstanceData->Add<1, FVector>({ InIdentifier, TEXT("Velocity0") }, { InMaxInstanceNum }, FVector::ZeroVector);
+		Velocity1Handle = InstanceData->Add<1, FVector>({ InIdentifier, TEXT("Velocity1") }, { InMaxInstanceNum }, FVector::ZeroVector);
+
+		WeightHandle = InstanceData->Add<1, float>({ InIdentifier, TEXT("Weight") }, { InMaxInstanceNum }, InWeight);
+		ScaleHandle = InstanceData->Add<1, float>({ InIdentifier, TEXT("Scale") }, { InMaxInstanceNum }, InScale);
+		ThresholdHandle = InstanceData->Add<1, float>({ InIdentifier, TEXT("Threshold") }, { InMaxInstanceNum }, InThreshold);
+	}
+
+	void FVelocityDifferencePenalty::Evaluate(const FIndexSet Instances)
+	{
+		UE_LEARNING_TRACE_CPUPROFILER_EVENT_SCOPE(FVelocityDifferencePenalty::Evaluate);
+
+		const TLearningArrayView<1, const FVector> Velocity0 = InstanceData->ConstView(Velocity0Handle);
+		const TLearningArrayView<1, const FVector> Velocity1 = InstanceData->ConstView(Velocity1Handle);
+
+		const TLearningArrayView<1, const float> Scale = InstanceData->ConstView(ScaleHandle);
+		const TLearningArrayView<1, const float> Threshold = InstanceData->ConstView(ThresholdHandle);
+		const TLearningArrayView<1, const float> Weight = InstanceData->ConstView(WeightHandle);
+
+		TLearningArrayView<1, float> Reward = InstanceData->View(RewardHandle);
+
+		for (const int32 InstanceIdx : Instances)
+		{
+			Reward[InstanceIdx] = Reward::VelocityDifferencePenalty(
+				Velocity0[InstanceIdx],
+				Velocity1[InstanceIdx],
+				Scale[InstanceIdx],
+				Threshold[InstanceIdx],
+				Weight[InstanceIdx]);
+		}
+	}
+
+	//------------------------------------------------------------------
+
+	FFacingDirectionAngularDifferencePenalty::FFacingDirectionAngularDifferencePenalty(
+		const FName& InIdentifier,
+		const TSharedRef<FArrayMap>& InInstanceData,
+		const int32 InMaxInstanceNum,
+		const FVector InFacingDirection0,
+		const FVector InFacingDirection1,
+		const float InWeight,
+		const float InScale,
+		const float InThreshold)
+		: FRewardObject(InIdentifier, InInstanceData, InMaxInstanceNum)
+	{
+		Rotation0Handle = InstanceData->Add<1, FQuat>({ InIdentifier, TEXT("Rotation0") }, { InMaxInstanceNum }, FQuat::Identity);
+		Rotation1Handle = InstanceData->Add<1, FQuat>({ InIdentifier, TEXT("Rotation1") }, { InMaxInstanceNum }, FQuat::Identity);
+		FacingDirection0Handle = InstanceData->Add<1, FVector>({ InIdentifier, TEXT("FacingDirection0") }, { InMaxInstanceNum }, InFacingDirection0);
+		FacingDirection1Handle = InstanceData->Add<1, FVector>({ InIdentifier, TEXT("FacingDirection1") }, { InMaxInstanceNum }, InFacingDirection1);
+
+		WeightHandle = InstanceData->Add<1, float>({ InIdentifier, TEXT("Weight") }, { InMaxInstanceNum }, InWeight);
+		ScaleHandle = InstanceData->Add<1, float>({ InIdentifier, TEXT("Scale") }, { InMaxInstanceNum }, InScale);
+		ThresholdHandle = InstanceData->Add<1, float>({ InIdentifier, TEXT("Threshold") }, { InMaxInstanceNum }, InThreshold);
+	}
+
+	void FFacingDirectionAngularDifferencePenalty::Evaluate(const FIndexSet Instances)
+	{
+		UE_LEARNING_TRACE_CPUPROFILER_EVENT_SCOPE(FFacingDirectionAngularDifferencePenalty::Evaluate);
+
+		const TLearningArrayView<1, const FQuat> Rotation0 = InstanceData->ConstView(Rotation0Handle);
+		const TLearningArrayView<1, const FQuat> Rotation1 = InstanceData->ConstView(Rotation1Handle);
+		const TLearningArrayView<1, const FVector> FacingDirection0 = InstanceData->ConstView(FacingDirection0Handle);
+		const TLearningArrayView<1, const FVector> FacingDirection1 = InstanceData->ConstView(FacingDirection1Handle);
+
+		const TLearningArrayView<1, const float> Scale = InstanceData->ConstView(ScaleHandle);
+		const TLearningArrayView<1, const float> Threshold = InstanceData->ConstView(ThresholdHandle);
+		const TLearningArrayView<1, const float> Weight = InstanceData->ConstView(WeightHandle);
+
+		TLearningArrayView<1, float> Reward = InstanceData->View(RewardHandle);
+
+		for (const int32 InstanceIdx : Instances)
+		{
+			Reward[InstanceIdx] = Reward::DirectionAngleDifferencePenalty(
+				Rotation0[InstanceIdx].RotateVector(FacingDirection0[InstanceIdx]),
+				Rotation1[InstanceIdx].RotateVector(FacingDirection1[InstanceIdx]),
+				Scale[InstanceIdx],
+				Threshold[InstanceIdx],
+				Weight[InstanceIdx]);
+		}
+	}
+
+	FFacingTowardsTargetAngularDifferencePenalty::FFacingTowardsTargetAngularDifferencePenalty(
+		const FName& InIdentifier,
+		const TSharedRef<FArrayMap>& InInstanceData,
+		const int32 InMaxInstanceNum,
+		const FVector InFacingDirection,
+		const float InMinRadius,
+		const float InWeight,
+		const float InScale,
+		const float InThreshold)
+		: FRewardObject(InIdentifier, InInstanceData, InMaxInstanceNum)
+	{
+		TargetHandle = InstanceData->Add<1, FVector>({ InIdentifier, TEXT("Target") }, { InMaxInstanceNum }, FVector::ZeroVector);
+		PositionHandle = InstanceData->Add<1, FVector>({ InIdentifier, TEXT("Position") }, { InMaxInstanceNum }, FVector::ZeroVector);
+		RotationHandle = InstanceData->Add<1, FQuat>({ InIdentifier, TEXT("Rotation") }, { InMaxInstanceNum }, FQuat::Identity);
+		FacingDirectionHandle = InstanceData->Add<1, FVector>({ InIdentifier, TEXT("FacingDirection") }, { InMaxInstanceNum }, InFacingDirection);
+		MinRadiusHandle = InstanceData->Add<1, float>({ InIdentifier, TEXT("MinRadius") }, { InMaxInstanceNum }, InMinRadius);
+
+		WeightHandle = InstanceData->Add<1, float>({ InIdentifier, TEXT("Weight") }, { InMaxInstanceNum }, InWeight);
+		ScaleHandle = InstanceData->Add<1, float>({ InIdentifier, TEXT("Scale") }, { InMaxInstanceNum }, InScale);
+		ThresholdHandle = InstanceData->Add<1, float>({ InIdentifier, TEXT("Threshold") }, { InMaxInstanceNum }, InThreshold);
+	}
+
+	void FFacingTowardsTargetAngularDifferencePenalty::Evaluate(const FIndexSet Instances)
+	{
+		UE_LEARNING_TRACE_CPUPROFILER_EVENT_SCOPE(FFacingTowardsTargetAngularDifferencePenalty::Evaluate);
+
+		const TLearningArrayView<1, const FVector> Target = InstanceData->ConstView(TargetHandle);
+		const TLearningArrayView<1, const FVector> Position = InstanceData->ConstView(PositionHandle);
+		const TLearningArrayView<1, const FQuat> Rotation = InstanceData->ConstView(RotationHandle);
+		const TLearningArrayView<1, const FVector> FacingDirection = InstanceData->ConstView(FacingDirectionHandle);
+		const TLearningArrayView<1, const float> MinRadius = InstanceData->ConstView(MinRadiusHandle);
+
+		const TLearningArrayView<1, const float> Scale = InstanceData->ConstView(ScaleHandle);
+		const TLearningArrayView<1, const float> Threshold = InstanceData->ConstView(ThresholdHandle);
+		const TLearningArrayView<1, const float> Weight = InstanceData->ConstView(WeightHandle);
+
+		TLearningArrayView<1, float> Reward = InstanceData->View(RewardHandle);
+
+		for (const int32 InstanceIdx : Instances)
+		{
+			UE_LEARNING_ARRAY_VALUE_CHECK(MinRadius[InstanceIdx] > 0.0f);
+
+			const FVector TargetDifference = Target[InstanceIdx] - Position[InstanceIdx];
+
+			if (TargetDifference.Length() > FMath::Max(MinRadius[InstanceIdx], UE_SMALL_NUMBER))
+			{
+				Reward[InstanceIdx] = Reward::DirectionAngleDifferencePenalty(
+					Rotation[InstanceIdx].RotateVector(FacingDirection[InstanceIdx]),
+					TargetDifference.GetUnsafeNormal(),
+					Scale[InstanceIdx],
+					Threshold[InstanceIdx],
+					Weight[InstanceIdx]);
+			}
+			else
+			{
+				Reward[InstanceIdx] = 0.0f;
+			}
+		}
+	}
+
+	//------------------------------------------------------------------
+
+	FSpringInertializationPenalty::FSpringInertializationPenalty(
+		const FName& InIdentifier,
+		const TSharedRef<FArrayMap>& InInstanceData,
+		const int32 InMaxInstanceNum,
+		const float InWeight,
+		const float InScale,
+		const float InThreshold)
+		: FRewardObject(InIdentifier, InInstanceData, InMaxInstanceNum)
+	{
+		Position0Handle = InstanceData->Add<1, FVector>({ InIdentifier, TEXT("Position0") }, { InMaxInstanceNum }, FVector::ZeroVector);
+		Position1Handle = InstanceData->Add<1, FVector>({ InIdentifier, TEXT("Position1") }, { InMaxInstanceNum }, FVector::ZeroVector);
+		Velocity0Handle = InstanceData->Add<1, FVector>({ InIdentifier, TEXT("Velocity0") }, { InMaxInstanceNum }, FVector::ZeroVector);
+		Velocity1Handle = InstanceData->Add<1, FVector>({ InIdentifier, TEXT("Velocity1") }, { InMaxInstanceNum }, FVector::ZeroVector);
+
+		HalfLifeHandle = InstanceData->Add<1, float>({ InIdentifier, TEXT("HalfLife") }, { InMaxInstanceNum }, 0.1f);
+
+		WeightHandle = InstanceData->Add<1, float>({ InIdentifier, TEXT("Weight") }, { InMaxInstanceNum }, InWeight);
+		ScaleHandle = InstanceData->Add<1, float>({ InIdentifier, TEXT("Scale") }, { InMaxInstanceNum }, InScale);
+		ThresholdHandle = InstanceData->Add<1, float>({ InIdentifier, TEXT("Threshold") }, { InMaxInstanceNum }, InThreshold);
+	}
+
+	void FSpringInertializationPenalty::Evaluate(const FIndexSet Instances)
+	{
+		UE_LEARNING_TRACE_CPUPROFILER_EVENT_SCOPE(FSpringInertializationPenalty::Evaluate);
+
+		const TLearningArrayView<1, const FVector> Position0 = InstanceData->ConstView(Position0Handle);
+		const TLearningArrayView<1, const FVector> Position1 = InstanceData->ConstView(Position1Handle);
+		const TLearningArrayView<1, const FVector> Velocity0 = InstanceData->ConstView(Velocity0Handle);
+		const TLearningArrayView<1, const FVector> Velocity1 = InstanceData->ConstView(Velocity1Handle);
+
+		const TLearningArrayView<1, const float> HalfLife = InstanceData->ConstView(HalfLifeHandle);
+
+		const TLearningArrayView<1, const float> Scale = InstanceData->ConstView(ScaleHandle);
+		const TLearningArrayView<1, const float> Threshold = InstanceData->ConstView(ThresholdHandle);
+		const TLearningArrayView<1, const float> Weight = InstanceData->ConstView(WeightHandle);
+
+		TLearningArrayView<1, float> Reward = InstanceData->View(RewardHandle);
+
+		for (const int32 InstanceIdx : Instances)
+		{
+			Reward[InstanceIdx] = Reward::SpringInertializationPenalty(
+				Position0[InstanceIdx],
+				Position1[InstanceIdx],
+				Velocity0[InstanceIdx],
+				Velocity1[InstanceIdx],
+				HalfLife[InstanceIdx],
+				Scale[InstanceIdx],
+				Threshold[InstanceIdx],
+				Weight[InstanceIdx]);
+		}
+	}
+
+	//------------------------------------------------------------------
+
+	FPositionArraySimilarityReward::FPositionArraySimilarityReward(
+		const FName& InIdentifier,
+		const TSharedRef<FArrayMap>& InInstanceData,
+		const int32 InMaxInstanceNum,
+		const int32 InPositionNum,
+		const float InWeight,
+		const float InScale,
+		const float InThreshold)
+		: FRewardObject(InIdentifier, InInstanceData, InMaxInstanceNum)
+	{
+		Positions0Handle = InstanceData->Add<2, FVector>({ InIdentifier, TEXT("Positions0") }, { InMaxInstanceNum, InPositionNum }, FVector::ZeroVector);
+		Positions1Handle = InstanceData->Add<2, FVector>({ InIdentifier, TEXT("Positions1") }, { InMaxInstanceNum, InPositionNum }, FVector::ZeroVector);
+		RelativePosition0Handle = InstanceData->Add<1, FVector>({ InIdentifier, TEXT("RelativePosition0") }, { InMaxInstanceNum }, FVector::ZeroVector);
+		RelativePosition1Handle = InstanceData->Add<1, FVector>({ InIdentifier, TEXT("RelativePosition1") }, { InMaxInstanceNum }, FVector::ZeroVector);
+		RelativeRotation0Handle = InstanceData->Add<1, FQuat>({ InIdentifier, TEXT("RelativeRotation0") }, { InMaxInstanceNum }, FQuat::Identity);
+		RelativeRotation1Handle = InstanceData->Add<1, FQuat>({ InIdentifier, TEXT("RelativeRotation1") }, { InMaxInstanceNum }, FQuat::Identity);
+
+		WeightHandle = InstanceData->Add<1, float>({ InIdentifier, TEXT("Weight") }, { InMaxInstanceNum }, InWeight);
+		ScaleHandle = InstanceData->Add<1, float>({ InIdentifier, TEXT("Scale") }, { InMaxInstanceNum }, InScale);
+		ThresholdHandle = InstanceData->Add<1, float>({ InIdentifier, TEXT("Threshold") }, { InMaxInstanceNum }, InThreshold);
+	}
+
+	void FPositionArraySimilarityReward::Evaluate(const FIndexSet Instances)
+	{
+		UE_LEARNING_TRACE_CPUPROFILER_EVENT_SCOPE(FPositionArraySimilarityReward::Evaluate);
+
+		const TLearningArrayView<2, const FVector> Positions0 = InstanceData->ConstView(Positions0Handle);
+		const TLearningArrayView<2, const FVector> Positions1 = InstanceData->ConstView(Positions1Handle);
+		const TLearningArrayView<1, const FVector> RelativePosition0 = InstanceData->ConstView(RelativePosition0Handle);
+		const TLearningArrayView<1, const FVector> RelativePosition1 = InstanceData->ConstView(RelativePosition1Handle);
+		const TLearningArrayView<1, const FQuat> RelativeRotation0 = InstanceData->ConstView(RelativeRotation0Handle);
+		const TLearningArrayView<1, const FQuat> RelativeRotation1 = InstanceData->ConstView(RelativeRotation1Handle);
+
+		const TLearningArrayView<1, const float> Scale = InstanceData->ConstView(ScaleHandle);
+		const TLearningArrayView<1, const float> Threshold = InstanceData->ConstView(ThresholdHandle);
+		const TLearningArrayView<1, const float> Weight = InstanceData->ConstView(WeightHandle);
+
+		TLearningArrayView<1, float> Reward = InstanceData->View(RewardHandle);
+
+		const int32 PositionNum = Positions0.Num<1>();
+
+		for (const int32 InstanceIdx : Instances)
+		{
+			Reward[InstanceIdx] = 0.0f;
+
+			for (int32 PositionIdx = 0; PositionIdx < PositionNum; PositionIdx++)
+			{
+				Reward[InstanceIdx] += Reward::PositionSimilarityReward(
+					RelativeRotation0[InstanceIdx].UnrotateVector(Positions0[InstanceIdx][PositionIdx] - RelativePosition0[InstanceIdx]),
+					RelativeRotation1[InstanceIdx].UnrotateVector(Positions1[InstanceIdx][PositionIdx] - RelativePosition1[InstanceIdx]),
+					Scale[InstanceIdx],
+					Threshold[InstanceIdx],
+					Weight[InstanceIdx]) / PositionNum;
+			}
+		}
+	}
+
+	//------------------------------------------------------------------
+
+}
diff --git a/Engine/Plugins/Experimental/LearningAgents/Source/Learning/Public/LearningArray.h b/Engine/Plugins/Experimental/LearningAgents/Source/Learning/Public/LearningArray.h
new file mode 100644
--- /dev/null
+++ b/Engine/Plugins/Experimental/LearningAgents/Source/Learning/Public/LearningArray.h
@@ -0,0 +1,967 @@
+// Copyright Epic Games, Inc. All Rights Reserved.
+
+#pragma once
+
+#include "CoreTypes.h"
+
+#include "MultiArrayView.h"
+#include "MultiArray.h"
+
+// If to enable ISPC code in Learning
+#ifndef UE_LEARNING_ISPC
+#define UE_LEARNING_ISPC INTEL_ISPC
+//#define UE_LEARNING_ISPC 0
+#endif
+
+// If to enable profiling in Learning
+#ifndef UE_LEARNING_PROFILE
+#define UE_LEARNING_PROFILE (!UE_BUILD_SHIPPING)
+//#define UE_LEARNING_PROFILE 0
+#endif
+
+// Learning development mode enables various runtime sanity checks such as those described below. 
+// However, given that tools built with Learning are often designed to be performant in the editor, 
+// these checks can be disabled once the systems are stable even when not shipping.
+#ifndef UE_LEARNING_DEVELOPMENT
+#define UE_LEARNING_DEVELOPMENT (!UE_BUILD_SHIPPING)
+//#define UE_LEARNING_DEVELOPMENT 0
+#endif
+
+// There are various checks in Learning to see if arrays contain inf/nan/MAX_flt/-MAX_flt
+// which can be disabled or enabled with this flag. These will be disabled in 
+// Shipping mode but in certain cases you may want to disable them in editor 
+// mode too to improve performance.
+#ifndef UE_LEARNING_ARRAY_CHECK_VALUES
+#define UE_LEARNING_ARRAY_CHECK_VALUES UE_LEARNING_DEVELOPMENT
+#endif
+
+// There are various checks in Learning to see if arrays are the right shapes
+// which can be disabled or enabled with this flag. These will be disabled in 
+// Shipping mode but in certain cases you may want to disable them in editor 
+// mode too to improve performance as some tight loops do contain shape checks.
+#ifndef UE_LEARNING_ARRAY_CHECK_SHAPE
+#define UE_LEARNING_ARRAY_CHECK_SHAPE UE_LEARNING_DEVELOPMENT
+#endif
+
+// Enables/Disables index out-of-bounds checking in LearningArray and LearningArrayView.
+// Since so many of the editor tools use extensive numerical code having range checking
+// enabled can significantly slow things down so it is recommended to disable this in
+// editor mode when not debugging. Will be disabled automatically in Shipping.
+#ifndef UE_LEARNING_ARRAY_CHECK_INDEX
+#define UE_LEARNING_ARRAY_CHECK_INDEX UE_LEARNING_DEVELOPMENT
+#endif
+
+// Enables/Disables restrict in LearningArray and LearningArrayView. By convention functions
+// in Learning should not be called with aliasing arguments so it should be okay to enable
+// this but nonetheless it can be safer to leave it disabled in particular if you are
+// seeing odd behavior on certain platforms.
+#ifndef UE_LEARNING_ARRAY_RESTRICT
+#define UE_LEARNING_ARRAY_RESTRICT UE_LEARNING_DEVELOPMENT
+#endif
+
+
+// Profiling Macros
+
+#if UE_LEARNING_PROFILE == 1
+#define UE_LEARNING_TRACE_CPUPROFILER_EVENT_SCOPE(Name) TRACE_CPUPROFILER_EVENT_SCOPE(Name)
+#else
+#define UE_LEARNING_TRACE_CPUPROFILER_EVENT_SCOPE(Name)
+#endif
+
+
+// Assert Macros
+
+#if UE_LEARNING_DEVELOPMENT == 1
+#define UE_LEARNING_CHECK(Expr) check(Expr)
+#define UE_LEARNING_CHECKF(Expr, Format, ...) checkf(Expr, Format, ##__VA_ARGS__)
+#else
+#define UE_LEARNING_CHECK(...)
+#define UE_LEARNING_CHECKF(...)
+#endif
+
+#define UE_LEARNING_NOT_IMPLEMENTED() UE_LEARNING_CHECKF(false, TEXT("Not Implemented!"))
+
+#if UE_LEARNING_ARRAY_CHECK_SHAPE == 1
+#define UE_LEARNING_ARRAY_SHAPE_CHECK(Expr) check(Expr)
+#define UE_LEARNING_ARRAY_SHAPE_CHECKF(Expr, Format, ...) checkf(Expr, Format, ##__VA_ARGS__)
+#else
+#define UE_LEARNING_ARRAY_SHAPE_CHECK(...)
+#define UE_LEARNING_ARRAY_SHAPE_CHECKF(...)
+#endif
+
+#if UE_LEARNING_ARRAY_CHECK_VALUES == 1
+#define UE_LEARNING_ARRAY_VALUE_CHECK(Expr) check(Expr)
+#define UE_LEARNING_ARRAY_VALUE_CHECKF(Expr, Format, ...) checkf(Expr, Format, ##__VA_ARGS__)
+#else
+#define UE_LEARNING_ARRAY_VALUE_CHECK(...)
+#define UE_LEARNING_ARRAY_VALUE_CHECKF(...)
+#endif
+
+#if UE_LEARNING_ARRAY_RESTRICT == 1
+#define UE_LEARNING_ARRAY_RESTRICT_CHECK(PtrLhs, PtrRhs) check((PtrLhs == nullptr && PtrRhs == nullptr) || (PtrLhs != PtrRhs))
+#else
+#define UE_LEARNING_ARRAY_RESTRICT_CHECK(PtrLhs, PtrRhs)
+#endif
+
+
+// Learning Multi-Dimensional Array Types
+
+template<uint32 DimNum>
+using TLearningArrayShape = TMultiArrayShape<DimNum>;
+
+template<uint32 DimNum, typename ElementType>
+using TLearningArrayView = TMultiArrayView<
+	DimNum, 
+	ElementType, 
+	UE_LEARNING_ARRAY_CHECK_INDEX, 
+	UE_LEARNING_ARRAY_RESTRICT>;
+
+template<uint32 DimNum, typename ElementType, typename Allocator = FDefaultAllocator>
+using TLearningArray = TMultiArray<
+	DimNum, 
+	ElementType, 
+	Allocator, 
+	UE_LEARNING_ARRAY_CHECK_INDEX, 
+	UE_LEARNING_ARRAY_RESTRICT>;
+
+
+// Learning Array Functions
+
+namespace UE::Learning
+{
+	/**
+	* Iterator for FIndexSet (see below)
+	*/
+	struct LEARNING_API FIndexSetIterator
+	{
+		FIndexSetIterator(const int32* InPtr) : bUseIndex(false), Index(0), Ptr(InPtr) {}
+		FIndexSetIterator(int32 InIndex) : bUseIndex(true), Index(InIndex), Ptr(nullptr) {}
+
+		bool operator!=(const FIndexSetIterator& Rhs) const { return bUseIndex ? Index != Rhs.Index : Ptr != Rhs.Ptr; }
+		FIndexSetIterator& operator++() { if (bUseIndex) { Index++; } else { Ptr++; }; return *this; }
+		int32 operator*() const { return bUseIndex ? Index : *Ptr; }
+
+	private:
+		const bool bUseIndex = true;
+		int32 Index = 0;
+		const int32* Ptr = nullptr;
+	};
+
+	/**
+	* FIndexSet represents a set of indices - more specifically, either a slice of indices, 
+	* or an ArrayView of int32 which acts as the set of indices. This type is generally used 
+	* when we have allocated arrays of data but only need to iterate over or update a subset. 
+	* For example, if we allocate data for 100 instances, but only 10 of those instances are 
+	* active we can use an FIndexSet to indicate which ones.
+	* 
+	* Having this type act as either a slice or array of int32 lets us write code once, but 
+	* still allows the compiler to potentially generate both code paths (or we can write 
+	* both paths manually, using ispc to be more efficient when we have a slice and don't 
+	* require random reads/writes).
+	* 
+	* We can also attempt to automatically convert index sets which contain consecutive 
+	* indices from arrays of int32 to slices using the `TryMakeSlice` function for more 
+	* efficient processing.
+	*/
+	struct LEARNING_API FIndexSet
+	{
+		FIndexSet() = default;
+		FIndexSet(const int32 InSingleIndex) : SliceStart(InSingleIndex), SliceNum(1) {}
+		FIndexSet(const int32 InSliceStart, const int32 InSliceNum) : SliceStart(InSliceStart), SliceNum(InSliceNum) {}
+		FIndexSet(const TArrayView<const int32> InIndices) : Indices(InIndices) {}
+		FIndexSet(const TLearningArrayView<1, const int32> InIndices) : Indices(InIndices) {}
+		FIndexSet(std::initializer_list<int32> InIndices) : Indices(InIndices) {}
+		
+		template<typename InAllocatorType>
+		FIndexSet(const TArray<int32, InAllocatorType>& InIndices) : Indices(InIndices) {}
+
+		template<typename InAllocatorType>
+		FIndexSet(const TLearningArray<1, int32, InAllocatorType>& InIndices) : Indices(InIndices) {}
+
+		FIndexSet& operator=(const int32 InSingleIndex) { SliceStart = InSingleIndex; SliceNum = 1; Indices = TLearningArrayView<1, const int32>(); return *this; }
+		FIndexSet& operator=(const FIndexSet& Rhs) { SliceStart = Rhs.SliceStart; SliceNum = Rhs.SliceNum; Indices = Rhs.Indices; return *this; }
+		FIndexSet& operator=(const TArrayView<const int32> InIndices) { SliceStart = 0; SliceNum = 0; Indices = InIndices; return *this; }
+		FIndexSet& operator=(TLearningArrayView<1, const int32> InIndices) { SliceStart = 0; SliceNum = 0; Indices = InIndices; return *this; }
+
+		template<typename InAllocatorType>
+		FIndexSet& operator=(const TArray<int32, InAllocatorType>& InIndices) { SliceStart = 0; SliceNum = 0; Indices = InIndices; return *this; }
+
+		template<typename InAllocatorType>
+		FIndexSet& operator=(const TLearningArray<1, int32, InAllocatorType>& InIndices) { SliceStart = 0; SliceNum = 0; Indices = InIndices; return *this; }
+
+		TArray<int32> ToArray() const
+		{
+			const int32 IdxNum = Num();
+			TArray<int32> Out;
+			Out.SetNumUninitialized(IdxNum);
+			for (int32 Idx = 0; Idx < IdxNum; Idx++)
+			{
+				Out[Idx] = (*this)[Idx];
+			}
+			return Out;
+		}
+
+		/**
+		* Attempts to convert this IndexSet from a view of indices into a slice for more efficient
+		* processing. This works by finding the minimum and maximum elements in the array and comparing 
+		* the range of this min and max to the number of elements. When indices are consecutive `Max - Min + 1 == Num`.
+		* 
+		* Warning: this algorithm works on the assumption that there are no duplicate indices in the view. If there
+		* are duplicate indices in your set this will not give the correct result.
+		* 
+		* @returns true if conversion to a slice is successful
+		*/
+		bool TryMakeSlice();
+
+		inline int32 operator[](const int32 Index) const
+		{
+			return IsSlice() ? SliceStart + Index : Indices[Index];
+		}
+
+		inline int32 Num() const
+		{
+			return IsSlice() ? SliceNum : Indices.Num();
+		}
+
+		inline FIndexSet Slice(const int32 InStart, const int32 InNum) const
+		{
+
+#if UE_LEARNING_DEVELOPMENT
+			if (IsSlice())
+			{
+				UE_LEARNING_CHECK(InStart >= 0);
+				UE_LEARNING_CHECK(InNum >= 0);
+				UE_LEARNING_CHECK(InStart + InNum <= SliceNum);
+			}
+#endif
+
+			return IsSlice() ? FIndexSet(SliceStart + InStart, InNum) : Indices.Slice(InStart, InNum);
+		}
+
+		inline bool Contains(const int32 Index) const
+		{
+			return IsSlice() ? (Index >= SliceStart && Index < SliceStart + SliceNum) : Indices.Contains(Index);
+		}
+
+		inline FIndexSetIterator begin() const
+		{
+			return IsSlice() ?
+				FIndexSetIterator(SliceStart) :
+				FIndexSetIterator(Indices.GetData());
+		}
+
+		inline FIndexSetIterator end() const
+		{
+			return IsSlice() ?
+				FIndexSetIterator(SliceStart + SliceNum) :
+				FIndexSetIterator(Indices.GetData() + Indices.Num());
+		}
+
+		inline bool IsSlice() const { return Indices.GetData() == nullptr; }
+		inline int32 GetSliceStart() const { UE_LEARNING_CHECK(IsSlice()); return SliceStart; }
+		inline int32 GetSliceNum() const { UE_LEARNING_CHECK(IsSlice()); return SliceNum; }
+		inline TLearningArrayView<1, const int32> GetIndices() const { UE_LEARNING_CHECK(!IsSlice()); return Indices; }
+
+	private:
+
+		int32 SliceStart = 0, SliceNum = 0;
+		TLearningArrayView<1, const int32> Indices;
+	};
+
+	/**
+	* Similar to `ParallelFor` but instead of providing a single index to the callback it provides a slice for the
+	* callback to loop over. The advantage here is that it gives the compiler a chance of optimizing the loop code
+	* in the callback such as unrolling and auto-vectorizing. This is important when the number of items to loop
+	* over is large (e.g. ~1000) since otherwise this would produce an individual function call for every element
+	* in the loop.
+	*
+	* @param Num					Number of elements to iterate over
+	* @param MinSliceElementNum		Minimum number of elements to include in each call to Body
+	* @param Body					Callback taking start and stop index slice of Num
+	*/
+	LEARNING_API void SlicedParallelFor(const int32 Num, const int32 MinSliceElementNum, const TFunctionRef<void(const int32 Start, const int32 Length)> Body);
+
+	/**
+	* Some additional functions that act on arrays such as copy, set, zero, etc.
+	* 
+	* Unfortunately most of these functions need to be duplicated for Array, ArrayView and their
+	* const and non-const variations as template functions don't allow for implicit casts.
+	*/
+	namespace Array
+	{
+		/**
+		* Check that two array shapes are equal.
+		*/
+		template<uint8 InDimNum>
+		inline void CheckShapesEqual(const TLearningArrayShape<InDimNum>& Lhs, const TLearningArrayShape<InDimNum>& Rhs)
+		{
+#if UE_LEARNING_ARRAY_CHECK_SHAPE == 1
+			for (uint8 Idx = 0; Idx < InDimNum; Idx++)
+			{
+				UE_LEARNING_ARRAY_SHAPE_CHECKF(Lhs[Idx] == Rhs[Idx],
+					TEXT("Array Shapes don't match on dimension %i of %i (lhs: %i, rhs: %i)"), Idx + 1, InDimNum, Lhs[Idx], Rhs[Idx]);
+			}
+#endif
+		}
+
+		/**
+		* Check that an array view does not contain any NaN, Inf, -Inf, MAX_flt, or -MAX_flt
+		*/
+		template<uint8 InDimNum, typename InElementType>
+		inline void Check(const TLearningArrayView<InDimNum, InElementType> View)
+		{
+#if UE_LEARNING_ARRAY_CHECK_VALUES == 1
+			for (int32 Idx = 0; Idx < View.Num(); Idx++)
+			{
+				UE_LEARNING_ARRAY_VALUE_CHECKF(
+					FMath::IsFinite(View.GetData()[Idx]) && View.GetData()[Idx] != MAX_flt && View.GetData()[Idx] != -MAX_flt,
+					TEXT("Invalid value %f found at flat array index %i"), View.GetData()[Idx], Idx);
+			}
+#endif
+		}
+
+		/**
+		* Check that an array does not contain any NaN, Inf, -Inf, MAX_flt, or -MAX_flt
+		*/
+		template<uint8 InDimNum, typename InElementType, typename Allocator>
+		inline void Check(const TLearningArray<InDimNum, InElementType, Allocator>& View)
+		{
+#if UE_LEARNING_ARRAY_CHECK_VALUES == 1
+			for (int32 Idx = 0; Idx < View.Num(); Idx++)
+			{
+				UE_LEARNING_ARRAY_VALUE_CHECKF(
+					FMath::IsFinite(View.GetData()[Idx]) && View.GetData()[Idx] != MAX_flt && View.GetData()[Idx] != -MAX_flt,
+					TEXT("Invalid value %f found at flat array index %i"), View.GetData()[Idx], Idx);
+			}
+#endif
+		}
+
+
+		/**
+		* Check that an array view does not contain any NaN, Inf, -Inf, MAX_flt, or -MAX_flt
+		*/
+		template<uint8 InDimNum, typename InElementType>
+		inline void Check(const TLearningArrayView<InDimNum, const InElementType> View)
+		{
+#if UE_LEARNING_ARRAY_CHECK_VALUES == 1
+			for (int32 Idx = 0; Idx < View.Num(); Idx++)
+			{
+				UE_LEARNING_ARRAY_VALUE_CHECKF(
+					FMath::IsFinite(View.GetData()[Idx]) && View.GetData()[Idx] != MAX_flt && View.GetData()[Idx] != -MAX_flt,
+					TEXT("Invalid value %f found at flat array index %i"), View.GetData()[Idx], Idx);
+			}
+#endif
+		}
+
+		/**
+		* Check that an array view does not contain any NaN, Inf, -Inf, MAX_flt, or -MAX_flt
+		*/
+		template<uint8 InDimNum, typename InElementType>
+		inline void Check(const TLearningArrayView<InDimNum, const InElementType> View, const FIndexSet Indices)
+		{
+#if UE_LEARNING_ARRAY_CHECK_VALUES == 1
+			if constexpr (InDimNum > 1)
+			{
+				for (const int32 Idx : Indices)
+				{
+					Check(View[Idx]);
+				}
+			}
+			else
+			{
+				for (const int32 Idx : Indices)
+				{
+					UE_LEARNING_ARRAY_VALUE_CHECKF(
+						FMath::IsFinite(View.GetData()[Idx]) && View.GetData()[Idx] != MAX_flt && View.GetData()[Idx] != -MAX_flt,
+						TEXT("Invalid value %f found at flat array index %i"), View.GetData()[Idx], Idx);
+				}
+			}
+#endif
+		}
+
+		/**
+		* Check that an array view does not contain any NaN, Inf, -Inf, MAX_flt, or -MAX_flt
+		*/
+		template<uint8 InDimNum, typename InElementType>
+		inline void Check(const TLearningArrayView<InDimNum, InElementType> View, const FIndexSet Indices)
+		{
+#if UE_LEARNING_ARRAY_CHECK_VALUES == 1
+			if constexpr (InDimNum > 1)
+			{
+				for (const int32 Idx : Indices)
+				{
+					Check(View[Idx]);
+				}
+			}
+			else
+			{
+				for (const int32 Idx : Indices)
+				{
+					UE_LEARNING_ARRAY_VALUE_CHECKF(
+						FMath::IsFinite(View.GetData()[Idx]) && View.GetData()[Idx] != MAX_flt && View.GetData()[Idx] != -MAX_flt,
+						TEXT("Invalid value %f found at flat array index %i"), View.GetData()[Idx], Idx);
+				}
+			}
+#endif
+		}
+
+		/**
+		* Zero the memory of an array view
+		*/
+		template<uint8 InDimNum, typename InElementType>
+		inline void Zero(TLearningArrayView<InDimNum, InElementType> View)
+		{
+			FMemory::Memzero((void*)View.GetData(), View.Num() * sizeof(InElementType));
+		}
+
+		/**
+		* Zero the memory of an array
+		*/
+		template<uint8 InDimNum, typename InElementType, typename Allocator>
+		inline void Zero(TLearningArray<InDimNum, InElementType, Allocator>& Array)
+		{
+			FMemory::Memzero((void*)Array.GetData(), Array.Num() * sizeof(InElementType));
+		}
+
+		/**
+		* Zero the memory of an array view
+		*/
+		template<uint8 InDimNum, typename InElementType>
+		inline void Zero(TLearningArrayView<InDimNum, InElementType> View, const FIndexSet Indices)
+		{
+			if constexpr (InDimNum > 1)
+			{
+				for (const int32 Idx : Indices)
+				{
+					Zero(View[Idx]);
+				}
+			}
+			else
+			{
+				for (const int32 Idx : Indices)
+				{
+					FMemory::Memzero((void*)&View[Idx], sizeof(InElementType));
+				}
+			}
+		}
+
+		/**
+		* Zero the memory of an array view
+		*/
+		template<uint8 InDimNum, typename InElementType, typename Allocator>
+		inline void Zero(TLearningArray<InDimNum, InElementType, Allocator>& View, const FIndexSet Indices)
+		{
+			if constexpr (InDimNum > 1)
+			{
+				for (const int32 Idx : Indices)
+				{
+					Zero(View[Idx]);
+				}
+			}
+			else
+			{
+				for (const int32 Idx : Indices)
+				{
+					FMemory::Memzero((void*)&View[Idx], sizeof(InElementType));
+				}
+			}
+		}
+
+
+		/**
+		* Set each item of an array view to the given element
+		*/
+		template<uint8 InDimNum, typename InElementType>
+		inline void Set(TLearningArrayView<InDimNum, InElementType> View, const InElementType& Element)
+		{
+			for (int32 Idx = 0; Idx < View.Num(); Idx++)
+			{
+				View.GetData()[Idx] = Element;
+			}
+		}
+
+		/**
+		* Set each item of an array to the given element
+		*/
+		template<uint8 InDimNum, typename InElementType, typename Allocator>
+		inline void Set(TLearningArray<InDimNum, InElementType, Allocator>& View, const InElementType& Element)
+		{
+			for (int32 Idx = 0; Idx < View.Num(); Idx++)
+			{
+				View.GetData()[Idx] = Element;
+			}
+		}
+
+		/**
+		* Set each item of an array view to the given element
+		*/
+		template<uint8 InDimNum, typename InElementType>
+		inline void Set(TLearningArrayView<InDimNum, InElementType> View, const InElementType& Element, const FIndexSet Indices)
+		{
+			if constexpr (InDimNum > 1)
+			{
+				for (const int32 Idx : Indices)
+				{
+					Set(View[Idx], Element);
+				}
+			}
+			else
+			{
+				for (const int32 Idx : Indices)
+				{
+					View[Idx] = Element;
+				}
+			}
+		}
+
+		/**
+		* Set each item of an array to the given element
+		*/
+		template<uint8 InDimNum, typename InElementType, typename Allocator>
+		inline void Set(TLearningArray<InDimNum, InElementType, Allocator>& View, const InElementType& Element, const FIndexSet Indices)
+		{
+			if constexpr (InDimNum > 1)
+			{
+				for (const int32 Idx : Indices)
+				{
+					Set(View[Idx], Element);
+				}
+			}
+			else
+			{
+				for (const int32 Idx : Indices)
+				{
+					View[Idx] = Element;
+				}
+			}
+		}
+
+		/**
+		* Copy the contents of one array view into another
+		*/
+		template<uint8 InDimNum, typename InElementType>
+		inline void Copy(
+			TLearningArrayView<InDimNum, InElementType> Dst,
+			TLearningArrayView<InDimNum, InElementType> Src)
+		{
+			UE_LEARNING_ARRAY_RESTRICT_CHECK(Dst.GetData(), Src.GetData());
+			CheckShapesEqual<InDimNum>(Dst.Shape(), Src.Shape());
+
+			for (int32 Idx = 0; Idx < Dst.Num(); Idx++)
+			{
+				Dst.GetData()[Idx] = Src.GetData()[Idx];
+			}
+		}
+
+		/**
+		* Copy the contents of one array view into another
+		*/
+		template<uint8 InDimNum, typename InElementType>
+		inline void Copy(
+			TLearningArrayView<InDimNum, InElementType> Dst,
+			const TLearningArrayView<InDimNum, const InElementType> Src)
+		{
+			UE_LEARNING_ARRAY_RESTRICT_CHECK(Dst.GetData(), Src.GetData());
+			CheckShapesEqual<InDimNum>(Dst.Shape(), Src.Shape());
+
+			for (int32 Idx = 0; Idx < Dst.Num(); Idx++)
+			{
+				Dst.GetData()[Idx] = Src.GetData()[Idx];
+			}
+		}
+
+		/**
+		* Copy the contents of one array view into another array
+		*/
+		template<uint8 InDimNum, typename InElementType, typename Allocator>
+		inline void Copy(
+			TLearningArray<InDimNum, InElementType, Allocator>& Dst,
+			const TLearningArrayView<InDimNum, const InElementType> Src)
+		{
+			UE_LEARNING_ARRAY_RESTRICT_CHECK(Dst.GetData(), Src.GetData());
+			CheckShapesEqual<InDimNum>(Dst.Shape(), Src.Shape());
+
+			for (int32 Idx = 0; Idx < Dst.Num(); Idx++)
+			{
+				Dst.GetData()[Idx] = Src.GetData()[Idx];
+			}
+		}
+
+		/**
+		* Copy the contents of one array view into another array
+		*/
+		template<uint8 InDimNum, typename InElementType, typename Allocator>
+		inline void Copy(
+			TLearningArray<InDimNum, InElementType, Allocator>& Dst,
+			const TLearningArrayView<InDimNum, InElementType> Src)
+		{
+			UE_LEARNING_ARRAY_RESTRICT_CHECK(Dst.GetData(), Src.GetData());
+			CheckShapesEqual<InDimNum>(Dst.Shape(), Src.Shape());
+
+			for (int32 Idx = 0; Idx < Dst.Num(); Idx++)
+			{
+				Dst.GetData()[Idx] = Src.GetData()[Idx];
+			}
+		}
+
+		/**
+		* Copy the contents of one array into another array view
+		*/
+		template<uint8 InDimNum, typename InElementType, typename Allocator>
+		inline void Copy(
+			TLearningArrayView<InDimNum, InElementType> Dst,
+			const TLearningArray<InDimNum, InElementType, Allocator>& Src)
+		{
+			UE_LEARNING_ARRAY_RESTRICT_CHECK(Dst.GetData(), Src.GetData());
+			CheckShapesEqual<InDimNum>(Dst.Shape(), Src.Shape());
+
+			for (int32 Idx = 0; Idx < Dst.Num(); Idx++)
+			{
+				Dst.GetData()[Idx] = Src.GetData()[Idx];
+			}
+		}
+
+		/**
+		* Copy the contents of one array into another array
+		*/
+		template<uint8 InDimNum, typename InElementType, typename AllocatorLhs, typename AllocatorRhs>
+		inline void Copy(
+			TLearningArray<InDimNum, InElementType, AllocatorLhs>& Dst,
+			const TLearningArray<InDimNum, InElementType, AllocatorRhs>& Src)
+		{
+			UE_LEARNING_ARRAY_RESTRICT_CHECK(Dst.GetData(), Src.GetData());
+			CheckShapesEqual<InDimNum>(Dst.Shape(), Src.Shape());
+
+			for (int32 Idx = 0; Idx < Dst.Num(); Idx++)
+			{
+				Dst.GetData()[Idx] = Src.GetData()[Idx];
+			}
+		}
+
+		/**
+		* Copy the contents of one array view into another
+		*/
+		template<uint8 InDimNum, typename InElementType>
+		inline void Copy(
+			TLearningArrayView<InDimNum, InElementType> Dst,
+			const TLearningArrayView<InDimNum, const InElementType> Src,
+			const FIndexSet Indices)
+		{
+			UE_LEARNING_ARRAY_RESTRICT_CHECK(Dst.GetData(), Src.GetData());
+			CheckShapesEqual<InDimNum>(Dst.Shape(), Src.Shape());
+
+			if constexpr (InDimNum > 1)
+			{
+				for (const int32 Idx : Indices)
+				{
+					Copy(Dst[Idx], Src[Idx]);
+				}
+			}
+			else
+			{
+				for (const int32 Idx : Indices)
+				{
+					Dst[Idx] = Src[Idx];
+				}
+			}
+		}
+
+		/**
+		* Copy the contents of one array view into another
+		*/
+		template<uint8 InDimNum, typename InElementType, typename Allocator>
+		inline void Copy(
+			TLearningArrayView<InDimNum, InElementType> Dst,
+			const TLearningArray<InDimNum, InElementType, Allocator>& Src,
+			const FIndexSet Indices)
+		{
+			UE_LEARNING_ARRAY_RESTRICT_CHECK(Dst.GetData(), Src.GetData());
+			CheckShapesEqual<InDimNum>(Dst.Shape(), Src.Shape());
+
+			if constexpr (InDimNum > 1)
+			{
+				for (const int32 Idx : Indices)
+				{
+					Copy(Dst[Idx], Src[Idx]);
+				}
+			}
+			else
+			{
+				for (const int32 Idx : Indices)
+				{
+					Dst[Idx] = Src[Idx];
+				}
+			}
+		}
+
+		/**
+		* Copy the contents of one array view into another
+		*/
+		template<uint8 InDimNum, typename InElementType, typename AllocatorLhs, typename AllocatorRhs>
+		inline void Copy(
+			TLearningArray<InDimNum, InElementType, AllocatorLhs>& Dst,
+			const TLearningArray<InDimNum, InElementType, AllocatorRhs>& Src,
+			const FIndexSet Indices)
+		{
+			UE_LEARNING_ARRAY_RESTRICT_CHECK(Dst.GetData(), Src.GetData());
+			CheckShapesEqual<InDimNum>(Dst.Shape(), Src.Shape());
+
+			if constexpr (InDimNum > 1)
+			{
+				for (const int32 Idx : Indices)
+				{
+					Copy(Dst[Idx], Src[Idx]);
+				}
+			}
+			else
+			{
+				for (const int32 Idx : Indices)
+				{
+					Dst[Idx] = Src[Idx];
+				}
+			}
+		}
+
+		/**
+		* Copy the contents of one array view into another
+		*/
+		template<uint8 InDimNum, typename InElementType, typename Allocator>
+		inline void Copy(
+			TLearningArray<InDimNum, InElementType, Allocator>& Dst,
+			const TLearningArrayView<InDimNum, const InElementType> Src,
+			const FIndexSet Indices)
+		{
+			UE_LEARNING_ARRAY_RESTRICT_CHECK(Dst.GetData(), Src.GetData());
+			CheckShapesEqual<InDimNum>(Dst.Shape(), Src.Shape());
+
+			if constexpr (InDimNum > 1)
+			{
+				for (const int32 Idx : Indices)
+				{
+					Copy(Dst[Idx], Src[Idx]);
+				}
+			}
+			else
+			{
+				for (const int32 Idx : Indices)
+				{
+					Dst[Idx] = Src[Idx];
+				}
+			}
+		}
+
+		/**
+		* Check if the contents of one array equals the contents of another.
+		*/
+		template<uint8 InDimNum, typename InElementType, typename AllocatorLhs, typename AllocatorRhs>
+		inline bool Equal(
+			const TLearningArray<InDimNum, InElementType, AllocatorLhs>& Lhs,
+			const TLearningArray<InDimNum, InElementType, AllocatorRhs>& Rhs)
+		{
+			UE_LEARNING_ARRAY_RESTRICT_CHECK(Lhs.GetData(), Rhs.GetData());
+			CheckShapesEqual<InDimNum>(Lhs.Shape(), Rhs.Shape());
+
+			for (int32 Idx = 0; Idx < Lhs.Num(); Idx++)
+			{
+				if (Lhs.GetData()[Idx] != Rhs.GetData()[Idx])
+				{
+					return false;
+				}
+			}
+
+			return true;
+		}
+
+		/**
+		* Check if the contents of one array equals the contents of another.
+		*/
+		template<uint8 InDimNum, typename InElementType, typename AllocatorRhs>
+		inline bool Equal(
+			const TLearningArrayView<InDimNum, const InElementType> Lhs,
+			const TLearningArray<InDimNum, InElementType, AllocatorRhs>& Rhs)
+		{
+			UE_LEARNING_ARRAY_RESTRICT_CHECK(Lhs.GetData(), Rhs.GetData());
+			CheckShapesEqual<InDimNum>(Lhs.Shape(), Rhs.Shape());
+
+			for (int32 Idx = 0; Idx < Lhs.Num(); Idx++)
+			{
+				if (Lhs.GetData()[Idx] != Rhs.GetData()[Idx])
+				{
+					return false;
+				}
+			}
+
+			return true;
+		}
+
+		/**
+		* Check if the contents of one array equals the contents of another.
+		*/
+		template<uint8 InDimNum, typename InElementType, typename AllocatorRhs>
+		inline bool Equal(
+			const TLearningArrayView<InDimNum, InElementType> Lhs,
+			const TLearningArray<InDimNum, InElementType, AllocatorRhs>& Rhs)
+		{
+			UE_LEARNING_ARRAY_RESTRICT_CHECK(Lhs.GetData(), Rhs.GetData());
+			CheckShapesEqual<InDimNum>(Lhs.Shape(), Rhs.Shape());
+
+			for (int32 Idx = 0; Idx < Lhs.Num(); Idx++)
+			{
+				if (Lhs.GetData()[Idx] != Rhs.GetData()[Idx])
+				{
+					return false;
+				}
+			}
+
+			return true;
+		}
+
+		/**
+		* Check if the contents of one array equals the contents of another.
+		*/
+		template<uint8 InDimNum, typename InElementType, typename AllocatorLhs>
+		inline bool Equal(
+			const TLearningArray<InDimNum, InElementType, AllocatorLhs>& Lhs,
+			const TLearningArrayView<InDimNum, const InElementType>& Rhs)
+		{
+			UE_LEARNING_ARRAY_RESTRICT_CHECK(Lhs.GetData(), Rhs.GetData());
+			CheckShapesEqual<InDimNum>(Lhs.Shape(), Rhs.Shape());
+
+			for (int32 Idx = 0; Idx < Lhs.Num(); Idx++)
+			{
+				if (Lhs.GetData()[Idx] != Rhs.GetData()[Idx])
+				{
+					return false;
+				}
+			}
+
+			return true;
+		}
+
+		/**
+		* Check if the contents of one array equals the contents of another.
+		*/
+		template<uint8 InDimNum, typename InElementType>
+		inline bool Equal(
+			const TLearningArrayView<InDimNum, const InElementType> Lhs,
+			const TLearningArrayView<InDimNum, const InElementType> Rhs)
+		{
+			UE_LEARNING_ARRAY_RESTRICT_CHECK(Lhs.GetData(), Rhs.GetData());
+			CheckShapesEqual<InDimNum>(Lhs.Shape(), Rhs.Shape());
+
+			for (int32 Idx = 0; Idx < Lhs.Num(); Idx++)
+			{
+				if (Lhs.GetData()[Idx] != Rhs.GetData()[Idx])
+				{
+					return false;
+				}
+			}
+
+			return true;
+		}
+
+		/**
+		* Shift all elements in an array to the left. Data at the end of the array will remain the same.
+		*/
+		template<uint8 InDimNum, typename InElementType, typename Allocator>
+		inline void ShiftLeft(TLearningArray<InDimNum, InElementType, Allocator>& Array, const int32 ShiftNum)
+		{
+			for (int32 Idx = 0; Idx < (int32)Array.template Num<0>() - ShiftNum; Idx++)
+			{
+				Copy(Array[Idx], Array[Idx + ShiftNum]);
+			}
+		}
+
+		/**
+		* Shift all elements in an array to the left. Data at the end of the array will remain the same.
+		*/
+		template<typename InElementType, typename Allocator>
+		inline void ShiftLeft(TLearningArray<1, InElementType, Allocator>& Array, const int32 ShiftNum)
+		{
+			for (int32 Idx = 0; Idx < (int32)Array.Num() - ShiftNum; Idx++)
+			{
+				Array.GetData()[Idx] = Array.GetData()[Idx + ShiftNum];
+			}
+		}
+
+		/**
+		* Shift all elements in an array to the left. Data at the end of the array will remain the same.
+		*/
+		template<uint8 InDimNum, typename InElementType>
+		inline void ShiftLeft(TLearningArrayView<InDimNum, InElementType> Array, const int32 ShiftNum)
+		{
+			for (int32 Idx = 0; Idx < (int32)Array.template Num<0>() - ShiftNum; Idx++)
+			{
+				Copy(Array[Idx], Array[Idx + ShiftNum]);
+			}
+		}
+
+		/**
+		* Shift all elements in an array to the left. Data at the end of the array will remain the same.
+		*/
+		template<typename InElementType>
+		inline void ShiftLeft(TLearningArrayView<1, InElementType> Array, const int32 ShiftNum)
+		{
+			for (int32 Idx = 0; Idx < (int32)Array.Num() - ShiftNum; Idx++)
+			{
+				Array.GetData()[Idx] = Array.GetData()[Idx + ShiftNum];
+			}
+		}
+
+		/**
+		* Serialize an array
+		*/
+		template<uint8 InDimNum, typename InElementType, typename Allocator>
+		inline void Serialize(FArchive& Ar, TLearningArray<InDimNum, InElementType, Allocator>& Array)
+		{
+			if (Ar.IsLoading())
+			{
+				TLearningArrayShape<InDimNum> Shape;
+				for (uint8 ShapeIdx = 0; ShapeIdx < InDimNum; ShapeIdx++)
+				{
+					Ar << Shape[ShapeIdx];
+				}
+				Array.SetNumUninitialized(Shape);
+				Ar.Serialize(Array.GetData(), Array.Num() * sizeof(InElementType));
+			}
+			else if (Ar.IsSaving())
+			{
+				for (uint8 ShapeIdx = 0; ShapeIdx < InDimNum; ShapeIdx++)
+				{
+					FDefaultAllocator::SizeType Num = Array.Num(ShapeIdx);
+					Ar << Num;
+				}
+				Ar.Serialize(Array.GetData(), Array.Num() * sizeof(InElementType));
+			}
+		}
+
+		template<typename InElementType>
+		inline int32 IndexOf(const TLearningArrayView<1, const InElementType> Array, const InElementType& Element)
+		{
+			for (int32 Idx = 0; Idx < Array.Num(); Idx++)
+			{
+				if (Array.GetData()[Idx] == Element)
+				{
+					return Idx;
+				}
+			}
+			return INDEX_NONE;
+		}
+
+		template<typename InElementType>
+		inline int32 IndexOf(const TLearningArray<1, InElementType>& Array, const InElementType& Element)
+		{
+			for (int32 Idx = 0; Idx < Array.Num(); Idx++)
+			{
+				if (Array.GetData()[Idx] == Element)
+				{
+					return Idx;
+				}
+			}
+			return INDEX_NONE;
+		}
+
+	}
+}
+
+
