Learning Agents: Create initial experimental version

diff --git a/Engine/Plugins/Experimental/LearningAgents/Source/Learning/Public/LearningArrayMap.h b/Engine/Plugins/Experimental/LearningAgents/Source/Learning/Public/LearningArrayMap.h
new file mode 100644
--- /dev/null
+++ b/Engine/Plugins/Experimental/LearningAgents/Source/Learning/Public/LearningArrayMap.h
@@ -0,0 +1,511 @@
+// Copyright Epic Games, Inc. All Rights Reserved.
+
+#pragma once
+
+#include "LearningArray.h"
+
+#include "Containers/Map.h"
+#include "Containers/StaticArray.h"
+#include "Misc/GeneratedTypeName.h"
+#include "UObject/NameTypes.h"
+
+/**
+* `FArrayMap` is a Map from `FArrayMapKey` to `TLearningArray` where each array 
+* object stored as a value can be a different shape, dimensionality, and type.
+* 
+* To make this data structure easier to use, a type-safe `TArrayMapHandle` object 
+* is provided which is basically an integer index that can be used to get a view 
+* of an array safely, and without type-casting.
+* 
+* Arrays in this data structure can also be "linked", which makes the two 
+* Keys/Handles which are linked together return a view of the same array data. This 
+* is useful to avoid copying over array data between two entries.
+* 
+* In `Learning` this object is most commonly used to store "InstanceData" - essentially
+* all the data required for multiple instances (such as multiple characters) is 
+* stored together in a single `FArrayMap` as a set of large arrays. This makes this 
+* data extremely efficient and cache-friendly to access when we are processing data for 
+* multiple instances at once.
+* 
+* This structure also provides some basic optimizations for storing single-element arrays
+* to also make it a viable option for when you have (for example) a single instance to store
+* data for.
+*/
+
+class USplineComponent;
+
+namespace UE::Learning
+{
+	enum class ECompletionMode : uint8;
+
+	/**
+	* Key type for FArrayMap. A pair of FNames - generally we can use one to identify the object itself, and one for the variable/member/field.
+	*/
+	struct LEARNING_API FArrayMapKey
+	{
+		FName Namespace;
+		FName Variable;
+
+		bool operator==(const FArrayMapKey& Other) const
+		{
+			return (Namespace == Other.Namespace) && (Variable == Other.Variable);
+		}
+
+		friend uint32 GetTypeHash(const FArrayMapKey& Other)
+		{
+			return GetTypeHash(Other.Namespace) ^ GetTypeHash(Other.Variable);
+		}
+	};
+
+
+	/**
+	* Compile-time function to generate a type id from a C++ type. This is used to check at
+	* runtime that when we access an array from the FArrayMap we are casting it to the correct 
+	* type.
+	*/
+	template<typename T> struct TTypeId
+	{
+		static_assert(sizeof(T) != sizeof(T), "Please specialize this class with a unique id to allow your type to be used in a FArrayMap");
+	};
+	template<> struct TTypeId<float> { enum { Value = 0 }; };
+	template<> struct TTypeId<double> { enum { Value = 1 }; };
+	template<> struct TTypeId<bool> { enum { Value = 2 }; };
+	template<> struct TTypeId<int8> { enum { Value = 3 }; };
+	template<> struct TTypeId<int16> { enum { Value = 4 }; };
+	template<> struct TTypeId<int32> { enum { Value = 5 }; };
+	template<> struct TTypeId<int64> { enum { Value = 6 }; };
+	template<> struct TTypeId<uint8> { enum { Value = 7 }; };
+	template<> struct TTypeId<uint16> { enum { Value = 8 }; };
+	template<> struct TTypeId<uint32> { enum { Value = 9 }; };
+	template<> struct TTypeId<uint64> { enum { Value = 10 }; };
+	template<> struct TTypeId<FVector> { enum { Value = 11 }; };
+	template<> struct TTypeId<FQuat> { enum { Value = 12 }; };
+	template<> struct TTypeId<FTransform> { enum { Value = 13 }; };
+	template<> struct TTypeId<ECompletionMode> { enum { Value = 14 }; };
+	enum { CustomTypeIdStartIndex = 15 };
+
+	/**
+	* Compile-time function to get a debug name for a C++ type
+	*/
+	template<typename T>
+	const TCHAR* CompileTimeTypeName() { return GetGeneratedTypeName<T>(); }
+
+	/**
+	* Untyped handle to an array stored in an FArrayMap
+	*/
+	struct FArrayMapHandle
+	{
+		int16 Index = INDEX_NONE;
+
+		bool operator==(const FArrayMapHandle& Rhs) const { return Index == Rhs.Index; }
+	};
+
+	/**
+	* Type-safe handle to an array stored in an FArrayMap
+	*/
+	template<uint8 DimNum, typename ElementType>
+	struct TArrayMapHandle
+	{
+		int16 Index = INDEX_NONE;
+
+		inline operator FArrayMapHandle() const { return { Index }; }
+	};
+
+	// Some helper functions for FArrayMap
+	namespace ArrayMapPrivate
+	{
+		/**
+		* Destructor for an array of elements of a given type
+		*/
+		template<typename ElementType>
+		void DestructArrayItems(void* Elements, const int32 DimNum, const int32* Shape)
+		{
+			int32 Total = Shape[0];
+			for (uint8 DimIdx = 1; DimIdx < DimNum; DimIdx++)
+			{
+				Total *= Shape[DimIdx];
+			}
+
+			DestructItems<ElementType, int32>((ElementType*)Elements, Total);
+		}
+
+		static inline int32 RoundUpToMultiple(const int32 Value, const int32 Multiple)
+		{
+			return ((Value + Multiple - 1) / Multiple) * Multiple;
+		}
+	}
+
+	/**
+	* A map of differently sized and typed multi-dimensional arrays
+	*/
+	struct LEARNING_API FArrayMap
+	{
+		FArrayMap();
+		~FArrayMap();
+
+		/**
+		* Get a read-only view of an array.
+		*/
+		template<uint8 DimNum, typename ElementType>
+		TLearningArrayView<DimNum, const ElementType> ConstView(const FArrayMapKey Key) const
+		{
+			const FArrayMapHandle Handle = Lookup(Key);
+			CheckArrayWithKey<DimNum, ElementType>(Key, Handle.Index);
+			FDynamicArrayView View = Arrays.GetData()[Handle.Index];
+			return TLearningArrayView<DimNum, const ElementType>((ElementType*)View.Data, TLearningArrayShape<DimNum>(View.Shape));
+		}
+
+		/**
+		* Get a read-only view of an array.
+		*/
+		template<uint8 DimNum, typename ElementType>
+		TLearningArrayView<DimNum, const ElementType> ConstView(const FArrayMapHandle Handle) const
+		{
+			CheckArray<DimNum, ElementType>(Handle.Index);
+			FDynamicArrayView View = Arrays.GetData()[Handle.Index];
+			return TLearningArrayView<DimNum, const ElementType>((ElementType*)View.Data, TLearningArrayShape<DimNum>(View.Shape));
+		}
+
+		/**
+		* Get a read-only view of an array.
+		*/
+		template<uint8 DimNum, typename ElementType>
+		TLearningArrayView<DimNum, const ElementType> ConstView(const TArrayMapHandle<DimNum, ElementType> Handle) const
+		{
+			CheckArray<DimNum, ElementType>(Handle.Index);
+			FDynamicArrayView View = Arrays.GetData()[Handle.Index];
+			return TLearningArrayView<DimNum, const ElementType>((ElementType*)View.Data, TLearningArrayShape<DimNum>(View.Shape));
+		}
+
+		/**
+		* Get a view of an array.
+		*/
+		template<uint8 DimNum, typename ElementType>
+		TLearningArrayView<DimNum, ElementType> View(const FArrayMapKey Key) const
+		{
+			const FArrayMapHandle Handle = Lookup(Key);
+			CheckArrayWithKey<DimNum, ElementType>(Key, Handle.Index);
+			FDynamicArrayView View = Arrays.GetData()[Handle.Index];
+			return TLearningArrayView<DimNum, ElementType>((ElementType*)View.Data, TLearningArrayShape<DimNum>(View.Shape));
+		}
+
+		/**
+		* Get a view of an array.
+		*/
+		template<uint8 DimNum, typename ElementType>
+		TLearningArrayView<DimNum, ElementType> View(const FArrayMapHandle Handle) const
+		{
+			CheckArray<DimNum, ElementType>(Handle.Index);
+			FDynamicArrayView View = Arrays.GetData()[Handle.Index];
+			return TLearningArrayView<DimNum, ElementType>((ElementType*)View.Data, TLearningArrayShape<DimNum>(View.Shape));
+		}
+
+		/**
+		* Get a view of an array.
+		*/
+		template<uint8 DimNum, typename ElementType>
+		TLearningArrayView<DimNum, ElementType> View(const TArrayMapHandle<DimNum, ElementType> Handle) const
+		{
+			CheckArray<DimNum, ElementType>(Handle.Index);
+			FDynamicArrayView View = Arrays.GetData()[Handle.Index];
+			return TLearningArrayView<DimNum, ElementType>((ElementType*)View.Data, TLearningArrayShape<DimNum>(View.Shape));
+		}
+
+		/**
+		* Get the type id of an array in the map.
+		*/
+		int16 TypeId(const FArrayMapKey Key) const;
+
+		/**
+		* Get the type id of an array in the map.
+		*/
+		int16 TypeId(const FArrayMapHandle Handle) const;
+
+		/**
+		* Get the type name of an array in the map.
+		*/
+		const TCHAR* TypeName(const FArrayMapKey Key) const;
+
+		/**
+		* Get the type name of an array in the map.
+		*/
+		const TCHAR* TypeName(const FArrayMapHandle Handle) const;
+
+		/**
+		* Get the number of dimensions of an array in the map.
+		*/
+		uint8 DimNum(const FArrayMapKey Key) const;
+
+		/**
+		* Get the number of dimensions of an array in the map.
+		*/
+		uint8 DimNum(const FArrayMapHandle Handle) const;
+
+		/**
+		* Add a new array to the map
+		*
+		* @param Key		The array key
+		* @param Shape		The shape of the array
+		* @param Default	The default value of the array
+		* @returns			A handle to later retrieve a view of the array for reading or writing
+		*/
+		template<uint8 DimNum, typename ElementType>
+		TArrayMapHandle<DimNum, ElementType> Add(const FArrayMapKey Key, const TLearningArrayShape<DimNum> Shape, const ElementType Default = ElementType())
+		{
+			static_assert(DimNum <= MaxDimNum, "Cannot add an array of that many dimensions.");
+
+			UE_LEARNING_CHECKF(Arrays.Num() < INT16_MAX,
+				TEXT("Maximum number of arrays already registered"));
+
+			UE_LEARNING_CHECKF(!Contains(Key),
+				TEXT("There already exists an array with key { \"%s\", \"%s\" }"),
+				*Key.Namespace.ToString(), *Key.Variable.ToString());
+
+			TArrayMapHandle<DimNum, ElementType> Handle = { Arrays.Num() };
+			Handles.Emplace(Key, Handle);
+
+			/*
+			* The FArrayMap object contains a big buffer of raw data for storing scalar data. This is 
+			* so that when we are allocating arrays with a single element don't end up creating many small 
+			* allocations and indirections. Here we check if the array is a single item and if so we place 
+			* it in the big buffer of scalar data, otherwise we allocate some new memory for it as normal
+			* using Malloc.
+			*/
+
+			const int32 ScalarInsertionIndex = ArrayMapPrivate::RoundUpToMultiple(ScalarOffset, alignof(ElementType));
+
+			FDynamicArrayView ArrayView;
+			if (Shape.Total() == 1 && ScalarInsertionIndex + sizeof(ElementType) <= ScalarData.Num() && alignof(ElementType) <= MaxScalarAlignment)
+			{
+				// Place in Scalar Data
+				ArrayView.Data = &ScalarData[ScalarInsertionIndex];
+				Flags.Emplace(FlagConstructed);
+				ScalarOffset += sizeof(ElementType);
+			}
+			else
+			{
+				// Allocate new array data
+				ArrayView.Data = FMemory::Malloc(Shape.Total() * sizeof(ElementType), alignof(ElementType));
+				Flags.Emplace(FlagConstructed | FlagAllocated);
+			}
+
+			// Set Shape
+			for (uint8 DimIdx = 0; DimIdx < DimNum; DimIdx++)
+			{
+				ArrayView.Shape[DimIdx] = Shape[DimIdx];
+			}
+
+			// Construct elements for this type
+			DefaultConstructItems<ElementType, int32>((ElementType*)ArrayView.Data, Shape.Total());
+
+			// Add array type data to Map
+			Arrays.Emplace(ArrayView);
+			Destructors.Emplace(ArrayMapPrivate::DestructArrayItems<ElementType>);
+			DimNums.Emplace(DimNum);
+			TypeIds.Emplace(TTypeId<ElementType>::Value);
+			TypeNames.Emplace(CompileTimeTypeName<ElementType>());
+
+			// Set to the default value
+			Array::Set(View(Handle), Default);
+
+			return Handle;
+		}
+
+		/**
+		* Add a view to an existing array to the map. Does not allocate any new storage.
+		*
+		* @param Key		The array key
+		* @param View		View to an existing array
+		* @returns			A handle to later retrieve the value of the array for reading or writing
+		*/
+		template<uint8 DimNum, typename ElementType>
+		TArrayMapHandle<DimNum, ElementType> AddView(const FArrayMapKey Key, const TLearningArrayView<DimNum, ElementType> View)
+		{
+			static_assert(DimNum <= MaxDimNum, "Cannot add an array of that many dimensions");
+
+			UE_LEARNING_CHECKF(Arrays.Num() < INT16_MAX,
+				TEXT("Maximum number of arrays already registered"));
+
+			UE_LEARNING_CHECKF(!Contains(Key),
+				TEXT("There already exists an array with key { \"%s\", \"%s\" }"),
+				*Key.Namespace.ToString(), *Key.Variable.ToString());
+
+			TArrayMapHandle<DimNum, ElementType> Handle = { Arrays.Num() };
+			Handles.Emplace(Key, Handle);
+
+			// Since this is an existing array we just create a basic view of it
+			FDynamicArrayView ArrayView;
+			ArrayView.Data = View.GetData();
+			for (uint8 DimIdx = 0; DimIdx < DimNum; DimIdx++)
+			{
+				ArrayView.Shape[DimIdx] = View.Num(DimIdx);
+			}
+
+			// And push the array data as usual
+			Flags.Emplace(FlagNone);
+			Arrays.Emplace(ArrayView);
+			Destructors.Emplace(ArrayMapPrivate::DestructArrayItems<ElementType>);
+			DimNums.Emplace(DimNum);
+			TypeIds.Emplace(TTypeId<ElementType>::Value);
+			TypeNames.Emplace(CompileTimeTypeName<ElementType>());
+
+			return Handle;
+		}
+
+		/**
+		* Check if an array with the given name has already been added.
+		*/
+		bool Contains(const FArrayMapKey Key) const;
+
+		/**
+		* Lookup the array handle associated with a key.
+		*/
+		FArrayMapHandle Lookup(const FArrayMapKey Key) const;
+
+		/**
+		* Find the array handle associated with a key. Returns `nullptr` if Key is not in Map.
+		*/
+		const FArrayMapHandle* Find(const FArrayMapKey Key) const;
+
+		/**
+		* Clear all array data.
+		*/
+		void Empty();
+
+		/**
+		* Link one array to another such that they share the same data. 
+		* Concretely, this function will make asking for a View of the 
+		* Dst array return a view of the Src array instead.
+		*/
+		void Link(const FArrayMapKey Src, const FArrayMapKey Dst);
+
+		/**
+		* Link one array to another such that they share the same data.
+		* Concretely, this function will make asking for a View of the
+		* Dst array return a view of the Src array instead.
+		*/
+		void Link(const FArrayMapHandle Src, const FArrayMapHandle Dst);
+
+		/**
+		* Link one array to another such that they share the same data.
+		* Concretely, this function will make asking for a View of the
+		* Dst array return a view of the Src array instead.
+		*/
+		template<uint8 DimNum, typename ElementType>
+		void Link(const TArrayMapHandle<DimNum, ElementType> Src, const TArrayMapHandle<DimNum, ElementType> Dst)
+		{
+			Link((FArrayMapHandle)Src, (FArrayMapHandle)Dst);
+		}
+
+		/**
+		* Check if the given array has a link.
+		*/
+		bool HasLink(const FArrayMapKey Dst) const;
+
+		/**
+		* Check if the given array has a link.
+		*/
+		bool HasLink(const FArrayMapHandle Dst) const;
+
+		/**
+		* Check if the given array is linked to.
+		*/
+		bool IsLinkedTo(const FArrayMapKey Src) const;
+
+		/**
+		* Check if the given array is linked to.
+		*/
+		bool IsLinkedTo(const FArrayMapHandle Src) const;
+
+		/**
+		* Get all keys in the map.
+		*/
+		template<typename Allocator = FDefaultAllocator>
+		void GetKeys(TArray<FArrayMapKey, Allocator>& OutKeys) const
+		{
+			Handles.GetKeys(OutKeys);
+		}
+
+		/**
+		* Get the key associated with a particular handle
+		*/
+		const FArrayMapKey* FindKey(FArrayMapHandle Handle) const;
+
+	private:
+
+		static constexpr int32 MaxDimNum = 6;
+		static constexpr int32 MaxScalarAlignment = 64;
+		static constexpr int32 InlineArrayNum = 128;
+		static constexpr int32 InlineScalarBytesNum = 2048;
+
+		enum : uint8
+		{
+			FlagNone = 0,
+
+			// Array is allocated by this map, so must be deallocated
+			FlagAllocated = 1 << 0,
+
+			// Array is constructed by this map, so must be destructed
+			FlagConstructed = 1 << 1,
+
+			// Array is a link to another array
+			FlagLinked = 1 << 2,
+		};
+
+		struct FDynamicArrayView
+		{
+			void* Data = nullptr;
+			int32 Shape[MaxDimNum] = { 0 };
+		};
+
+		template<uint8 DimNum, typename ElementType>
+		void CheckArray(const int32 Index) const
+		{
+			UE_LEARNING_CHECKF(Index != INDEX_NONE, TEXT("Invalid Index!"));
+
+			UE_LEARNING_CHECKF(DimNums[Index] == DimNum,
+				TEXT("Asked for array of %i dimensions, but array has %i dimensions"),
+				DimNum, DimNums[Index]);
+
+			UE_LEARNING_CHECKF(TypeIds[Index] == TTypeId<ElementType>::Value,
+				TEXT("Asked for array using type %s, but array data is of type %s"),
+				CompileTimeTypeName<ElementType>(), TypeNames[Index]);
+		}
+
+		template<uint8 DimNum, typename ElementType>
+		void CheckArrayWithKey(const FArrayMapKey Key, const int32 Index) const
+		{
+			UE_LEARNING_CHECKF(Index != INDEX_NONE, TEXT("Invalid Index!"));
+
+			UE_LEARNING_CHECKF(DimNums[Index] == DimNum,
+				TEXT("Asked for array { \"%s\", \"%s\" } with %i dimensions, but array has %i dimensions"),
+				*Key.Namespace.ToString(), *Key.Variable.ToString(), DimNum, DimNums[Index]);
+
+			UE_LEARNING_CHECKF(TypeIds[Index] == TTypeId<ElementType>::Value,
+				TEXT("Asked for array { \"%s\", \"%s\" } using type %s, but array data is of type %s"),
+				*Key.Namespace.ToString(), *Key.Variable.ToString(), CompileTimeTypeName<ElementType>(), TypeNames[Index]);
+		}
+
+		// Array Data
+
+		TArray<FDynamicArrayView, TInlineAllocator<InlineArrayNum>> Arrays;
+
+		// Handle Map
+
+		TMap<FArrayMapKey, FArrayMapHandle, TInlineSetAllocator<InlineArrayNum>> Handles;
+
+		// Array Metadata
+
+		TArray<TFunctionRef<void(void* Elements, const int32 DimNum, const int32* Shape)>, TInlineAllocator<InlineArrayNum>> Destructors;
+		TArray<uint8, TInlineAllocator<InlineArrayNum>> DimNums;
+		TArray<uint8, TInlineAllocator<InlineArrayNum>> Flags;
+		TArray<uint16, TInlineAllocator<InlineArrayNum>> TypeIds;
+		TArray<const TCHAR*, TInlineAllocator<InlineArrayNum>> TypeNames;
+
+		// Scalar Data
+
+		int32 ScalarOffset = 0;
+		TStaticArray<uint8, InlineScalarBytesNum, MaxScalarAlignment> ScalarData;
+	};
+}
+
diff --git a/Engine/Plugins/Experimental/LearningAgents/Source/Learning/Public/LearningCompletion.h b/Engine/Plugins/Experimental/LearningAgents/Source/Learning/Public/LearningCompletion.h
new file mode 100644
--- /dev/null
+++ b/Engine/Plugins/Experimental/LearningAgents/Source/Learning/Public/LearningCompletion.h
@@ -0,0 +1,79 @@
+// Copyright Epic Games, Inc. All Rights Reserved.
+
+#pragma once
+
+#include "LearningArray.h"
+
+namespace UE::Learning
+{
+	/**
+	* Completion Mode
+	*/
+	enum class ECompletionMode : uint8
+	{
+		// Episode is still running
+		Running = 0,
+
+		// Episode ended early but was still in progress. Value function will be used to estimate final return.
+		Truncated = 1,
+
+		// Episode ended early and zero reward was expected for all future steps.
+		Terminated = 2,
+	};
+
+	/**
+	* This object stores essentially a list of instances that need to be reset.
+	* It can be filled based on the completion status or manually.
+	*/
+	struct LEARNING_API FResetInstanceBuffer
+	{
+		void Resize(const int32 InMaxInstanceNum);
+
+		void SetResetInstances(const FIndexSet Instances);
+		void SetResetInstancesFromCompletions(const TLearningArrayView<1, const ECompletionMode> Completions, const FIndexSet Instances);
+		const int32 GetResetInstanceNum() const;
+		const FIndexSet GetResetInstances() const;
+
+	private:
+
+		int32 MaxInstanceNum = 0;
+		int32 ResetInstanceNum = 0;
+		FIndexSet ResetInstancesSet;
+		TLearningArray<1, int32, TInlineAllocator<1>> ResetInstances;
+	};
+
+	namespace Completion
+	{
+		/**
+		* Converts a ECompletionMode into a string.
+		*/
+		LEARNING_API const TCHAR* CompletionModeString(const ECompletionMode Completion);
+
+		/**
+		* Takes the logical Or of completions. More specifically, if either completion is `Terminated` that 
+		* will be the result of this operator, otherwise either completion being `Truncated` takes priority 
+		* over either being `Running`. Put simply: Terminated > Truncated > Running
+		*/
+		LEARNING_API ECompletionMode Or(const ECompletionMode Lhs, const ECompletionMode Rhs);
+
+		/**
+		* Set completions for all instances whose episode has reached the max number of steps.
+		*
+		* @param InOutCompletions				Output buffer to write completions to
+		* @param EpisodeStepNums				Number of steps taken by each instance
+		* @param EpisodeMaxStepNum				Maximum number of allowed steps
+		* @param EpisodeEndCompletionMode		Completion mode to use for instances who reach the end of the episode
+		* @param Instances						Instances to process
+		*
+		*/
+		LEARNING_API void EvaluateEndOfEpisodeCompletions(
+			TLearningArrayView<1, ECompletionMode> InOutCompletions,
+			const TLearningArrayView<1, const int32> EpisodeStepNums,
+			const int32 EpisodeMaxStepNum,
+			const ECompletionMode EpisodeEndCompletionMode,
+			const FIndexSet Instances);
+
+	}
+}
+
+
diff --git a/Engine/Plugins/Experimental/LearningAgents/Source/Learning/Public/LearningCompletionObject.h b/Engine/Plugins/Experimental/LearningAgents/Source/Learning/Public/LearningCompletionObject.h
new file mode 100644
--- /dev/null
+++ b/Engine/Plugins/Experimental/LearningAgents/Source/Learning/Public/LearningCompletionObject.h
@@ -0,0 +1,194 @@
+// Copyright Epic Games, Inc. All Rights Reserved.
+
+#pragma once
+
+#include "LearningArrayMap.h"
+#include "LearningCompletion.h"
+
+#include "Math/Vector.h"
+#include "Templates/SharedPointer.h"
+
+namespace UE::Learning
+{
+	namespace Completion
+	{
+		static constexpr float DefaultThresholdPosition = 100.0f; // cm
+		static constexpr float DefaultThresholdVelocity = 200.0f; // cm/s
+		static constexpr float DefaultThresholdAngle = UE_PI / 2.0f; // rad
+		static constexpr float DefaultThresholdAngularVelocity = UE_PI; // rad/s
+		static constexpr float DefaultThresholdHeight = 0.0f; // cm
+	}
+
+	//------------------------------------------------------------------
+
+	/**
+	* Base class for an object which computes if an instance should be completion 
+	* from a set of other arrays. Here, all data is assumed to be stored in a 
+	* `FArrayMap` object to make the processing of multiple instances efficient.
+	*/
+	struct LEARNING_API FCompletionObject
+	{
+		FCompletionObject(
+			const FName& InIdentifier,
+			const TSharedRef<FArrayMap>& InInstanceData,
+			const int32 InMaxInstanceNum,
+			const ECompletionMode InCompletionMode = ECompletionMode::Terminated);
+
+		virtual ~FCompletionObject() {}
+
+		virtual void Evaluate(const FIndexSet Instances) = 0;
+
+		TLearningArrayView<1, ECompletionMode> CompletionBuffer();
+
+		TSharedRef<FArrayMap> InstanceData;
+		ECompletionMode CompletionMode;
+
+		TArrayMapHandle<1, ECompletionMode> CompletionHandle;
+	};
+
+	//------------------------------------------------------------------
+
+	/**
+	* Terminates if any of the provided completions terminate
+	*/
+	struct LEARNING_API FAnyCompletion : public FCompletionObject
+	{
+		FAnyCompletion(
+			const FName& InIdentifier,
+			const TLearningArrayView<1, const TSharedRef<FCompletionObject>> InCompletions,
+			const TSharedRef<FArrayMap>& InInstanceData,
+			const int32 InMaxInstanceNum);
+
+		virtual void Evaluate(const FIndexSet Instances) override final;
+
+		TLearningArray<1, TSharedRef<FCompletionObject>, TInlineAllocator<32>> Completions;
+	};
+
+	//------------------------------------------------------------------
+
+	/**
+	* Completion due to some condition being true
+	*/
+	struct LEARNING_API FConditionalCompletion : public FCompletionObject
+	{
+		FConditionalCompletion(
+			const FName& InIdentifier,
+			const TSharedRef<FArrayMap>& InInstanceData,
+			const int32 InMaxInstanceNum,
+			const ECompletionMode InCompletionMode = ECompletionMode::Terminated);
+
+		virtual void Evaluate(const FIndexSet Instances) override final;
+
+		TArrayMapHandle<1, bool> ConditionHandle;
+	};
+
+	//------------------------------------------------------------------
+
+	/**
+	* Completion due to the difference between two scalar positions exceeding a given threshold
+	*/
+	struct LEARNING_API FScalarPositionDifferenceCompletion : public FCompletionObject
+	{
+		FScalarPositionDifferenceCompletion(
+			const FName& InIdentifier,
+			const TSharedRef<FArrayMap>& InInstanceData,
+			const int32 InMaxInstanceNum,
+			const float InThreshold = Completion::DefaultThresholdPosition,
+			const ECompletionMode InCompletionMode = ECompletionMode::Terminated);
+
+		virtual void Evaluate(const FIndexSet Instances) override final;
+
+		TArrayMapHandle<1, float> Position0Handle;
+		TArrayMapHandle<1, float> Position1Handle;
+		TArrayMapHandle<1, float> ThresholdHandle;
+	};
+
+	/**
+	* Completion due to the difference between two planar positions exceeding a given threshold
+	*/
+	struct LEARNING_API FPlanarPositionDifferenceCompletion : public FCompletionObject
+	{
+		FPlanarPositionDifferenceCompletion(
+			const FName& InIdentifier,
+			const TSharedRef<FArrayMap>& InInstanceData,
+			const int32 InMaxInstanceNum,
+			const float InThreshold = Completion::DefaultThresholdPosition,
+			const ECompletionMode InCompletionMode = ECompletionMode::Terminated,
+			const FVector InAxis0 = FVector::ForwardVector,
+			const FVector InAxis1 = FVector::RightVector);
+
+		virtual void Evaluate(const FIndexSet Instances) override final;
+
+		FVector Axis0 = FVector::ForwardVector;
+		FVector Axis1 = FVector::RightVector;
+
+		TArrayMapHandle<1, FVector> Position0Handle;
+		TArrayMapHandle<1, FVector> Position1Handle;
+		TArrayMapHandle<1, float> ThresholdHandle;
+	};
+
+	//------------------------------------------------------------------
+
+	/**
+	* Completion due to the difference between two scalar velocities exceeding a given threshold
+	*/
+	struct LEARNING_API FScalarVelocityDifferenceCompletion : public FCompletionObject
+	{
+		FScalarVelocityDifferenceCompletion(
+			const FName& InIdentifier,
+			const TSharedRef<FArrayMap>& InInstanceData,
+			const int32 InMaxInstanceNum,
+			const float InThreshold = Completion::DefaultThresholdVelocity,
+			const ECompletionMode InCompletionMode = ECompletionMode::Terminated);
+
+		virtual void Evaluate(const FIndexSet Instances) override final;
+
+		TArrayMapHandle<1, float> Velocity0Handle;
+		TArrayMapHandle<1, float> Velocity1Handle;
+		TArrayMapHandle<1, float> ThresholdHandle;
+	};
+
+	//------------------------------------------------------------------
+
+	/**
+	* Completion due to the difference between two scalar rotations exceeding a given threshold
+	*/
+	struct LEARNING_API FScalarRotationDifferenceCompletion : public FCompletionObject
+	{
+		FScalarRotationDifferenceCompletion(
+			const FName& InIdentifier,
+			const TSharedRef<FArrayMap>& InInstanceData,
+			const int32 InMaxInstanceNum,
+			const float InThreshold = Completion::DefaultThresholdAngle,
+			const ECompletionMode InCompletionMode = ECompletionMode::Terminated);
+
+		virtual void Evaluate(const FIndexSet Instances) override final;
+
+		TArrayMapHandle<1, float> Rotation0Handle;
+		TArrayMapHandle<1, float> Rotation1Handle;
+		TArrayMapHandle<1, float> ThresholdHandle;
+	};
+
+	//------------------------------------------------------------------
+
+	/**
+	* Completion due to the difference between two scalar angular velocities exceeding a given threshold
+	*/
+	struct LEARNING_API FScalarAngularVelocityDifferenceCompletion : public FCompletionObject
+	{
+		FScalarAngularVelocityDifferenceCompletion(
+			const FName& InIdentifier,
+			const TSharedRef<FArrayMap>& InInstanceData,
+			const int32 InMaxInstanceNum,
+			const float InThreshold = Completion::DefaultThresholdAngularVelocity,
+			const ECompletionMode InCompletionMode = ECompletionMode::Terminated);
+
+		virtual void Evaluate(const FIndexSet Instances) override final;
+
+		TArrayMapHandle<1, float> AngularVelocity0Handle;
+		TArrayMapHandle<1, float> AngularVelocity1Handle;
+		TArrayMapHandle<1, float> ThresholdHandle;
+	};
+
+	//------------------------------------------------------------------
+}
diff --git a/Engine/Plugins/Experimental/LearningAgents/Source/Learning/Public/LearningFeatureObject.h b/Engine/Plugins/Experimental/LearningAgents/Source/Learning/Public/LearningFeatureObject.h
new file mode 100644
--- /dev/null
+++ b/Engine/Plugins/Experimental/LearningAgents/Source/Learning/Public/LearningFeatureObject.h
@@ -0,0 +1,990 @@
+// Copyright Epic Games, Inc. All Rights Reserved.
+
+#pragma once
+
+#include "LearningArrayMap.h"
+
+#include "Math/Quat.h"
+#include "Math/Vector.h"
+#include "Templates/SharedPointer.h"
+
+namespace UE::Learning
+{
+	namespace Feature
+	{
+		static constexpr float DefaultScale = 1.0f;
+		static constexpr float DefaultScaleTime = 1.0f; // s
+		static constexpr float DefaultScaleDirection = 1.0f;
+
+		static constexpr float DefaultScalePosition = 100.0f; // cm
+		static constexpr float DefaultScaleVelocity = 200.0f; // cm/s
+		static constexpr float DefaultScaleAcceleration = 400.0f; // cm/s^2
+
+		static constexpr float DefaultScaleAngle = UE_PI / 2.0f; // rad
+		static constexpr float DefaultScaleAngularVelocity = UE_PI; // rad/s
+		static constexpr float DefaultScaleAngularAcceleration = 2.0f * UE_PI; // rad/s
+
+		static constexpr float DefaultScaleScale = 1.0f;
+
+		namespace Encode
+		{
+			LEARNING_API void Float(
+				float& Output,
+				const float Input,
+				const float Scale = DefaultScale,
+				const float Epsilon = UE_SMALL_NUMBER);
+
+
+			LEARNING_API void Time(
+				float& Output,
+				const float Time,
+				const float RelativeTime = 0.0f,
+				const float Scale = DefaultScaleTime,
+				const float Epsilon = UE_SMALL_NUMBER);
+
+
+			LEARNING_API void PlanarDirection(
+				TLearningArrayView<1, float> Output,
+				const FVector Direction,
+				const FQuat RelativeRotation = FQuat::Identity,
+				const float Scale = DefaultScaleDirection,
+				const FVector Axis0 = FVector::ForwardVector,
+				const FVector Axis1 = FVector::RightVector,
+				const float Epsilon = UE_SMALL_NUMBER);
+
+			LEARNING_API void Direction(
+				TLearningArrayView<1, float> Output,
+				const FVector Direction,
+				const FQuat RelativeRotation = FQuat::Identity,
+				const float Scale = DefaultScaleDirection,
+				const float Epsilon = UE_SMALL_NUMBER);
+
+
+			LEARNING_API void ScalarPosition(
+				float& Output,
+				const float Position,
+				const float RelativePosition = 0.0f,
+				const float Scale = DefaultScalePosition,
+				const float Epsilon = UE_SMALL_NUMBER);
+
+			LEARNING_API void PlanarPosition(
+				TLearningArrayView<1, float> Output,
+				const FVector Position,
+				const FVector RelativePosition = FVector::ZeroVector,
+				const FQuat RelativeRotation = FQuat::Identity,
+				const float Scale = DefaultScalePosition,
+				const FVector Axis0 = FVector::ForwardVector,
+				const FVector Axis1 = FVector::RightVector,
+				const float Epsilon = UE_SMALL_NUMBER);
+
+			LEARNING_API void Position(
+				TLearningArrayView<1, float> Output,
+				const FVector Position,
+				const FVector RelativePosition = FVector::ZeroVector,
+				const FQuat RelativeRotation = FQuat::Identity,
+				const float Scale = DefaultScalePosition,
+				const float Epsilon = UE_SMALL_NUMBER);
+
+
+			LEARNING_API void ScalarVelocity(
+				float& Output,
+				const float Velocity,
+				const float Scale = DefaultScaleVelocity,
+				const float Epsilon = UE_SMALL_NUMBER);
+
+			LEARNING_API void PlanarVelocity(
+				TLearningArrayView<1, float> Output,
+				const FVector Velocity,
+				const FQuat RelativeRotation = FQuat::Identity,
+				const float Scale = DefaultScaleVelocity,
+				const FVector Axis0 = FVector::ForwardVector,
+				const FVector Axis1 = FVector::RightVector,
+				const float Epsilon = UE_SMALL_NUMBER);
+
+			LEARNING_API void Velocity(
+				TLearningArrayView<1, float> Output,
+				const FVector Velocity,
+				const FQuat RelativeRotation = FQuat::Identity,
+				const float Scale = DefaultScaleVelocity,
+				const float Epsilon = UE_SMALL_NUMBER);
+
+
+			LEARNING_API void ScalarAcceleration(
+				float& Output,
+				const float Acceleration,
+				const float Scale = DefaultScaleAcceleration,
+				const float Epsilon = UE_SMALL_NUMBER);
+
+			LEARNING_API void PlanarAcceleration(
+				TLearningArrayView<1, float> Output,
+				const FVector Acceleration,
+				const FQuat RelativeRotation = FQuat::Identity,
+				const float Scale = DefaultScaleAcceleration,
+				const FVector Axis0 = FVector::ForwardVector,
+				const FVector Axis1 = FVector::RightVector,
+				const float Epsilon = UE_SMALL_NUMBER);
+
+			LEARNING_API void Acceleration(
+				TLearningArrayView<1, float> Output,
+				const FVector Acceleration,
+				const FQuat RelativeRotation = FQuat::Identity,
+				const float Scale = DefaultScaleAcceleration,
+				const float Epsilon = UE_SMALL_NUMBER);
+
+
+			LEARNING_API void Angle(
+				TLearningArrayView<1, float> Output,
+				const float Angle,
+				const float RelativeRotation = 0.0f,
+				const float Scale = DefaultScaleDirection, // angles are encoded as directions
+				const float Epsilon = UE_SMALL_NUMBER);
+
+			LEARNING_API void PlanarRotation(
+				TLearningArrayView<1, float> Output,
+				const FQuat Rotation,
+				const FQuat RelativeRotation = FQuat::Identity,
+				const float Scale = DefaultScaleDirection,  // rotations are encoded as directions
+				const FVector Axis = FVector::UpVector,
+				const float Epsilon = UE_SMALL_NUMBER);
+
+			LEARNING_API void Rotation(
+				TLearningArrayView<1, float> Output,
+				const FQuat Rotation,
+				const FQuat RelativeRotation = FQuat::Identity,
+				const float Scale = DefaultScaleDirection, // rotations are encoded as directions
+				const float Epsilon = UE_SMALL_NUMBER);
+
+			LEARNING_API void RotationVector(
+				TLearningArrayView<1, float> Output,
+				const FVector RotationVector,
+				const float Scale = DefaultScaleAngle,
+				const float Epsilon = UE_SMALL_NUMBER);
+
+
+			LEARNING_API void ScalarAngularVelocity(
+				float& Output,
+				const float AngularVelocity,
+				const float Scale = DefaultScaleAngularVelocity,
+				const float Epsilon = UE_SMALL_NUMBER);
+
+			LEARNING_API void PlanarAngularVelocity(
+				float& Output,
+				const FVector AngularVelocity,
+				const FQuat RelativeRotation = FQuat::Identity,
+				const float Scale = DefaultScaleAngularVelocity,
+				const FVector Axis = FVector::UpVector,
+				const float Epsilon = UE_SMALL_NUMBER);
+
+			LEARNING_API void AngularVelocity(
+				TLearningArrayView<1, float> Output,
+				const FVector AngularVelocity,
+				const FQuat RelativeRotation = FQuat::Identity,
+				const float Scale = DefaultScaleAngularVelocity,
+				const float Epsilon = UE_SMALL_NUMBER);
+
+
+			LEARNING_API void ScalarAngularAcceleration(
+				float& Output,
+				const float AngularAcceleration,
+				const float Scale = DefaultScaleAngularAcceleration,
+				const float Epsilon = UE_SMALL_NUMBER);
+
+			LEARNING_API void PlanarAngularAcceleration(
+				float& Output,
+				const FVector AngularAcceleration,
+				const FQuat RelativeRotation = FQuat::Identity,
+				const float Scale = DefaultScaleAngularAcceleration,
+				const FVector Axis = FVector::UpVector,
+				const float Epsilon = UE_SMALL_NUMBER);
+
+			LEARNING_API void AngularAcceleration(
+				TLearningArrayView<1, float> Output,
+				const FVector AngularAcceleration,
+				const FQuat RelativeRotation = FQuat::Identity,
+				const float Scale =DefaultScaleAngularAcceleration,
+				const float Epsilon = UE_SMALL_NUMBER);
+
+
+			LEARNING_API void Scale(
+				float& Output,
+				const float InScale,
+				const float RelativeScale = 1.0f,
+				const float Scale = DefaultScaleScale,
+				const float Epsilon = UE_SMALL_NUMBER);
+		}
+
+		namespace Decode
+		{
+			LEARNING_API void Float(
+				float& Value,
+				const TLearningArrayView<1, const float> Input,
+				const float Scale = DefaultScale);
+
+
+			LEARNING_API void Time(
+				float& OutTime,
+				const TLearningArrayView<1, const float> Input,
+				const float RelativeTime,
+				const float Scale = DefaultScaleTime);
+
+
+			LEARNING_API void PlanarDirection(
+				FVector& OutPlanarDirection,
+				const TLearningArrayView<1, const float> Input,
+				const FQuat RelativeRotation = FQuat::Identity,
+				const float Scale = DefaultScaleDirection,
+				const FVector Axis0 = FVector::ForwardVector,
+				const FVector Axis1 = FVector::RightVector,
+				const float Epsilon = UE_SMALL_NUMBER);
+
+			LEARNING_API void Direction(
+				FVector& OutDirection,
+				const TLearningArrayView<1, const float> Input,
+				const FQuat RelativeRotation = FQuat::Identity,
+				const float Scale = DefaultScaleDirection,
+				const float Epsilon = UE_SMALL_NUMBER);
+
+
+			LEARNING_API void ScalarPosition(
+				float& OutScalarPosition,
+				const float Input,
+				const float RelativePosition = 0.0f,
+				const float Scale = DefaultScalePosition);
+
+			LEARNING_API void PlanarPosition(
+				FVector& OutPlanarPosition,
+				const TLearningArrayView<1, const float> Input,
+				const FVector RelativePosition = FVector::ZeroVector,
+				const FQuat RelativeRotation = FQuat::Identity,
+				const float Scale = DefaultScalePosition,
+				const FVector Axis0 = FVector::ForwardVector,
+				const FVector Axis1 = FVector::RightVector);
+
+			LEARNING_API void Position(
+				FVector& OutPosition,
+				const TLearningArrayView<1, const float> Input,
+				const FVector RelativePosition = FVector::ZeroVector,
+				const FQuat RelativeRotation = FQuat::Identity,
+				const float Scale = DefaultScalePosition);
+
+
+			LEARNING_API void ScalarVelocity(
+				float& OutScalarVelocity,
+				const float Input,
+				const float Scale = DefaultScaleVelocity);
+
+			LEARNING_API void PlanarVelocity(
+				FVector& OutPlanarVelocity,
+				const TLearningArrayView<1, const float> Input,
+				const FQuat RelativeRotation = FQuat::Identity,
+				const float Scale = DefaultScaleVelocity,
+				const FVector Axis0 = FVector::ForwardVector,
+				const FVector Axis1 = FVector::RightVector);
+
+			LEARNING_API void Velocity(
+				FVector& OutVelocity,
+				const TLearningArrayView<1, const float> Input,
+				const FQuat RelativeRotation = FQuat::Identity,
+				const float Scale = DefaultScaleVelocity);
+
+
+			LEARNING_API void ScalarAcceleration(
+				float& OutScalarAcceleration,
+				const float Input,
+				const float Scale = DefaultScaleAcceleration);
+
+			LEARNING_API void PlanarAcceleration(
+				FVector& OutPlanarAcceleration,
+				const TLearningArrayView<1, const float> Input,
+				const FQuat RelativeRotation = FQuat::Identity,
+				const float Scale = DefaultScaleAcceleration,
+				const FVector Axis0 = FVector::ForwardVector,
+				const FVector Axis1 = FVector::RightVector);
+
+			LEARNING_API void Acceleration(
+				FVector& OutAcceleration,
+				const TLearningArrayView<1, const float> Input,
+				const FQuat RelativeRotation = FQuat::Identity,
+				const float Scale = DefaultScaleAcceleration);
+
+
+			LEARNING_API void Angle(
+				float& OutScalarRotation,
+				const TLearningArrayView<1, const float> Input,
+				const float RelativeRotation = 0.0f,
+				const float Scale = DefaultScaleDirection); // angles are encoded as directions
+
+			LEARNING_API void PlanarRotation(
+				FQuat& OutPlanarRotation,
+				const TLearningArrayView<1, const float> Input,
+				const FQuat RelativeRotation = FQuat::Identity,
+				const float Scale = DefaultScaleDirection, // rotations are encoded as directions
+				const FVector Axis = FVector::UpVector);
+
+			LEARNING_API void Rotation(
+				FQuat& OutRotation,
+				const TLearningArrayView<1, const float> Input,
+				const FQuat RelativeRotation = FQuat::Identity,
+				const float Scale = DefaultScaleDirection, // rotations are encoded as directions
+				const float Epsilon = UE_SMALL_NUMBER);
+
+			LEARNING_API void RotationVector(
+				FVector& OutRotationVector,
+				const TLearningArrayView<1, const float> Input,
+				const float Scale = DefaultScaleAngle);
+
+
+			LEARNING_API void ScalarAngularVelocity(
+				float& OutScalarVelocity,
+				const float Input,
+				const float Scale = DefaultScaleAngularVelocity);
+
+			LEARNING_API void PlanarAngularVelocity(
+				FVector& OutPlanarVelocity,
+				const float Input,
+				const FQuat RelativeRotation = FQuat::Identity,
+				const float Scale = DefaultScaleAngularVelocity,
+				const FVector Axis = FVector::UpVector);
+
+			LEARNING_API void AngularVelocity(
+				FVector& OutAngularVelocity,
+				const TLearningArrayView<1, const float> Input,
+				const FQuat RelativeRotation = FQuat::Identity,
+				const float Scale = DefaultScaleAngularVelocity);
+
+
+			LEARNING_API void ScalarAngularAcceleration(
+				float& OutScalarAngularAcceleration,
+				const float Input,
+				const float Scale = DefaultScaleAngularAcceleration);
+
+			LEARNING_API void PlanarAngularAcceleration(
+				FVector& OutPlanarAngularAcceleration,
+				const float Input,
+				const FQuat RelativeRotation = FQuat::Identity,
+				const float Scale = DefaultScaleAngularAcceleration,
+				const FVector Axis = FVector::UpVector);
+
+			LEARNING_API void AngularAcceleration(
+				FVector& OutAngularAcceleration,
+				const TLearningArrayView<1, const float> Input,
+				const FQuat RelativeRotation = FQuat::Identity,
+				const float Scale = DefaultScaleAngularAcceleration);
+
+
+			LEARNING_API void Scale(
+				float& OutScale,
+				const float Input,
+				const float RelativeScale = 1.0f,
+				const float Scale = DefaultScaleScale,
+				const float Epsilon = UE_SMALL_NUMBER);
+		}
+	}
+
+	/**
+	* Base class for an object which can do the encoding or decoding of
+	* features from a set of other arrays. Here, all data is assumed
+	* to be stored in a `FArrayMap` object to make the processing of 
+	* multiple instances efficient.
+	*/
+	struct LEARNING_API FFeatureObject
+	{
+		FFeatureObject(
+			const FName& InIdentifier,
+			const TSharedRef<FArrayMap>& InInstanceData,
+			const int32 InMaxInstanceNum,
+			const int32 InFeatureDimNum,
+			const float InScale = Feature::DefaultScale);
+
+		virtual ~FFeatureObject() {}
+
+		int32 DimNum() const;
+		TLearningArrayView<2, float> FeatureBuffer();
+
+		virtual bool IsEncodable() const = 0;
+		virtual bool IsDecodable() const = 0;
+		virtual void Encode(const FIndexSet Instances) = 0;
+		virtual void Decode(const FIndexSet Instances) = 0;
+
+		TSharedRef<FArrayMap> InstanceData;
+		int32 FeatureDimNum = 0;
+		float Scale = Feature::DefaultScale;
+		TArrayMapHandle<2, float> FeatureHandle;
+	};
+
+	//------------------------------------------------------------------
+
+	/**
+	* Feature that concatenates multiple other features
+	*/
+	struct LEARNING_API FConcatenateFeature : public FFeatureObject
+	{
+		FConcatenateFeature(
+			const FName& InIdentifier,
+			const TLearningArrayView<1, const TSharedRef<FFeatureObject>> InFeatures,
+			const TSharedRef<FArrayMap>& InInstanceData,
+			const int32 InMaxInstanceNum,
+			const float InScale = Feature::DefaultScale);
+
+		virtual bool IsEncodable() const override final;
+		virtual bool IsDecodable() const override final;
+		virtual void Encode(const FIndexSet Instances) override final;
+		virtual void Decode(const FIndexSet Instances) override final;
+
+		TLearningArray<1, TSharedRef<FFeatureObject>, TInlineAllocator<32>> Features;
+		TLearningArray<1, int32, TInlineAllocator<32>> Offsets;
+		TLearningArray<1, int32, TInlineAllocator<32>> Sizes;
+	};
+
+	//------------------------------------------------------------------
+
+	/**
+	* Feature for an array of floats
+	*/
+	struct LEARNING_API FFloatFeature : public FFeatureObject
+	{
+		FFloatFeature(
+			const FName& InIdentifier,
+			const TSharedRef<FArrayMap>& InInstanceData,
+			const int32 InMaxInstanceNum,
+			const int32 InFloatNum,
+			const float InScale = Feature::DefaultScale);
+
+		virtual bool IsEncodable() const override final { return true; }
+		virtual bool IsDecodable() const override final { return true; }
+		virtual void Encode(const FIndexSet Instances) override final;
+		virtual void Decode(const FIndexSet Instances) override final;
+
+		TArrayMapHandle<2, float> ValueHandle;
+	};
+
+	//------------------------------------------------------------------
+
+	/**
+	* Feature for an array of times
+	*/
+	struct LEARNING_API FTimeFeature : public FFeatureObject
+	{
+		FTimeFeature(
+			const FName& InIdentifier,
+			const TSharedRef<FArrayMap>& InInstanceData,
+			const int32 InMaxInstanceNum,
+			const int32 InTimeNum,
+			const float InScale = Feature::DefaultScaleTime);
+
+		virtual bool IsEncodable() const override final { return true; }
+		virtual bool IsDecodable() const override final { return true; }
+		virtual void Encode(const FIndexSet Instances) override final;
+		virtual void Decode(const FIndexSet Instances) override final;
+
+		TArrayMapHandle<2, float> TimeHandle;
+		TArrayMapHandle<1, float> RelativeTimeHandle;
+	};
+
+	//------------------------------------------------------------------
+
+	/**
+	* Feature for an array of planar directions
+	*/
+	struct LEARNING_API FPlanarDirectionFeature : public FFeatureObject
+	{
+		FPlanarDirectionFeature(
+			const FName& InIdentifier,
+			const TSharedRef<FArrayMap>& InInstanceData,
+			const int32 InMaxInstanceNum,
+			const int32 InDirectionNum,
+			const float InScale = Feature::DefaultScaleDirection,
+			const FVector InAxis0 = FVector::ForwardVector,
+			const FVector InAxis1 = FVector::RightVector);
+
+		virtual bool IsEncodable() const override final { return true; }
+		virtual bool IsDecodable() const override final { return true; }
+		virtual void Encode(const FIndexSet Instances) override final;
+		virtual void Decode(const FIndexSet Instances) override final;
+
+		FVector Axis0 = FVector::ForwardVector;
+		FVector Axis1 = FVector::RightVector;
+
+		TArrayMapHandle<2, FVector> DirectionHandle;
+		TArrayMapHandle<1, FQuat> RelativeRotationHandle;
+	};
+
+	/**
+	* Feature for an array of directions
+	*/
+	struct LEARNING_API FDirectionFeature : public FFeatureObject
+	{
+		FDirectionFeature(
+			const FName& InIdentifier,
+			const TSharedRef<FArrayMap>& InInstanceData,
+			const int32 InMaxInstanceNum,
+			const int32 InDirectionNum,
+			const float InScale = Feature::DefaultScaleDirection);
+
+		virtual bool IsEncodable() const override final { return true; }
+		virtual bool IsDecodable() const override final { return true; }
+		virtual void Encode(const FIndexSet Instances) override final;
+		virtual void Decode(const FIndexSet Instances) override final;
+
+		TArrayMapHandle<2, FVector> DirectionHandle;
+		TArrayMapHandle<1, FQuat> RelativeRotationHandle;
+	};
+
+	//------------------------------------------------------------------
+
+	/**
+	* Feature for an array of scalar positions
+	*/
+	struct LEARNING_API FScalarPositionFeature : public FFeatureObject
+	{
+		FScalarPositionFeature(
+			const FName& InIdentifier,
+			const TSharedRef<FArrayMap>& InInstanceData,
+			const int32 InMaxInstanceNum,
+			const int32 InPositionNum,
+			const float InScale = Feature::DefaultScalePosition);
+
+		virtual bool IsEncodable() const override final { return true; }
+		virtual bool IsDecodable() const override final { return true; }
+		virtual void Encode(const FIndexSet Instances) override final;
+		virtual void Decode(const FIndexSet Instances) override final;
+
+		TArrayMapHandle<2, float> PositionHandle;
+		TArrayMapHandle<1, float> RelativePositionHandle;
+	};
+
+	/**
+	* Feature for an array of planar positions
+	*/
+	struct LEARNING_API FPlanarPositionFeature : public FFeatureObject
+	{
+		FPlanarPositionFeature(
+			const FName& InIdentifier,
+			const TSharedRef<FArrayMap>& InInstanceData,
+			const int32 InMaxInstanceNum,
+			const int32 InPositionNum,
+			const float InScale = Feature::DefaultScalePosition,
+			const FVector InAxis0 = FVector::ForwardVector,
+			const FVector InAxis1 = FVector::RightVector);
+
+		virtual bool IsEncodable() const override final { return true; }
+		virtual bool IsDecodable() const override final { return true; }
+		virtual void Encode(const FIndexSet Instances) override final;
+		virtual void Decode(const FIndexSet Instances) override final;
+
+		FVector Axis0 = FVector::ForwardVector;
+		FVector Axis1 = FVector::RightVector;
+
+		TArrayMapHandle<2, FVector> PositionHandle;
+		TArrayMapHandle<1, FVector> RelativePositionHandle;
+		TArrayMapHandle<1, FQuat> RelativeRotationHandle;
+	};
+
+	/**
+	* Feature for an array of positions
+	*/
+	struct LEARNING_API FPositionFeature : public FFeatureObject
+	{
+		FPositionFeature(
+			const FName& InIdentifier,
+			const TSharedRef<FArrayMap>& InInstanceData,
+			const int32 InMaxInstanceNum,
+			const int32 InPositionNum,
+			const float InScale = Feature::DefaultScalePosition);
+
+		virtual bool IsEncodable() const override final { return true; }
+		virtual bool IsDecodable() const override final { return true; }
+		virtual void Encode(const FIndexSet Instances) override final;
+		virtual void Decode(const FIndexSet Instances) override final;
+
+		TArrayMapHandle<2, FVector> PositionHandle;
+		TArrayMapHandle<1, FVector> RelativePositionHandle;
+		TArrayMapHandle<1, FQuat> RelativeRotationHandle;
+	};
+
+	//------------------------------------------------------------------
+
+	/**
+	* Feature for an array of scalar velocities
+	*/
+	struct LEARNING_API FScalarVelocityFeature : public FFeatureObject
+	{
+		FScalarVelocityFeature(
+			const FName& InIdentifier,
+			const TSharedRef<FArrayMap>& InInstanceData,
+			const int32 InMaxInstanceNum,
+			const int32 InVelocityNum,
+			const float InScale = Feature::DefaultScaleVelocity);
+
+		virtual bool IsEncodable() const override final { return true; }
+		virtual bool IsDecodable() const override final { return true; }
+		virtual void Encode(const FIndexSet Instances) override final;
+		virtual void Decode(const FIndexSet Instances) override final;
+
+		TArrayMapHandle<2, float> VelocityHandle;
+	};
+
+	/**
+	* Feature for an array of planar velocities
+	*/
+	struct LEARNING_API FPlanarVelocityFeature : public FFeatureObject
+	{
+		FPlanarVelocityFeature(
+			const FName& InIdentifier,
+			const TSharedRef<FArrayMap>& InInstanceData,
+			const int32 InMaxInstanceNum,
+			const int32 InVelocityNum,
+			const float InScale = Feature::DefaultScaleVelocity,
+			const FVector InAxis0 = FVector::ForwardVector,
+			const FVector InAxis1 = FVector::RightVector);
+
+		virtual bool IsEncodable() const override final { return true; }
+		virtual bool IsDecodable() const override final { return true; }
+		virtual void Encode(const FIndexSet Instances) override final;
+		virtual void Decode(const FIndexSet Instances) override final;
+
+		FVector Axis0 = FVector::ForwardVector;
+		FVector Axis1 = FVector::RightVector;
+
+		TArrayMapHandle<2, FVector> VelocityHandle;
+		TArrayMapHandle<1, FQuat> RelativeRotationHandle;
+	};
+
+	/**
+	* Feature for an array of velocities
+	*/
+	struct LEARNING_API FVelocityFeature : public FFeatureObject
+	{
+		FVelocityFeature(
+			const FName& InIdentifier,
+			const TSharedRef<FArrayMap>& InInstanceData,
+			const int32 InMaxInstanceNum,
+			const int32 InVelocityNum,
+			const float InScale = Feature::DefaultScaleVelocity);
+
+		virtual bool IsEncodable() const override final { return true; }
+		virtual bool IsDecodable() const override final { return true; }
+		virtual void Encode(const FIndexSet Instances) override final;
+		virtual void Decode(const FIndexSet Instances) override final;
+
+		TArrayMapHandle<2, FVector> VelocityHandle;
+		TArrayMapHandle<1, FQuat> RelativeRotationHandle;
+	};
+
+	//------------------------------------------------------------------
+
+	/**
+	* Feature for an array of scalar accelerations
+	*/
+	struct LEARNING_API FScalarAccelerationFeature : public FFeatureObject
+	{
+		FScalarAccelerationFeature(
+			const FName& InIdentifier,
+			const TSharedRef<FArrayMap>& InInstanceData,
+			const int32 InMaxInstanceNum,
+			const int32 InAccelerationNum,
+			const float InScale = Feature::DefaultScaleAcceleration);
+
+		virtual bool IsEncodable() const override final { return true; }
+		virtual bool IsDecodable() const override final { return true; }
+		virtual void Encode(const FIndexSet Instances) override final;
+		virtual void Decode(const FIndexSet Instances) override final;
+
+		TArrayMapHandle<2, float> AccelerationHandle;
+	};
+
+	/**
+	* Feature for an array of planar accelerations
+	*/
+	struct LEARNING_API FPlanarAccelerationFeature : public FFeatureObject
+	{
+		FPlanarAccelerationFeature(
+			const FName& InIdentifier,
+			const TSharedRef<FArrayMap>& InInstanceData,
+			const int32 InMaxInstanceNum,
+			const int32 InAccelerationNum,
+			const float InScale = Feature::DefaultScaleAcceleration,
+			const FVector InAxis0 = FVector::ForwardVector,
+			const FVector InAxis1 = FVector::RightVector);
+
+		virtual bool IsEncodable() const override final { return true; }
+		virtual bool IsDecodable() const override final { return true; }
+		virtual void Encode(const FIndexSet Instances) override final;
+		virtual void Decode(const FIndexSet Instances) override final;
+
+		FVector Axis0 = FVector::ForwardVector;
+		FVector Axis1 = FVector::RightVector;
+
+		TArrayMapHandle<2, FVector> AccelerationHandle;
+		TArrayMapHandle<1, FQuat> RelativeRotationHandle;
+	};
+
+	/**
+	* Feature for an array of accelerations
+	*/
+	struct LEARNING_API FAccelerationFeature : public FFeatureObject
+	{
+		FAccelerationFeature(
+			const FName& InIdentifier,
+			const TSharedRef<FArrayMap>& InInstanceData,
+			const int32 InMaxInstanceNum,
+			const int32 InAccelerationNum,
+			const float InScale = Feature::DefaultScaleAcceleration);
+
+		virtual bool IsEncodable() const override final { return true; }
+		virtual bool IsDecodable() const override final { return true; }
+		virtual void Encode(const FIndexSet Instances) override final;
+		virtual void Decode(const FIndexSet Instances) override final;
+
+		TArrayMapHandle<2, FVector> AccelerationHandle;
+		TArrayMapHandle<1, FQuat> RelativeRotationHandle;
+	};
+
+	//------------------------------------------------------------------
+
+	/**
+	* Feature for an array of angles
+	*/
+	struct LEARNING_API FAngleFeature : public FFeatureObject
+	{
+		FAngleFeature(
+			const FName& InIdentifier,
+			const TSharedRef<FArrayMap>& InInstanceData,
+			const int32 InMaxInstanceNum,
+			const int32 InAngleNum,
+			const float InScale = Feature::DefaultScaleDirection); // angles are encoded as directions
+
+		virtual bool IsEncodable() const override final { return true; }
+		virtual bool IsDecodable() const override final { return true; }
+		virtual void Encode(const FIndexSet Instances) override final;
+		virtual void Decode(const FIndexSet Instances) override final;
+
+		TArrayMapHandle<2, float> AngleHandle;
+		TArrayMapHandle<1, float> RelativeAngleHandle;
+	};
+
+	/**
+	* Feature for an array of rotations
+	*/
+	struct LEARNING_API FPlanarRotationFeature : public FFeatureObject
+	{
+		FPlanarRotationFeature(
+			const FName& InIdentifier,
+			const TSharedRef<FArrayMap>& InInstanceData,
+			const int32 InMaxInstanceNum,
+			const int32 InRotationNum,
+			const float InScale = Feature::DefaultScaleDirection, // rotations are encoded as directions
+			const FVector Axis = FVector::UpVector);
+
+		virtual bool IsEncodable() const override final { return true; }
+		virtual bool IsDecodable() const override final { return true; }
+		virtual void Encode(const FIndexSet Instances) override final;
+		virtual void Decode(const FIndexSet Instances) override final;
+
+		FVector Axis = FVector::UpVector;
+
+		TArrayMapHandle<2, FQuat> RotationHandle;
+		TArrayMapHandle<1, FQuat> RelativeRotationHandle;
+	};
+
+	/**
+	* Feature for an array of rotations
+	*/
+	struct LEARNING_API FRotationFeature : public FFeatureObject
+	{
+		FRotationFeature(
+			const FName& InIdentifier,
+			const TSharedRef<FArrayMap>& InInstanceData,
+			const int32 InMaxInstanceNum,
+			const int32 InRotationNum,
+			const float InScale = Feature::DefaultScaleDirection); // rotations are encoded as directions
+
+		virtual bool IsEncodable() const override final { return true; }
+		virtual bool IsDecodable() const override final { return true; }
+		virtual void Encode(const FIndexSet Instances) override final;
+		virtual void Decode(const FIndexSet Instances) override final;
+
+		TArrayMapHandle<2, FQuat> RotationHandle;
+		TArrayMapHandle<1, FQuat> RelativeRotationHandle;
+	};
+
+	/**
+	* Feature for an array of rotation vectors
+	*/
+	struct LEARNING_API FRotationVectorFeature : public FFeatureObject
+	{
+		FRotationVectorFeature(
+			const FName& InIdentifier,
+			const TSharedRef<FArrayMap>& InInstanceData,
+			const int32 InMaxInstanceNum,
+			const int32 InRotationVectorNum,
+			const float InScale = Feature::DefaultScaleAngle);
+
+		virtual bool IsEncodable() const override final { return true; }
+		virtual bool IsDecodable() const override final { return true; }
+		virtual void Encode(const FIndexSet Instances) override final;
+		virtual void Decode(const FIndexSet Instances) override final;
+
+		TArrayMapHandle<2, FVector> RotationVectorsHandle;
+	};
+
+	//------------------------------------------------------------------
+
+	/**
+	* Feature for an array of scalar angular velocities
+	*/
+	struct LEARNING_API FScalarAngularVelocityFeature : public FFeatureObject
+	{
+		FScalarAngularVelocityFeature(
+			const FName& InIdentifier,
+			const TSharedRef<FArrayMap>& InInstanceData,
+			const int32 InMaxInstanceNum,
+			const int32 InAngularAccelerationNum,
+			const float InScale = Feature::DefaultScaleAngularVelocity);
+
+		virtual bool IsEncodable() const override final { return true; }
+		virtual bool IsDecodable() const override final { return true; }
+		virtual void Encode(const FIndexSet Instances) override final;
+		virtual void Decode(const FIndexSet Instances) override final;
+
+		TArrayMapHandle<2, float> AngularVelocityHandle;
+	};
+
+	/**
+	* Feature for an array of planar angular velocities
+	*/
+	struct LEARNING_API FPlanarAngularVelocityFeature : public FFeatureObject
+	{
+		FPlanarAngularVelocityFeature(
+			const FName& InIdentifier,
+			const TSharedRef<FArrayMap>& InInstanceData,
+			const int32 InMaxInstanceNum,
+			const int32 InAngularAccelerationNum,
+			const float InScale = Feature::DefaultScaleAngularVelocity,
+			const FVector InAxis = FVector::UpVector);
+
+		virtual bool IsEncodable() const override final { return true; }
+		virtual bool IsDecodable() const override final { return true; }
+		virtual void Encode(const FIndexSet Instances) override final;
+		virtual void Decode(const FIndexSet Instances) override final;
+
+		FVector Axis = FVector::UpVector;
+
+		TArrayMapHandle<2, FVector> AngularVelocityHandle;
+		TArrayMapHandle<1, FQuat> RelativeRotationHandle;
+	};
+
+	/**
+	* Feature for an array of angular velocities
+	*/
+	struct LEARNING_API FAngularVelocityFeature : public FFeatureObject
+	{
+		FAngularVelocityFeature(
+			const FName& InIdentifier,
+			const TSharedRef<FArrayMap>& InInstanceData,
+			const int32 InMaxInstanceNum,
+			const int32 InAngularAccelerationNum,
+			const float InScale = Feature::DefaultScaleAngularVelocity);
+
+		virtual bool IsEncodable() const override final { return true; }
+		virtual bool IsDecodable() const override final { return true; }
+		virtual void Encode(const FIndexSet Instances) override final;
+		virtual void Decode(const FIndexSet Instances) override final;
+
+		TArrayMapHandle<2, FVector> AngularVelocityHandle;
+		TArrayMapHandle<1, FQuat> RelativeRotationHandle;
+	};
+
+	//------------------------------------------------------------------
+
+	/**
+	* Feature for an array of scalar angular accelerations
+	*/
+	struct LEARNING_API FScalarAngularAccelerationFeature : public FFeatureObject
+	{
+		FScalarAngularAccelerationFeature(
+			const FName& InIdentifier,
+			const TSharedRef<FArrayMap>& InInstanceData,
+			const int32 InMaxInstanceNum,
+			const int32 InAngularAccelerationNum,
+			const float InScale = Feature::DefaultScaleAngularAcceleration);
+
+		virtual bool IsEncodable() const override final { return true; }
+		virtual bool IsDecodable() const override final { return true; }
+		virtual void Encode(const FIndexSet Instances) override final;
+		virtual void Decode(const FIndexSet Instances) override final;
+
+		TArrayMapHandle<2, float> AngularAccelerationHandle;
+	};
+
+	/**
+	* Feature for an array of planar angular accelerations
+	*/
+	struct LEARNING_API FPlanarAngularAccelerationFeature : public FFeatureObject
+	{
+		FPlanarAngularAccelerationFeature(
+			const FName& InIdentifier,
+			const TSharedRef<FArrayMap>& InInstanceData,
+			const int32 InMaxInstanceNum,
+			const int32 InAngularAccelerationNum,
+			const float InScale = Feature::DefaultScaleAngularAcceleration,
+			const FVector InAxis = FVector::UpVector);
+
+		virtual bool IsEncodable() const override final { return true; }
+		virtual bool IsDecodable() const override final { return true; }
+		virtual void Encode(const FIndexSet Instances) override final;
+		virtual void Decode(const FIndexSet Instances) override final;
+
+		FVector Axis = FVector::UpVector;
+
+		TArrayMapHandle<2, FVector> AngularAccelerationHandle;
+		TArrayMapHandle<1, FQuat> RelativeRotationHandle;
+	};
+
+	/**
+	* Feature for an array of angular accelerations
+	*/
+	struct LEARNING_API FAngularAccelerationFeature : public FFeatureObject
+	{
+		FAngularAccelerationFeature(
+			const FName& InIdentifier,
+			const TSharedRef<FArrayMap>& InInstanceData,
+			const int32 InMaxInstanceNum,
+			const int32 InAngularAccelerationNum,
+			const float InScale = Feature::DefaultScaleAngularAcceleration);
+
+		virtual bool IsEncodable() const override final { return true; }
+		virtual bool IsDecodable() const override final { return true; }
+		virtual void Encode(const FIndexSet Instances) override final;
+		virtual void Decode(const FIndexSet Instances) override final;
+
+		TArrayMapHandle<2, FVector> AngularAccelerationHandle;
+		TArrayMapHandle<1, FQuat> RelativeRotationHandle;
+	};
+
+	//------------------------------------------------------------------
+
+	/**
+	* Feature for an array of scales
+	*/
+	struct LEARNING_API FScaleFeature : public FFeatureObject
+	{
+		FScaleFeature(
+			const FName& InIdentifier,
+			const TSharedRef<FArrayMap>& InInstanceData,
+			const int32 InMaxInstanceNum,
+			const int32 InScaleNum,
+			const float InScale = Feature::DefaultScaleScale);
+
+		virtual bool IsEncodable() const override final { return true; }
+		virtual bool IsDecodable() const override final { return true; }
+		virtual void Encode(const FIndexSet Instances) override final;
+		virtual void Decode(const FIndexSet Instances) override final;
+
+		TArrayMapHandle<2, float> ScaleHandle;
+		TArrayMapHandle<1, float> RelativeScaleHandle;
+	};
+
+	//------------------------------------------------------------------
+
+}
diff --git a/Engine/Plugins/Experimental/LearningAgents/Source/Learning/Public/LearningFunctionObject.h b/Engine/Plugins/Experimental/LearningAgents/Source/Learning/Public/LearningFunctionObject.h
new file mode 100644
--- /dev/null
+++ b/Engine/Plugins/Experimental/LearningAgents/Source/Learning/Public/LearningFunctionObject.h
@@ -0,0 +1,179 @@
+// Copyright Epic Games, Inc. All Rights Reserved.
+
+#pragma once
+
+#include "LearningArrayMap.h"
+
+#include "Math/Quat.h"
+#include "Math/Transform.h"
+#include "Math/Vector.h"
+#include "Templates/SharedPointer.h"
+
+namespace UE::Learning
+{
+	/**
+	* Base class for an object which outputs to some arrays stored in
+	* a `FArrayMap`, using other arrays in the `FArrayMap` as input. 
+	* Here, all data is assumed to be stored in a `FArrayMap` object 
+	* to make the processing of multiple instances efficient.
+	*/
+	struct LEARNING_API FFunctionObject
+	{
+		FFunctionObject(const TSharedRef<FArrayMap>& InInstanceData);
+		virtual ~FFunctionObject() {}
+
+		virtual void Evaluate(const FIndexSet Instances) = 0;
+
+		TSharedRef<FArrayMap> InstanceData;
+	};
+
+	/**
+	* Executes another set of function objects in sequence
+	*/
+	struct LEARNING_API FSequentialFunction : public FFunctionObject
+	{
+		FSequentialFunction(
+			const TLearningArrayView<1, const TSharedRef<FFunctionObject>> InFunctions,
+			const TSharedRef<FArrayMap>& InInstanceData);
+
+		virtual void Evaluate(const FIndexSet Instances) override final;
+
+		TLearningArray<1, TSharedRef<FFunctionObject>, TInlineAllocator<32>> Functions;
+	};
+
+	/**
+	* Copies vector data from one array to another
+	*/
+	struct LEARNING_API FCopyVectorsFunction : public FFunctionObject
+	{
+		FCopyVectorsFunction(
+			const FName& InIdentifier,
+			const TSharedRef<FArrayMap>& InInstanceData,
+			const int32 InMaxInstanceNum,
+			const int32 InDimensionsNum);
+
+		virtual void Evaluate(const FIndexSet Instances) override final;
+
+		TArrayMapHandle<2, float> InputHandle;
+		TArrayMapHandle<2, float> OutputHandle;
+	};
+
+	/**
+	* Computes rotations given transforms
+	*/
+	struct LEARNING_API FExtractRotationsFromTransformsFunction : public FFunctionObject
+	{
+		FExtractRotationsFromTransformsFunction(
+			const FName& InIdentifier,
+			const TSharedRef<FArrayMap>& InInstanceData,
+			const int32 InMaxInstanceNum);
+
+		virtual void Evaluate(const FIndexSet Instances) override final;
+
+		TArrayMapHandle<1, FTransform> TransformHandle;
+		TArrayMapHandle<1, FQuat> RotationHandle;
+	};
+
+	/**
+	* Computes positions and rotations given transforms
+	*/
+	struct LEARNING_API FExtractPositionsAndRotationsFromTransformsFunction : public FFunctionObject
+	{
+		FExtractPositionsAndRotationsFromTransformsFunction(
+			const FName& InIdentifier,
+			const TSharedRef<FArrayMap>& InInstanceData,
+			const int32 InMaxInstanceNum);
+
+		virtual void Evaluate(const FIndexSet Instances) override final;
+
+		TArrayMapHandle<1, FTransform> TransformHandle;
+		TArrayMapHandle<1, FQuat> RotationHandle;
+		TArrayMapHandle<1, FVector> PositionHandle;
+	};
+
+	/**
+	* Samples a scalar value from a uniform distribution
+	*/
+	struct LEARNING_API FRandomUniformFunction : public FFunctionObject
+	{
+		FRandomUniformFunction(
+			const FName& InIdentifier,
+			const TSharedRef<FArrayMap>& InInstanceData,
+			const int32 InMaxInstanceNum,
+			const uint32 InSeed,
+			const float InMin = 0.0f,
+			const float InMax = 1.0f);
+
+		virtual void Evaluate(const FIndexSet Instances) override final;
+
+		TArrayMapHandle<1, uint32> SeedHandle;
+		TArrayMapHandle<1, float> MinHandle;
+		TArrayMapHandle<1, float> MaxHandle;
+		TArrayMapHandle<1, float> ValueHandle;
+	};
+
+	/**
+	* Samples values from a planar clipped Gaussian distribution.
+	*/
+	struct LEARNING_API FRandomPlanarClippedGaussianFunction : public FFunctionObject
+	{
+		FRandomPlanarClippedGaussianFunction(
+			const FName& InIdentifier,
+			const TSharedRef<FArrayMap>& InInstanceData,
+			const int32 InMaxInstanceNum,
+			const uint32 InSeed,
+			const float InMean = 0.0f,
+			const float InStd = 1.0f,
+			const float InClip = 10.0f);
+
+		virtual void Evaluate(const FIndexSet Instances) override final;
+
+		TArrayMapHandle<1, uint32> SeedHandle;
+		TArrayMapHandle<1, float> MeanHandle;
+		TArrayMapHandle<1, float> StdHandle;
+		TArrayMapHandle<1, float> ClipHandle;
+		TArrayMapHandle<1, FVector> ValueHandle;
+	};
+
+	/**
+	* Samples planar directions.
+	*/
+	struct LEARNING_API FRandomPlanarDirectionFunction : public FFunctionObject
+	{
+		FRandomPlanarDirectionFunction(
+			const FName& InIdentifier,
+			const TSharedRef<FArrayMap>& InInstanceData,
+			const uint32 InSeed,
+			const int32 InMaxInstanceNum);
+
+		virtual void Evaluate(const FIndexSet Instances) override final;
+
+		TArrayMapHandle<1, uint32> SeedHandle;
+		TArrayMapHandle<1, FVector> DirectionHandle;
+	};
+
+	/**
+	* Samples planar directions with corresponding velocities.
+	*/
+	struct LEARNING_API FRandomPlanarDirectionVelocityFunction : public FFunctionObject
+	{
+		FRandomPlanarDirectionVelocityFunction(
+			const FName& InIdentifier,
+			const TSharedRef<FArrayMap>& InInstanceData,
+			const int32 InMaxInstanceNum,
+			const uint32 InSeed,
+			const float InVelocityScale = 1.0f);
+
+		virtual void Evaluate(const FIndexSet Instances) override final;
+
+		TArrayMapHandle<1, uint32> SeedHandle;
+		TArrayMapHandle<1, float> VelocityScaleHandle;
+		TArrayMapHandle<1, FVector> DirectionHandle;
+		TArrayMapHandle<1, FVector> VelocityHandle;
+	};
+
+}
+
+
+
+
diff --git a/Engine/Plugins/Experimental/LearningAgents/Source/Learning/Public/LearningNeuralNetwork.h b/Engine/Plugins/Experimental/LearningAgents/Source/Learning/Public/LearningNeuralNetwork.h
new file mode 100644
--- /dev/null
+++ b/Engine/Plugins/Experimental/LearningAgents/Source/Learning/Public/LearningNeuralNetwork.h
@@ -0,0 +1,50 @@
+// Copyright Epic Games, Inc. All Rights Reserved.
+
+#pragma once
+
+#include "LearningArray.h"
+
+namespace UE::Learning
+{
+	/**
+	* Activation Function for use in a Neural Network
+	*/
+	enum class EActivationFunction : uint8
+	{
+		// ReLU Activation - Fast to train and evaluate but occasionally causes gradient collapse and untrainable networks.
+		ReLU = 0,
+
+		// ELU Activation - Generally performs better than ReLU and is not prone to gradient collapse but slower to evaluate.
+		ELU = 1,
+
+		// TanH Activation - Smooth activation function that is slower to train and evaluate but sometimes more stable for certain tasks.
+		TanH = 2,
+	};
+
+	LEARNING_API const TCHAR* ActivationFunctionString(const EActivationFunction ActivationFunction);
+
+	/**
+	* Simple feed-forward network data.
+	*/
+	struct LEARNING_API FNeuralNetwork
+	{
+		void DeserializeFromBytes(const TLearningArrayView<1, const uint8> RawBytes);
+		void SerializeToBytes(TLearningArrayView<1, uint8> OutRawBytes) const;
+
+		void Resize(
+			const int32 InputNum,
+			const int32 OutputNum,
+			const int32 HiddenNum,
+			const int32 LayerNum);
+
+		int32 GetInputNum() const;
+		int32 GetOutputNum() const;
+		int32 GetLayerNum() const;
+		int32 GetHiddenNum() const;
+		int32 GetTotalByteNum() const;
+
+		EActivationFunction ActivationFunction = EActivationFunction::ELU;
+		TArray<TLearningArray<2, float>, TInlineAllocator<16>> Weights;
+		TArray<TLearningArray<1, float>, TInlineAllocator<16>> Biases;
+	};
+}
diff --git a/Engine/Plugins/Experimental/LearningAgents/Source/Learning/Public/LearningPolicyObject.h b/Engine/Plugins/Experimental/LearningAgents/Source/Learning/Public/LearningPolicyObject.h
new file mode 100644
--- /dev/null
+++ b/Engine/Plugins/Experimental/LearningAgents/Source/Learning/Public/LearningPolicyObject.h
@@ -0,0 +1,62 @@
+// Copyright Epic Games, Inc. All Rights Reserved.
+
+#pragma once
+
+#include "LearningArrayMap.h"
+#include "LearningFunctionObject.h"
+
+#include "Templates/SharedPointer.h"
+
+namespace UE::Learning
+{
+	struct FNeuralNetwork;
+
+	/**
+	* Settings object for a neural network based policy function
+	*/
+	struct LEARNING_API FNeuralNetworkPolicyFunctionSettings
+	{
+		// Minimum amount of action noise to allow
+		float ActionNoiseMin = 0.0f;
+
+		// Maximum amount of action noise to allow
+		float ActionNoiseMax = 0.0f;
+
+		// Overall scale of the action noise
+		float ActionNoiseScale = 1.0f;
+
+		// If to allow for multi-threaded evaluation
+		bool bParallelEvaluation = true;
+
+		// Minimum batch size to use for multi-threaded evaluation
+		uint16 MinParallelBatchSize = 16;
+	};
+
+	/**
+	* Neural-network based policy object. Stores various settings and intermediate
+	* storage required to evaluate the given network for the provided number of instances.
+	*/
+	struct LEARNING_API FNeuralNetworkPolicyFunction : public FFunctionObject
+	{
+		FNeuralNetworkPolicyFunction(
+			const FName& InIdentifier,
+			const TSharedRef<FArrayMap>& InInstanceData,
+			const int32 InMaxInstanceNum,
+			const TSharedRef<FNeuralNetwork>& InNeuralNetwork,
+			const uint32 InSeed,
+			const FNeuralNetworkPolicyFunctionSettings& InSettings = FNeuralNetworkPolicyFunctionSettings());
+
+		virtual void Evaluate(const FIndexSet Instances) override final;
+
+		TSharedRef<FNeuralNetwork> NeuralNetwork;
+		FNeuralNetworkPolicyFunctionSettings Settings;
+
+		TArray<TLearningArray<2, float, TInlineAllocator<128>>, TInlineAllocator<16>> Activations;
+
+		TArrayMapHandle<1, uint32> SeedHandle;
+		TArrayMapHandle<2, float> InputHandle;
+		TArrayMapHandle<2, float> OutputHandle;
+		TArrayMapHandle<1, float> ActionNoiseScaleHandle;
+	};
+
+}
diff --git a/Engine/Plugins/Experimental/LearningAgents/Source/Learning/Public/LearningProgress.h b/Engine/Plugins/Experimental/LearningAgents/Source/Learning/Public/LearningProgress.h
new file mode 100644
--- /dev/null
+++ b/Engine/Plugins/Experimental/LearningAgents/Source/Learning/Public/LearningProgress.h
@@ -0,0 +1,68 @@
+// Copyright Epic Games, Inc. All Rights Reserved.
+
+#pragma once
+
+#include "HAL/CriticalSection.h"
+#include "Internationalization/Text.h"
+
+namespace UE::Learning
+{
+	/**
+	* Simple thread-safe structure to record progress of some long-running computation.
+	*/
+	struct LEARNING_API FProgress
+	{
+		int32 GetProgress() const;
+		void SetProgress(const int32 InProgress);
+
+		void GetMessage(FText& OutMessage);
+		void SetMessage(const FText& InMessage);
+
+		void Decrement();
+		void Decrement(const int32 Num);
+		void Done();
+
+	private:
+		FRWLock Lock;
+		FText Message;
+		TAtomic<int32> Progress = 0;
+	};
+
+	/**
+	* Scoped read lock, which can optionally be null
+	*/
+	struct LEARNING_API FScopeNullableReadLock
+	{
+	public:
+
+		UE_NODISCARD_CTOR FScopeNullableReadLock(FRWLock* InLock);
+		~FScopeNullableReadLock();
+
+		FScopeNullableReadLock() = delete;
+		FScopeNullableReadLock(const FScopeNullableReadLock& InScopeLock) = delete;
+		FScopeNullableReadLock& operator=(FScopeNullableReadLock& InScopeLock) = delete;
+
+	private:
+		FRWLock* Lock = nullptr;
+	};
+
+	/**
+	* Scoped write lock, which can optionally be null
+	*/
+	struct LEARNING_API FScopeNullableWriteLock
+	{
+	public:
+
+		UE_NODISCARD_CTOR FScopeNullableWriteLock(FRWLock* InLock);
+		~FScopeNullableWriteLock();
+
+		FScopeNullableWriteLock() = delete;
+		FScopeNullableWriteLock(const FScopeNullableWriteLock& InScopeLock) = delete;
+		FScopeNullableWriteLock& operator=(FScopeNullableWriteLock& InScopeLock) = delete;
+
+	private:
+		FRWLock* Lock = nullptr;
+	};
+
+}
+
diff --git a/Engine/Plugins/Experimental/LearningAgents/Source/Learning/Public/LearningRecorderObject.h b/Engine/Plugins/Experimental/LearningAgents/Source/Learning/Public/LearningRecorderObject.h
new file mode 100644
--- /dev/null
+++ b/Engine/Plugins/Experimental/LearningAgents/Source/Learning/Public/LearningRecorderObject.h
@@ -0,0 +1,123 @@
+// Copyright Epic Games, Inc. All Rights Reserved.
+
+#pragma once
+
+#include "LearningArrayMap.h"
+
+#include "Templates/SharedPointer.h"
+
+namespace UE::Learning
+{
+	namespace RecorderObjectPrivate
+	{
+		template<uint8 DimNum, typename ElementType>
+		void Copy(
+			TLearningArrayView<1, int32> InOutStepNums,
+			const FArrayMap& SrcArrayMap,
+			const FArrayMap& LogArrayMap,
+			const FArrayMapHandle SrcHandle,
+			const FArrayMapHandle LogHandle,
+			const int32 MaxStepNum,
+			const FIndexSet Instances)
+		{
+			const TLearningArrayView<DimNum, const ElementType> SrcView = SrcArrayMap.ConstView<DimNum, ElementType>(SrcHandle);
+			TLearningArrayView<DimNum + 1, ElementType> LogView = LogArrayMap.View<DimNum + 1, ElementType>(LogHandle);
+
+			for (const int32 InstanceIdx : Instances)
+			{
+				const int32 StepNum = InOutStepNums[InstanceIdx];
+
+				if (StepNum < MaxStepNum)
+				{
+					if constexpr (DimNum == 1)
+					{
+						LogView[StepNum][InstanceIdx] = SrcView[InstanceIdx];
+					}
+					else
+					{
+						Array::Copy(LogView[StepNum][InstanceIdx], SrcView[InstanceIdx]);
+					}
+
+					InOutStepNums[InstanceIdx]++;
+				}
+			}
+		}
+	}
+
+	/**
+	* The `RecorderObject` can be used to record the values of certain arrays in an `FArrayMap`. 
+	* This can be useful for debugging since we can keep track of certain arrays over time and 
+	* compute statistics or visualize them.
+	* 
+	* This object assumes the first dimension on all arrays corresponds to `Instances`, as it works by
+	* only record updates and resets for certain `Instances` given by an `FIndexSet`. This makes it 
+	* easier to use in conjunction with the other objects that use `FArrayMap` such as `FRewardObject`.
+	*/
+	struct LEARNING_API FRecorderObject
+	{
+		FRecorderObject(
+			const int32 InMaxInstanceNum,
+			const int32 InMaxStepNum,
+			const TSharedRef<FArrayMap>& InSrcArrayMap,
+			const TSharedRef<FArrayMap>& InLogArrayMap);
+
+		~FRecorderObject();
+
+		/**
+		* Record an array from the provided `SrcArrayMap` in the `LogArrayMap`
+		*/
+		template<uint8 DimNum, typename ElementType>
+		void Add(const FArrayMapKey Key)
+		{
+			const FArrayMapHandle SrcHandle = SrcArrayMap->Lookup(Key);
+
+			const TLearningArrayView<DimNum, const ElementType> SrcView = SrcArrayMap->ConstView<DimNum, ElementType>(SrcHandle);
+			UE_LEARNING_ARRAY_SHAPE_CHECK(SrcView.template Num<0>() == MaxInstanceNum);
+
+			TMultiArrayShape<DimNum + 1> LogShape;
+			LogShape[0] = MaxStepNum;
+			for (int32 DimIdx = 0; DimIdx < DimNum; DimIdx++)
+			{
+				LogShape[1 + DimIdx] = SrcView.Num(DimIdx);
+			}
+
+			const FArrayMapHandle LogHandle = LogArrayMap->Add<DimNum + 1, ElementType>(Key, LogShape);
+
+			SrcHandles.Emplace(SrcHandle);
+			LogHandles.Emplace(LogHandle);
+			CopyFunctions.Emplace(RecorderObjectPrivate::Copy<DimNum, ElementType>);
+		}
+
+		TLearningArrayView<1, const int32> GetStepNums() const { return StepNums; }
+
+		/**
+		* Reset all recorded arrays
+		*/
+		void Reset(const FIndexSet Instances);
+
+		/**
+		* Update all recorded arrays (record their current value)
+		*/
+		void Update(const FIndexSet Instances);
+
+	private:
+
+		int32 MaxInstanceNum = 0;
+		int32 MaxStepNum = 0;
+		TSharedRef<FArrayMap> SrcArrayMap;
+		TSharedRef<FArrayMap> LogArrayMap;
+
+		TLearningArray<1, int32> StepNums;
+		TArray<FArrayMapHandle, TInlineAllocator<32>> SrcHandles;
+		TArray<FArrayMapHandle, TInlineAllocator<32>> LogHandles;
+
+		TArray<TFunctionRef<void(
+			TLearningArrayView<1, int32> InOutStepNums,
+			const FArrayMap& SrcArrayMap,
+			const FArrayMap& LogArrayMap,
+			const FArrayMapHandle SrcHandle,
+			const FArrayMapHandle LogHandle,
+			const int32 MaxStepNum,
+			const FIndexSet Instances)>, TInlineAllocator<32>> CopyFunctions;
+	};
+}
diff --git a/Engine/Plugins/Experimental/LearningAgents/Source/Learning/Public/LearningRewardObject.h b/Engine/Plugins/Experimental/LearningAgents/Source/Learning/Public/LearningRewardObject.h
new file mode 100644
--- /dev/null
+++ b/Engine/Plugins/Experimental/LearningAgents/Source/Learning/Public/LearningRewardObject.h
@@ -0,0 +1,605 @@
+// Copyright Epic Games, Inc. All Rights Reserved.
+
+#pragma once
+
+#include "LearningArrayMap.h"
+
+#include "Math/Quat.h"
+#include "Math/Vector.h"
+#include "Templates/SharedPointer.h"
+
+namespace UE::Learning
+{
+	namespace Reward
+	{
+		static constexpr float DefaultScale = 1.0f;
+		static constexpr float DefaultScalePosition = 100.0f; // cm
+		static constexpr float DefaultScaleVelocity = 200.0f; // cm/s
+		static constexpr float DefaultScaleAngle = UE_PI / 2.0f; // rad
+		static constexpr float DefaultScaleAngularVelocity = UE_PI; // rad/s
+		static constexpr float DefaultScaleSpringInertialization = 100.0f; // cm
+
+		LEARNING_API float DistanceToReward(const float Distance);
+		LEARNING_API float DistanceToPenalty(const float Distance);
+		LEARNING_API float SquaredDistanceToReward(const float SquaredDistance);
+		LEARNING_API float SquaredDistanceToPenalty(const float SquaredDistance);
+
+		// Reward Functions
+
+		LEARNING_API float ScalarSimilarityReward(
+			const float Scalar0,
+			const float Scalar1,
+			const float Scale = DefaultScale,
+			const float Threshold = 0.0f,
+			const float Weight = 1.0f,
+			const float Epsilon = UE_SMALL_NUMBER);
+
+		LEARNING_API float ScalarPositionSimilarityReward(
+			const float Position0,
+			const float Position1,
+			const float Scale = DefaultScalePosition,
+			const float Threshold = 0.0f,
+			const float Weight = 1.0f,
+			const float Epsilon = UE_SMALL_NUMBER);
+
+		LEARNING_API float PositionSimilarityReward(
+			const FVector Position0,
+			const FVector Position1,
+			const float Scale = DefaultScalePosition,
+			const float Threshold = 0.0f,
+			const float Weight = 1.0f,
+			const float Epsilon = UE_SMALL_NUMBER);
+
+		LEARNING_API float ScalarRotationSimilarityReward(
+			const float Rotation0,
+			const float Rotation1,
+			const float Scale = DefaultScaleAngle,
+			const float Threshold = 0.0f,
+			const float Weight = 1.0f,
+			const float Epsilon = UE_SMALL_NUMBER);
+
+		LEARNING_API float ScalarAngularVelocitySimilarityReward(
+			const float AngularVelocity0,
+			const float AngularVelocity1,
+			const float Scale = DefaultScaleAngularVelocity,
+			const float Threshold = 0.0f,
+			const float Weight = 1.0f,
+			const float Epsilon = UE_SMALL_NUMBER);
+
+		LEARNING_API float ScalarVelocityReward(
+			const float Velocity,
+			const float Scale = DefaultScaleVelocity,
+			const float Weight = 1.0f,
+			const float Epsilon = UE_SMALL_NUMBER);
+
+		LEARNING_API float LocalDirectionalVelocityReward(
+			const FVector Velocity,
+			const FQuat RelativeRotation,
+			const float Scale = DefaultScaleVelocity,
+			const float Weight = 1.0f,
+			const FVector Axis = FVector::ForwardVector,
+			const float Epsilon = UE_SMALL_NUMBER);
+
+		// Penalty Functions
+
+		LEARNING_API float ScalarDifferencePenalty(
+			const float Scalar0,
+			const float Scalar1,
+			const float Scale = DefaultScale,
+			const float Threshold = 0.0f,
+			const float Weight = 1.0f,
+			const float Epsilon = UE_SMALL_NUMBER);
+
+		LEARNING_API float ScalarMagnitudePenalty(
+			const float Scalar,
+			const float Scale = DefaultScale,
+			const float Threshold = 0.0f,
+			const float Weight = 1.0f,
+			const float Epsilon = UE_SMALL_NUMBER);
+
+		LEARNING_API float DirectionAngleDifferencePenalty(
+			const FVector Direction0,
+			const FVector Direction1,
+			const float Scale = DefaultScaleAngle, // difference between directions is an angle
+			const float Threshold = 0.0f,
+			const float Weight = 1.0f,
+			const float Epsilon = UE_SMALL_NUMBER);
+
+		LEARNING_API float PlanarPositionDifferencePenalty(
+			const FVector Position0,
+			const FVector Position1,
+			const float Scale = DefaultScalePosition,
+			const float Threshold = 0.0f,
+			const float Weight = 1.0f,
+			const FVector Axis0 = FVector::ForwardVector,
+			const FVector Axis1 = FVector::UpVector,
+			const float Epsilon = UE_SMALL_NUMBER);
+
+		LEARNING_API float PositionDifferencePenalty(
+			const FVector Position0,
+			const FVector Position1,
+			const float Scale = DefaultScalePosition,
+			const float Threshold = 0.0f,
+			const float Weight = 1.0f,
+			const float Epsilon = UE_SMALL_NUMBER);
+
+		LEARNING_API float VelocityDifferencePenalty(
+			const FVector Velocity0,
+			const FVector Velocity1,
+			const float Scale = DefaultScaleVelocity,
+			const float Threshold = 0.0f,
+			const float Weight = 1.0f,
+			const float Epsilon = UE_SMALL_NUMBER);
+
+		LEARNING_API float SpringInertializationPenalty(
+			const FVector Position0,
+			const FVector Position1,
+			const FVector Velocity0,
+			const FVector Velocity1,
+			const float HalfLife,
+			const float Scale = DefaultScaleSpringInertialization,
+			const float Threshold = 0.0f,
+			const float Weight = 1.0f,
+			const float Epsilon = UE_SMALL_NUMBER);
+	}
+
+	//------------------------------------------------------------------
+
+	/**
+	* Base class for an object which computes a reward or penalty value 
+	* from a set of other arrays. Here, all data is assumed to be stored 
+	* in a `FArrayMap` object to make the processing of multiple instances efficient.
+	*/
+	struct LEARNING_API FRewardObject
+	{
+		FRewardObject(
+			const FName& InIdentifier,
+			const TSharedRef<FArrayMap>& InInstanceData,
+			const int32 InMaxInstanceNum);
+
+		virtual ~FRewardObject() {}
+
+		virtual void Evaluate(const FIndexSet Instances) = 0;
+
+		TLearningArrayView<1, float> RewardBuffer();
+
+		TSharedRef<FArrayMap> InstanceData;
+		TArrayMapHandle<1, float> RewardHandle;
+	};
+
+	//------------------------------------------------------------------
+
+	/**
+	* Reward that sums the results of multiple other reward functions
+	*/
+	struct LEARNING_API FSumReward : public FRewardObject
+	{
+		FSumReward(
+			const FName& InIdentifier,
+			const TLearningArrayView<1, const TSharedRef<FRewardObject>> InRewards,
+			const TSharedRef<FArrayMap>& InInstanceData,
+			const int32 InMaxInstanceNum);
+
+		virtual void Evaluate(const FIndexSet Instances) override final;
+
+		TLearningArray<1, TSharedRef<FRewardObject>, TInlineAllocator<32>> Rewards;
+	};
+
+	//------------------------------------------------------------------
+
+	/**
+	* Constant Reward
+	*/
+	struct LEARNING_API FConstantReward : public FRewardObject
+	{
+		FConstantReward(
+			const FName& InIdentifier,
+			const TSharedRef<FArrayMap>& InInstanceData,
+			const int32 InMaxInstanceNum,
+			const float InValue = 1.0f);
+
+		virtual void Evaluate(const FIndexSet Instances) override final;
+
+		float Value = 1.0f;
+	};
+
+	/**
+	* Constant Penalty
+	*/
+	struct LEARNING_API FConstantPenalty : public FRewardObject
+	{
+		FConstantPenalty(
+			const FName& InIdentifier,
+			const TSharedRef<FArrayMap>& InInstanceData,
+			const int32 InMaxInstanceNum,
+			const float InValue = 1.0f);
+
+		virtual void Evaluate(const FIndexSet Instances) override final;
+
+		float Value = 1.0f;
+	};
+
+	/**
+	* Conditional Constant Reward
+	*/
+	struct LEARNING_API FConditionalConstantReward : public FRewardObject
+	{
+		FConditionalConstantReward(
+			const FName& InIdentifier,
+			const TSharedRef<FArrayMap>& InInstanceData,
+			const int32 InMaxInstanceNum,
+			const float InValue = 1.0f);
+
+		virtual void Evaluate(const FIndexSet Instances) override final;
+
+		float Value = 1.0f;
+
+		TArrayMapHandle<1, bool> ConditionHandle;
+	};
+
+	//------------------------------------------------------------------
+
+	/**
+	* Basic Float Reward
+	*/
+	struct LEARNING_API FFloatReward : public FRewardObject
+	{
+		FFloatReward(
+			const FName& InIdentifier,
+			const TSharedRef<FArrayMap>& InInstanceData,
+			const int32 InMaxInstanceNum,
+			const float InWeight = 1.0f);
+
+		virtual void Evaluate(const FIndexSet Instances) override final;
+
+		TArrayMapHandle<1, float> ValueHandle;
+		TArrayMapHandle<1, float> WeightHandle;
+	};
+
+	//------------------------------------------------------------------
+
+	/**
+	* Penalty based on the average absolute magnitude of each dimension of a vector
+	*/
+	struct LEARNING_API FVectorAverageMagnitudePenalty : public FRewardObject
+	{
+		FVectorAverageMagnitudePenalty(
+			const FName& InIdentifier,
+			const TSharedRef<FArrayMap>& InInstanceData,
+			const int32 InMaxInstanceNum,
+			const int32 InDimNum,
+			const float InWeight = 1.0f,
+			const float InScale = Reward::DefaultScale,
+			const float InThreshold = 0.0f);
+
+		virtual void Evaluate(const FIndexSet Instances) override final;
+
+		int32 DimNum = 0;
+
+		TArrayMapHandle<2, float> ValueHandle;
+
+		TArrayMapHandle<1, float> WeightHandle;
+		TArrayMapHandle<1, float> ScaleHandle;
+		TArrayMapHandle<1, float> ThresholdHandle;
+	};
+
+	//------------------------------------------------------------------
+
+	/**
+	* Reward based on the similarity between two scalar position
+	*/
+	struct LEARNING_API FScalarPositionSimilarityReward : public FRewardObject
+	{
+		FScalarPositionSimilarityReward(
+			const FName& InIdentifier,
+			const TSharedRef<FArrayMap>& InInstanceData,
+			const int32 InMaxInstanceNum,
+			const float InWeight = 1.0f,
+			const float InScale = Reward::DefaultScalePosition,
+			const float InThreshold = 0.0f);
+
+		virtual void Evaluate(const FIndexSet Instances) override final;
+
+		TArrayMapHandle<1, float> Position0Handle;
+		TArrayMapHandle<1, float> Position1Handle;
+
+		TArrayMapHandle<1, float> WeightHandle;
+		TArrayMapHandle<1, float> ScaleHandle;
+		TArrayMapHandle<1, float> ThresholdHandle;
+	};
+
+	//------------------------------------------------------------------
+
+	/**
+	* Reward based on the similarity between two rotations
+	*/
+	struct LEARNING_API FScalarRotationSimilarityReward : public FRewardObject
+	{
+		FScalarRotationSimilarityReward(
+			const FName& InIdentifier,
+			const TSharedRef<FArrayMap>& InInstanceData,
+			const int32 InMaxInstanceNum,
+			const float InWeight = 1.0f,
+			const float InScale = Reward::DefaultScaleAngle,
+			const float InThreshold = 0.0f);
+
+		virtual void Evaluate(const FIndexSet Instances) override final;
+
+		TArrayMapHandle<1, float> Rotation0Handle;
+		TArrayMapHandle<1, float> Rotation1Handle;
+
+		TArrayMapHandle<1, float> WeightHandle;
+		TArrayMapHandle<1, float> ScaleHandle;
+		TArrayMapHandle<1, float> ThresholdHandle;
+	};
+
+	//------------------------------------------------------------------
+
+	/**
+	* Reward based on the similarity between two scalar angular velocities
+	*/
+	struct LEARNING_API FScalarAngularVelocitySimilarityReward : public FRewardObject
+	{
+		FScalarAngularVelocitySimilarityReward(
+			const FName& InIdentifier,
+			const TSharedRef<FArrayMap>& InInstanceData,
+			const int32 InMaxInstanceNum,
+			const float InWeight = 1.0f,
+			const float InScale = Reward::DefaultScaleAngularVelocity,
+			const float InThreshold = 0.0f);
+
+		virtual void Evaluate(const FIndexSet Instances) override final;
+
+		TArrayMapHandle<1, float> AngularVelocity0Handle;
+		TArrayMapHandle<1, float> AngularVelocity1Handle;
+
+		TArrayMapHandle<1, float> WeightHandle;
+		TArrayMapHandle<1, float> ScaleHandle;
+		TArrayMapHandle<1, float> ThresholdHandle;
+	};
+
+	//------------------------------------------------------------------
+
+	/**
+	* Reward based on how large (and positive) a scalar velocity is
+	*/
+	struct LEARNING_API FScalarVelocityReward : public FRewardObject
+	{
+		FScalarVelocityReward(
+			const FName& InIdentifier,
+			const TSharedRef<FArrayMap>& InInstanceData,
+			const int32 InMaxInstanceNum,
+			const float InWeight = 1.0f,
+			const float InScale = Reward::DefaultScaleVelocity);
+
+		virtual void Evaluate(const FIndexSet Instances) override final;
+
+		TArrayMapHandle<1, float> VelocityHandle;
+
+		TArrayMapHandle<1, float> WeightHandle;
+		TArrayMapHandle<1, float> ScaleHandle;
+	};
+
+	/**
+	* Reward based on how large a velocity is in a particular local direction
+	*/
+	struct LEARNING_API FLocalDirectionalVelocityReward : public FRewardObject
+	{
+		FLocalDirectionalVelocityReward(
+			const FName& InIdentifier,
+			const TSharedRef<FArrayMap>& InInstanceData,
+			const int32 InMaxInstanceNum,
+			const float InWeight = 1.0f,
+			const float InScale = Reward::DefaultScaleVelocity,
+			const FVector InAxis = FVector::ForwardVector);
+
+		virtual void Evaluate(const FIndexSet Instances) override final;
+
+		FVector Axis = FVector::ForwardVector;
+
+		TArrayMapHandle<1, FVector> VelocityHandle;
+		TArrayMapHandle<1, FQuat> RelativeRotationHandle;
+
+		TArrayMapHandle<1, float> WeightHandle;
+		TArrayMapHandle<1, float> ScaleHandle;
+	};
+
+	//------------------------------------------------------------------
+
+	/**
+	* Penalty based on the difference between two planar positions
+	*/
+	struct LEARNING_API FPlanarPositionDifferencePenalty : public FRewardObject
+	{
+		FPlanarPositionDifferencePenalty(
+			const FName& InIdentifier,
+			const TSharedRef<FArrayMap>& InInstanceData,
+			const int32 InMaxInstanceNum,
+			const float InWeight = 1.0f,
+			const float InScale = Reward::DefaultScalePosition,
+			const float InThreshold = 0.0f,
+			const FVector InAxis0 = FVector::ForwardVector,
+			const FVector InAxis1 = FVector::RightVector);
+
+		virtual void Evaluate(const FIndexSet Instances) override final;
+
+		FVector Axis0 = FVector::ForwardVector;
+		FVector Axis1 = FVector::RightVector;
+
+		TArrayMapHandle<1, FVector> Position0Handle;
+		TArrayMapHandle<1, FVector> Position1Handle;
+
+		TArrayMapHandle<1, float> WeightHandle;
+		TArrayMapHandle<1, float> ScaleHandle;
+		TArrayMapHandle<1, float> ThresholdHandle;
+	};
+
+	/**
+	* Penalty based on the difference between two positions
+	*/
+	struct LEARNING_API FPositionDifferencePenalty : public FRewardObject
+	{
+		FPositionDifferencePenalty(
+			const FName& InIdentifier,
+			const TSharedRef<FArrayMap>& InInstanceData,
+			const int32 InMaxInstanceNum,
+			const float InWeight = 1.0f,
+			const float InScale = Reward::DefaultScalePosition,
+			const float InThreshold = 0.0f);
+
+		virtual void Evaluate(const FIndexSet Instances) override final;
+
+		TArrayMapHandle<1, FVector> Position0Handle;
+		TArrayMapHandle<1, FVector> Position1Handle;
+
+		TArrayMapHandle<1, float> WeightHandle;
+		TArrayMapHandle<1, float> ScaleHandle;
+		TArrayMapHandle<1, float> ThresholdHandle;
+	};
+
+	//------------------------------------------------------------------
+
+	/**
+	* Penalty based on the difference between two velocities
+	*/
+	struct LEARNING_API FVelocityDifferencePenalty : public FRewardObject
+	{
+		FVelocityDifferencePenalty(
+			const FName& InIdentifier,
+			const TSharedRef<FArrayMap>& InInstanceData,
+			const int32 InMaxInstanceNum,
+			const float InWeight = 1.0f,
+			const float InScale = Reward::DefaultScaleVelocity,
+			const float InThreshold = 0.0f);
+
+		virtual void Evaluate(const FIndexSet Instances) override final;
+
+		TArrayMapHandle<1, FVector> Velocity0Handle;
+		TArrayMapHandle<1, FVector> Velocity1Handle;
+
+		TArrayMapHandle<1, float> WeightHandle;
+		TArrayMapHandle<1, float> ScaleHandle;
+		TArrayMapHandle<1, float> ThresholdHandle;
+	};
+
+	//------------------------------------------------------------------
+
+	/**
+	* Penalty based on the difference between two rotations facing directions
+	*/
+	struct LEARNING_API FFacingDirectionAngularDifferencePenalty : public FRewardObject
+	{
+		FFacingDirectionAngularDifferencePenalty(
+			const FName& InIdentifier,
+			const TSharedRef<FArrayMap>& InInstanceData,
+			const int32 InMaxInstanceNum,
+			const FVector InFacingDirection0 = FVector::ForwardVector,
+			const FVector InFacingDirection1 = FVector::ForwardVector,
+			const float InWeight = 1.0f,
+			const float InScale = Reward::DefaultScaleAngle, // difference between directions is an angle
+			const float InThreshold = 0.0f);
+
+		virtual void Evaluate(const FIndexSet Instances) override final;
+
+		TArrayMapHandle<1, FQuat> Rotation0Handle;
+		TArrayMapHandle<1, FQuat> Rotation1Handle;
+		TArrayMapHandle<1, FVector> FacingDirection0Handle;
+		TArrayMapHandle<1, FVector> FacingDirection1Handle;
+
+		TArrayMapHandle<1, float> WeightHandle;
+		TArrayMapHandle<1, float> ScaleHandle;
+		TArrayMapHandle<1, float> ThresholdHandle;
+	};
+
+	/**
+	* Penalty based on the difference between a rotation facing direction and a position
+	*/
+	struct LEARNING_API FFacingTowardsTargetAngularDifferencePenalty : public FRewardObject
+	{
+		FFacingTowardsTargetAngularDifferencePenalty(
+			const FName& InIdentifier,
+			const TSharedRef<FArrayMap>& InInstanceData,
+			const int32 InMaxInstanceNum,
+			const FVector InFacingDirection = FVector::ForwardVector,
+			const float InMinRadius = UE_KINDA_SMALL_NUMBER,
+			const float InWeight = 1.0f,
+			const float InScale = Reward::DefaultScaleAngle, // difference between directions is an angle
+			const float InThreshold = 0.0f); 
+
+		virtual void Evaluate(const FIndexSet Instances) override final;
+
+		TArrayMapHandle<1, FVector> PositionHandle;
+		TArrayMapHandle<1, FVector> TargetHandle;
+		TArrayMapHandle<1, FQuat> RotationHandle;
+		TArrayMapHandle<1, FVector> FacingDirectionHandle;
+		TArrayMapHandle<1, float> MinRadiusHandle;
+
+		TArrayMapHandle<1, float> WeightHandle;
+		TArrayMapHandle<1, float> ScaleHandle;
+		TArrayMapHandle<1, float> ThresholdHandle;
+	};
+
+	//------------------------------------------------------------------
+
+	/**
+	* Penalty based on the average spring inertialization displacement introduced when transitioning between two moving points
+	*
+	* See: https://theorangeduck.com/page/inertialization-transition-cost
+	*/
+	struct LEARNING_API FSpringInertializationPenalty : public FRewardObject
+	{
+		FSpringInertializationPenalty(
+			const FName& InIdentifier,
+			const TSharedRef<FArrayMap>& InInstanceData,
+			const int32 InMaxInstanceNum,
+			const float InWeight = 1.0f,
+			const float InScale = Reward::DefaultScaleSpringInertialization,
+			const float InThreshold = 0.0f);
+
+		virtual void Evaluate(const FIndexSet Instances) override final;
+
+		TArrayMapHandle<1, FVector> Position0Handle;
+		TArrayMapHandle<1, FVector> Position1Handle;
+		TArrayMapHandle<1, FVector> Velocity0Handle;
+		TArrayMapHandle<1, FVector> Velocity1Handle;
+
+		TArrayMapHandle<1, float> HalfLifeHandle;
+
+		TArrayMapHandle<1, float> WeightHandle;
+		TArrayMapHandle<1, float> ScaleHandle;
+		TArrayMapHandle<1, float> ThresholdHandle;
+	};
+
+	//------------------------------------------------------------------
+
+	/**
+	* Reward based on the similarity between two arrays of position
+	*/
+	struct LEARNING_API FPositionArraySimilarityReward : public FRewardObject
+	{
+		FPositionArraySimilarityReward(
+			const FName& InIdentifier,
+			const TSharedRef<FArrayMap>& InInstanceData,
+			const int32 InMaxInstanceNum,
+			const int32 InPositionNum,
+			const float InWeight = 1.0f,
+			const float InScale = Reward::DefaultScalePosition,
+			const float InThreshold = 0.0f);
+
+		virtual void Evaluate(const FIndexSet Instances) override final;
+
+		TArrayMapHandle<2, FVector> Positions0Handle;
+		TArrayMapHandle<2, FVector> Positions1Handle;
+		TArrayMapHandle<1, FVector> RelativePosition0Handle;
+		TArrayMapHandle<1, FVector> RelativePosition1Handle;
+		TArrayMapHandle<1, FQuat> RelativeRotation0Handle;
+		TArrayMapHandle<1, FQuat> RelativeRotation1Handle;
+
+		TArrayMapHandle<1, float> WeightHandle;
+		TArrayMapHandle<1, float> ScaleHandle;
+		TArrayMapHandle<1, float> ThresholdHandle;
+	};
+
+	//------------------------------------------------------------------
+
+}
