Learning Agents: Create initial experimental version

diff --git a/Engine/Plugins/Experimental/LearningAgents/Source/Learning/Public/MultiArray.h b/Engine/Plugins/Experimental/LearningAgents/Source/Learning/Public/MultiArray.h
new file mode 100644
--- /dev/null
+++ b/Engine/Plugins/Experimental/LearningAgents/Source/Learning/Public/MultiArray.h
@@ -0,0 +1,1529 @@
+// Copyright Epic Games, Inc. All Rights Reserved.
+
+#pragma once
+
+#include "CoreTypes.h"
+#include "Misc/AssertionMacros.h"
+#include "HAL/UnrealMemory.h"
+#include "Templates/IsSigned.h"
+#include "Templates/UnrealTypeTraits.h"
+#include "Templates/UnrealTemplate.h"
+#include "Containers/ContainerAllocationPolicies.h"
+#include "Containers/ContainerElementTypeCompatibility.h"
+
+#include "Templates/AndOrNot.h"
+#include "Templates/IdentityFunctor.h"
+#include "Templates/Invoke.h"
+#include "Templates/Less.h"
+#include "Templates/ChooseClass.h"
+#include "Templates/Sorting.h"
+#include "Templates/AlignmentTemplates.h"
+#include "Templates/IsConstructible.h"
+
+#include <type_traits>
+
+#include "MultiArrayView.h"
+
+// Disable static analysis warnings
+// PVS does not correctly infer the use of the template specialization for single dimension 
+// arrays and so throws up a bunch of warnings from assuming DimNum <= 1 in the generic template
+//-V::654,621,547
+
+namespace UE::MultiArray::Private
+{
+	// Copied from Core/Public/Containers/Array.h
+	//
+	// Assume elements are compatible with themselves - avoids problems with generated copy
+	// constructors of arrays of forwarded types, e.g.:
+	//
+	// struct FThing;
+	//
+	// struct FOuter
+	// {
+	//     TMultiArray<2, FThing> Arr; // this will cause errors without this workaround
+	// };
+	//
+	// This should be changed to use std::disjunction and std::is_constructible, and the usage
+	// changed to use ::value instead of ::Value, when std::disjunction (C++17) is available everywhere.
+	template <typename DestType, typename SourceType>
+	using TMultiArrayElementsAreCompatible = TOrValue<std::is_same<DestType, std::decay_t<SourceType>>::value, TIsConstructible<DestType, SourceType>>;
+}
+
+
+/**
+ * Templated dynamic multi-dimensional array
+ *
+ * A dynamically sized array of typed elements.  Makes the assumption that your elements are relocatable;
+ * i.e. that they can be transparently moved to new memory without a copy constructor. 
+ *
+ **/
+template<uint8 InDimNum, typename InElementType, typename InAllocatorType = FDefaultAllocator, bool bInIsChecked = true, bool bInIsRestrict = false>
+class TMultiArray
+{
+public:
+	static constexpr uint8 DimNum = InDimNum;
+	using ElementType = InElementType;
+	using AllocatorType = InAllocatorType;
+	static constexpr bool bIsChecked = bInIsChecked;
+	static constexpr bool bIsRestrict = bInIsRestrict;
+	using SizeType = typename InAllocatorType::SizeType;
+
+	using PointerType = typename TChooseClass<
+		bIsRestrict,
+		ElementType* RESTRICT,
+		ElementType*>::Result;
+
+	using ElementAllocatorType = typename TChooseClass<
+		AllocatorType::NeedsElementType,
+		typename AllocatorType::template ForElementType<ElementType>,
+		typename AllocatorType::ForAnyElementType
+	>::Result;
+
+	static_assert(DimNum > 1, "TMultiArray requires a positive, non-zero number of dimensions");
+	static_assert(TIsSigned<SizeType>::Value, "TMultiArray only supports signed index types");
+
+public:
+
+	/**
+	 * Array bracket operator. Returns reference to element at given index.
+	 *
+	 * @returns Reference to indexed element.
+	 */
+	FORCEINLINE TMultiArrayView<DimNum - 1, ElementType, bIsChecked, bIsRestrict, SizeType> operator[](SizeType Index)
+	{
+		RangeCheck(0, Index);
+
+		TMultiArrayShape<DimNum - 1, SizeType> NewShape;
+		for (uint8 Idx = 0; Idx < DimNum - 1; Idx++)
+		{
+			NewShape[Idx] = ArrayShape[Idx + 1];
+		}
+
+		return TMultiArrayView<DimNum - 1, ElementType, bIsChecked, bIsRestrict, SizeType>(GetData() + Index * Stride(0), NewShape);
+	}
+
+	/**
+	 * Array bracket operator. Returns reference to element at given index.
+	 * 
+	 * Const version of the above.
+	 *
+	 * @returns Reference to indexed element.
+	 */
+	FORCEINLINE TMultiArrayView<DimNum - 1, const ElementType, bIsChecked, bIsRestrict, SizeType> operator[](SizeType Index) const
+	{
+		RangeCheck(0, Index);
+
+		TMultiArrayShape<DimNum - 1, SizeType> NewShape;
+		for (uint8 Idx = 0; Idx < DimNum - 1; Idx++)
+		{
+			NewShape[Idx] = ArrayShape[Idx + 1];
+		}
+
+		return TMultiArrayView<DimNum - 1, const ElementType, bIsChecked, bIsRestrict, SizeType>(GetData() + Index * Stride(0), NewShape);
+	}
+
+	/**
+	 * Flattens the array view into a single dimension.
+	 *
+	 * @returns Flattened array view.
+	 */
+	FORCEINLINE TMultiArrayView<1, ElementType, bIsChecked, bIsRestrict, SizeType> Flatten()
+	{
+		return TMultiArrayView<1, ElementType, bIsChecked, bIsRestrict, SizeType>(GetData(), Num());
+	}
+
+	/**
+	 * Flattens the array view into a single dimension.
+	 *
+	 * @returns Flattened array view.
+	 */
+	FORCEINLINE TMultiArrayView<1, const ElementType, bIsChecked, bIsRestrict, SizeType> Flatten() const
+	{
+		return TMultiArrayView<1, const ElementType, bIsChecked, bIsRestrict, SizeType>(GetData(), Num());
+	}
+
+	/**
+	 * Flattens the array on a given dimension, merging that dimension and the following one.
+	 * e.g. a 3D MultiArrayView with shape [10, 5, 3] flattened on dimension 0 will become a 2D MultiArrayView with shape [10 * 5, 3]
+	 *
+	 * @returns View flattened on the given dimension.
+	 */
+	FORCEINLINE TMultiArrayView<DimNum - 1, ElementType, bIsChecked, bIsRestrict, SizeType> Flatten(uint8 Dimension)
+	{
+		if (bIsChecked)
+		{
+			checkf((Dimension < DimNum - 1), TEXT("MultiArray flatten dimension out of bounds: %lld from a rank of %lld"), (long long)Dimension, (long long)DimNum); // & for one branch
+		}
+
+		TMultiArrayShape<DimNum - 1, SizeType> NewShape;
+		uint8 SrcIdx = 0;
+		for (uint8 Idx = 0; Idx < DimNum - 1; Idx++)
+		{
+			if (Idx == Dimension)
+			{
+				NewShape[Idx] = ArrayShape[SrcIdx] * ArrayShape[SrcIdx + 1];
+				SrcIdx += 2;
+			}
+			else
+			{
+				NewShape[Idx] = ArrayShape[SrcIdx];
+				SrcIdx += 1;
+			}
+		}
+
+		return TMultiArrayView<DimNum - 1, ElementType, bIsChecked, bIsRestrict, SizeType>(GetData(), NewShape);
+	}
+
+	/**
+	 * Flattens the array on a given dimension, merging that dimension and the following one.
+	 * e.g. a 3D MultiArrayView with shape [10, 5, 3] flattened on dimension 0 will become a 2D MultiArrayView with shape [10 * 5, 3]
+	 *
+	 * @returns View flattened on the given dimension.
+	 */
+	FORCEINLINE TMultiArrayView<DimNum - 1, const ElementType, bIsChecked, bIsRestrict, SizeType> Flatten(uint8 Dimension) const
+	{
+		if (bIsChecked)
+		{
+			checkf((Dimension < DimNum - 1), TEXT("MultiArray flatten dimension out of bounds: %lld from a rank of %lld"), (long long)Dimension, (long long)DimNum); // & for one branch
+		}
+
+		TMultiArrayShape<DimNum - 1, SizeType> NewShape;
+		uint8 SrcIdx = 0;
+		for (uint8 Idx = 0; Idx < DimNum - 1; Idx++)
+		{
+			if (Idx == Dimension)
+			{
+				NewShape[Idx] = ArrayShape[SrcIdx] * ArrayShape[SrcIdx + 1];
+				SrcIdx += 2;
+			}
+			else
+			{
+				NewShape[Idx] = ArrayShape[SrcIdx];
+				SrcIdx += 1;
+			}
+		}
+
+		return TMultiArrayView<DimNum - 1, const ElementType, bIsChecked, bIsRestrict, SizeType>(GetData(), NewShape);
+	}
+
+	/**
+	 * Flattens the array on a given dimension, merging that dimension and the following one.
+	 * e.g. a 3D MultiArrayView with shape [10, 5, 3] flattened on dimension 0 will become a 2D MultiArrayView with shape [10 * 5, 3]
+	 *
+	 * @returns View flattened on the given dimension.
+	 */
+	template<uint8 Dimension>
+	FORCEINLINE TMultiArrayView<DimNum - 1, ElementType, bIsChecked, bIsRestrict, SizeType> Flatten()
+	{
+		static_assert(Dimension < DimNum - 1, "MultiArray flatten dimension out of bounds");
+
+		TMultiArrayShape<DimNum - 1, SizeType> NewShape;
+		uint8 SrcIdx = 0;
+		for (uint8 Idx = 0; Idx < DimNum - 1; Idx++)
+		{
+			if (Idx == Dimension)
+			{
+				NewShape[Idx] = ArrayShape[SrcIdx] * ArrayShape[SrcIdx + 1];
+				SrcIdx += 2;
+			}
+			else
+			{
+				NewShape[Idx] = ArrayShape[SrcIdx];
+				SrcIdx += 1;
+			}
+		}
+
+		return TMultiArrayView<DimNum - 1, ElementType, bIsChecked, bIsRestrict, SizeType>(GetData(), NewShape);
+	}
+
+	/**
+	 * Flattens the array on a given dimension, merging that dimension and the following one.
+	 * e.g. a 3D MultiArrayView with shape [10, 5, 3] flattened on dimension 0 will become a 2D MultiArrayView with shape [10 * 5, 3]
+	 *
+	 * @returns View flattened on the given dimension.
+	 */
+	template<uint8 Dimension>
+	FORCEINLINE TMultiArrayView<DimNum - 1, const ElementType, bIsChecked, bIsRestrict, SizeType> Flatten() const
+	{
+		static_assert(Dimension < DimNum - 1, "MultiArray flatten dimension out of bounds");
+
+		TMultiArrayShape<DimNum - 1, SizeType> NewShape;
+		uint8 SrcIdx = 0;
+		for (uint8 Idx = 0; Idx < DimNum - 1; Idx++)
+		{
+			if (Idx == Dimension)
+			{
+				NewShape[Idx] = ArrayShape[SrcIdx] * ArrayShape[SrcIdx + 1];
+				SrcIdx += 2;
+			}
+			else
+			{
+				NewShape[Idx] = ArrayShape[SrcIdx];
+				SrcIdx += 1;
+			}
+		}
+
+		return TMultiArrayView<DimNum - 1, const ElementType, bIsChecked, bIsRestrict, SizeType>(GetData(), NewShape);
+	}
+
+
+public:
+
+	/**
+	 * Constructor, initializes element number counters.
+	 */
+	FORCEINLINE TMultiArray() = default;
+
+	/**
+	 * Constructor from a raw array of elements.
+	 *
+	 * @param Ptr   A pointer to an array of elements to copy.
+	 * @param Nums The number of elements on each dimension.
+	 */
+	FORCEINLINE TMultiArray(const ElementType* Ptr, TMultiArrayShape<DimNum, SizeType> Shape)
+	{
+		check(Ptr != nullptr || Shape.Total() == 0);
+
+		CopyToEmpty(Ptr, Shape);
+	}
+
+	template <typename OtherElementType, bool bOtherIsChecked, bool bOtherIsRestrict, typename OtherSizeType>
+	FORCEINLINE explicit TMultiArray(const TMultiArrayView<DimNum, OtherElementType, bOtherIsChecked, bOtherIsRestrict, OtherSizeType>& Other)
+		: TMultiArray(Other.GetData(), Other.Shape()) {}
+
+	/**
+	 * Copy constructor with changed allocator. Use the common routine to perform the copy.
+	 *
+	 * @param Other The source array to copy.
+	 */
+	template <
+		typename OtherElementType,
+		typename OtherAllocator,
+		bool bOtherIsChecked,
+		std::enable_if_t<UE::MultiArray::Private::TMultiArrayElementsAreCompatible<ElementType, const OtherElementType&>::Value>* = nullptr
+	>
+		FORCEINLINE explicit TMultiArray(const TMultiArray<DimNum, OtherElementType, OtherAllocator, bOtherIsChecked>& Other)
+	{
+		CopyToEmpty(Other.GetData(), Other.Shape());
+	}
+
+	/**
+	 * Copy constructor. Use the common routine to perform the copy.
+	 *
+	 * @param Other The source array to copy.
+	 */
+	FORCEINLINE TMultiArray(const TMultiArray& Other)
+	{
+		CopyToEmpty(Other.GetData(), Other.Shape());
+	}
+
+	/**
+	 * Assignment operator. First deletes all currently contained elements
+	 * and then copies from other array.
+	 *
+	 * AllocatorType changing version.
+	 *
+	 * @param Other The source array to assign from.
+	 */
+	template<typename OtherAllocatorType, bool bOtherIsChecked>
+	TMultiArray& operator=(const TMultiArray<DimNum, ElementType, OtherAllocatorType, bOtherIsChecked>& Other)
+	{
+		DestructItems(GetData(), Num());
+		CopyToEmpty(Other.GetData(), Other.Shape());
+		return *this;
+	}
+
+	/**
+	 * Assignment operator. First deletes all currently contained elements
+	 * and then copies from other array.
+	 *
+	 * @param Other The source array to assign from.
+	 */
+	TMultiArray& operator=(const TMultiArray& Other)
+	{
+		if (this != &Other)
+		{
+			DestructItems(GetData(), Num());
+			CopyToEmpty(Other.GetData(), Other.Shape());
+		}
+		return *this;
+	}
+
+	template <typename OtherElementType, bool bOtherIsChecked, bool bOtherIsRestrict, typename OtherSizeType>
+	TMultiArray& operator=(const TMultiArrayView<DimNum, OtherElementType, bOtherIsChecked, bOtherIsRestrict, OtherSizeType>& Other);
+
+private:
+
+	/**
+	 * Moves or copies array. Depends on the array type traits.
+	 *
+	 * This override copies.
+	 *
+	 * @param ToArray Array to move into.
+	 * @param FromArray Array to move from.
+	 */
+	template <typename FromArrayType, typename ToArrayType>
+	static FORCEINLINE void MoveOrCopy(ToArrayType& ToArray, FromArrayType& FromArray)
+	{
+		ToArray.CopyToEmpty(FromArray.GetData(), FromArray.Shape());
+	}
+
+public:
+	/**
+	 * Move constructor.
+	 *
+	 * @param Other Array to move from.
+	 */
+	FORCEINLINE TMultiArray(TMultiArray&& Other)
+	{
+		MoveOrCopy(*this, Other);
+	}
+
+	/**
+	 * Move constructor.
+	 *
+	 * @param Other Array to move from.
+	 */
+	template <
+		typename OtherElementType,
+		typename OtherAllocator,
+		bool bOtherIsChecked,
+		std::enable_if_t<UE::MultiArray::Private::TMultiArrayElementsAreCompatible<ElementType, OtherElementType&&>::Value>* = nullptr
+	>
+		FORCEINLINE explicit TMultiArray(TMultiArray<DimNum, OtherElementType, OtherAllocator, bOtherIsChecked>&& Other)
+	{
+		MoveOrCopy(*this, Other);
+	}
+
+	/**
+	 * Move assignment operator.
+	 *
+	 * @param Other Array to assign and move from.
+	 */
+	TMultiArray& operator=(TMultiArray&& Other)
+	{
+		if (this != &Other)
+		{
+			DestructItems(GetData(), Num());
+			MoveOrCopy(*this, Other);
+		}
+		return *this;
+	}
+
+	/** Destructor. */
+	~TMultiArray()
+	{
+		DestructItems(GetData(), Num());
+
+		// note ArrayShape and data pointer are not invalidated
+		// they are left unchanged and use-after-destruct will see them the same as before destruct
+	}
+
+	/**
+	 * Helper function for returning a typed pointer to the first array entry.
+	 *
+	 * @returns Pointer to first array entry or nullptr if empty
+	 */
+	FORCEINLINE PointerType GetData()
+	{
+		return (PointerType)AllocatorInstance.GetAllocation();
+	}
+
+	/**
+	 * Helper function for returning a typed pointer to the first array entry.
+	 *
+	 * @returns Pointer to first array entry or nullptr if empty
+	 */
+	FORCEINLINE const PointerType GetData() const
+	{
+		return (const PointerType)AllocatorInstance.GetAllocation();
+	}
+
+	/**
+	 * DO NOT USE DIRECTLY
+	 * STL-like iterators to enable range-based for loop support.
+	 */
+	FORCEINLINE ElementType* begin() const { return GetData(); }
+	FORCEINLINE ElementType* end() const { return GetData() + Num(); }
+
+	/**
+	 * Helper function returning the size of the inner type.
+	 *
+	 * @returns Size in bytes of array type.
+	 */
+	FORCEINLINE uint32 GetTypeSize() const
+	{
+		return sizeof(ElementType);
+	}
+
+	/**
+	 * Checks array invariants: if array size is greater than zero and less
+	 * than maximum.
+	 */
+	FORCEINLINE void CheckInvariants() const
+	{
+		if (bIsChecked)
+		{
+			for (SizeType Idx = 0; Idx < DimNum; Idx++)
+			{
+				checkSlow(ArrayShape[Idx] >= 0);
+			}
+		}
+	}
+
+	/**
+	 * Checks if a dimension is within the allowed number of dimensions
+	 *
+	 * @param Dimension Dimension of the array.
+	 */
+	FORCEINLINE void DimensionCheck(uint8 Dimension) const
+	{
+		if (bIsChecked)
+		{
+			checkf((Dimension < DimNum), TEXT("MultiArray dimension out of bounds: %lld from a rank of %lld"), (long long)Dimension, (long long)DimNum); // & for one branch
+		}
+	}
+
+	/**
+	 * Checks if index is in dimension range.
+	 *
+	 * @param Dimension Dimension of the array.
+	 * @param Index Index to check.
+	 */
+	FORCEINLINE void RangeCheck(uint8 Dimension, SizeType Index) const
+	{
+		DimensionCheck(Dimension);
+		CheckInvariants();
+
+		if (bIsChecked)
+		{
+			checkf((Index >= 0) & (Index < ArrayShape[Dimension]), TEXT("Array index out of bounds: %lld from a dimension of size %lld"), (long long)Index, (long long)ArrayShape[Dimension]); // & for one branch
+		}
+	}
+
+	/**
+	 * Checks if a slice range [Index, Index+InNum) is in dimension range.
+	 * Length is 0 is allowed on empty dimensions; Index must be 0 in that case.
+	 *
+	 * @param Dimension Dimension of the array.
+	 * @param Index Starting index of the slice.
+	 * @param InNum Length of the slice.
+	 */
+	FORCEINLINE void SliceRangeCheck(uint8 Dimension, SizeType Index, SizeType InNum) const
+	{
+		DimensionCheck(Dimension);
+
+		if (bIsChecked)
+		{
+			checkf(Index >= 0, TEXT("Invalid index (%lld)"), (long long)Index);
+			checkf(InNum >= 0, TEXT("Invalid count (%lld)"), (long long)InNum);
+			checkf(Index + InNum <= ArrayShape[Dimension], TEXT("Range (index: %lld, count: %lld) lies outside the view of %lld elements"), (long long)Index, (long long)InNum, (long long)ArrayShape[Dimension]);
+		}
+	}
+
+	/**
+	 * Tests if dimension is valid, i.e. less than the number of dimensions.
+	 *
+	 * @param Dimension Dimension of the array.
+	 *
+	 * @returns True if dimension is valid. False otherwise.
+	 */
+	FORCEINLINE bool IsValidDimension(uint8 Dimension) const
+	{
+		return Dimension < DimNum;
+	}
+
+	/**
+	 * Tests if index is valid, i.e. greater than or equal to zero, and less than the number of elements in the dimension.
+	 *
+	 * @param Dimension Dimension of the array.
+	 * @param Index Index to test.
+	 *
+	 * @returns True if index is valid. False otherwise.
+	 */
+	FORCEINLINE bool IsValidIndex(uint8 Dimension, SizeType Index) const
+	{
+		DimensionCheck(Dimension);
+
+		return (Index >= 0) && (Index < ArrayShape[Dimension]);
+	}
+
+	/**
+	 * Returns true if the array is empty and contains no elements.
+	 *
+	 * @returns True if the array is empty.
+	 * @see Num
+	 */
+	bool IsEmpty() const
+	{
+		return Num() == 0;
+	}
+
+	/**
+	 * Returns true if the dimension is empty and contains no elements.
+	 *
+	 * @param Dimension Dimension of the array.
+	 * @returns True if the dimension is empty.
+	 * @see Num
+	 */
+	bool IsEmpty(uint8 Dimension) const
+	{
+		DimensionCheck(Dimension);
+
+		return ArrayShape[Dimension] == 0;
+	}
+
+	/**
+	 * Returns number of dimensions.
+	 *
+	 * @returns Number of dimensions in array.
+	 */
+	FORCEINLINE SizeType Rank() const
+	{
+		return DimNum;
+	}
+
+	/**
+	 * Returns total number of elements
+	 *
+	 * @returns Total number of elements in array.
+	 */
+	FORCEINLINE SizeType Num() const
+	{
+		return ArrayShape.Total();
+	}
+
+	/**
+	 * Returns number of elements in a dimension.
+	 *
+	 * @returns Number of elements in array.
+	 */
+	template<uint8 InDimIdx>
+	FORCEINLINE SizeType Num() const
+	{
+		static_assert(InDimIdx < DimNum);
+		return ArrayShape[InDimIdx];
+	}
+
+	/**
+	 * Returns number of elements in a dimension.
+	 *
+	 * @param Dimension Dimension of the array.
+	 * @returns Number of elements in array.
+	 */
+	FORCEINLINE SizeType Num(uint8 Dimension) const
+	{
+		DimensionCheck(Dimension);
+
+		return ArrayShape[Dimension];
+	}
+
+	/**
+	 * Returns the number of elements in each dimension.
+	 *
+	 * @returns Number of elements in each dimension.
+	 */
+	FORCEINLINE TMultiArrayShape<DimNum, SizeType> Shape() const
+	{
+		return ArrayShape;
+	}
+
+	/**
+	 * Returns the stride for a dimension.
+	 *
+	 * @param Dimension Dimension of the array.
+	 * @returns Stride of that dimension.
+	 */
+	FORCEINLINE SizeType Stride(uint8 Dimension) const
+	{
+		DimensionCheck(Dimension);
+
+		SizeType Total = 1;
+		for (uint8 Idx = Dimension + 1; Idx < DimNum; Idx++)
+		{
+			Total *= ArrayShape[Idx];
+		}
+		return Total;
+	}
+
+	/**
+	 * Returns a sliced view. Slicing outside of the range of the view is illegal.
+	 *
+	 * @param Index Starting index of the new view
+	 * @param InNum Number of elements in the new view
+	 * @returns Sliced view
+	 */
+	[[nodiscard]] FORCEINLINE TMultiArrayView<DimNum, ElementType, bIsChecked, bIsRestrict, SizeType> Slice(SizeType Index, SizeType InNum)
+	{
+		SliceRangeCheck(0, Index, InNum);
+		TMultiArrayShape NewShape = ArrayShape;
+		NewShape[0] = InNum;
+		return TMultiArrayView<DimNum, ElementType, bIsChecked, bIsRestrict, SizeType>(GetData() + Index * Stride(0), NewShape);
+	}
+
+	/**
+	 * Returns a sliced view. Slicing outside of the range of the view is illegal.
+	 *
+	 * @param Index Starting index of the new view
+	 * @param InNum Number of elements in the new view
+	 * @returns Sliced view
+	 */
+	[[nodiscard]] FORCEINLINE TMultiArrayView<DimNum, const ElementType, bIsChecked, bIsRestrict, SizeType> Slice(SizeType Index, SizeType InNum) const
+	{
+		SliceRangeCheck(0, Index, InNum);
+		TMultiArrayShape NewShape = ArrayShape;
+		NewShape[0] = InNum;
+		return TMultiArrayView<DimNum, const ElementType, bIsChecked, bIsRestrict, SizeType>(GetData() + Index * Stride(0), NewShape);
+	}
+
+	/**
+	 * Checks if this array contains the element.
+	 *
+	 * @returns	True if found. False otherwise.
+	 */
+	template <typename ComparisonType>
+	bool Contains(const ComparisonType& Item) const
+	{
+		for (const ElementType* RESTRICT Data = GetData(), *RESTRICT DataEnd = Data + Num(); Data != DataEnd; ++Data)
+		{
+			if (*Data == Item)
+			{
+				return true;
+			}
+		}
+		return false;
+	}
+
+public:
+
+	/**
+	 * Empties the array. It calls the destructors on held items if needed.
+	 */
+	void Empty()
+	{
+		SizeType CurrTotal = Num();
+		DestructItems(GetData(), CurrTotal);
+		ArrayShape = TMultiArrayShape<DimNum, SizeType>();
+		AllocatorResizeAllocation(CurrTotal, 0);
+	}
+
+	/**
+	 * Resizes array to given number of elements.
+	 *
+	 * @param NewShape New shape of the array.
+	 */
+	void SetNum(TMultiArrayShape<DimNum, SizeType> NewShape)
+	{
+		SizeType CurrTotal = ArrayShape.Total();
+		SizeType NewTotal = NewShape.Total();
+
+		if (NewTotal > CurrTotal)
+		{
+			AllocatorResizeAllocation(CurrTotal, NewTotal);
+			ConstructItems<ElementType>(GetData() + CurrTotal, GetData(), NewTotal - CurrTotal);
+		}
+		else if (NewTotal < CurrTotal)
+		{
+			DestructItems<ElementType>(GetData() + NewTotal, CurrTotal - NewTotal);
+			AllocatorResizeAllocation(CurrTotal, NewTotal);
+		}
+
+		ArrayShape = NewShape;
+	}
+
+	/**
+	 * Resizes array to given number of elements. New elements will be zeroed.
+	 *
+	 * @param NewShape New shape of the array.
+	 */
+	void SetNumZeroed(TMultiArrayShape<DimNum, SizeType> NewShape)
+	{
+		SizeType CurrTotal = ArrayShape.Total();
+		SizeType NewTotal = NewShape.Total();
+
+		if (NewTotal > CurrTotal)
+		{
+			AllocatorResizeAllocation(CurrTotal, NewTotal);
+			FMemory::Memzero((void*)(GetData() + CurrTotal), (CurrTotal - NewTotal) * sizeof(ElementType));
+		}
+		else if (NewTotal < CurrTotal)
+		{
+			DestructItems<ElementType>(GetData() + NewTotal, CurrTotal - NewTotal);
+			AllocatorResizeAllocation(CurrTotal, NewTotal);
+		}
+
+		ArrayShape = NewShape;
+	}
+
+	/**
+	 * Resizes array to given number of elements. New elements will be uninitialized.
+	 *
+	 * @param NewShape New shape of the array.
+	 */
+	void SetNumUninitialized(TMultiArrayShape<DimNum, SizeType> NewShape)
+	{
+		SizeType CurrTotal = ArrayShape.Total();
+		SizeType NewTotal = NewShape.Total();
+
+		if (NewTotal > CurrTotal)
+		{
+			AllocatorResizeAllocation(CurrTotal, NewTotal);
+		}
+		else if (NewTotal < CurrTotal)
+		{
+			DestructItems<ElementType>(GetData() + NewTotal, CurrTotal - NewTotal);
+			AllocatorResizeAllocation(CurrTotal, NewTotal);
+		}
+
+		ArrayShape = NewShape;
+	}
+
+private:
+	void AllocatorResizeAllocation(SizeType CurrentNum, SizeType NewNum)
+	{
+		if constexpr (TAllocatorTraits<AllocatorType>::SupportsElementAlignment)
+		{
+			AllocatorInstance.ResizeAllocation(CurrentNum, NewNum, sizeof(ElementType), alignof(ElementType));
+		}
+		else
+		{
+			AllocatorInstance.ResizeAllocation(CurrentNum, NewNum, sizeof(ElementType));
+		}
+	}
+
+	/**
+	 * Copies data from one array into this array. Uses the fast path if the
+	 * data in question does not need a constructor.
+	 *
+	 * @param OtherData A pointer to the data to copy
+	 * @param OtherNums The dimensions of the data to copy
+	 */
+	template <typename OtherElementType, typename OtherSizeType>
+	void CopyToEmpty(const OtherElementType* OtherData, TMultiArrayShape<DimNum, OtherSizeType> OtherShape)
+	{
+		SizeType CurrTotal = ArrayShape.Total();
+		OtherSizeType OtherTotal = OtherShape.Total();
+
+		for (uint8 Idx = 0; Idx < DimNum; Idx++)
+		{
+			SizeType NewNum = (SizeType)OtherShape[Idx];
+			
+			if (bIsChecked)
+			{
+				checkf((OtherSizeType)NewNum == OtherShape[Idx], TEXT("Invalid number of elements to add to this array type: %lld"), (long long)NewNum);
+			}
+
+			ArrayShape[Idx] = NewNum;
+		}
+
+		SizeType NewTotal = ArrayShape.Total();
+
+		AllocatorResizeAllocation(CurrTotal, NewTotal);
+		ConstructItems<ElementType>(GetData(), OtherData, OtherTotal);
+	}
+
+public:
+
+	const ElementAllocatorType& GetAllocatorInstance() const { return AllocatorInstance; }
+	ElementAllocatorType& GetAllocatorInstance() { return AllocatorInstance; }
+
+public:
+
+	FORCEINLINE operator TMultiArrayView<DimNum, ElementType, bIsChecked, bIsRestrict, SizeType>()
+	{
+		return TMultiArrayView<DimNum, ElementType, bIsChecked, bIsRestrict, SizeType>(GetData(), ArrayShape);
+	}
+
+	FORCEINLINE operator TMultiArrayView<DimNum, const ElementType, bIsChecked, bIsRestrict, SizeType>() const
+	{
+		return TMultiArrayView<DimNum, const ElementType, bIsChecked, bIsRestrict, SizeType>(GetData(), ArrayShape);
+	}
+
+protected:
+
+	ElementAllocatorType AllocatorInstance;
+	TMultiArrayShape<DimNum, SizeType> ArrayShape;
+};
+
+
+/**
+* Specialization for single dimensional MultiArray
+*/
+template<typename InElementType, typename InAllocatorType, bool bInIsChecked, bool bInIsRestrict>
+class TMultiArray<1, InElementType, InAllocatorType, bInIsChecked, bInIsRestrict>
+{
+public:
+	static constexpr uint8 DimNum = 1;
+	using ElementType = InElementType;
+	using AllocatorType = InAllocatorType;
+	static constexpr bool bIsChecked = bInIsChecked;
+	static constexpr bool bIsRestrict = bInIsRestrict;
+	using SizeType = typename InAllocatorType::SizeType;
+
+	using PointerType = typename TChooseClass<
+		bIsRestrict,
+		ElementType* RESTRICT,
+		ElementType*>::Result;
+
+	using ElementAllocatorType = typename TChooseClass<
+		AllocatorType::NeedsElementType,
+		typename AllocatorType::template ForElementType<ElementType>,
+		typename AllocatorType::ForAnyElementType
+	>::Result;
+
+	static_assert(TIsSigned<SizeType>::Value, "TMultiArray only supports signed index types");
+
+	/**
+	 * Initializer list constructor
+	 */
+	TMultiArray(std::initializer_list<InElementType> InitList)
+	{
+		// This is not strictly legal, as std::initializer_list's iterators are not guaranteed to be pointers, but
+		// this appears to be the case on all of our implementations.  Also, if it's not true on a new implementation,
+		// it will fail to compile rather than behave badly.
+		CopyToEmpty<InElementType, SizeType>(InitList.begin(), { (SizeType)InitList.size() });
+	}
+
+	/**
+	 * Initializer list assignment operator. First deletes all currently contained elements
+	 * and then copies from initializer list.
+	 *
+	 * @param InitList The initializer_list to copy from.
+	 */
+	TMultiArray& operator=(std::initializer_list<InElementType> InitList)
+	{
+		DestructItems(GetData(), Num());
+		// This is not strictly legal, as std::initializer_list's iterators are not guaranteed to be pointers, but
+		// this appears to be the case on all of our implementations.  Also, if it's not true on a new implementation,
+		// it will fail to compile rather than behave badly.
+		CopyToEmpty<InElementType, SizeType>(InitList.begin(), { (SizeType)InitList.size() });
+		return *this;
+	}
+
+	/**
+	 * Array bracket operator. Returns reference to element at given index.
+	 *
+	 * @returns Reference to indexed element.
+	 */
+	FORCEINLINE ElementType& operator[](SizeType Index)
+	{
+		RangeCheck(0, Index);
+		return GetData()[Index];
+	}
+
+	/**
+	 * Array bracket operator. Returns reference to element at given index.
+	 *
+	 * Const version of the above.
+	 *
+	 * @returns Reference to indexed element.
+	 */
+	FORCEINLINE const ElementType& operator[](SizeType Index) const
+	{
+		RangeCheck(0, Index);
+		return GetData()[Index];
+	}
+
+	/**
+	 * Flattens the array view into a single dimension.
+	 *
+	 * @returns Flattened array view.
+	 */
+	FORCEINLINE TMultiArrayView<1, ElementType, bIsChecked, bIsRestrict, SizeType> Flatten()
+	{
+		return TMultiArrayView<1, ElementType, bIsChecked, bIsRestrict, SizeType>(GetData(), Num());
+	}
+
+	/**
+	 * Flattens the array view into a single dimension.
+	 *
+	 * @returns Flattened array view.
+	 */
+	FORCEINLINE TMultiArrayView<1, const ElementType, bIsChecked, bIsRestrict, SizeType> Flatten() const
+	{
+		return TMultiArrayView<1, const ElementType, bIsChecked, bIsRestrict, SizeType>(GetData(), Num());
+	}
+
+	/** Implicit cast to TArrayView. */
+	FORCEINLINE operator TArrayView<ElementType, SizeType>() const
+	{
+		return TArrayView(GetData(), Num());
+	}
+
+public:
+
+	/**
+	 * Constructor, initializes element number counters.
+	 */
+	FORCEINLINE TMultiArray() = default;
+
+	/**
+	 * Constructor from a raw array of elements.
+	 *
+	 * @param Ptr   A pointer to an array of elements to copy.
+	 * @param Nums The number of elements on each dimension.
+	 */
+	FORCEINLINE TMultiArray(const ElementType* Ptr, TMultiArrayShape<DimNum, SizeType> Shape)
+	{
+		check(Ptr != nullptr || Shape.Total() == 0);
+
+		CopyToEmpty(Ptr, Shape);
+	}
+
+	template <typename OtherElementType, bool bOtherIsChecked, bool bOtherIsRestrict, typename OtherSizeType>
+	explicit TMultiArray(const TMultiArrayView<DimNum, OtherElementType, bOtherIsChecked, bOtherIsRestrict, OtherSizeType>& Other)
+		: TMultiArray(Other.GetData(), Other.Shape()) {}
+
+	/**
+	 * Copy constructor with changed allocator. Use the common routine to perform the copy.
+	 *
+	 * @param Other The source array to copy.
+	 */
+	template <
+		typename OtherElementType,
+		typename OtherAllocator,
+		bool bOtherIsChecked,
+		std::enable_if_t<UE::MultiArray::Private::TMultiArrayElementsAreCompatible<ElementType, const OtherElementType&>::Value>* = nullptr
+	>
+		FORCEINLINE explicit TMultiArray(const TMultiArray<DimNum, OtherElementType, OtherAllocator, bOtherIsChecked>& Other)
+	{
+		CopyToEmpty(Other.GetData(), Other.Shape());
+	}
+
+	/**
+	 * Copy constructor. Use the common routine to perform the copy.
+	 *
+	 * @param Other The source array to copy.
+	 */
+	FORCEINLINE TMultiArray(const TMultiArray& Other)
+	{
+		CopyToEmpty(Other.GetData(), Other.Shape());
+	}
+
+	/**
+	 * Assignment operator. First deletes all currently contained elements
+	 * and then copies from other array.
+	 *
+	 * AllocatorType changing version.
+	 *
+	 * @param Other The source array to assign from.
+	 */
+	template<typename OtherAllocatorType, bool bOtherIsChecked>
+	TMultiArray& operator=(const TMultiArray<DimNum, ElementType, OtherAllocatorType, bOtherIsChecked>& Other)
+	{
+		DestructItems(GetData(), Num());
+		CopyToEmpty(Other.GetData(), Other.Shape());
+		return *this;
+	}
+
+	/**
+	 * Assignment operator. First deletes all currently contained elements
+	 * and then copies from other array.
+	 *
+	 * @param Other The source array to assign from.
+	 */
+	TMultiArray& operator=(const TMultiArray& Other)
+	{
+		if (this != &Other)
+		{
+			DestructItems(GetData(), Num());
+			CopyToEmpty(Other.GetData(), Other.Shape());
+		}
+		return *this;
+	}
+
+	template <typename OtherElementType, bool bOtherIsChecked, bool bOtherIsRestrict, typename OtherSizeType>
+	TMultiArray& operator=(const TMultiArrayView<DimNum, OtherElementType, bOtherIsChecked, bOtherIsRestrict, OtherSizeType>& Other);
+
+private:
+
+	/**
+	 * Moves or copies array. Depends on the array type traits.
+	 *
+	 * This override copies.
+	 *
+	 * @param ToArray Array to move into.
+	 * @param FromArray Array to move from.
+	 */
+	template <typename FromArrayType, typename ToArrayType>
+	static FORCEINLINE void MoveOrCopy(ToArrayType& ToArray, FromArrayType& FromArray)
+	{
+		ToArray.CopyToEmpty(FromArray.GetData(), FromArray.Shape());
+	}
+
+public:
+	/**
+	 * Move constructor.
+	 *
+	 * @param Other Array to move from.
+	 */
+	FORCEINLINE TMultiArray(TMultiArray&& Other)
+	{
+		MoveOrCopy(*this, Other);
+	}
+
+	/**
+	 * Move constructor.
+	 *
+	 * @param Other Array to move from.
+	 */
+	template <
+		typename OtherElementType,
+		typename OtherAllocator,
+		bool bOtherIsChecked,
+		std::enable_if_t<UE::MultiArray::Private::TMultiArrayElementsAreCompatible<ElementType, OtherElementType&&>::Value>* = nullptr
+	>
+		FORCEINLINE explicit TMultiArray(TMultiArray<DimNum, OtherElementType, OtherAllocator, bOtherIsChecked>&& Other)
+	{
+		MoveOrCopy(*this, Other);
+	}
+
+	/**
+	 * Move assignment operator.
+	 *
+	 * @param Other Array to assign and move from.
+	 */
+	TMultiArray& operator=(TMultiArray&& Other)
+	{
+		if (this != &Other)
+		{
+			DestructItems(GetData(), Num());
+			MoveOrCopy(*this, Other);
+		}
+		return *this;
+	}
+
+	/** Destructor. */
+	~TMultiArray()
+	{
+		DestructItems(GetData(), Num());
+
+		// Note: ArrayShape and data pointer are not invalidated.
+		// They are left unchanged and use-after-destruct will see them the same as before destruct.
+	}
+
+	/**
+	 * Helper function for returning a typed pointer to the first array entry.
+	 *
+	 * @returns Pointer to first array entry or nullptr if empty
+	 */
+	FORCEINLINE PointerType GetData()
+	{
+		return (PointerType)AllocatorInstance.GetAllocation();
+	}
+
+	/**
+	 * Helper function for returning a typed pointer to the first array entry.
+	 *
+	 * @returns Pointer to first array entry or nullptr if empty
+	 */
+	FORCEINLINE const PointerType GetData() const
+	{
+		return (const PointerType)AllocatorInstance.GetAllocation();
+	}
+
+	/**
+	 * DO NOT USE DIRECTLY
+	 * STL-like iterators to enable range-based for loop support.
+	 */
+	FORCEINLINE ElementType* begin() const { return GetData(); }
+	FORCEINLINE ElementType* end() const { return GetData() + Num(); }
+
+	/**
+	 * Helper function returning the size of the inner type.
+	 *
+	 * @returns Size in bytes of array type.
+	 */
+	FORCEINLINE uint32 GetTypeSize() const
+	{
+		return sizeof(ElementType);
+	}
+
+	/**
+	 * Checks array invariants: if array size is greater than zero and less
+	 * than maximum.
+	 */
+	FORCEINLINE void CheckInvariants() const
+	{
+		if (bIsChecked)
+		{
+			for (SizeType Idx = 0; Idx < DimNum; Idx++)
+			{
+				checkSlow(ArrayShape[Idx] >= 0);
+			}
+		}
+	}
+
+	/**
+	 * Checks if a dimension is within the allowed number of dimensions
+	 *
+	 * @param Dimension Dimension of the array.
+	 */
+	FORCEINLINE void DimensionCheck(uint8 Dimension) const
+	{
+		if (bIsChecked)
+		{
+			checkf((Dimension < DimNum), TEXT("MultiArray dimension out of bounds: %lld from a rank of %lld"), (long long)Dimension, (long long)DimNum); // & for one branch
+		}
+	}
+
+	/**
+	 * Checks if index is in dimension range.
+	 *
+	 * @param Dimension Dimension of the array.
+	 * @param Index Index to check.
+	 */
+	FORCEINLINE void RangeCheck(uint8 Dimension, SizeType Index) const
+	{
+		DimensionCheck(Dimension);
+		CheckInvariants();
+
+		if (bIsChecked)
+		{
+			checkf((Index >= 0) & (Index < ArrayShape[Dimension]), TEXT("Array index out of bounds: %lld from a dimension of size %lld"), (long long)Index, (long long)ArrayShape[Dimension]); // & for one branch
+		}
+	}
+
+	/**
+	 * Checks if a slice range [Index, Index+InNum) is in dimension range.
+	 * Length is 0 is allowed on empty dimensions; Index must be 0 in that case.
+	 *
+	 * @param Dimension Dimension of the array.
+	 * @param Index Starting index of the slice.
+	 * @param InNum Length of the slice.
+	 */
+	FORCEINLINE void SliceRangeCheck(uint8 Dimension, SizeType Index, SizeType InNum) const
+	{
+		DimensionCheck(Dimension);
+
+		if (bIsChecked)
+		{
+			checkf(Index >= 0, TEXT("Invalid index (%lld)"), (long long)Index);
+			checkf(InNum >= 0, TEXT("Invalid count (%lld)"), (long long)InNum);
+			checkf(Index + InNum <= ArrayShape[Dimension], TEXT("Range (index: %lld, count: %lld) lies outside the view of %lld elements"), (long long)Index, (long long)InNum, (long long)ArrayShape[Dimension]);
+		}
+	}
+
+	/**
+	 * Tests if dimension is valid, i.e. less than the number of dimensions.
+	 *
+	 * @param Dimension Dimension of the array.
+	 *
+	 * @returns True if dimension is valid. False otherwise.
+	 */
+	FORCEINLINE bool IsValidDimension(uint8 Dimension) const
+	{
+		return Dimension < DimNum;
+	}
+
+	/**
+	 * Tests if index is valid, i.e. greater than or equal to zero, and less than the number of elements in the dimension.
+	 *
+	 * @param Dimension Dimension of the array.
+	 * @param Index Index to test.
+	 *
+	 * @returns True if index is valid. False otherwise.
+	 */
+	FORCEINLINE bool IsValidIndex(uint8 Dimension, SizeType Index) const
+	{
+		DimensionCheck(Dimension);
+
+		return (Index >= 0) && (Index < ArrayShape[Dimension]);
+	}
+
+	/**
+	 * Returns true if the array is empty and contains no elements.
+	 *
+	 * @returns True if the array is empty.
+	 * @see Num
+	 */
+	bool IsEmpty() const
+	{
+		return Num() == 0;
+	}
+
+	/**
+	 * Returns true if the dimension is empty and contains no elements.
+	 *
+	 * @param Dimension Dimension of the array.
+	 * @returns True if the dimension is empty.
+	 * @see Num
+	 */
+	bool IsEmpty(uint8 Dimension) const
+	{
+		DimensionCheck(Dimension);
+
+		return ArrayShape[Dimension] == 0;
+	}
+
+	/**
+	 * Returns number of dimensions.
+	 *
+	 * @returns Number of dimensions in array.
+	 */
+	FORCEINLINE SizeType Rank() const
+	{
+		return DimNum;
+	}
+
+	/**
+	 * Returns total number of elements
+	 *
+	 * @returns Total number of elements in array.
+	 */
+	FORCEINLINE SizeType Num() const
+	{
+		return ArrayShape.Total();
+	}
+
+	/**
+	 * Returns number of elements in a dimension.
+	 *
+	 * @returns Number of elements in array.
+	 */
+	template<uint8 InDimIdx>
+	FORCEINLINE SizeType Num() const
+	{
+		static_assert(InDimIdx < DimNum);
+		return ArrayShape[InDimIdx];
+	}
+
+	/**
+	 * Returns number of elements in a dimension.
+	 *
+	 * @param Dimension Dimension of the array.
+	 * @returns Number of elements in array.
+	 */
+	FORCEINLINE SizeType Num(uint8 Dimension) const
+	{
+		DimensionCheck(Dimension);
+
+		return ArrayShape[Dimension];
+	}
+
+	/**
+	 * Returns the number of elements in each dimension.
+	 *
+	 * @returns Number of elements in each dimension.
+	 */
+	FORCEINLINE TMultiArrayShape<DimNum, SizeType> Shape() const
+	{
+		return ArrayShape;
+	}
+
+	/**
+	 * Returns the stride for a dimension.
+	 *
+	 * @param Dimension Dimension of the array.
+	 * @returns Stride of that dimension.
+	 */
+	FORCEINLINE SizeType Stride(uint8 Dimension) const
+	{
+		DimensionCheck(Dimension);
+
+		SizeType Total = 1;
+		for (uint8 Idx = Dimension + 1; Idx < DimNum; Idx++)
+		{
+			Total *= ArrayShape[Idx];
+		}
+		return Total;
+	}
+
+	/**
+	 * Returns a sliced view. Slicing outside of the range of the view is illegal.
+	 *
+	 * @param Index Starting index of the new view
+	 * @param InNum Number of elements in the new view
+	 * @returns Sliced view
+	 */
+	[[nodiscard]] FORCEINLINE TMultiArrayView<DimNum, ElementType, bIsChecked, bIsRestrict, SizeType> Slice(SizeType Index, SizeType InNum)
+	{
+		SliceRangeCheck(0, Index, InNum);
+		TMultiArrayShape NewShape = ArrayShape;
+		NewShape[0] = InNum;
+		return TMultiArrayView<DimNum, ElementType, bIsChecked, bIsRestrict, SizeType>(GetData() + Index * Stride(0), NewShape);
+	}
+
+	/**
+	 * Returns a sliced view. Slicing outside of the range of the view is illegal.
+	 *
+	 * @param Index Starting index of the new view
+	 * @param InNum Number of elements in the new view
+	 * @returns Sliced view
+	 */
+	[[nodiscard]] FORCEINLINE TMultiArrayView<DimNum, const ElementType, bIsChecked, bIsRestrict, SizeType> Slice(SizeType Index, SizeType InNum) const
+	{
+		SliceRangeCheck(0, Index, InNum);
+		TMultiArrayShape NewShape = ArrayShape;
+		NewShape[0] = InNum;
+		return TMultiArrayView<DimNum, const ElementType, bIsChecked, bIsRestrict, SizeType>(GetData() + Index * Stride(0), NewShape);
+	}
+
+	/**
+	 * Checks if this array contains the element.
+	 *
+	 * @returns	True if found. False otherwise.
+	 */
+	template <typename ComparisonType>
+	bool Contains(const ComparisonType& Item) const
+	{
+		for (const ElementType* RESTRICT Data = GetData(), *RESTRICT DataEnd = Data + Num(); Data != DataEnd; ++Data)
+		{
+			if (*Data == Item)
+			{
+				return true;
+			}
+		}
+		return false;
+	}
+
+public:
+
+	/**
+	 * Empties the array. It calls the destructors on held items if needed.
+	 */
+	void Empty()
+	{
+		SizeType CurrTotal = Num();
+		DestructItems(GetData(), CurrTotal);
+		ArrayShape = TMultiArrayShape<DimNum, SizeType>();
+		AllocatorResizeAllocation(CurrTotal, 0);
+	}
+
+	/**
+	 * Resizes array to given number of elements.
+	 *
+	 * @param NewShape New shape of the array.
+	 */
+	void SetNum(TMultiArrayShape<DimNum, SizeType> NewShape)
+	{
+		SizeType CurrTotal = ArrayShape.Total();
+		SizeType NewTotal = NewShape.Total();
+
+		if (NewTotal > CurrTotal)
+		{
+			AllocatorResizeAllocation(CurrTotal, NewTotal);
+			ConstructItems<ElementType>(GetData() + CurrTotal, GetData(), NewTotal - CurrTotal);
+		}
+		else if (NewTotal < CurrTotal)
+		{
+			DestructItems<ElementType>(GetData() + NewTotal, CurrTotal - NewTotal);
+			AllocatorResizeAllocation(CurrTotal, NewTotal);
+		}
+
+		ArrayShape = NewShape;
+	}
+
+	/**
+	 * Resizes array to given number of elements. New elements will be zeroed.
+	 *
+	 * @param NewShape New shape of the array.
+	 */
+	void SetNumZeroed(TMultiArrayShape<DimNum, SizeType> NewShape)
+	{
+		SizeType CurrTotal = ArrayShape.Total();
+		SizeType NewTotal = NewShape.Total();
+
+		if (NewTotal > CurrTotal)
+		{
+			AllocatorResizeAllocation(CurrTotal, NewTotal);
+			FMemory::Memzero((void*)(GetData() + CurrTotal), (NewTotal - CurrTotal) * sizeof(ElementType));
+		}
+		else if (NewTotal < CurrTotal)
+		{
+			DestructItems<ElementType>(GetData() + NewTotal, CurrTotal - NewTotal);
+			AllocatorResizeAllocation(CurrTotal, NewTotal);
+		}
+
+		ArrayShape = NewShape;
+	}
+
+	/**
+	 * Resizes array to given number of elements. New elements will be uninitialized.
+	 *
+	 * @param NewShape New shape of the array.
+	 */
+	void SetNumUninitialized(TMultiArrayShape<DimNum, SizeType> NewShape)
+	{
+		SizeType CurrTotal = ArrayShape.Total();
+		SizeType NewTotal = NewShape.Total();
+
+		if (NewTotal > CurrTotal)
+		{
+			AllocatorResizeAllocation(CurrTotal, NewTotal);
+		}
+		else if (NewTotal < CurrTotal)
+		{
+			DestructItems<ElementType>(GetData() + NewTotal, CurrTotal - NewTotal);
+			AllocatorResizeAllocation(CurrTotal, NewTotal);
+		}
+
+		ArrayShape = NewShape;
+	}
+
+private:
+	void AllocatorResizeAllocation(SizeType CurrentNum, SizeType NewNum)
+	{
+		if constexpr (TAllocatorTraits<AllocatorType>::SupportsElementAlignment)
+		{
+			AllocatorInstance.ResizeAllocation(CurrentNum, NewNum, sizeof(ElementType), alignof(ElementType));
+		}
+		else
+		{
+			AllocatorInstance.ResizeAllocation(CurrentNum, NewNum, sizeof(ElementType));
+		}
+	}
+
+	/**
+	 * Copies data from one array into this array. Uses the fast path if the
+	 * data in question does not need a constructor.
+	 *
+	 * @param OtherData A pointer to the data to copy
+	 * @param OtherNums The dimensions of the data to copy
+	 */
+	template <typename OtherElementType, typename OtherSizeType>
+	void CopyToEmpty(const OtherElementType* OtherData, TMultiArrayShape<DimNum, OtherSizeType> OtherShape)
+	{
+		SizeType CurrTotal = ArrayShape.Total();
+		OtherSizeType OtherTotal = OtherShape.Total();
+
+		for (uint8 Idx = 0; Idx < DimNum; Idx++)
+		{
+			SizeType NewNum = (SizeType)OtherShape[Idx];
+
+			if (bIsChecked)
+			{
+				checkf((OtherSizeType)NewNum == OtherShape[Idx], TEXT("Invalid number of elements to add to this array type: %lld"), (long long)NewNum);
+			}
+
+			ArrayShape[Idx] = NewNum;
+		}
+
+		SizeType NewTotal = ArrayShape.Total();
+
+		AllocatorResizeAllocation(CurrTotal, NewTotal);
+		ConstructItems<ElementType>(GetData(), OtherData, OtherTotal);
+	}
+
+public:
+
+	const ElementAllocatorType& GetAllocatorInstance() const { return AllocatorInstance; }
+	ElementAllocatorType& GetAllocatorInstance() { return AllocatorInstance; }
+
+public:
+
+	FORCEINLINE operator TMultiArrayView<DimNum, ElementType, bIsChecked, bIsRestrict, SizeType>()
+	{
+		return TMultiArrayView<DimNum, ElementType, bIsChecked, bIsRestrict, SizeType>(GetData(), ArrayShape);
+	}
+
+	FORCEINLINE operator TMultiArrayView<DimNum, const ElementType, bIsChecked, bIsRestrict, SizeType>() const
+	{
+		return TMultiArrayView<DimNum, const ElementType, bIsChecked, bIsRestrict, SizeType>(GetData(), ArrayShape);
+	}
+
+protected:
+
+	ElementAllocatorType AllocatorInstance;
+	TMultiArrayShape<DimNum, SizeType> ArrayShape;
+};
diff --git a/Engine/Plugins/Experimental/LearningAgents/Source/Learning/Public/MultiArrayView.h b/Engine/Plugins/Experimental/LearningAgents/Source/Learning/Public/MultiArrayView.h
new file mode 100644
--- /dev/null
+++ b/Engine/Plugins/Experimental/LearningAgents/Source/Learning/Public/MultiArrayView.h
@@ -0,0 +1,1251 @@
+// Copyright Epic Games, Inc. All Rights Reserved.
+
+#pragma once
+
+#include "CoreTypes.h"
+#include "Containers/ArrayView.h"
+#include "Templates/IsConst.h"
+#include "Templates/IsSigned.h"
+#include "Templates/PointerIsConvertibleFromTo.h"
+#include "Misc/AssertionMacros.h"
+#include "Templates/Invoke.h"
+#include "Templates/UnrealTypeTraits.h"
+#include "Traits/ElementType.h"
+#include "Math/UnrealMathUtility.h"
+
+namespace UE::MultiArrayView::Private
+{
+	/**
+	 * Copied from Core/Public/Containers/ArrayView.h
+	 *
+	 * Trait testing whether a type is compatible with the view type
+	 */
+	template <typename T, typename ElementType>
+	struct TIsCompatibleElementType
+	{
+	public:
+		/** NOTE:
+		 * The stars in the TPointerIsConvertibleFromTo test are *IMPORTANT*
+		 * They prevent TArrayView<Base>(TArray<Derived>&) from compiling!
+		 */
+		enum { Value = TPointerIsConvertibleFromTo<T*, ElementType* const>::Value };
+	};
+
+	// Copied from Core/Public/Containers/ArrayView.h
+	//
+	// Simply forwards to an unqualified GetData(), but can be called from within TArrayView
+	// where GetData() is already a member and so hides any others.
+	template <typename T>
+	FORCEINLINE decltype(auto) GetDataHelper(T&& Arg)
+	{
+		return GetData(Forward<T>(Arg));
+	}
+
+	// Copied from Core/Public/Containers/ArrayView.h
+	//
+	// Gets the data from the passed argument and proceeds to reinterpret the resulting elements
+	template <typename T>
+	FORCEINLINE decltype(auto) GetReinterpretedDataHelper(T&& Arg)
+	{
+		auto NaturalPtr = GetData(Forward<T>(Arg));
+		using NaturalElementType = typename TRemovePointer<decltype(NaturalPtr)>::Type;
+
+		auto EndPtr = NaturalPtr + GetNum(Arg);
+		TContainerElementTypeCompatibility<NaturalElementType>::ReinterpretRange(NaturalPtr, EndPtr);
+
+		return reinterpret_cast<typename TContainerElementTypeCompatibility<NaturalElementType>::ReinterpretType*>(NaturalPtr);
+	}
+
+	/**
+	 * Copied from Core/Public/Containers/ArrayView.h
+	 *
+	 * Trait testing whether a type is compatible with the view type
+	 */
+	template <typename RangeType, typename ElementType>
+	struct TIsCompatibleRangeType
+	{
+		static constexpr bool Value = TIsCompatibleElementType<typename TRemovePointer<decltype(GetData(DeclVal<RangeType&>()))>::Type, ElementType>::Value;
+
+		template <typename T>
+		static decltype(auto) GetData(T&& Arg)
+		{
+			return UE::MultiArrayView::Private::GetDataHelper(Forward<T>(Arg));
+		}
+	};
+
+	/**
+	 * Copied from Core/Public/Containers/ArrayView.h
+	 *
+	 * Trait testing whether a type is reinterpretable in a way that permits use with the view type
+	 */
+	template <typename RangeType, typename ElementType>
+	struct TIsReinterpretableRangeType
+	{
+	private:
+		using NaturalElementType = typename TRemovePointer<decltype(GetData(DeclVal<RangeType&>()))>::Type;
+
+	public:
+		static constexpr bool Value =
+			!std::is_same_v<typename TContainerElementTypeCompatibility<NaturalElementType>::ReinterpretType, NaturalElementType>
+			&&
+			TIsCompatibleElementType<typename TContainerElementTypeCompatibility<NaturalElementType>::ReinterpretType, ElementType>::Value;
+
+		template <typename T>
+		static decltype(auto) GetData(T&& Arg)
+		{
+			return UE::MultiArrayView::Private::GetReinterpretedDataHelper(Forward<T>(Arg));
+		}
+	};
+}
+
+template<uint8 InDimNum, typename InSizeType = FDefaultAllocator::SizeType>
+class TMultiArrayShape
+{
+public:
+	static constexpr uint8 DimNum = InDimNum;
+	using SizeType = InSizeType;
+
+	static_assert(DimNum >= 1, "TMultiArrayShape requires a positive, non-zero number of dimensions");
+	static_assert(TIsSigned<SizeType>::Value, "TMultiArrayShape only supports signed index types");
+
+public:
+
+	TMultiArrayShape() = default;
+
+	TMultiArrayShape(std::initializer_list<SizeType> InNums)
+	{
+		checkf(InNums.size() == DimNum, TEXT("Wrong number of elements in constructor"));
+
+		uint8 Idx = 0;
+		for (SizeType InNum : InNums)
+		{
+			Nums[Idx] = InNum;
+			Idx++;
+		}
+	}
+
+	TMultiArrayShape(const SizeType* InNums)
+	{
+		for (uint8 Idx = 0; Idx < DimNum; Idx++)
+		{
+			Nums[Idx] = InNums[Idx];
+		}
+	}
+
+	SizeType& operator[](uint8 Dimension)
+	{
+		return Nums[Dimension];
+	}
+
+	const SizeType& operator[](uint8 Dimension) const
+	{
+		return Nums[Dimension];
+	}
+
+	SizeType Total() const
+	{
+		if constexpr (DimNum == 1)
+		{
+			return Nums[0];
+		}
+		else
+		{
+			InSizeType Total = Nums[0];
+			for (uint8 Idx = 1; Idx < DimNum; Idx++)
+			{
+				Total *= Nums[Idx];
+			}
+			return Total;
+		}
+	}
+
+private:
+	SizeType Nums[DimNum] = { 0 };
+};
+
+/**
+ * Templated fixed-size view of multi-dimensional array
+ *
+ * A statically sized view of a multi-dimensional array of typed elements. 
+ *
+ * @see TArrayView
+ */
+template<uint8 InDimNum, typename InElementType, bool bInIsChecked = true, bool bInIsRestrict = false, typename InSizeType = FDefaultAllocator::SizeType>
+class TMultiArrayView
+{
+public:
+	static constexpr uint8 DimNum = InDimNum;
+	using ElementType = InElementType;
+	static constexpr bool bIsChecked = bInIsChecked;
+	static constexpr bool bIsRestrict = bInIsRestrict;
+	using SizeType = InSizeType;
+
+	using PointerType = typename TChooseClass<
+		bIsRestrict,
+		ElementType* RESTRICT,
+		ElementType*>::Result;
+
+	static_assert(DimNum > 1, "TMultiArrayView requires a positive, non-zero number of dimensions");
+	static_assert(TIsSigned<SizeType>::Value, "TMultiArrayView only supports signed index types");
+
+public:
+
+	/**
+	 * Constructor.
+	 */
+	TMultiArrayView() : DataPtr(nullptr) { }
+
+	/**
+	 * Construct a view of an arbitrary pointer
+	 *
+	 * @param InData	The data to view
+	 * @param InShape	The number of elements on each dimension
+	 */
+	FORCEINLINE TMultiArrayView(ElementType* InData, TMultiArrayShape<DimNum, SizeType> InShape)
+		: DataPtr(InData)
+		, ArrayShape(InShape)
+	{
+		CheckInvariants();
+	}
+
+	/**
+	 * Array bracket operator. Returns reference to sub-view at given index.
+	 *
+	 * @returns Reference to indexed element.
+	 */
+	FORCEINLINE TMultiArrayView<DimNum - 1, ElementType, bIsChecked, bIsRestrict, SizeType> operator[](SizeType Index) const
+	{
+		RangeCheck(0, Index);
+
+		TMultiArrayShape<DimNum - 1, SizeType> NewShape;
+		for (uint8 Idx = 0; Idx < DimNum - 1; Idx++)
+		{
+			NewShape[Idx] = ArrayShape[Idx + 1];
+		}
+
+		return TMultiArrayView<DimNum - 1, ElementType, bIsChecked, bIsRestrict, SizeType>(DataPtr + Index * Stride(0), NewShape);
+	}
+
+	/**
+	 * Flattens the array view into a single dimension.
+	 *
+	 * @returns Flattened array view.
+	 */
+	FORCEINLINE TMultiArrayView<1, ElementType, bIsChecked, bIsRestrict, SizeType> Flatten() const
+	{
+		return TMultiArrayView<1, ElementType, bIsChecked, bIsRestrict, SizeType>(DataPtr, Num());
+	}
+
+	/**
+	 * Flattens the array on a given dimension, merging that dimension and the following one.
+	 * e.g. a 3D MultiArrayView with shape [10, 5, 3] flattened on dimension 0 will become a 2D MultiArrayView with shape [10 * 5, 3]
+	 *
+	 * @returns View flattened on the given dimension.
+	 */
+	FORCEINLINE TMultiArrayView<DimNum - 1, ElementType, bIsChecked, bIsRestrict, SizeType> Flatten(uint8 Dimension) const
+	{
+		if (bIsChecked)
+		{
+			checkf((Dimension < DimNum - 1), TEXT("MultiArray flatten dimension out of bounds: %lld from a rank of %lld"), (long long)Dimension, (long long)DimNum); // & for one branch
+		}
+
+		TMultiArrayShape<DimNum - 1, SizeType> NewShape;
+		uint8 SrcIdx = 0;
+		for (uint8 Idx = 0; Idx < DimNum - 1; Idx++)
+		{
+			if (Idx == Dimension)
+			{
+				NewShape[Idx] = ArrayShape[SrcIdx] * ArrayShape[SrcIdx+1];
+				SrcIdx += 2;
+			}
+			else
+			{
+				NewShape[Idx] = ArrayShape[SrcIdx];
+				SrcIdx += 1;
+			}
+		}
+
+		return TMultiArrayView<DimNum - 1, ElementType, bIsChecked, bIsRestrict, SizeType>(DataPtr, NewShape);
+	}
+
+	/**
+	 * Flattens the array on a given dimension, merging that dimension and the following one.
+	 * e.g. a 3D MultiArrayView with shape [10, 5, 3] flattened on dimension 0 will become a 2D MultiArrayView with shape [10 * 5, 3]
+	 *
+	 * @returns View flattened on the given dimension.
+	 */
+	template<uint8 Dimension>
+	FORCEINLINE TMultiArrayView<DimNum - 1, ElementType, bIsChecked, bIsRestrict, SizeType> Flatten() const
+	{
+		static_assert(Dimension < DimNum - 1, "MultiArray flatten dimension out of bounds");
+
+		TMultiArrayShape<DimNum - 1, SizeType> NewShape;
+		uint8 SrcIdx = 0;
+		for (uint8 Idx = 0; Idx < DimNum - 1; Idx++)
+		{
+			if (Idx == Dimension)
+			{
+				NewShape[Idx] = ArrayShape[SrcIdx] * ArrayShape[SrcIdx + 1];
+				SrcIdx += 2;
+			}
+			else
+			{
+				NewShape[Idx] = ArrayShape[SrcIdx];
+				SrcIdx += 1;
+			}
+		}
+
+		return TMultiArrayView<DimNum - 1, ElementType, bIsChecked, bIsRestrict, SizeType>(DataPtr, NewShape);
+	}
+
+public:
+
+	/**
+	 * Helper function for returning a typed pointer to the first array entry.
+	 *
+	 * @returns Pointer to first array entry or nullptr if empty
+	 */
+	FORCEINLINE PointerType GetData()
+	{
+		return DataPtr;
+	}
+
+	/**
+	 * Helper function for returning a typed pointer to the first array entry.
+	 *
+	 * @returns Pointer to first array entry or nullptr if empty
+	 */
+	FORCEINLINE const PointerType GetData() const
+	{
+		return DataPtr;
+	}
+
+
+	/**
+	 * DO NOT USE DIRECTLY
+	 * STL-like iterators to enable range-based for loop support.
+	 */
+	FORCEINLINE ElementType* begin() const { return GetData(); }
+	FORCEINLINE ElementType* end() const { return GetData() + Num(); }
+
+	/**
+	 * Helper function returning the size of the inner type.
+	 *
+	 * @returns Size in bytes of array type.
+	 */
+	FORCEINLINE static constexpr size_t GetTypeSize()
+	{
+		return sizeof(ElementType);
+	}
+
+	/**
+	 * Helper function returning the alignment of the inner type.
+	 */
+	FORCEINLINE static constexpr size_t GetTypeAlignment()
+	{
+		return alignof(ElementType);
+	}
+
+	/**
+	 * Checks array invariants: if array size is greater than zero on each dimension
+	 */
+	FORCEINLINE void CheckInvariants() const
+	{
+		if (bIsChecked)
+		{
+			for (SizeType Idx = 0; Idx < DimNum; Idx++)
+			{
+				checkSlow(ArrayShape[Idx] >= 0);
+			}
+		}
+	}
+
+	/**
+	 * Checks if a dimension is within the allowed number of dimensions
+	 *
+	 * @param Dimension Dimension of the array.
+	 */
+	FORCEINLINE void DimensionCheck(uint8 Dimension) const
+	{
+		if (bIsChecked)
+		{
+			checkf((Dimension < DimNum), TEXT("MultiArray dimension out of bounds: %lld from a rank of %lld"), (long long)Dimension, (long long)DimNum); // & for one branch
+		}
+	}
+
+	/**
+	 * Checks if index is in dimension range.
+	 *
+	 * @param Dimension Dimension of the array.
+	 * @param Index Index to check.
+	 */
+	FORCEINLINE void RangeCheck(uint8 Dimension, SizeType Index) const
+	{
+		DimensionCheck(Dimension);
+		CheckInvariants();
+
+		if (bIsChecked)
+		{
+			checkf((Index >= 0) & (Index < ArrayShape[Dimension]), TEXT("MultiArray index out of bounds: %lld from a dimension of size %lld"), (long long)Index, (long long)ArrayShape[Dimension]); // & for one branch
+		}
+	}
+
+	/**
+	 * Checks if a slice range [Index, Index+InNum) is in dimension range.
+	 * Length is 0 is allowed on empty dimensions; Index must be 0 in that case.
+	 *
+	 * @param Dimension Dimension of the array.
+	 * @param Index Starting index of the slice.
+	 * @param InNum Length of the slice.
+	 */
+	FORCEINLINE void SliceRangeCheck(uint8 Dimension, SizeType Index, SizeType InNum) const
+	{
+		DimensionCheck(Dimension);
+
+		if (bIsChecked)
+		{
+			checkf(Index >= 0, TEXT("Invalid index (%lld)"), (long long)Index);
+			checkf(InNum >= 0, TEXT("Invalid count (%lld)"), (long long)InNum);
+			checkf(Index + InNum <= ArrayShape[Dimension], TEXT("Range (index: %lld, count: %lld) lies outside the view of %lld elements"), (long long)Index, (long long)InNum, (long long)ArrayShape[Dimension]);
+		}
+	}
+
+	/**
+	 * Tests if dimension is valid, i.e. less than the number of dimensions.
+	 *
+	 * @param Dimension Dimension of the array.
+	 *
+	 * @returns True if dimension is valid. False otherwise.
+	 */
+	FORCEINLINE bool IsValidDimension(uint8 Dimension) const
+	{
+		return Dimension < DimNum;
+	}
+
+	/**
+	 * Tests if index is valid, i.e. greater than or equal to zero, and less than the number of elements in the dimension.
+	 *
+	 * @param Dimension Dimension of the array.
+	 * @param Index Index to test.
+	 *
+	 * @returns True if index is valid. False otherwise.
+	 */
+	FORCEINLINE bool IsValidIndex(uint8 Dimension, SizeType Index) const
+	{
+		DimensionCheck(Dimension);
+
+		return (Index >= 0) && (Index < ArrayShape[Dimension]);
+	}
+
+	/**
+	 * Returns true if the array is empty and contains no elements.
+	 *
+	 * @returns True if the array is empty.
+	 * @see Num
+	 */
+	bool IsEmpty() const
+	{
+		return Num() == 0;
+	}
+
+	/**
+	 * Returns true if the dimension is empty and contains no elements.
+	 *
+	 * @param Dimension Dimension of the array.
+	 * @returns True if the dimension is empty.
+	 * @see Num
+	 */
+	bool IsEmpty(uint8 Dimension) const
+	{
+		DimensionCheck(Dimension);
+
+		return ArrayShape[Dimension] == 0;
+	}
+
+	/**
+	 * Returns the number of dimensions.
+	 *
+	 * @returns Number of dimensions in array.
+	 */
+	FORCEINLINE SizeType Rank() const
+	{
+		return DimNum;
+	}
+
+	/**
+	 * Returns the total number of elements
+	 *
+	 * @returns Total number of elements in array.
+	 */
+	FORCEINLINE SizeType Num() const
+	{
+		return ArrayShape.Total();
+	}
+
+	/**
+	 * Returns the number of elements in a dimension.
+	 *
+	 * @returns Number of elements in array.
+	 */
+	template<uint8 InDimIdx>
+	FORCEINLINE SizeType Num() const
+	{
+		static_assert(InDimIdx < DimNum);
+		return ArrayShape[InDimIdx];
+	}
+
+	/**
+	 * Returns the number of elements in a dimension.
+	 *
+	 * @param Dimension Dimension of the array.
+	 * @returns Number of elements in array.
+	 */
+	FORCEINLINE SizeType Num(uint8 Dimension) const
+	{
+		DimensionCheck(Dimension);
+
+		return ArrayShape[Dimension];
+	}
+
+	/**
+	 * Returns the number of elements in each dimension.
+	 *
+	 * @returns Number of elements in each dimension.
+	 */
+	FORCEINLINE TMultiArrayShape<DimNum, SizeType> Shape() const
+	{
+		return ArrayShape;
+	}
+
+	/**
+	 * Returns the stride for a dimension.
+	 *
+	 * @param Dimension Dimension of the array.
+	 * @returns Stride of that dimension.
+	 */
+	FORCEINLINE SizeType Stride(uint8 Dimension) const
+	{
+		DimensionCheck(Dimension);
+
+		SizeType Total = 1;
+		for (uint8 Idx = Dimension + 1; Idx < DimNum; Idx++)
+		{
+			Total *= ArrayShape[Idx];
+		}
+		return Total;
+	}
+
+	/**
+	 * Returns a sliced view. Slicing outside of the range of the view is illegal.
+	 *
+	 * @param Index Starting index of the new view
+	 * @param InNum Number of elements in the new view
+	 * @returns Sliced view
+	 */
+	[[nodiscard]] FORCEINLINE TMultiArrayView Slice(SizeType Index, SizeType InNum) const
+	{
+		SliceRangeCheck(0, Index, InNum);
+		TMultiArrayShape NewShape = ArrayShape;
+		NewShape[0] = InNum;
+		return TMultiArrayView(DataPtr + Index * Stride(0), NewShape);
+	}
+
+	/** Returns the left-most part of the view by taking the given number of elements from the left. */
+	[[nodiscard]] inline TMultiArrayView Left(SizeType Count) const
+	{
+		TMultiArrayShape NewShape = ArrayShape;
+		NewShape[0] = FMath::Clamp(Count, 0, ArrayShape[0]);
+		return TMultiArrayView(DataPtr, NewShape);
+	}
+
+	/** Returns the left-most part of the view by chopping the given number of elements from the right. */
+	[[nodiscard]] inline TMultiArrayView LeftChop(SizeType Count) const
+	{
+		TMultiArrayShape NewShape = ArrayShape;
+		NewShape[0] = FMath::Clamp(ArrayShape[0] - Count, 0, ArrayShape[0]);
+		return TMultiArrayView(DataPtr, NewShape);
+	}
+
+	/** Returns the right-most part of the view by taking the given number of elements from the right. */
+	[[nodiscard]] inline TMultiArrayView Right(SizeType Count) const
+	{
+		const SizeType OutLen = FMath::Clamp(Count, 0, ArrayShape[0]);
+		TMultiArrayShape NewShape = ArrayShape;
+		NewShape[0] = OutLen;
+		return TMultiArrayView(DataPtr + (ArrayShape[0] - OutLen) * Stride(0), NewShape);
+	}
+
+	/** Returns the right-most part of the view by chopping the given number of elements from the left. */
+	[[nodiscard]] inline TMultiArrayView RightChop(SizeType Count) const
+	{
+		const SizeType OutLen = FMath::Clamp(ArrayShape[0] - Count, 0, ArrayShape[0]);
+		TMultiArrayShape NewShape = ArrayShape;
+		NewShape[0] = OutLen;
+		return TMultiArrayView(DataPtr + (ArrayShape[0] - OutLen) * Stride(0), NewShape);
+	}
+
+	/** Returns the middle part of the view by taking up to the given number of elements from the given position. */
+	[[nodiscard]] inline TMultiArrayView Mid(SizeType Index, SizeType Count = TNumericLimits<SizeType>::Max()) const
+	{
+		ElementType* const CurrentStart = DataPtr;
+		const SizeType     CurrentLength = Num(0);
+
+		// Clamp minimum index at the start of the range, adjusting the length down if necessary
+		const SizeType NegativeIndexOffset = (Index < 0) ? Index : 0;
+		Count += NegativeIndexOffset;
+		Index -= NegativeIndexOffset;
+
+		// Clamp maximum index at the end of the range
+		Index = (Index > CurrentLength) ? CurrentLength : Index;
+
+		// Clamp count between 0 and the distance to the end of the range
+		Count = FMath::Clamp(Count, 0, (CurrentLength - Index));
+
+		TMultiArrayShape NewShape = ArrayShape;
+		NewShape[0] = Count;
+		TMultiArrayView Result = TMultiArrayView(CurrentStart + Index * Stride(0), NewShape);
+		return Result;
+	}
+
+	/** Modifies the view to be the given number of elements from the left. */
+	inline void LeftInline(SizeType CharCount)
+	{
+		*this = Left(CharCount);
+	}
+
+	/** Modifies the view by chopping the given number of elements from the right. */
+	inline void LeftChopInline(SizeType CharCount)
+	{
+		*this = LeftChop(CharCount);
+	}
+
+	/** Modifies the view to be the given number of elements from the right. */
+	inline void RightInline(SizeType CharCount)
+	{
+		*this = Right(CharCount);
+	}
+
+	/** Modifies the view by chopping the given number of elements from the left. */
+	inline void RightChopInline(SizeType CharCount)
+	{
+		*this = RightChop(CharCount);
+	}
+
+	/** Modifies the view to be the middle part by taking up to the given number of elements from the given position. */
+	inline void MidInline(SizeType Position, SizeType CharCount = TNumericLimits<SizeType>::Max())
+	{
+		*this = Mid(Position, CharCount);
+	}
+
+	/**
+	 * Checks if this array contains the element.
+	 *
+	 * @returns	True if found. False otherwise.
+	 */
+	template <typename ComparisonType>
+	bool Contains(const ComparisonType& Item) const
+	{
+		for (const ElementType* RESTRICT Data = GetData(), *RESTRICT DataEnd = Data + Num(); Data != DataEnd; ++Data)
+		{
+			if (*Data == Item)
+			{
+				return true;
+			}
+		}
+		return false;
+	}
+
+public:
+
+	/** Implicit cast for Checked and Restrict. */
+	FORCEINLINE operator TMultiArrayView<DimNum, ElementType, true, false, SizeType>()
+	{
+		return TMultiArrayView<DimNum, ElementType, true, false, SizeType>(DataPtr, ArrayShape);
+	}
+
+	/** Implicit cast for Checked and Restrict. */
+	FORCEINLINE operator TMultiArrayView<DimNum, ElementType, false, false, SizeType>()
+	{
+		return TMultiArrayView<DimNum, ElementType, false, false, SizeType>(DataPtr, ArrayShape);
+	}
+
+	/** Implicit cast for Checked and Restrict. */
+	FORCEINLINE operator TMultiArrayView<DimNum, ElementType, true, true, SizeType>()
+	{
+		return TMultiArrayView<DimNum, ElementType, true, true, SizeType>(DataPtr, ArrayShape);
+	}
+
+	/** Implicit cast for Checked and Restrict. */
+	FORCEINLINE operator TMultiArrayView<DimNum, ElementType, false, true, SizeType>()
+	{
+		return TMultiArrayView<DimNum, ElementType, false, true, SizeType>(DataPtr, ArrayShape);
+	}
+
+	/** Implicit cast for Checked and Restrict. */
+	FORCEINLINE operator TMultiArrayView<DimNum, const ElementType, true, false, SizeType>() const
+	{
+		return TMultiArrayView<DimNum, const ElementType, true, false, SizeType>(DataPtr, ArrayShape);
+	}
+
+	/** Implicit cast for Checked and Restrict. */
+	FORCEINLINE operator TMultiArrayView<DimNum, const ElementType, false, false, SizeType>() const
+	{
+		return TMultiArrayView<DimNum, const ElementType, false, false, SizeType>(DataPtr, ArrayShape);
+	}
+
+	/** Implicit cast for Checked and Restrict. */
+	FORCEINLINE operator TMultiArrayView<DimNum, const ElementType, true, true, SizeType>() const
+	{
+		return TMultiArrayView<DimNum, const ElementType, true, true, SizeType>(DataPtr, ArrayShape);
+	}
+
+	/** Implicit cast for Checked and Restrict. */
+	FORCEINLINE operator TMultiArrayView<DimNum, const ElementType, false, true, SizeType>() const
+	{
+		return TMultiArrayView<DimNum, const ElementType, false, true, SizeType>(DataPtr, ArrayShape);
+	}
+
+
+private:
+
+	PointerType DataPtr;
+	TMultiArrayShape<DimNum, SizeType> ArrayShape;
+};
+
+
+/**
+* Specialization for single dimensional MultiArrayView
+*/
+template<typename InElementType, bool bInIsChecked, bool bInIsRestrict, typename InSizeType>
+class TMultiArrayView<1, InElementType, bInIsChecked, bInIsRestrict, InSizeType>
+{
+public:
+	static constexpr uint8 DimNum = 1;
+	using ElementType = InElementType;
+	static constexpr bool bIsChecked = bInIsChecked;
+	static constexpr bool bIsRestrict = bInIsRestrict;
+	using SizeType = InSizeType;
+
+	using PointerType = typename TChooseClass<
+		bIsRestrict,
+		ElementType* RESTRICT,
+		ElementType*>::Result;
+
+	static_assert(TIsSigned<SizeType>::Value, "TMultiArrayView only supports signed index types");
+
+	/**
+	 * Constructor.
+	 */
+	TMultiArrayView() : DataPtr(nullptr) { }
+
+private:
+	template <typename T>
+	using TIsCompatibleElementType = UE::MultiArrayView::Private::TIsCompatibleElementType<T, ElementType>;
+
+	template <typename T>
+	using TIsCompatibleRangeType = UE::MultiArrayView::Private::TIsCompatibleRangeType<T, ElementType>;
+
+	template <typename T>
+	using TIsReinterpretableRangeType = UE::MultiArrayView::Private::TIsReinterpretableRangeType<T, ElementType>;
+
+public:
+	/**
+	 * Constructor from another range
+	 *
+	 * @param Other The source range to copy
+	 */
+	template <
+		typename OtherRangeType,
+		typename CVUnqualifiedOtherRangeType = std::remove_cv_t<typename TRemoveReference<OtherRangeType>::Type>,
+		typename = typename TEnableIf<
+			TAnd<
+				TIsContiguousContainer<CVUnqualifiedOtherRangeType>,
+				TOr<
+					TIsCompatibleRangeType<OtherRangeType>,
+					TIsReinterpretableRangeType<OtherRangeType>
+				>
+			>::Value
+		>::Type
+	>
+		FORCEINLINE TMultiArrayView(OtherRangeType&& Other)
+		: DataPtr(TChooseClass<
+			TIsCompatibleRangeType<OtherRangeType>::Value,
+			TIsCompatibleRangeType<OtherRangeType>,
+			TIsReinterpretableRangeType<OtherRangeType>
+		>::Result::GetData(Forward<OtherRangeType>(Other)))
+	{
+		const auto InCount = GetNum(Forward<OtherRangeType>(Other));
+		check((InCount >= 0) && ((sizeof(InCount) < sizeof(SizeType)) || (InCount <= static_cast<decltype(InCount)>(TNumericLimits<SizeType>::Max()))));
+		ArrayShape[0] = (SizeType)InCount;
+	}
+
+	/**
+	 * Construct a view of an arbitrary pointer
+	 *
+	 * @param InData	The data to view
+	 * @param InNums	The number of elements on each dimension
+	 */
+	FORCEINLINE TMultiArrayView(ElementType* InData, TMultiArrayShape<DimNum, SizeType> InShape)
+		: DataPtr(InData)
+		, ArrayShape(InShape)
+	{
+		CheckInvariants();
+	}
+
+	/**
+	 * Construct a view of an arbitrary pointer
+	 *
+	 * @param InData	The data to view
+	 * @param InNum	The number of elements
+	 */
+	template <typename OtherElementType,
+		typename = typename TEnableIf<TIsCompatibleElementType<OtherElementType>::Value>::Type>
+		FORCEINLINE TMultiArrayView(OtherElementType* InData, SizeType InNum)
+		: DataPtr(InData)
+	{
+		ArrayShape[0] = InNum;
+		CheckInvariants();
+	}
+
+	/**
+	 * Construct a view of an initializer list.
+	 *
+	 * The caller is responsible for ensuring that the view does not outlive the initializer list.
+	 */
+	FORCEINLINE TMultiArrayView(std::initializer_list<ElementType> List)
+		: DataPtr(UE::MultiArrayView::Private::GetDataHelper(List))
+	{
+		ArrayShape[0] = GetNum(List);
+		CheckInvariants();
+	}
+
+public:
+
+	/**
+	 * Array bracket operator. Returns a reference to the element at given index.
+	 *
+	 * @returns Reference to indexed element.
+	 */
+	FORCEINLINE ElementType& operator[](SizeType Index) const
+	{
+		RangeCheck(0, Index);
+		return DataPtr[Index];
+	}
+
+	/** Implicit cast to TArrayView. */
+	FORCEINLINE operator TArrayView<ElementType, SizeType>() const
+	{
+		return TArrayView(DataPtr, ArrayShape[0]);
+	}
+
+public:
+
+	/**
+	 * Helper function for returning a typed pointer to the first array entry.
+	 *
+	 * @returns Pointer to first array entry or nullptr if empty
+	 */
+	FORCEINLINE PointerType GetData()
+	{
+		return DataPtr;
+	}
+
+	/**
+	 * Helper function for returning a typed pointer to the first array entry.
+	 *
+	 * @returns Pointer to first array entry or nullptr if empty
+	 */
+	FORCEINLINE const PointerType GetData() const
+	{
+		return DataPtr;
+	}
+
+	/**
+	 * DO NOT USE DIRECTLY
+	 * STL-like iterators to enable range-based for loop support.
+	 */
+	FORCEINLINE ElementType* begin() const { return GetData(); }
+	FORCEINLINE ElementType* end() const { return GetData() + Num(); }
+
+	/**
+	 * Helper function returning the size of the inner type.
+	 *
+	 * @returns Size in bytes of array type.
+	 */
+	FORCEINLINE static constexpr size_t GetTypeSize()
+	{
+		return sizeof(ElementType);
+	}
+
+	/**
+	 * Helper function returning the alignment of the inner type.
+	 */
+	FORCEINLINE static constexpr size_t GetTypeAlignment()
+	{
+		return alignof(ElementType);
+	}
+
+	/**
+	 * Checks array invariants: if array size is greater than zero on each dimension
+	 */
+	FORCEINLINE void CheckInvariants() const
+	{
+		if (bIsChecked)
+		{
+			for (SizeType Idx = 0; Idx < DimNum; Idx++)
+			{
+				checkSlow(ArrayShape[Idx] >= 0);
+			}
+		}
+	}
+
+	/**
+	 * Checks if a dimension is within the allowed number of dimensions
+	 *
+	 * @param Dimension Dimension of the array.
+	 */
+	FORCEINLINE void DimensionCheck(uint8 Dimension) const
+	{
+		if (bIsChecked)
+		{
+			checkf((Dimension < DimNum), TEXT("MultiArray dimension out of bounds: %lld from a rank of %lld"), (long long)Dimension, (long long)DimNum); // & for one branch
+		}
+	}
+
+	/**
+	 * Checks if index is in dimension range.
+	 *
+	 * @param Dimension Dimension of the array.
+	 * @param Index Index to check.
+	 */
+	FORCEINLINE void RangeCheck(uint8 Dimension, SizeType Index) const
+	{
+		DimensionCheck(Dimension);
+		CheckInvariants();
+
+		if (bIsChecked)
+		{
+			checkf((Index >= 0) & (Index < ArrayShape[Dimension]), TEXT("MultiArray index out of bounds: %lld from a dimension of size %lld"), (long long)Index, (long long)ArrayShape[Dimension]); // & for one branch
+		}
+	}
+
+	/**
+	 * Checks if a slice range [Index, Index+InNum) is in dimension range.
+	 * Length is 0 is allowed on empty dimensions; Index must be 0 in that case.
+	 *
+	 * @param Dimension Dimension of the array.
+	 * @param Index Starting index of the slice.
+	 * @param InNum Length of the slice.
+	 */
+	FORCEINLINE void SliceRangeCheck(uint8 Dimension, SizeType Index, SizeType InNum) const
+	{
+		DimensionCheck(Dimension);
+
+		if (bIsChecked)
+		{
+			checkf(Index >= 0, TEXT("Invalid index (%lld)"), (long long)Index);
+			checkf(InNum >= 0, TEXT("Invalid count (%lld)"), (long long)InNum);
+			checkf(Index + InNum <= ArrayShape[Dimension], TEXT("Range (index: %lld, count: %lld) lies outside the view of %lld elements"), (long long)Index, (long long)InNum, (long long)ArrayShape[Dimension]);
+		}
+	}
+
+	/**
+	 * Tests if dimension is valid, i.e. less than the number of dimensions.
+	 *
+	 * @param Dimension Dimension of the array.
+	 *
+	 * @returns True if dimension is valid. False otherwise.
+	 */
+	FORCEINLINE bool IsValidDimension(uint8 Dimension) const
+	{
+		return Dimension < DimNum;
+	}
+
+	/**
+	 * Tests if index is valid, i.e. greater than or equal to zero, and less than the number of elements in the dimension.
+	 *
+	 * @param Dimension Dimension of the array.
+	 * @param Index Index to test.
+	 *
+	 * @returns True if index is valid. False otherwise.
+	 */
+	FORCEINLINE bool IsValidIndex(uint8 Dimension, SizeType Index) const
+	{
+		DimensionCheck(Dimension);
+
+		return (Index >= 0) && (Index < ArrayShape[Dimension]);
+	}
+
+	/**
+	 * Returns true if the array is empty and contains no elements.
+	 *
+	 * @returns True if the array is empty.
+	 * @see Num
+	 */
+	bool IsEmpty() const
+	{
+		return Num() == 0;
+	}
+
+	/**
+	 * Returns true if the dimension is empty and contains no elements.
+	 *
+	 * @param Dimension Dimension of the array.
+	 * @returns True if the dimension is empty.
+	 * @see Num
+	 */
+	bool IsEmpty(uint8 Dimension) const
+	{
+		DimensionCheck(Dimension);
+
+		return ArrayShape[Dimension] == 0;
+	}
+
+	/**
+	 * Returns the number of dimensions.
+	 *
+	 * @returns Number of dimensions in array.
+	 */
+	FORCEINLINE SizeType Rank() const
+	{
+		return DimNum;
+	}
+
+	/**
+	 * Returns the total number of elements
+	 *
+	 * @returns Total number of elements in array.
+	 */
+	FORCEINLINE SizeType Num() const
+	{
+		return ArrayShape.Total();
+	}
+
+	/**
+	 * Returns the number of elements in a dimension.
+	 *
+	 * @returns Number of elements in array.
+	 */
+	template<uint8 InDimIdx>
+	FORCEINLINE SizeType Num() const
+	{
+		static_assert(InDimIdx < DimNum);
+		return ArrayShape[InDimIdx];
+	}
+
+	/**
+	 * Returns the number of elements in a dimension.
+	 *
+	 * @param Dimension Dimension of the array.
+	 * @returns Number of elements in array.
+	 */
+	FORCEINLINE SizeType Num(uint8 Dimension) const
+	{
+		DimensionCheck(Dimension);
+
+		return ArrayShape[Dimension];
+	}
+
+	/**
+	 * Returns the number of elements in each dimension.
+	 *
+	 * @returns Number of elements in each dimension.
+	 */
+	FORCEINLINE TMultiArrayShape<DimNum, SizeType> Shape() const
+	{
+		return ArrayShape;
+	}
+
+	/**
+	 * Returns the stride for a dimension.
+	 *
+	 * @param Dimension Dimension of the array.
+	 * @returns Stride of that dimension.
+	 */
+	FORCEINLINE SizeType Stride(uint8 Dimension) const
+	{
+		DimensionCheck(Dimension);
+
+		SizeType Total = 1;
+		for (uint8 Idx = Dimension + 1; Idx < DimNum; Idx++)
+		{
+			Total *= ArrayShape[Idx];
+		}
+		return Total;
+	}
+
+	/**
+	 * Returns a sliced view. Slicing outside of the range of the view is illegal.
+	 *
+	 * @param Index Starting index of the new view
+	 * @param InNum Number of elements in the new view
+	 * @returns Sliced view
+	 */
+	[[nodiscard]] FORCEINLINE TMultiArrayView Slice(SizeType Index, SizeType InNum) const
+	{
+		SliceRangeCheck(0, Index, InNum);
+		TMultiArrayShape NewShape = ArrayShape;
+		NewShape[0] = InNum;
+		return TMultiArrayView(DataPtr + Index * Stride(0), NewShape);
+	}
+
+	/** Returns the left-most part of the view by taking the given number of elements from the left. */
+	[[nodiscard]] inline TMultiArrayView Left(SizeType Count) const
+	{
+		TMultiArrayShape NewShape = ArrayShape;
+		NewShape[0] = FMath::Clamp(Count, 0, ArrayShape[0]);
+		return TMultiArrayView(DataPtr, NewShape);
+	}
+
+	/** Returns the left-most part of the view by chopping the given number of elements from the right. */
+	[[nodiscard]] inline TMultiArrayView LeftChop(SizeType Count) const
+	{
+		TMultiArrayShape NewShape = ArrayShape;
+		NewShape[0] = FMath::Clamp(ArrayShape[0] - Count, 0, ArrayShape[0]);
+		return TMultiArrayView(DataPtr, NewShape);
+	}
+
+	/** Returns the right-most part of the view by taking the given number of elements from the right. */
+	[[nodiscard]] inline TMultiArrayView Right(SizeType Count) const
+	{
+		const SizeType OutLen = FMath::Clamp(Count, 0, ArrayShape[0]);
+		TMultiArrayShape NewShape = ArrayShape;
+		NewShape[0] = OutLen;
+		return TMultiArrayView(DataPtr + (ArrayShape[0] - OutLen) * Stride(0), NewShape);
+	}
+
+	/** Returns the right-most part of the view by chopping the given number of elements from the left. */
+	[[nodiscard]] inline TMultiArrayView RightChop(SizeType Count) const
+	{
+		const SizeType OutLen = FMath::Clamp(ArrayShape[0] - Count, 0, ArrayShape[0]);
+		TMultiArrayShape NewShape = ArrayShape;
+		NewShape[0] = OutLen;
+		return TMultiArrayView(DataPtr + (ArrayShape[0] - OutLen) * Stride(0), NewShape);
+	}
+
+	/** Returns the middle part of the view by taking up to the given number of elements from the given position. */
+	[[nodiscard]] inline TMultiArrayView Mid(SizeType Index, SizeType Count = TNumericLimits<SizeType>::Max()) const
+	{
+		ElementType* const CurrentStart = DataPtr;
+		const SizeType     CurrentLength = Num(0);
+
+		// Clamp minimum index at the start of the range, adjusting the length down if necessary
+		const SizeType NegativeIndexOffset = (Index < 0) ? Index : 0;
+		Count += NegativeIndexOffset;
+		Index -= NegativeIndexOffset;
+
+		// Clamp maximum index at the end of the range
+		Index = (Index > CurrentLength) ? CurrentLength : Index;
+
+		// Clamp count between 0 and the distance to the end of the range
+		Count = FMath::Clamp(Count, 0, (CurrentLength - Index));
+
+		TMultiArrayShape NewShape = ArrayShape;
+		NewShape[0] = Count;
+		TMultiArrayView Result = TMultiArrayView(CurrentStart + Index * Stride(0), NewShape);
+		return Result;
+	}
+
+	/** Modifies the view to be the given number of elements from the left. */
+	inline void LeftInline(SizeType CharCount)
+	{
+		*this = Left(CharCount);
+	}
+
+	/** Modifies the view by chopping the given number of elements from the right. */
+	inline void LeftChopInline(SizeType CharCount)
+	{
+		*this = LeftChop(CharCount);
+	}
+
+	/** Modifies the view to be the given number of elements from the right. */
+	inline void RightInline(SizeType CharCount)
+	{
+		*this = Right(CharCount);
+	}
+
+	/** Modifies the view by chopping the given number of elements from the left. */
+	inline void RightChopInline(SizeType CharCount)
+	{
+		*this = RightChop(CharCount);
+	}
+
+	/** Modifies the view to be the middle part by taking up to the given number of elements from the given position. */
+	inline void MidInline(SizeType Position, SizeType CharCount = TNumericLimits<SizeType>::Max())
+	{
+		*this = Mid(Position, CharCount);
+	}
+
+	/**
+	 * Checks if this array contains the element.
+	 *
+	 * @returns	True if found. False otherwise.
+	 */
+	template <typename ComparisonType>
+	bool Contains(const ComparisonType& Item) const
+	{
+		for (const ElementType* RESTRICT Data = GetData(), *RESTRICT DataEnd = Data + Num(); Data != DataEnd; ++Data)
+		{
+			if (*Data == Item)
+			{
+				return true;
+			}
+		}
+		return false;
+	}
+
+public:
+
+	/** Implicit cast for Checked and Restrict. */
+	FORCEINLINE operator TMultiArrayView<DimNum, ElementType, true, false, SizeType>()
+	{
+		return TMultiArrayView<DimNum, ElementType, true, false, SizeType>(DataPtr, ArrayShape);
+	}
+
+	/** Implicit cast for Checked and Restrict. */
+	FORCEINLINE operator TMultiArrayView<DimNum, ElementType, false, false, SizeType>()
+	{
+		return TMultiArrayView<DimNum, ElementType, false, false, SizeType>(DataPtr, ArrayShape);
+	}
+
+	/** Implicit cast for Checked and Restrict. */
+	FORCEINLINE operator TMultiArrayView<DimNum, ElementType, true, true, SizeType>()
+	{
+		return TMultiArrayView<DimNum, ElementType, true, true, SizeType>(DataPtr, ArrayShape);
+	}
+
+	/** Implicit cast for Checked and Restrict. */
+	FORCEINLINE operator TMultiArrayView<DimNum, ElementType, false, true, SizeType>()
+	{
+		return TMultiArrayView<DimNum, ElementType, false, true, SizeType>(DataPtr, ArrayShape);
+	}
+
+	/** Implicit cast for Checked and Restrict. */
+	FORCEINLINE operator TMultiArrayView<DimNum, const ElementType, true, false, SizeType>() const
+	{
+		return TMultiArrayView<DimNum, const ElementType, true, false, SizeType>(DataPtr, ArrayShape);
+	}
+
+	/** Implicit cast for Checked and Restrict. */
+	FORCEINLINE operator TMultiArrayView<DimNum, const ElementType, false, false, SizeType>() const
+	{
+		return TMultiArrayView<DimNum, const ElementType, false, false, SizeType>(DataPtr, ArrayShape);
+	}
+
+	/** Implicit cast for Checked and Restrict. */
+	FORCEINLINE operator TMultiArrayView<DimNum, const ElementType, true, true, SizeType>() const
+	{
+		return TMultiArrayView<DimNum, const ElementType, true, true, SizeType>(DataPtr, ArrayShape);
+	}
+
+	/** Implicit cast for Checked and Restrict. */
+	FORCEINLINE operator TMultiArrayView<DimNum, const ElementType, false, true, SizeType>() const
+	{
+		return TMultiArrayView<DimNum, const ElementType, false, true, SizeType>(DataPtr, ArrayShape);
+	}
+
+private:
+
+	PointerType DataPtr;
+	TMultiArrayShape<DimNum, SizeType> ArrayShape;
+};
diff --git a/Engine/Plugins/Experimental/LearningAgents/Source/LearningAgents/LearningAgents.Build.cs b/Engine/Plugins/Experimental/LearningAgents/Source/LearningAgents/LearningAgents.Build.cs
new file mode 100644
--- /dev/null
+++ b/Engine/Plugins/Experimental/LearningAgents/Source/LearningAgents/LearningAgents.Build.cs
@@ -0,0 +1,43 @@
+// Copyright Epic Games, Inc. All Rights Reserved.
+
+using UnrealBuildTool;
+
+public class LearningAgents : ModuleRules
+{
+	public LearningAgents(ReadOnlyTargetRules Target) : base(Target)
+	{
+		PCHUsage = ModuleRules.PCHUsageMode.UseExplicitOrSharedPCHs;
+
+		PublicIncludePaths.AddRange(
+			new string[] {
+				// ... add public include paths required here ...
+			});
+		
+		PrivateIncludePaths.AddRange(
+			new string[] {
+				// ... add other private include paths required here ...
+			});
+		
+		PublicDependencyModuleNames.AddRange(
+			new string[]
+			{
+				"Core",
+				"Learning",
+				// ... add other public dependencies that you statically link with here ...
+			});
+		
+		PrivateDependencyModuleNames.AddRange(
+			new string[]
+			{
+				"CoreUObject",
+				"Engine",
+				// ... add private dependencies that you statically link with here ...	
+			});
+		
+		DynamicallyLoadedModuleNames.AddRange(
+			new string[]
+			{
+				// ... add any modules that your module loads dynamically here ...
+			});
+	}
+}
diff --git a/Engine/Plugins/Experimental/LearningAgents/Source/LearningAgents/Private/LearningAgentsActions.cpp b/Engine/Plugins/Experimental/LearningAgents/Source/LearningAgents/Private/LearningAgentsActions.cpp
new file mode 100644
--- /dev/null
+++ b/Engine/Plugins/Experimental/LearningAgents/Source/LearningAgents/Private/LearningAgentsActions.cpp
@@ -0,0 +1,359 @@
+// Copyright Epic Games, Inc. All Rights Reserved.
+
+#include "LearningAgentsActions.h"
+
+#include "LearningAgentsType.h"
+#include "LearningArray.h"
+#include "LearningArrayMap.h"
+#include "LearningFeatureObject.h"
+#include "LearningLog.h"
+
+#include "GameFramework/Actor.h"
+#include "VisualLogger/VisualLogger.h"
+
+#define UE_LEARNING_AGENTS_VLOG_STRING(Owner, Category, Verbosity, Location, Color, Format, ...) \
+	UE_VLOG_LOCATION(Owner, Category, Verbosity, Location, 0.0f, Color, Format, ##__VA_ARGS__)
+
+#define UE_LEARNING_AGENTS_VLOG_TRANSFORM(Owner, Category, Verbosity, Location, Rotation, Color, Format, ...) \
+	UE_VLOG_SEGMENT(Owner, Category, Verbosity, Location, Location + 15.0f * Rotation.RotateVector(FVector::ForwardVector), FColor::Red, TEXT("")); \
+	UE_VLOG_SEGMENT(Owner, Category, Verbosity, Location, Location + 15.0f * Rotation.RotateVector(FVector::RightVector), FColor::Green, TEXT("")); \
+	UE_VLOG_SEGMENT(Owner, Category, Verbosity, Location, Location + 15.0f * Rotation.RotateVector(FVector::UpVector), FColor::Blue, TEXT("")); \
+	UE_VLOG_OBOX(Owner, Category, Verbosity, FBox(10.0f * FVector(-1, -1, -1), 10.0f * FVector(1, 1, 1)), FTransform(Rotation, Location, FVector::OneVector).ToMatrixNoScale(), Color, TEXT("")); \
+	UE_LEARNING_AGENTS_VLOG_STRING(Owner, Category, Verbosity, Location + FVector(0.0f, 0.0f, 20.0f), Color, Format, ##__VA_ARGS__)
+
+#define UE_LEARNING_AGENTS_VLOG_PLANE(Owner, Category, Verbosity, Location, Rotation, Axis0, Axis1, Color, Format, ...) \
+	UE_VLOG_OBOX(Owner, Category, Verbosity, FBox(25.0f * FVector(-1, -1, 0), 25.0f * FVector(1, 1, 0)), FTransform(Rotation, Location, FVector::OneVector).ToMatrixNoScale(), Color, Format, ##__VA_ARGS__)
+
+namespace UE::Learning::Agents::Private
+{
+	template<typename ActionUObject, typename ActionFObject, typename... InArgTypes>
+	ActionUObject* AddAction(ULearningAgentsType* AgentType, FName Name, InArgTypes&& ...Args)
+	{
+		if (!AgentType)
+		{
+			UE_LOG(LogLearning, Error, TEXT("AgentType is nullptr"));
+			return nullptr;
+		}
+
+		ActionUObject* Action = NewObject<ActionUObject>(AgentType, Name);
+
+		Action->FeatureObject = MakeShared<ActionFObject>(
+			Action->GetFName(),
+			AgentType->GetInstanceData().ToSharedRef(),
+			AgentType->GetMaxInstanceNum(),
+			Forward<InArgTypes>(Args)...);
+
+		AgentType->AddAction(Action, Action->FeatureObject.ToSharedRef());
+
+		return Action;
+	}
+}
+
+//------------------------------------------------------------------
+
+UFloatAction* UFloatAction::AddFloatAction(ULearningAgentsType* AgentType, FName Name, float Scale)
+{
+	return UE::Learning::Agents::Private::AddAction<UFloatAction, UE::Learning::FFloatFeature>(AgentType, Name, 1, Scale);
+}
+
+float UFloatAction::GetFloatAction(int32 AgentId)
+{
+	const TLearningArrayView<2, const float> View = FeatureObject->InstanceData->ConstView(FeatureObject->ValueHandle);
+
+	if (AgentId == INDEX_NONE)
+	{
+		UE_LOG(LogLearning, Error, TEXT("AgentId is invalid (INDEX_NONE)"));
+		return 0.0f;
+	}
+
+	if (AgentId < 0 || AgentId >= View.Num<0>())
+	{
+		UE_LOG(LogLearning, Error, TEXT("AgentId %d is out of index. Valid range [0, %d]."), AgentId, View.Num<0>() - 1);
+		return 0.0f;
+	}
+
+	return View[AgentId][0];
+}
+
+void UFloatAction::SetFloatAction(int32 AgentId, float Value)
+{
+	TLearningArrayView<2, float> View = FeatureObject->InstanceData->View(FeatureObject->ValueHandle);
+
+	if (AgentId == INDEX_NONE)
+	{
+		UE_LOG(LogLearning, Error, TEXT("AgentId is invalid (INDEX_NONE)"));
+		return;
+	}
+
+	if (AgentId < 0 || AgentId >= View.Num<0>())
+	{
+		UE_LOG(LogLearning, Error, TEXT("AgentId %d is out of index. Valid range [0, %d]."), AgentId, View.Num<0>() - 1);
+		return;
+	}
+
+	View[AgentId][0] = Value;
+}
+
+#if ENABLE_VISUAL_LOG
+void UFloatAction::VisualLog(const UE::Learning::FIndexSet Instances) const
+{
+	UE_LEARNING_TRACE_CPUPROFILER_EVENT_SCOPE(UFloatAction::VisualLog);
+
+	const TLearningArrayView<2, const float> ValueView = FeatureObject->InstanceData->ConstView(FeatureObject->ValueHandle);
+	const TLearningArrayView<2, const float> FeatureView = FeatureObject->InstanceData->ConstView(FeatureObject->FeatureHandle);
+
+	if (const ULearningAgentsType* AgentType = Cast<ULearningAgentsType>(GetOuter()))
+	{
+		for (const int32 Instance : Instances)
+		{
+			if (const AActor* Actor = Cast<AActor>(AgentType->GetAgent(Instance)))
+			{
+				UE_LEARNING_AGENTS_VLOG_STRING(this, LogLearning, Display,
+					Actor->GetActorLocation(),
+					VisualLogColor.ToFColor(true),
+					TEXT("Agent %i\nScale: [% 6.2f]\nValue: [% 6.2f]\nEncoded: [% 6.3f]"),
+					Instance,
+					FeatureObject->Scale,
+					ValueView[Instance][0],
+					FeatureView[Instance][0]);
+			}
+		}
+	}
+}
+#endif
+
+//------------------------------------------------------------------
+
+UVectorAction* UVectorAction::AddVectorAction(ULearningAgentsType* AgentType, FName Name, float Scale)
+{
+	return UE::Learning::Agents::Private::AddAction<UVectorAction, UE::Learning::FFloatFeature>(AgentType, Name, 3, Scale);
+}
+
+FVector UVectorAction::GetVectorAction(int32 AgentId)
+{
+	const TLearningArrayView<2, const float> View = FeatureObject->InstanceData->ConstView(FeatureObject->ValueHandle);
+
+	if (AgentId == INDEX_NONE)
+	{
+		UE_LOG(LogLearning, Error, TEXT("AgentId is invalid (INDEX_NONE)"));
+		return FVector::ZeroVector;
+	}
+
+	if (AgentId < 0 || AgentId >= View.Num<0>())
+	{
+		UE_LOG(LogLearning, Error, TEXT("AgentId %d is out of index. Valid range [0, %d]."), AgentId, View.Num<0>() - 1);
+		return FVector::ZeroVector;
+	}
+
+	return FVector(View[AgentId][0], View[AgentId][1], View[AgentId][2]);
+}
+
+void UVectorAction::SetVectorAction(int32 AgentId, FVector InAction)
+{
+	TLearningArrayView<2, float> View = FeatureObject->InstanceData->View(FeatureObject->ValueHandle);
+
+	if (AgentId == INDEX_NONE)
+	{
+		UE_LOG(LogLearning, Error, TEXT("AgentId is invalid (INDEX_NONE)"));
+		return;
+	}
+
+	if (AgentId < 0 || AgentId >= View.Num<0>())
+	{
+		UE_LOG(LogLearning, Error, TEXT("AgentId %d is out of index. Valid range [0, %d]."), AgentId, View.Num<0>() - 1);
+		return;
+	}
+
+	View[AgentId][0] = InAction.X;
+	View[AgentId][1] = InAction.Y;
+	View[AgentId][2] = InAction.Z;
+}
+
+#if ENABLE_VISUAL_LOG
+void UVectorAction::VisualLog(const UE::Learning::FIndexSet Instances) const
+{
+	UE_LEARNING_TRACE_CPUPROFILER_EVENT_SCOPE(UVectorAction::VisualLog);
+
+	const TLearningArrayView<2, const float> ValueView = FeatureObject->InstanceData->ConstView(FeatureObject->ValueHandle);
+	const TLearningArrayView<2, const float> FeatureView = FeatureObject->InstanceData->ConstView(FeatureObject->FeatureHandle);
+
+	if (const ULearningAgentsType* AgentType = Cast<ULearningAgentsType>(GetOuter()))
+	{
+		for (const int32 Instance : Instances)
+		{
+			if (const AActor* Actor = Cast<AActor>(AgentType->GetAgent(Instance)))
+			{
+				const FVector Vector(ValueView[Instance][0], ValueView[Instance][1], ValueView[Instance][2]);
+
+				UE_VLOG_ARROW(this, LogLearning, Display,
+					Actor->GetActorLocation(),
+					Actor->GetActorLocation() + Vector,
+					VisualLogColor.ToFColor(true),
+					TEXT(""));
+
+				UE_LEARNING_AGENTS_VLOG_STRING(this, LogLearning, Display,
+					Actor->GetActorLocation() + Vector,
+					VisualLogColor.ToFColor(true),
+					TEXT("Vector: [% 6.4f % 6.4f % 6.4f]"),
+					Vector.X, Vector.Y, Vector.Z);
+
+				UE_LEARNING_AGENTS_VLOG_STRING(this, LogLearning, Display,
+					Actor->GetActorLocation(),
+					VisualLogColor.ToFColor(true),
+					TEXT("Agent %i\nScale: [% 6.2f]\nEncoded: [% 6.3f % 6.3f % 6.3f]"),
+					Instance,
+					FeatureObject->Scale,
+					FeatureView[Instance][0], FeatureView[Instance][1], FeatureView[Instance][2]);
+			}
+		}
+	}
+}
+#endif
+
+//------------------------------------------------------------------
+
+URotationVectorArrayAction* URotationVectorArrayAction::AddRotationVectorArrayAction(ULearningAgentsType* AgentType, FName Name, int32 RotationVectorNum, float Scale)
+{
+	return UE::Learning::Agents::Private::AddAction<URotationVectorArrayAction, UE::Learning::FRotationVectorFeature>(AgentType, Name, RotationVectorNum, FMath::DegreesToRadians(Scale));
+}
+
+void URotationVectorArrayAction::GetRotationVectorArrayAction(int32 AgentId, TArray<FVector>& OutRotationVectors)
+{
+	const TLearningArrayView<2, const FVector> View = FeatureObject->InstanceData->ConstView(FeatureObject->RotationVectorsHandle);
+
+	if (AgentId == INDEX_NONE)
+	{
+		UE_LOG(LogLearning, Error, TEXT("AgentId is invalid (INDEX_NONE)"));
+		return;
+	}
+
+	if (AgentId < 0 || AgentId >= View.Num<0>())
+	{
+		UE_LOG(LogLearning, Error, TEXT("AgentId %d is out of index. Valid range [0, %d]."), AgentId, View.Num<0>() - 1);
+		return;
+	}
+
+	OutRotationVectors.SetNumUninitialized(View.Num<1>());
+	for (int32 RotationVectorIdx = 0; RotationVectorIdx < View.Num<1>(); RotationVectorIdx++)
+	{
+		OutRotationVectors[RotationVectorIdx] = View[AgentId][RotationVectorIdx];
+	}
+}
+
+void URotationVectorArrayAction::GetRotationVectorArrayActionAsQuats(int32 AgentId, TArray<FQuat>& OutRotations)
+{
+	const TLearningArrayView<2, const FVector> View = FeatureObject->InstanceData->ConstView(FeatureObject->RotationVectorsHandle);
+
+	if (AgentId == INDEX_NONE)
+	{
+		UE_LOG(LogLearning, Error, TEXT("AgentId is invalid (INDEX_NONE)"));
+		return;
+	}
+
+	if (AgentId < 0 || AgentId >= View.Num<0>())
+	{
+		UE_LOG(LogLearning, Error, TEXT("AgentId %d is out of index. Valid range [0, %d]."), AgentId, View.Num<0>() - 1);
+		return;
+	}
+
+	OutRotations.SetNumUninitialized(View.Num<1>());
+	for (int32 RotationVectorIdx = 0; RotationVectorIdx < View.Num<1>(); RotationVectorIdx++)
+	{
+		OutRotations[RotationVectorIdx] = FQuat::MakeFromRotationVector(View[AgentId][RotationVectorIdx]);
+	}
+}
+
+void URotationVectorArrayAction::GetRotationVectorArrayActionAsRotators(int32 AgentId, TArray<FRotator>& OutRotations)
+{
+	const TLearningArrayView<2, const FVector> View = FeatureObject->InstanceData->ConstView(FeatureObject->RotationVectorsHandle);
+
+	if (AgentId == INDEX_NONE)
+	{
+		UE_LOG(LogLearning, Error, TEXT("AgentId is invalid (INDEX_NONE)"));
+		return;
+	}
+
+	if (AgentId < 0 || AgentId >= View.Num<0>())
+	{
+		UE_LOG(LogLearning, Error, TEXT("AgentId %d is out of index. Valid range [0, %d]."), AgentId, View.Num<0>() - 1);
+		return;
+	}
+
+	OutRotations.SetNumUninitialized(View.Num<1>());
+	for (int32 RotationVectorIdx = 0; RotationVectorIdx < View.Num<1>(); RotationVectorIdx++)
+	{
+		OutRotations[RotationVectorIdx] = FQuat::MakeFromRotationVector(View[AgentId][RotationVectorIdx]).Rotator();
+	}
+}
+
+#if ENABLE_VISUAL_LOG
+void URotationVectorArrayAction::VisualLog(const UE::Learning::FIndexSet Instances) const
+{
+	UE_LEARNING_TRACE_CPUPROFILER_EVENT_SCOPE(URotationVectorArrayAction::VisualLog);
+
+	const TLearningArrayView<2, const FVector> ValueView = FeatureObject->InstanceData->ConstView(FeatureObject->RotationVectorsHandle);
+	const TLearningArrayView<2, const float> FeatureView = FeatureObject->InstanceData->ConstView(FeatureObject->FeatureHandle);
+
+	const int32 RotationVectorNum = ValueView.Num<1>();
+
+	if (const ULearningAgentsType* AgentType = Cast<ULearningAgentsType>(GetOuter()))
+	{
+		for (const int32 Instance : Instances)
+		{
+			if (const AActor* Actor = Cast<AActor>(AgentType->GetAgent(Instance)))
+			{
+				for (int32 RotationVectorIdx = 0; RotationVectorIdx < RotationVectorNum; RotationVectorIdx++)
+				{
+					const FVector Offset = FVector(0.0, 10.0f * RotationVectorIdx, 0.0f);
+					const FVector RotationVector = ValueView[Instance][RotationVectorIdx];
+
+					UE_VLOG_LOCATION(this, LogLearning, Display,
+						Actor->GetActorLocation() + Offset,
+						2.5f,
+						VisualLogColor.ToFColor(true),
+						TEXT(""));
+
+					UE_VLOG_ARROW(this, LogLearning, Display,
+						Actor->GetActorLocation() + Offset,
+						Actor->GetActorLocation() + Offset + RotationVector,
+						VisualLogColor.ToFColor(true),
+						TEXT(""));
+
+					UE_LEARNING_AGENTS_VLOG_STRING(this, LogLearning, Display,
+						Actor->GetActorLocation() + Offset + RotationVector,
+						VisualLogColor.ToFColor(true),
+						TEXT("Rotation Vector: [% 6.4f % 6.4f % 6.4f]"),
+						RotationVector.X, RotationVector.Y, RotationVector.Z);
+
+				}
+
+				if (RotationVectorNum > 0)
+				{
+					UE_LEARNING_AGENTS_VLOG_STRING(this, LogLearning, Display,
+						Actor->GetActorLocation(),
+						VisualLogColor.ToFColor(true),
+						TEXT("Agent %i\nScale: [% 6.2f]\nEncoded: [% 6.3f % 6.3f % 6.3f ...]"),
+						Instance,
+						FeatureObject->Scale,
+						FeatureView[Instance][0], FeatureView[Instance][1], FeatureView[Instance][2]);
+				}
+				else
+				{
+					UE_LEARNING_AGENTS_VLOG_STRING(this, LogLearning, Display,
+						Actor->GetActorLocation(),
+						VisualLogColor.ToFColor(true),
+						TEXT("Agent %i\nScale: [% 6.2f]\nEncoded: []"),
+						Instance,
+						FeatureObject->Scale);
+				}
+			}
+		}
+	}
+}
+#endif
+
+//------------------------------------------------------------------
+
+#undef UE_LEARNING_AGENTS_VLOG_STRING
+#undef UE_LEARNING_AGENTS_VLOG_TRANSFORM
+#undef UE_LEARNING_AGENTS_VLOG_PLANE
