Learning Agents: Create initial experimental version

diff --git a/Engine/Plugins/Experimental/LearningAgents/Source/LearningAgents/Private/LearningAgentsHelpers.cpp b/Engine/Plugins/Experimental/LearningAgents/Source/LearningAgents/Private/LearningAgentsHelpers.cpp
new file mode 100644
--- /dev/null
+++ b/Engine/Plugins/Experimental/LearningAgents/Source/LearningAgents/Private/LearningAgentsHelpers.cpp
@@ -0,0 +1,138 @@
+// Copyright Epic Games, Inc. All Rights Reserved.
+
+#include "LearningAgentsHelpers.h"
+
+#include "LearningLog.h"
+
+#include "Components/MeshComponent.h"
+#include "Components/SplineComponent.h"
+
+//------------------------------------------------------------------
+
+FTransform ULearningAgentsHelpers::ProjectTransformOntoGroundPlane(const FTransform Transform, const FVector LocalForwardVector)
+{
+	FVector Position = Transform.GetLocation();
+	Position.Z = 0.0f;
+
+	const FVector Direction = (FVector(1.0f, 1.0f, 0.0f) * Transform.TransformVectorNoScale(LocalForwardVector)).GetSafeNormal(UE_SMALL_NUMBER, FVector::ForwardVector);
+
+	return FTransform(FQuat::FindBetweenNormals(FVector::ForwardVector, Direction), Position);
+}
+
+void ULearningAgentsHelpers::ProjectPositionRotationOntoGroundPlane(FVector& OutPosition, FRotator& OutRotation, const FVector InPosition, const FRotator InRotation, const FVector LocalForwardVector)
+{
+	OutPosition = InPosition;
+	OutPosition.Z = 0.0f;
+
+	const FVector Direction = (FVector(1.0f, 1.0f, 0.0f) * InRotation.RotateVector(LocalForwardVector)).GetSafeNormal(UE_SMALL_NUMBER, FVector::ForwardVector);
+	OutRotation = FQuat::FindBetweenNormals(FVector::ForwardVector, Direction).Rotator();
+}
+
+//------------------------------------------------------------------
+
+void ULearningAgentsHelpers::GetMeshBonePositions(const UMeshComponent* MeshComponent, const TArray<FName>& BoneNames, TArray<FVector>& OutBonePositions)
+{
+	if (!MeshComponent)
+	{
+		UE_LOG(LogLearning, Warning, TEXT("Mesh Component was nullptr."));
+		return;
+	}
+
+	const int32 BoneNum = BoneNames.Num();
+
+	OutBonePositions.SetNumUninitialized(BoneNum);
+
+	for (int32 BoneIdx = 0; BoneIdx < BoneNum; BoneIdx++)
+	{
+		OutBonePositions[BoneIdx] = MeshComponent->GetSocketLocation(BoneNames[BoneIdx]);
+	}
+}
+
+//------------------------------------------------------------------
+
+float ULearningAgentsHelpers::GetProportionAlongSpline(const USplineComponent* SplineComponent, const float DistanceAlongSpline)
+{
+	if (!SplineComponent)
+	{
+		UE_LOG(LogLearning, Warning, TEXT("Spline Component was nullptr."));
+		return 0.0f;
+	}
+
+	return FMath::Clamp(DistanceAlongSpline / FMath::Max(SplineComponent->GetSplineLength(), UE_SMALL_NUMBER), 0.0f, 1.0f);
+}
+
+float ULearningAgentsHelpers::GetProportionAlongSplineAsAngle(const USplineComponent* SplineComponent, const float DistanceAlongSpline)
+{
+	if (!SplineComponent)
+	{
+		UE_LOG(LogLearning, Warning, TEXT("Spline Component was nullptr."));
+		return 0.0f;
+	}
+
+	if (!SplineComponent->IsClosedLoop())
+	{
+		UE_LOG(LogLearning, Warning, TEXT("Getting proportion along spline as angle, but spline is not closed loop. Consider using GetPropotionAlongSpline instead."));
+	}
+
+	const float TotalDistance = SplineComponent->GetSplineLength();
+
+	const float Ratio = UE_TWO_PI * (FMath::Wrap(DistanceAlongSpline, 0.0f, TotalDistance) / FMath::Max(TotalDistance, UE_SMALL_NUMBER)) - UE_PI;
+
+	return FMath::RadiansToDegrees(Ratio);
+}
+
+void ULearningAgentsHelpers::GetPositionsAlongSpline(TArray<FVector>& OutPositions, const USplineComponent* SplineComponent, const int32 PositionNum, const float StartDistanceAlongSpline, const float StopDistanceAlongSpline, const ESplineCoordinateSpace::Type CoordinateSpace)
+{
+	if (!SplineComponent)
+	{
+		UE_LOG(LogLearning, Warning, TEXT("Spline Component was nullptr."));
+		return;
+	}
+
+	OutPositions.SetNumUninitialized(PositionNum);
+
+	const float TotalDistance = SplineComponent->GetSplineLength();
+	const bool bIsClosedLoop = SplineComponent->IsClosedLoop();
+
+	for (int32 PositionIdx = 0; PositionIdx < PositionNum; PositionIdx++)
+	{
+		float PositionDistance = PositionNum == 1 ? 
+			(StartDistanceAlongSpline + StopDistanceAlongSpline) / 2.0f : FMath::Lerp(
+				StartDistanceAlongSpline,
+				StopDistanceAlongSpline,
+				((float)PositionIdx) / (PositionNum - 1));
+
+		PositionDistance = bIsClosedLoop ? FMath::Wrap(PositionDistance, 0.0f, TotalDistance) : PositionDistance;
+
+		OutPositions[PositionIdx] = SplineComponent->GetLocationAtDistanceAlongSpline(PositionDistance, CoordinateSpace);
+	}
+}
+
+float ULearningAgentsHelpers::GetVelocityAlongSpline(const USplineComponent* SplineComponent, const FVector Position, const FVector Velocity, const float FiniteDifferenceDelta, const ESplineCoordinateSpace::Type CoordinateSpace)
+{
+	if (!SplineComponent)
+	{
+		UE_LOG(LogLearning, Warning, TEXT("Spline Component was nullptr."));
+		return 0.0f;
+	}
+
+	float Distance0 = SplineComponent->GetDistanceAlongSplineAtLocation(Position, CoordinateSpace);
+	float Distance1 = SplineComponent->GetDistanceAlongSplineAtLocation(Position + FiniteDifferenceDelta * Velocity, CoordinateSpace);
+
+	if (SplineComponent->IsClosedLoop())
+	{
+		const float SplineDistance = SplineComponent->GetSplineLength();
+
+		if (FMath::Abs(Distance0 - (Distance1 + SplineDistance)) < FMath::Abs(Distance0 - Distance1))
+		{
+			Distance1 = Distance1 + SplineDistance;
+		}
+		else if (FMath::Abs((Distance0 + SplineDistance) - Distance1) < FMath::Abs(Distance0 - Distance1))
+		{
+			Distance0 = Distance0 + SplineDistance;
+		}
+	}
+
+	return (Distance1 - Distance0) / FMath::Max(FiniteDifferenceDelta, UE_SMALL_NUMBER);
+}
+
diff --git a/Engine/Plugins/Experimental/LearningAgents/Source/LearningAgents/Private/LearningAgentsObservations.cpp b/Engine/Plugins/Experimental/LearningAgents/Source/LearningAgents/Private/LearningAgentsObservations.cpp
new file mode 100644
--- /dev/null
+++ b/Engine/Plugins/Experimental/LearningAgents/Source/LearningAgents/Private/LearningAgentsObservations.cpp
@@ -0,0 +1,962 @@
+// Copyright Epic Games, Inc. All Rights Reserved.
+
+#include "LearningAgentsObservations.h"
+
+#include "LearningAgentsType.h"
+#include "LearningArray.h"
+#include "LearningArrayMap.h"
+#include "LearningFeatureObject.h"
+#include "LearningLog.h"
+
+#include "GameFramework/Actor.h"
+#include "VisualLogger/VisualLogger.h"
+
+#define UE_LEARNING_AGENTS_VLOG_STRING(Owner, Category, Verbosity, Location, Color, Format, ...) \
+	UE_VLOG_LOCATION(Owner, Category, Verbosity, Location, 0.0f, Color, Format, ##__VA_ARGS__)
+
+#define UE_LEARNING_AGENTS_VLOG_TRANSFORM(Owner, Category, Verbosity, Location, Rotation, Color, Format, ...) \
+	UE_VLOG_SEGMENT(Owner, Category, Verbosity, Location, Location + 15.0f * Rotation.RotateVector(FVector::ForwardVector), FColor::Red, TEXT("")); \
+	UE_VLOG_SEGMENT(Owner, Category, Verbosity, Location, Location + 15.0f * Rotation.RotateVector(FVector::RightVector), FColor::Green, TEXT("")); \
+	UE_VLOG_SEGMENT(Owner, Category, Verbosity, Location, Location + 15.0f * Rotation.RotateVector(FVector::UpVector), FColor::Blue, TEXT("")); \
+	UE_VLOG_OBOX(Owner, Category, Verbosity, FBox(10.0f * FVector(-1, -1, -1), 10.0f * FVector(1, 1, 1)), FTransform(Rotation, Location, FVector::OneVector).ToMatrixNoScale(), Color, TEXT("")); \
+	UE_LEARNING_AGENTS_VLOG_STRING(Owner, Category, Verbosity, Location + FVector(0.0f, 0.0f, 20.0f), Color, Format, ##__VA_ARGS__)
+
+#define UE_LEARNING_AGENTS_VLOG_PLANE(Owner, Category, Verbosity, Location, Rotation, Axis0, Axis1, Color, Format, ...) \
+	UE_VLOG_OBOX(Owner, Category, Verbosity, FBox(25.0f * FVector(-1, -1, 0), 25.0f * FVector(1, 1, 0)), FTransform(Rotation, Location, FVector::OneVector).ToMatrixNoScale(), Color, Format, ##__VA_ARGS__)
+
+namespace UE::Learning::Agents::Private
+{
+	template<typename ObservationUObject, typename ObservationFObject, typename... InArgTypes>
+	ObservationUObject* AddObservation(ULearningAgentsType* AgentType, FName Name, InArgTypes&& ...Args)
+	{
+		if (!AgentType)
+		{
+			UE_LOG(LogLearning, Error, TEXT("AgentType is nullptr"));
+			return nullptr;
+		}
+
+		ObservationUObject* Observation = NewObject<ObservationUObject>(AgentType, Name);
+
+		Observation->FeatureObject = MakeShared<ObservationFObject>(
+			Observation->GetFName(),
+			AgentType->GetInstanceData().ToSharedRef(),
+			AgentType->GetMaxInstanceNum(),
+			Forward<InArgTypes>(Args)...);
+
+		AgentType->AddObservation(Observation, Observation->FeatureObject.ToSharedRef());
+
+		return Observation;
+	}
+}
+
+//------------------------------------------------------------------
+
+UFloatObservation* UFloatObservation::AddFloatObservation(ULearningAgentsType* AgentType, FName Name, float Scale)
+{
+	return UE::Learning::Agents::Private::AddObservation<UFloatObservation, UE::Learning::FFloatFeature>(AgentType, Name, 1, Scale);
+}
+
+void UFloatObservation::SetFloatObservation(int32 AgentId, float Value)
+{
+	TLearningArrayView<2, float> View = FeatureObject->InstanceData->View(FeatureObject->ValueHandle);
+
+	if (AgentId == INDEX_NONE)
+	{
+		UE_LOG(LogLearning, Error, TEXT("AgentId is invalid (INDEX_NONE)"));
+		return;
+	}
+
+	if (AgentId < 0 || AgentId >= View.Num<0>())
+	{
+		UE_LOG(LogLearning, Error, TEXT("AgentId %d is out of index. Valid range [0, %d]."), AgentId, View.Num<0>() - 1);
+		return;
+	}
+
+	View[AgentId][0] = Value;
+}
+
+#if ENABLE_VISUAL_LOG
+void UFloatObservation::VisualLog(const UE::Learning::FIndexSet Instances) const
+{
+	UE_LEARNING_TRACE_CPUPROFILER_EVENT_SCOPE(UFloatObservation::VisualLog);
+
+	const TLearningArrayView<2, const float> ValueView = FeatureObject->InstanceData->ConstView(FeatureObject->ValueHandle);
+	const TLearningArrayView<2, const float> FeatureView = FeatureObject->InstanceData->ConstView(FeatureObject->FeatureHandle);
+
+	if (const ULearningAgentsType* AgentType = Cast<ULearningAgentsType>(GetOuter()))
+	{
+		for (const int32 Instance : Instances)
+		{
+			if (const AActor* Actor = Cast<AActor>(AgentType->GetAgent(Instance)))
+			{
+				UE_LEARNING_AGENTS_VLOG_STRING(this, LogLearning, Display,
+					Actor->GetActorLocation(),
+					VisualLogColor.ToFColor(true),
+					TEXT("Agent %i\nScale: [% 6.2f]\nValue: [% 6.2f]\nEncoded: [% 6.3f]"),
+					Instance,
+					ValueView[Instance][0],
+					FeatureObject->Scale,
+					FeatureView[Instance][0]);
+			}
+		}
+	}
+}
+#endif
+
+//------------------------------------------------------------------
+
+UVectorObservation* UVectorObservation::AddVectorObservation(ULearningAgentsType* AgentType, FName Name, float Scale)
+{
+	return UE::Learning::Agents::Private::AddObservation<UVectorObservation, UE::Learning::FFloatFeature>(AgentType, Name, 3, Scale);
+}
+
+void UVectorObservation::SetVectorObservation(int32 AgentId, FVector Vector)
+{
+	TLearningArrayView<2, float> View = FeatureObject->InstanceData->View(FeatureObject->ValueHandle);
+
+	if (AgentId == INDEX_NONE)
+	{
+		UE_LOG(LogLearning, Error, TEXT("AgentId is invalid (INDEX_NONE)"));
+		return;
+	}
+
+	if (AgentId < 0 || AgentId >= View.Num<0>())
+	{
+		UE_LOG(LogLearning, Error, TEXT("AgentId %d is out of index. Valid range [0, %d]."), AgentId, View.Num<0>() - 1);
+		return;
+	}
+
+	View[AgentId][0] = Vector.X;
+	View[AgentId][1] = Vector.Y;
+	View[AgentId][2] = Vector.Z;
+}
+
+#if ENABLE_VISUAL_LOG
+void UVectorObservation::VisualLog(const UE::Learning::FIndexSet Instances) const
+{
+	UE_LEARNING_TRACE_CPUPROFILER_EVENT_SCOPE(UVectorObservation::VisualLog);
+
+	const TLearningArrayView<2, const float> ValueView = FeatureObject->InstanceData->ConstView(FeatureObject->ValueHandle);
+	const TLearningArrayView<2, const float> FeatureView = FeatureObject->InstanceData->ConstView(FeatureObject->FeatureHandle);
+
+	if (const ULearningAgentsType* AgentType = Cast<ULearningAgentsType>(GetOuter()))
+	{
+		for (const int32 Instance : Instances)
+		{
+			if (const AActor* Actor = Cast<AActor>(AgentType->GetAgent(Instance)))
+			{
+				const FVector Vector(ValueView[Instance][0], ValueView[Instance][1], ValueView[Instance][2]);
+
+				UE_VLOG_ARROW(this, LogLearning, Display,
+					Actor->GetActorLocation(),
+					Actor->GetActorLocation() + Vector,
+					VisualLogColor.ToFColor(true),
+					TEXT(""));
+
+				UE_LEARNING_AGENTS_VLOG_STRING(this, LogLearning, Display,
+					Actor->GetActorLocation() + Vector,
+					VisualLogColor.ToFColor(true),
+					TEXT("Vector: [% 6.4f % 6.4f % 6.4f]"),
+					Vector.X, Vector.Y, Vector.Z);
+
+				UE_LEARNING_AGENTS_VLOG_STRING(this, LogLearning, Display,
+					Actor->GetActorLocation(),
+					VisualLogColor.ToFColor(true),
+					TEXT("Agent %i\nScale: [% 6.2f]\nEncoded: [% 6.3f % 6.3f % 6.3f]"),
+					Instance,
+					FeatureObject->Scale,
+					FeatureView[Instance][0], FeatureView[Instance][1], FeatureView[Instance][2]);
+			}
+		}
+	}
+}
+#endif
+
+//------------------------------------------------------------------
+
+UAngleObservation* UAngleObservation::AddAngleObservation(ULearningAgentsType* AgentType, FName Name, float Scale)
+{
+	return UE::Learning::Agents::Private::AddObservation<UAngleObservation, UE::Learning::FAngleFeature>(AgentType, Name, 1, Scale);
+}
+
+void UAngleObservation::SetAngleObservation(int32 AgentId, float Angle, float RelativeAngle)
+{
+	TLearningArrayView<2, float> AngleView = FeatureObject->InstanceData->View(FeatureObject->AngleHandle);
+	TLearningArrayView<1, float> RelativeAngleView = FeatureObject->InstanceData->View(FeatureObject->RelativeAngleHandle);
+
+	if (AgentId == INDEX_NONE)
+	{
+		UE_LOG(LogLearning, Error, TEXT("AgentId is invalid (INDEX_NONE)"));
+		return;
+	}
+
+	if (AgentId < 0 || AgentId >= AngleView.Num<0>())
+	{
+		UE_LOG(LogLearning, Error, TEXT("AgentId %d is out of index. Valid range [0, %d]."), AgentId, AngleView.Num<0>() - 1);
+		return;
+	}
+
+	AngleView[AgentId][0] = FMath::DegreesToRadians(Angle);
+	RelativeAngleView[AgentId] = FMath::DegreesToRadians(RelativeAngle);
+}
+
+#if ENABLE_VISUAL_LOG
+void UAngleObservation::VisualLog(const UE::Learning::FIndexSet Instances) const
+{
+	UE_LEARNING_TRACE_CPUPROFILER_EVENT_SCOPE(UAngleObservation::VisualLog);
+
+	const TLearningArrayView<2, const float> AngleView = FeatureObject->InstanceData->ConstView(FeatureObject->AngleHandle);
+	const TLearningArrayView<1, const float> RelativeAngleView = FeatureObject->InstanceData->ConstView(FeatureObject->RelativeAngleHandle);
+	const TLearningArrayView<2, const float> FeatureView = FeatureObject->InstanceData->ConstView(FeatureObject->FeatureHandle);
+
+	if (const ULearningAgentsType* AgentType = Cast<ULearningAgentsType>(GetOuter()))
+	{
+		for (const int32 Instance : Instances)
+		{
+			if (const AActor* Actor = Cast<AActor>(AgentType->GetAgent(Instance)))
+			{
+				const float Angle = AngleView[Instance][0];
+				const float RelativeAngle = RelativeAngleView[Instance];
+
+				UE_VLOG_CIRCLE(this, LogLearning, Display,
+					Actor->GetActorLocation(),
+					FVector::UpVector,
+					50.0f,
+					VisualLogColor.ToFColor(true),
+					TEXT(""));
+
+				UE_VLOG_SEGMENT(this, LogLearning, Display,
+					Actor->GetActorLocation(),
+					Actor->GetActorLocation() + 50.0f * FVector(FMath::Sin(RelativeAngle), FMath::Cos(RelativeAngle), 0.0f),
+					VisualLogColor.ToFColor(true),
+					TEXT(""));
+
+				UE_VLOG_LOCATION(this, LogLearning, Display,
+					Actor->GetActorLocation() + 50.0f * FVector(FMath::Sin(RelativeAngle), FMath::Cos(RelativeAngle), 0.0f),
+					2.5f,
+					VisualLogColor.ToFColor(true),
+					TEXT("Relative Angle: [% 6.1f]"),
+					RelativeAngle);
+
+				UE_VLOG_SEGMENT(this, LogLearning, Display,
+					Actor->GetActorLocation(),
+					Actor->GetActorLocation() + 50.0f * FVector(FMath::Sin(Angle), FMath::Cos(Angle), 0.0f),
+					VisualLogColor.ToFColor(true),
+					TEXT(""));
+
+				UE_VLOG_LOCATION(this, LogLearning, Display,
+					Actor->GetActorLocation() + 50.0f * FVector(FMath::Sin(Angle), FMath::Cos(Angle), 0.0f),
+					2.5f,
+					VisualLogColor.ToFColor(true),
+					TEXT("Angle: [% 6.1f]"),
+					Angle);
+
+				UE_VLOG_LOCATION(this, LogLearning, Display,
+					Actor->GetActorLocation(),
+					5.0f,
+					VisualLogColor.ToFColor(true),
+					TEXT("Agent %i\nScale: [% 6.2f]\nEncoded: [% 6.3f % 6.3f]"),
+					Instance,
+					FeatureObject->Scale,
+					FeatureView[Instance][0], FeatureView[Instance][1]);
+			}
+		}
+	}
+}
+#endif
+
+//------------------------------------------------------------------
+
+UPlanarDirectionObservation* UPlanarDirectionObservation::AddPlanarDirectionObservation(ULearningAgentsType* AgentType, FName Name, float Scale, FVector Axis0, FVector Axis1)
+{
+	return UE::Learning::Agents::Private::AddObservation<UPlanarDirectionObservation, UE::Learning::FPlanarDirectionFeature>(AgentType, Name, 1, Scale, Axis0, Axis1);
+}
+
+void UPlanarDirectionObservation::SetPlanarDirectionObservation(int32 AgentId, FVector Direction, FRotator RelativeRotation)
+{
+	TLearningArrayView<2, FVector> DirectionView = FeatureObject->InstanceData->View(FeatureObject->DirectionHandle);
+	TLearningArrayView<1, FQuat> RelativeRotationView = FeatureObject->InstanceData->View(FeatureObject->RelativeRotationHandle);
+
+	if (AgentId == INDEX_NONE)
+	{
+		UE_LOG(LogLearning, Error, TEXT("AgentId is invalid (INDEX_NONE)"));
+		return;
+	}
+
+	if (AgentId < 0 || AgentId >= DirectionView.Num<0>())
+	{
+		UE_LOG(LogLearning, Error, TEXT("AgentId %d is out of index. Valid range [0, %d]."), AgentId, DirectionView.Num<0>() - 1);
+		return;
+	}
+
+	RelativeRotationView[AgentId] = FQuat::MakeFromRotator(RelativeRotation);
+	DirectionView[AgentId][0] = Direction;
+}
+
+#if ENABLE_VISUAL_LOG
+void UPlanarDirectionObservation::VisualLog(const UE::Learning::FIndexSet Instances) const
+{
+	UE_LEARNING_TRACE_CPUPROFILER_EVENT_SCOPE(UPlanarDirectionObservation::VisualLog);
+
+	const TLearningArrayView<2, const FVector> DirectionView = FeatureObject->InstanceData->ConstView(FeatureObject->DirectionHandle);
+	const TLearningArrayView<1, const FQuat> RelativeRotationView = FeatureObject->InstanceData->ConstView(FeatureObject->RelativeRotationHandle);
+	const TLearningArrayView<2, const float> FeatureView = FeatureObject->InstanceData->ConstView(FeatureObject->FeatureHandle);
+
+	if (const ULearningAgentsType* AgentType = Cast<ULearningAgentsType>(GetOuter()))
+	{
+		for (const int32 Instance : Instances)
+		{
+			if (const AActor* Actor = Cast<AActor>(AgentType->GetAgent(Instance)))
+			{
+				const FVector Direction = DirectionView[Instance][0];
+				const FQuat RelativeRotation = RelativeRotationView[Instance];
+				const FVector LocalDirection = RelativeRotation.UnrotateVector(Direction);
+
+				UE_VLOG_ARROW(this, LogLearning, Display,
+					Actor->GetActorLocation(),
+					Actor->GetActorLocation() + 100.0f * Direction,
+					VisualLogColor.ToFColor(true),
+					TEXT(""));
+
+				UE_LEARNING_AGENTS_VLOG_STRING(this, LogLearning, Display,
+					Actor->GetActorLocation() + 100.0f * Direction,
+					VisualLogColor.ToFColor(true),
+					TEXT("Direction: [% 6.3f % 6.3f % 6.3f]\nLocal Direction: [% 6.3f % 6.3f % 6.3f]"),
+					Direction.X, Direction.Y, Direction.Z,
+					LocalDirection.X, LocalDirection.Y, LocalDirection.Z);
+
+				UE_LEARNING_AGENTS_VLOG_PLANE(this, LogLearning, Display,
+					Actor->GetActorLocation(),
+					RelativeRotation,
+					FeatureObject->Axis0,
+					FeatureObject->Axis1,
+					VisualLogColor.ToFColor(true),
+					TEXT(""));
+
+				UE_LEARNING_AGENTS_VLOG_TRANSFORM(this, LogLearning, Display,
+					Actor->GetActorLocation(),
+					RelativeRotation,
+					VisualLogColor.ToFColor(true),
+					TEXT("Agent %i\nScale: [% 6.2f]\nEncoded: [% 6.3f % 6.3f]"),
+					Instance,
+					FeatureObject->Scale,
+					FeatureView[Instance][0], FeatureView[Instance][1]);
+			}
+		}
+	}
+}
+#endif
+
+UDirectionObservation* UDirectionObservation::AddDirectionObservation(ULearningAgentsType* AgentType, FName Name, float Scale)
+{
+	return UE::Learning::Agents::Private::AddObservation<UDirectionObservation, UE::Learning::FDirectionFeature>(AgentType, Name, 1, Scale);
+}
+
+void UDirectionObservation::SetDirectionObservation(int32 AgentId, FVector Direction, FRotator RelativeRotation)
+{
+	TLearningArrayView<2, FVector> DirectionView = FeatureObject->InstanceData->View(FeatureObject->DirectionHandle);
+	TLearningArrayView<1, FQuat> RelativeRotationView = FeatureObject->InstanceData->View(FeatureObject->RelativeRotationHandle);
+
+	if (AgentId == INDEX_NONE)
+	{
+		UE_LOG(LogLearning, Error, TEXT("AgentId is invalid (INDEX_NONE)"));
+		return;
+	}
+
+	if (AgentId < 0 || AgentId >= DirectionView.Num<0>())
+	{
+		UE_LOG(LogLearning, Error, TEXT("AgentId %d is out of index. Valid range [0, %d]."), AgentId, DirectionView.Num<0>() - 1);
+		return;
+	}
+
+	RelativeRotationView[AgentId] = FQuat::MakeFromRotator(RelativeRotation);
+	DirectionView[AgentId][0] = Direction;
+}
+
+#if ENABLE_VISUAL_LOG
+void UDirectionObservation::VisualLog(const UE::Learning::FIndexSet Instances) const
+{
+	UE_LEARNING_TRACE_CPUPROFILER_EVENT_SCOPE(UDirectionObservation::VisualLog);
+
+	const TLearningArrayView<2, const FVector> DirectionView = FeatureObject->InstanceData->ConstView(FeatureObject->DirectionHandle);
+	const TLearningArrayView<1, const FQuat> RelativeRotationView = FeatureObject->InstanceData->ConstView(FeatureObject->RelativeRotationHandle);
+	const TLearningArrayView<2, const float> FeatureView = FeatureObject->InstanceData->ConstView(FeatureObject->FeatureHandle);
+
+	if (const ULearningAgentsType* AgentType = Cast<ULearningAgentsType>(GetOuter()))
+	{
+		for (const int32 Instance : Instances)
+		{
+			if (const AActor* Actor = Cast<AActor>(AgentType->GetAgent(Instance)))
+			{
+				const FVector Direction = DirectionView[Instance][0];
+				const FQuat RelativeRotation = RelativeRotationView[Instance];
+				const FVector LocalDirection = RelativeRotation.UnrotateVector(Direction);
+
+				UE_VLOG_ARROW(this, LogLearning, Display,
+					Actor->GetActorLocation(),
+					Actor->GetActorLocation() + 100.0f * Direction,
+					VisualLogColor.ToFColor(true),
+					TEXT(""));
+
+				UE_LEARNING_AGENTS_VLOG_STRING(this, LogLearning, Display,
+					Actor->GetActorLocation() + 100.0f * Direction,
+					VisualLogColor.ToFColor(true),
+					TEXT("Direction: [% 6.3f % 6.3f % 6.3f]\nLocal Direction: [% 6.3f % 6.3f % 6.3f]"),
+					Direction.X, Direction.Y, Direction.Z,
+					LocalDirection.X, LocalDirection.Y, LocalDirection.Z);
+
+				UE_LEARNING_AGENTS_VLOG_TRANSFORM(this, LogLearning, Display,
+					Actor->GetActorLocation(),
+					RelativeRotation,
+					VisualLogColor.ToFColor(true),
+					TEXT("Agent %i\nScale: [% 6.2f]\nEncoded: [% 6.3f % 6.3f % 6.3f]"),
+					Instance,
+					FeatureObject->Scale,
+					FeatureView[Instance][0], FeatureView[Instance][1], FeatureView[Instance][2]);
+			}
+		}
+	}
+}
+#endif
+
+//------------------------------------------------------------------
+
+UPlanarPositionObservation* UPlanarPositionObservation::AddPlanarPositionObservation(ULearningAgentsType* AgentType, FName Name, float Scale, FVector Axis0, FVector Axis1)
+{
+	return UE::Learning::Agents::Private::AddObservation<UPlanarPositionObservation, UE::Learning::FPlanarPositionFeature>(AgentType, Name, 1, Scale, Axis0, Axis1);
+}
+
+void UPlanarPositionObservation::SetPlanarPositionObservation(int32 AgentId, FVector Position, FVector RelativePosition, FRotator RelativeRotation)
+{
+	TLearningArrayView<2, FVector> PositionView = FeatureObject->InstanceData->View(FeatureObject->PositionHandle);
+	TLearningArrayView<1, FVector> RelativePositionView = FeatureObject->InstanceData->View(FeatureObject->RelativePositionHandle);
+	TLearningArrayView<1, FQuat> RelativeRotationView = FeatureObject->InstanceData->View(FeatureObject->RelativeRotationHandle);
+
+	if (AgentId == INDEX_NONE)
+	{
+		UE_LOG(LogLearning, Error, TEXT("AgentId is invalid (INDEX_NONE)"));
+		return;
+	}
+
+	if (AgentId < 0 || AgentId >= PositionView.Num<0>())
+	{
+		UE_LOG(LogLearning, Error, TEXT("AgentId %d is out of index. Valid range [0, %d]."), AgentId, PositionView.Num<0>() - 1);
+		return;
+	}
+
+	RelativePositionView[AgentId] = RelativePosition;
+	RelativeRotationView[AgentId] = FQuat::MakeFromRotator(RelativeRotation);
+	PositionView[AgentId][0] = Position;
+}
+
+#if ENABLE_VISUAL_LOG
+void UPlanarPositionObservation::VisualLog(const UE::Learning::FIndexSet Instances) const
+{
+	UE_LEARNING_TRACE_CPUPROFILER_EVENT_SCOPE(UPlanarPositionObservation::VisualLog);
+
+	const TLearningArrayView<2, const FVector> PositionView = FeatureObject->InstanceData->ConstView(FeatureObject->PositionHandle);
+	const TLearningArrayView<1, const FVector> RelativePositionView = FeatureObject->InstanceData->ConstView(FeatureObject->RelativePositionHandle);
+	const TLearningArrayView<1, const FQuat> RelativeRotationView = FeatureObject->InstanceData->ConstView(FeatureObject->RelativeRotationHandle);
+	const TLearningArrayView<2, const float> FeatureView = FeatureObject->InstanceData->ConstView(FeatureObject->FeatureHandle);
+
+	if (const ULearningAgentsType* AgentType = Cast<ULearningAgentsType>(GetOuter()))
+	{
+		for (const int32 Instance : Instances)
+		{
+			if (const AActor* Actor = Cast<AActor>(AgentType->GetAgent(Instance)))
+			{
+				const FVector Position = PositionView[Instance][0];
+				const FVector RelativePosition = RelativePositionView[Instance];
+				const FQuat RelativeRotation = RelativeRotationView[Instance];
+				const FVector LocalPosition = RelativeRotation.UnrotateVector(Position - RelativePosition);
+
+				UE_VLOG_LOCATION(this, LogLearning, Display,
+					Position,
+					10.0f,
+					VisualLogColor.ToFColor(true),
+					TEXT("Position: [% 6.1f % 6.1f % 6.1f]\nLocal Position: [% 6.1f % 6.1f % 6.1f]"),
+					Position.X, Position.Y, Position.Z,
+					LocalPosition.X, LocalPosition.Y, LocalPosition.Z);
+
+				UE_VLOG_SEGMENT(this, LogLearning, Display,
+					RelativePosition,
+					Position,
+					VisualLogColor.ToFColor(true),
+					TEXT(""));
+
+				UE_LEARNING_AGENTS_VLOG_PLANE(this, LogLearning, Display,
+					RelativePosition,
+					RelativeRotation,
+					FeatureObject->Axis0,
+					FeatureObject->Axis1,
+					VisualLogColor.ToFColor(true),
+					TEXT(""));
+
+				UE_LEARNING_AGENTS_VLOG_TRANSFORM(this, LogLearning, Display,
+					RelativePosition,
+					RelativeRotation,
+					VisualLogColor.ToFColor(true),
+					TEXT("Agent %i\nScale: [% 6.2f]\nEncoded: [% 6.3f % 6.3f]"),
+					Instance,
+					FeatureObject->Scale,
+					FeatureView[Instance][0], FeatureView[Instance][1]);
+			}
+		}
+	}
+}
+#endif
+
+UPositionObservation* UPositionObservation::AddPositionObservation(ULearningAgentsType* AgentType, FName Name, float Scale)
+{
+	return UE::Learning::Agents::Private::AddObservation<UPositionObservation, UE::Learning::FPositionFeature>(AgentType, Name, 1, Scale);
+}
+
+void UPositionObservation::SetPositionObservation(int32 AgentId, FVector Position, FVector RelativePosition, FRotator RelativeRotation)
+{
+	TLearningArrayView<2, FVector> PositionView = FeatureObject->InstanceData->View(FeatureObject->PositionHandle);
+	TLearningArrayView<1, FVector> RelativePositionView = FeatureObject->InstanceData->View(FeatureObject->RelativePositionHandle);
+	TLearningArrayView<1, FQuat> RelativeRotationView = FeatureObject->InstanceData->View(FeatureObject->RelativeRotationHandle);
+
+	if (AgentId == INDEX_NONE)
+	{
+		UE_LOG(LogLearning, Error, TEXT("AgentId is invalid (INDEX_NONE)"));
+		return;
+	}
+
+	if (AgentId < 0 || AgentId >= PositionView.Num<0>())
+	{
+		UE_LOG(LogLearning, Error, TEXT("AgentId %d is out of index. Valid range [0, %d]."), AgentId, PositionView.Num<0>() - 1);
+		return;
+	}
+
+	RelativePositionView[AgentId] = RelativePosition;
+	RelativeRotationView[AgentId] = FQuat::MakeFromRotator(RelativeRotation);
+	PositionView[AgentId][0] = Position;
+}
+
+#if ENABLE_VISUAL_LOG
+void UPositionObservation::VisualLog(const UE::Learning::FIndexSet Instances) const
+{
+	UE_LEARNING_TRACE_CPUPROFILER_EVENT_SCOPE(UPositionObservation::VisualLog);
+
+	const TLearningArrayView<2, const FVector> PositionView = FeatureObject->InstanceData->ConstView(FeatureObject->PositionHandle);
+	const TLearningArrayView<1, const FVector> RelativePositionView = FeatureObject->InstanceData->ConstView(FeatureObject->RelativePositionHandle);
+	const TLearningArrayView<1, const FQuat> RelativeRotationView = FeatureObject->InstanceData->ConstView(FeatureObject->RelativeRotationHandle);
+	const TLearningArrayView<2, const float> FeatureView = FeatureObject->InstanceData->ConstView(FeatureObject->FeatureHandle);
+
+	if (const ULearningAgentsType* AgentType = Cast<ULearningAgentsType>(GetOuter()))
+	{
+		for (const int32 Instance : Instances)
+		{
+			if (const AActor* Actor = Cast<AActor>(AgentType->GetAgent(Instance)))
+			{
+				const FVector Position = PositionView[Instance][0];
+				const FVector RelativePosition = RelativePositionView[Instance];
+				const FQuat RelativeRotation = RelativeRotationView[Instance];
+				const FVector LocalPosition = RelativeRotation.UnrotateVector(Position - RelativePosition);
+
+				UE_VLOG_LOCATION(this, LogLearning, Display,
+					Position,
+					10.0f,
+					VisualLogColor.ToFColor(true),
+					TEXT("Position: [% 6.1f % 6.1f % 6.1f]\nLocal Position: [% 6.1f % 6.1f % 6.1f]"),
+					Position.X, Position.Y, Position.Z,
+					LocalPosition.X, LocalPosition.Y, LocalPosition.Z);
+
+				UE_VLOG_SEGMENT(this, LogLearning, Display,
+					RelativePosition,
+					Position,
+					VisualLogColor.ToFColor(true),
+					TEXT(""));
+
+				UE_LEARNING_AGENTS_VLOG_TRANSFORM(this, LogLearning, Display,
+					RelativePosition,
+					RelativeRotation,
+					VisualLogColor.ToFColor(true),
+					TEXT("Agent %i\nScale: [% 6.2f]\nEncoded: [% 6.3f % 6.3f % 6.3f]"),
+					Instance,
+					FeatureObject->Scale,
+					FeatureView[Instance][0], FeatureView[Instance][1], FeatureView[Instance][2]);
+			}
+		}
+	}
+}
+#endif
+
+UPlanarPositionArrayObservation* UPlanarPositionArrayObservation::AddPlanarPositionArrayObservation(ULearningAgentsType* AgentType, FName Name, int32 PositionNum, float Scale, FVector Axis0, FVector Axis1)
+{
+	return UE::Learning::Agents::Private::AddObservation<UPlanarPositionArrayObservation, UE::Learning::FPlanarPositionFeature>(AgentType, Name, PositionNum, Scale, Axis0, Axis1);
+}
+
+void UPlanarPositionArrayObservation::SetPlanarPositionArrayObservation(int32 AgentId, const TArray<FVector>& Positions, FVector RelativePosition, FRotator RelativeRotation)
+{
+	TLearningArrayView<2, FVector> PositionView = FeatureObject->InstanceData->View(FeatureObject->PositionHandle);
+	TLearningArrayView<1, FVector> RelativePositionView = FeatureObject->InstanceData->View(FeatureObject->RelativePositionHandle);
+	TLearningArrayView<1, FQuat> RelativeRotationView = FeatureObject->InstanceData->View(FeatureObject->RelativeRotationHandle);
+
+	if (AgentId == INDEX_NONE)
+	{
+		UE_LOG(LogLearning, Error, TEXT("AgentId is invalid (INDEX_NONE)"));
+		return;
+	}
+
+	if (AgentId < 0 || AgentId >= PositionView.Num<0>())
+	{
+		UE_LOG(LogLearning, Error, TEXT("AgentId %d is out of index. Valid range [0, %d]."), AgentId, PositionView.Num<0>() - 1);
+		return;
+	}
+
+	const int32 PositionNum = PositionView.Num<1>();
+
+	if (Positions.Num() != PositionNum)
+	{
+		UE_LOG(LogLearning, Error, TEXT("Incorrect number of positions in array. Got %i, expected %i."), Positions.Num(), PositionNum);
+		return;
+	}
+
+	RelativePositionView[AgentId] = RelativePosition;
+	RelativeRotationView[AgentId] = FQuat::MakeFromRotator(RelativeRotation);
+	UE::Learning::Array::Copy(PositionView[AgentId], TLearningArrayView<1, const FVector>(Positions));
+}
+
+#if ENABLE_VISUAL_LOG
+void UPlanarPositionArrayObservation::VisualLog(const UE::Learning::FIndexSet Instances) const
+{
+	UE_LEARNING_TRACE_CPUPROFILER_EVENT_SCOPE(UPlanarPositionArrayObservation::VisualLog);
+
+	const TLearningArrayView<2, const FVector> PositionView = FeatureObject->InstanceData->ConstView(FeatureObject->PositionHandle);
+	const TLearningArrayView<1, const FVector> RelativePositionView = FeatureObject->InstanceData->ConstView(FeatureObject->RelativePositionHandle);
+	const TLearningArrayView<1, const FQuat> RelativeRotationView = FeatureObject->InstanceData->ConstView(FeatureObject->RelativeRotationHandle);
+	const TLearningArrayView<2, const float> FeatureView = FeatureObject->InstanceData->ConstView(FeatureObject->FeatureHandle);
+
+	const int32 PositionNum = PositionView.Num<1>();
+
+	if (const ULearningAgentsType* AgentType = Cast<ULearningAgentsType>(GetOuter()))
+	{
+		for (const int32 Instance : Instances)
+		{
+			if (const AActor* Actor = Cast<AActor>(AgentType->GetAgent(Instance)))
+			{
+				const FVector RelativePosition = RelativePositionView[Instance];
+				const FQuat RelativeRotation = RelativeRotationView[Instance];
+
+				for (int32 PositionIdx = 0; PositionIdx < PositionNum; PositionIdx++)
+				{
+					const FVector LocalPosition = RelativeRotation.UnrotateVector(PositionView[Instance][PositionIdx] - RelativePosition);
+
+					UE_VLOG_LOCATION(this, LogLearning, Display,
+						PositionView[Instance][PositionIdx],
+						10.0f,
+						VisualLogColor.ToFColor(true),
+						TEXT("Position: [% 6.1f % 6.1f % 6.1f]\nLocal Position: [% 6.1f % 6.1f % 6.1f]"),
+						PositionView[Instance][PositionIdx].X, 
+						PositionView[Instance][PositionIdx].Y, 
+						PositionView[Instance][PositionIdx].Z,
+						LocalPosition.X, LocalPosition.Y, LocalPosition.Z);
+
+					UE_VLOG_SEGMENT(this, LogLearning, Display,
+						RelativePosition,
+						PositionView[Instance][PositionIdx],
+						VisualLogColor.ToFColor(true),
+						TEXT(""));
+				}
+
+				UE_LEARNING_AGENTS_VLOG_PLANE(this, LogLearning, Display,
+					RelativePosition,
+					RelativeRotation,
+					FeatureObject->Axis0,
+					FeatureObject->Axis1,
+					VisualLogColor.ToFColor(true),
+					TEXT(""));
+
+				if (PositionNum > 0)
+				{
+					UE_LEARNING_AGENTS_VLOG_TRANSFORM(this, LogLearning, Display,
+						RelativePosition,
+						RelativeRotation,
+						VisualLogColor.ToFColor(true),
+						TEXT("Agent %i\nEncoded: [% 6.3f % 6.3f ...]"),
+						Instance,
+						FeatureObject->Scale,
+						FeatureView[Instance][0], FeatureView[Instance][1]);
+				}
+				else
+				{
+					UE_LEARNING_AGENTS_VLOG_TRANSFORM(this, LogLearning, Display,
+						RelativePosition,
+						RelativeRotation,
+						VisualLogColor.ToFColor(true),
+						TEXT("Agent %i\nScale: [% 6.2f]\nEncoded: []"),
+						Instance,
+						FeatureObject->Scale);
+				}
+			}
+		}
+	}
+}
+#endif
+
+UPositionArrayObservation* UPositionArrayObservation::AddPositionArrayObservation(ULearningAgentsType* AgentType, FName Name, int32 PositionNum, float Scale)
+{
+	return UE::Learning::Agents::Private::AddObservation<UPositionArrayObservation, UE::Learning::FPositionFeature>(AgentType, Name, PositionNum, Scale);
+}
+
+void UPositionArrayObservation::SetPositionArrayObservation(int32 AgentId, const TArray<FVector>& Positions, FVector RelativePosition, FRotator RelativeRotation)
+{
+	TLearningArrayView<2, FVector> PositionView = FeatureObject->InstanceData->View(FeatureObject->PositionHandle);
+	TLearningArrayView<1, FVector> RelativePositionView = FeatureObject->InstanceData->View(FeatureObject->RelativePositionHandle);
+	TLearningArrayView<1, FQuat> RelativeRotationView = FeatureObject->InstanceData->View(FeatureObject->RelativeRotationHandle);
+	
+	if (AgentId == INDEX_NONE)
+	{
+		UE_LOG(LogLearning, Error, TEXT("AgentId is invalid (INDEX_NONE)"));
+		return;
+	}
+
+	if (AgentId < 0 || AgentId >= PositionView.Num<0>())
+	{
+		UE_LOG(LogLearning, Error, TEXT("AgentId %d is out of index. Valid range [0, %d]."), AgentId, PositionView.Num<0>() - 1);
+		return;
+	}
+
+	const int32 PositionNum = PositionView.Num<1>();
+
+	if (Positions.Num() != PositionNum)
+	{
+		UE_LOG(LogLearning, Error, TEXT("Incorrect number of positions in array. Got %i, expected %i."), Positions.Num(), PositionNum);
+		return;
+	}
+
+	RelativePositionView[AgentId] = RelativePosition;
+	RelativeRotationView[AgentId] = FQuat::MakeFromRotator(RelativeRotation);
+	UE::Learning::Array::Copy(PositionView[AgentId], TLearningArrayView<1, const FVector>(Positions));
+}
+
+#if ENABLE_VISUAL_LOG
+void UPositionArrayObservation::VisualLog(const UE::Learning::FIndexSet Instances) const
+{
+	UE_LEARNING_TRACE_CPUPROFILER_EVENT_SCOPE(UPositionArrayObservation::VisualLog);
+
+	const TLearningArrayView<2, const FVector> PositionView = FeatureObject->InstanceData->ConstView(FeatureObject->PositionHandle);
+	const TLearningArrayView<1, const FVector> RelativePositionView = FeatureObject->InstanceData->ConstView(FeatureObject->RelativePositionHandle);
+	const TLearningArrayView<1, const FQuat> RelativeRotationView = FeatureObject->InstanceData->ConstView(FeatureObject->RelativeRotationHandle);
+	const TLearningArrayView<2, const float> FeatureView = FeatureObject->InstanceData->ConstView(FeatureObject->FeatureHandle);
+
+	const int32 PositionNum = PositionView.Num<1>();
+
+	if (const ULearningAgentsType* AgentType = Cast<ULearningAgentsType>(GetOuter()))
+	{
+		for (const int32 Instance : Instances)
+		{
+			if (const AActor* Actor = Cast<AActor>(AgentType->GetAgent(Instance)))
+			{
+				const FVector RelativePosition = RelativePositionView[Instance];
+				const FQuat RelativeRotation = RelativeRotationView[Instance];
+
+				for (int32 PositionIdx = 0; PositionIdx < PositionNum; PositionIdx++)
+				{
+					const FVector LocalPosition = RelativeRotation.UnrotateVector(PositionView[Instance][PositionIdx] - RelativePosition);
+
+					UE_VLOG_LOCATION(this, LogLearning, Display,
+						PositionView[Instance][PositionIdx],
+						10.0f,
+						VisualLogColor.ToFColor(true),
+						TEXT("Position: [% 6.1f % 6.1f % 6.1f]\nLocal Position: [% 6.1f % 6.1f % 6.1f]"),
+						PositionView[Instance][PositionIdx].X, 
+						PositionView[Instance][PositionIdx].Y, 
+						PositionView[Instance][PositionIdx].Z,
+						LocalPosition.X, LocalPosition.Y, LocalPosition.Z);
+
+					UE_VLOG_SEGMENT(this, LogLearning, Display,
+						RelativePosition,
+						PositionView[Instance][PositionIdx],
+						VisualLogColor.ToFColor(true),
+						TEXT(""));
+				}
+
+				if (PositionNum > 0)
+				{
+					UE_LEARNING_AGENTS_VLOG_TRANSFORM(this, LogLearning, Display,
+						RelativePosition,
+						RelativeRotation,
+						VisualLogColor.ToFColor(true),
+						TEXT("Agent %i\nEncoded: [% 6.3f % 6.3f % 6.3f ...]"),
+						Instance,
+						FeatureObject->Scale,
+						FeatureView[Instance][0], FeatureView[Instance][1], FeatureView[Instance][2]);
+				}
+				else
+				{
+					UE_LEARNING_AGENTS_VLOG_TRANSFORM(this, LogLearning, Display,
+						RelativePosition,
+						RelativeRotation,
+						VisualLogColor.ToFColor(true),
+						TEXT("Agent %i\nScale: [% 6.2f]\nEncoded: []"),
+						Instance,
+						FeatureObject->Scale);
+				}
+			}
+		}
+	}
+}
+#endif
+
+//------------------------------------------------------------------
+
+UPlanarVelocityObservation* UPlanarVelocityObservation::AddPlanarVelocityObservation(ULearningAgentsType* AgentType, FName Name, float Scale, FVector Axis0, FVector Axis1)
+{
+	return UE::Learning::Agents::Private::AddObservation<UPlanarVelocityObservation, UE::Learning::FPlanarVelocityFeature>(AgentType, Name, 1, Scale, Axis0, Axis1);
+}
+
+void UPlanarVelocityObservation::SetPlanarVelocityObservation(int32 AgentId, FVector Velocity, FRotator RelativeRotation)
+{
+	TLearningArrayView<2, FVector> VelocityView = FeatureObject->InstanceData->View(FeatureObject->VelocityHandle);
+	TLearningArrayView<1, FQuat> RelativeRotationView = FeatureObject->InstanceData->View(FeatureObject->RelativeRotationHandle);
+
+	if (AgentId == INDEX_NONE)
+	{
+		UE_LOG(LogLearning, Error, TEXT("AgentId is invalid (INDEX_NONE)"));
+		return;
+	}
+
+	if (AgentId < 0 || AgentId >= VelocityView.Num<0>())
+	{
+		UE_LOG(LogLearning, Error, TEXT("AgentId %d is out of index. Valid range [0, %d]."), AgentId, VelocityView.Num<0>() - 1);
+		return;
+	}
+
+	VelocityView[AgentId][0] = Velocity;
+	RelativeRotationView[AgentId] = FQuat::MakeFromRotator(RelativeRotation);
+}
+
+#if ENABLE_VISUAL_LOG
+void UPlanarVelocityObservation::VisualLog(const UE::Learning::FIndexSet Instances) const
+{
+	UE_LEARNING_TRACE_CPUPROFILER_EVENT_SCOPE(UPlanarVelocityObservation::VisualLog);
+
+	const TLearningArrayView<2, const FVector> VelocityView = FeatureObject->InstanceData->ConstView(FeatureObject->VelocityHandle);
+	const TLearningArrayView<1, const FQuat> RelativeRotationView = FeatureObject->InstanceData->ConstView(FeatureObject->RelativeRotationHandle);
+	const TLearningArrayView<2, const float> FeatureView = FeatureObject->InstanceData->ConstView(FeatureObject->FeatureHandle);
+
+	if (const ULearningAgentsType* AgentType = Cast<ULearningAgentsType>(GetOuter()))
+	{
+		for (const int32 Instance : Instances)
+		{
+			if (const AActor* Actor = Cast<AActor>(AgentType->GetAgent(Instance)))
+			{
+				const FVector Velocity = VelocityView[Instance][0];
+				const FQuat RelativeRotation = RelativeRotationView[Instance];
+				const FVector LocalVelocity = RelativeRotation.UnrotateVector(Velocity);
+
+				UE_VLOG_ARROW(this, LogLearning, Display,
+					Actor->GetActorLocation(),
+					Actor->GetActorLocation() + Velocity,
+					VisualLogColor.ToFColor(true),
+					TEXT(""));
+
+				UE_LEARNING_AGENTS_VLOG_STRING(this, LogLearning, Display,
+					Actor->GetActorLocation() + Velocity,
+					VisualLogColor.ToFColor(true),
+					TEXT("Velocity: [% 6.3f % 6.3f % 6.3f]\nLocal Velocity: [% 6.3f % 6.3f % 6.3f]"),
+					Velocity.X, Velocity.Y, Velocity.Z,
+					LocalVelocity.X, LocalVelocity.Y, LocalVelocity.Z);
+
+				UE_LEARNING_AGENTS_VLOG_PLANE(this, LogLearning, Display,
+					Actor->GetActorLocation(),
+					RelativeRotation,
+					FeatureObject->Axis0,
+					FeatureObject->Axis1,
+					VisualLogColor.ToFColor(true),
+					TEXT(""));
+
+				UE_LEARNING_AGENTS_VLOG_TRANSFORM(this, LogLearning, Display,
+					Actor->GetActorLocation(),
+					RelativeRotation,
+					VisualLogColor.ToFColor(true),
+					TEXT("Agent %i\nScale: [% 6.2f]\nEncoded: [% 6.3f % 6.3f]"),
+					Instance,
+					FeatureObject->Scale,
+					FeatureView[Instance][0], FeatureView[Instance][1]);
+			}
+		}
+	}
+}
+#endif
+
+UVelocityObservation* UVelocityObservation::AddVelocityObservation(ULearningAgentsType* AgentType, FName Name, float Scale)
+{
+	return UE::Learning::Agents::Private::AddObservation<UVelocityObservation, UE::Learning::FVelocityFeature>(AgentType, Name, 1, Scale);
+}
+
+void UVelocityObservation::SetVelocityObservation(int32 AgentId, FVector Velocity, FRotator RelativeRotation)
+{
+	TLearningArrayView<2, FVector> VelocityView = FeatureObject->InstanceData->View(FeatureObject->VelocityHandle);
+	TLearningArrayView<1, FQuat> RelativeRotationView = FeatureObject->InstanceData->View(FeatureObject->RelativeRotationHandle);
+
+	if (AgentId == INDEX_NONE)
+	{
+		UE_LOG(LogLearning, Error, TEXT("AgentId is invalid (INDEX_NONE)"));
+		return;
+	}
+
+	if (AgentId < 0 || AgentId >= VelocityView.Num<0>())
+	{
+		UE_LOG(LogLearning, Error, TEXT("AgentId %d is out of index. Valid range [0, %d]."), AgentId, VelocityView.Num<0>() - 1);
+		return;
+	}
+
+	VelocityView[AgentId][0] = Velocity;
+	RelativeRotationView[AgentId] = FQuat::MakeFromRotator(RelativeRotation);
+}
+
+#if ENABLE_VISUAL_LOG
+void UVelocityObservation::VisualLog(const UE::Learning::FIndexSet Instances) const
+{
+	UE_LEARNING_TRACE_CPUPROFILER_EVENT_SCOPE(UVelocityObservation::VisualLog);
+
+	const TLearningArrayView<2, const FVector> VelocityView = FeatureObject->InstanceData->ConstView(FeatureObject->VelocityHandle);
+	const TLearningArrayView<1, const FQuat> RelativeRotationView = FeatureObject->InstanceData->ConstView(FeatureObject->RelativeRotationHandle);
+	const TLearningArrayView<2, const float> FeatureView = FeatureObject->InstanceData->ConstView(FeatureObject->FeatureHandle);
+
+	if (const ULearningAgentsType* AgentType = Cast<ULearningAgentsType>(GetOuter()))
+	{
+		for (const int32 Instance : Instances)
+		{
+			if (const AActor* Actor = Cast<AActor>(AgentType->GetAgent(Instance)))
+			{
+				const FVector Velocity = VelocityView[Instance][0];
+				const FQuat RelativeRotation = RelativeRotationView[Instance];
+				const FVector LocalVelocity = RelativeRotation.UnrotateVector(Velocity);
+
+				UE_VLOG_ARROW(this, LogLearning, Display,
+					Actor->GetActorLocation(),
+					Actor->GetActorLocation() + Velocity,
+					VisualLogColor.ToFColor(true),
+					TEXT(""));
+
+				UE_LEARNING_AGENTS_VLOG_STRING(this, LogLearning, Display,
+					Actor->GetActorLocation() + Velocity,
+					VisualLogColor.ToFColor(true),
+					TEXT("Velocity: [% 6.3f % 6.3f % 6.3f]\nLocal Velocity: [% 6.3f % 6.3f % 6.3f]"),
+					Velocity.X, Velocity.Y, Velocity.Z,
+					LocalVelocity.X, LocalVelocity.Y, LocalVelocity.Z);
+
+				UE_LEARNING_AGENTS_VLOG_TRANSFORM(this, LogLearning, Display,
+					Actor->GetActorLocation(),
+					RelativeRotation,
+					VisualLogColor.ToFColor(true),
+					TEXT("Agent %i\nScale: [% 6.2f]\nEncoded: [% 6.3f % 6.3f % 6.3f]"),
+					Instance,
+					FeatureObject->Scale,
+					FeatureView[Instance][0], FeatureView[Instance][1], FeatureView[Instance][2]);
+			}
+		}
+	}
+}
+#endif
+
+//------------------------------------------------------------------
+
+#undef UE_LEARNING_AGENTS_VLOG_STRING
+#undef UE_LEARNING_AGENTS_VLOG_TRANSFORM
+#undef UE_LEARNING_AGENTS_VLOG_PLANE
diff --git a/Engine/Plugins/Experimental/LearningAgents/Source/LearningAgents/Private/LearningAgentsType.cpp b/Engine/Plugins/Experimental/LearningAgents/Source/LearningAgents/Private/LearningAgentsType.cpp
new file mode 100644
--- /dev/null
+++ b/Engine/Plugins/Experimental/LearningAgents/Source/LearningAgents/Private/LearningAgentsType.cpp
@@ -0,0 +1,396 @@
+// Copyright Epic Games, Inc. All Rights Reserved.
+
+#include "LearningAgentsType.h"
+#include "LearningAgentsObservations.h"
+#include "LearningAgentsActions.h"
+
+#include "LearningArray.h"
+#include "LearningArrayMap.h"
+#include "LearningFeatureObject.h"
+#include "LearningLog.h"
+#include "LearningNeuralNetwork.h"
+#include "LearningPolicyObject.h"
+#include "Misc/FileHelper.h"
+#include "Misc/Paths.h"
+#include "EngineDefines.h"
+
+ULearningAgentsType::ULearningAgentsType()
+{
+	PrimaryComponentTick.bCanEverTick = true;
+	PrimaryComponentTick.bStartWithTickEnabled = false;
+}
+
+ULearningAgentsType::ULearningAgentsType(FVTableHelper& Helper) : ULearningAgentsType() {}
+ULearningAgentsType::~ULearningAgentsType() {}
+
+void ULearningAgentsType::PostInitProperties()
+{
+	Super::PostInitProperties();
+
+	if (HasAnyFlags(RF_ClassDefaultObject))
+	{
+		return;
+	}
+
+	// Pre-populate the vacant ids
+	OccupiedAgentIds.Reserve(MaxInstanceNum);
+	VacantAgentIds.Reserve(MaxInstanceNum);
+	for (int32 i = MaxInstanceNum - 1; i >= 0; i--)
+	{
+		VacantAgentIds.Push(i);
+		Agents.Add(nullptr);
+	}
+
+	UpdateAgentSets();
+}
+
+void ULearningAgentsType::SetupAgentType(
+	const FLearningAgentsTypeSettings& Settings,
+	const FLearningAgentsNetworkSettings& NetworkSettings)
+{
+	// Reset Setup Flag
+	bSetupPerformed = false;
+
+	// Allocate Instance Data
+	InstanceData = MakeShared<UE::Learning::FArrayMap>();
+
+	// Setup Observations
+	ObservationObjects.Empty();
+	ObservationFeatures.Empty();
+	SetupObservations(this);
+	Observations = MakeShared<UE::Learning::FConcatenateFeature>(TEXT("Observations"),
+		TLearningArrayView<1, const TSharedRef<UE::Learning::FFeatureObject>>(ObservationFeatures),
+		InstanceData.ToSharedRef(),
+		MaxInstanceNum);
+
+	// Setup Actions
+	ActionObjects.Empty();
+	ActionFeatures.Empty();
+	SetupActions(this);
+	Actions = MakeShared<UE::Learning::FConcatenateFeature>(TEXT("Actions"),
+		TLearningArrayView<1, const TSharedRef<UE::Learning::FFeatureObject>>(ActionFeatures),
+		InstanceData.ToSharedRef(),
+		MaxInstanceNum);
+
+	// Create Neural Network
+	NeuralNetwork = MakeShared<UE::Learning::FNeuralNetwork>();
+	NeuralNetwork->Resize(
+		Observations->DimNum(), 
+		2 * Actions->DimNum(), 
+		NetworkSettings.HiddenLayerSize, 
+		NetworkSettings.LayerNum);
+
+	switch (NetworkSettings.ActivationFunction)
+	{
+		case ELearningAgentsActivationFunction::ReLU:
+			NeuralNetwork->ActivationFunction = UE::Learning::EActivationFunction::ReLU;
+			break;
+		case ELearningAgentsActivationFunction::ELU:
+			NeuralNetwork->ActivationFunction = UE::Learning::EActivationFunction::ELU;
+			break;
+		case ELearningAgentsActivationFunction::TanH:
+			NeuralNetwork->ActivationFunction = UE::Learning::EActivationFunction::TanH;
+			break;
+		default:
+			UE_LOG(LogLearning, Error, TEXT("NetworkSetting's activation function was not found. Defaulting to ELU."));
+			NeuralNetwork->ActivationFunction = UE::Learning::EActivationFunction::ELU;
+			break;
+	}
+
+	// Create Policy
+	UE::Learning::FNeuralNetworkPolicyFunctionSettings PolicySettings;
+	PolicySettings.ActionNoiseMin = Settings.ActionNoiseMin;
+	PolicySettings.ActionNoiseMax = Settings.ActionNoiseMax;
+
+	Policy = MakeUnique<UE::Learning::FNeuralNetworkPolicyFunction>(
+		TEXT("Policy"), 
+		InstanceData.ToSharedRef(), 
+		MaxInstanceNum, 
+		NeuralNetwork.ToSharedRef(), 
+		Settings.ActionNoiseSeed,
+		PolicySettings);
+
+	InstanceData->Link(Observations->FeatureHandle, Policy->InputHandle);
+	InstanceData->Link(Policy->OutputHandle, Actions->FeatureHandle);
+
+	// Done!
+	bSetupPerformed = true;
+}
+
+const bool ULearningAgentsType::IsSetupPerformed() const
+{
+	return bSetupPerformed;
+}
+
+const int32 ULearningAgentsType::GetMaxInstanceNum() const
+{
+	return MaxInstanceNum;
+}
+
+const TSharedPtr<UE::Learning::FArrayMap>& ULearningAgentsType::GetInstanceData() const
+{
+	return InstanceData;
+}
+
+UE::Learning::FFeatureObject& ULearningAgentsType::GetObservationFeature() const
+{
+	return *Observations;
+}
+
+UE::Learning::FFeatureObject& ULearningAgentsType::GetActionFeature() const
+{
+	return *Actions;
+}
+
+UE::Learning::FNeuralNetwork& ULearningAgentsType::GetNeuralNetwork() const
+{
+	return *NeuralNetwork;
+}
+
+const UE::Learning::FNeuralNetworkPolicyFunction& ULearningAgentsType::GetPolicy() const
+{
+	return *Policy;
+}
+
+const TConstArrayView<TObjectPtr<UObject>> ULearningAgentsType::GetAgents() const
+{
+	return Agents;
+}
+
+const UE::Learning::FIndexSet ULearningAgentsType::GetOccupiedAgentSet() const
+{
+	return OccupiedAgentSet;
+}
+
+const UE::Learning::FIndexSet ULearningAgentsType::GetVacantAgentSet() const
+{
+	return VacantAgentSet;
+}
+
+int32 ULearningAgentsType::AddAgent(UObject* Agent)
+{
+	if (Agent == nullptr)
+	{
+		UE_LOG(LogLearning, Error, TEXT("Attempted to add an agent but agent is nullptr."));
+		return INDEX_NONE;
+	}
+
+	if (VacantAgentIds.IsEmpty())
+	{
+		UE_LOG(LogLearning, Error, TEXT("Attempting to add an agent but we have no more vacant ids. Increase MaxInstanceNum (%d) or remove unused agents."), Agents.Num());
+		return INDEX_NONE;
+	}
+
+	// Add Agent
+	const int32 NewAgentId = VacantAgentIds.Pop();
+	Agents[NewAgentId] = Agent;
+
+	OccupiedAgentIds.Add(NewAgentId);
+
+	UpdateAgentSets();
+
+	OnAgentAdded.Broadcast(NewAgentId, Agent);
+
+	return NewAgentId;
+}
+
+void ULearningAgentsType::RemoveAgentById(int32 AgentId)
+{
+	if (AgentId == INDEX_NONE)
+	{
+		UE_LOG(LogLearning, Warning, TEXT("Attempting to remove an agent with id of INDEX_NONE."));
+		return;
+	}
+
+	int32 RemovedCount = OccupiedAgentIds.RemoveSingleSwap(AgentId, false);
+
+	if (RemovedCount == 0)
+	{
+		UE_LOG(LogLearning, Warning, TEXT("Trying to remove an agent but its Id (%d) is not in the occupied agents."), AgentId);
+		return;
+	}
+
+	UObject* RemovedAgent = Agents[AgentId];
+
+	// Remove Agent
+	VacantAgentIds.Push(AgentId);
+	Agents[AgentId] = nullptr;
+
+	UpdateAgentSets();
+
+	OnAgentRemoved.Broadcast(AgentId, RemovedAgent);
+}
+
+void ULearningAgentsType::RemoveAgent(UObject* Agent)
+{
+	if (Agent == nullptr)
+	{
+		UE_LOG(LogLearning, Error, TEXT("Attempted to remove an agent but agent is nullptr."));
+		return;
+	}
+
+	int32 AgentId = INDEX_NONE;
+	bool bIsFound = Agents.Find(Agent, AgentId);
+
+	if (!bIsFound)
+	{
+		UE_LOG(LogLearning, Warning, TEXT("Trying to remove an agent but it was not found."), AgentId);
+		return;
+	}
+
+	RemoveAgentById(AgentId);
+}
+
+UObject* ULearningAgentsType::GetAgent(int32 AgentId, TSubclassOf<UObject> AgentClass)
+{
+	if (AgentId < 0 || AgentId >= Agents.Num())
+	{
+		UE_LOG(LogLearning, Warning, TEXT("AgentId %d outside valid range [0, %d]"), AgentId, Agents.Num() - 1);
+		return nullptr;
+	}
+
+	return Agents[AgentId];
+}
+
+const UObject* ULearningAgentsType::GetAgent(int32 AgentId) const
+{
+	if (AgentId < 0 || AgentId >= Agents.Num())
+	{
+		UE_LOG(LogLearning, Warning, TEXT("AgentId %d outside valid range [0, %d]"), AgentId, Agents.Num() - 1);
+		return nullptr;
+	}
+
+	return Agents[AgentId];
+}
+
+void ULearningAgentsType::SetupObservations_Implementation(ULearningAgentsType* AgentType)
+{
+	// Can be overridden to setup observations without blueprints
+}
+
+void ULearningAgentsType::SetObservations_Implementation(const TArray<int32>& AgentIds)
+{
+	// Can be overridden to set observations without blueprints
+}
+
+void ULearningAgentsType::AddObservation(TObjectPtr<ULearningAgentsObservation> Object, const TSharedRef<UE::Learning::FFeatureObject>& Feature)
+{
+	check(!bSetupPerformed);
+	ObservationObjects.Add(Object);
+	ObservationFeatures.Add(Feature);
+}
+
+void ULearningAgentsType::SetupActions_Implementation(ULearningAgentsType* AgentType)
+{
+	// Can be overridden to setup actions without blueprints
+}
+
+void ULearningAgentsType::GetActions_Implementation(const TArray<int32>& AgentIds)
+{
+	// Can be overridden to get actions without blueprints
+}
+
+void ULearningAgentsType::AddAction(TObjectPtr<ULearningAgentsAction> Object, const TSharedRef<UE::Learning::FFeatureObject>& Feature)
+{
+	check(!bSetupPerformed);
+	ActionObjects.Add(Object);
+	ActionFeatures.Add(Feature);
+}
+
+void ULearningAgentsType::LoadNetwork(const FDirectoryPath& Directory, const FString Filename)
+{
+	if (!bSetupPerformed)
+	{
+		UE_LOG(LogLearning, Error, TEXT("Setup must be run before network can be loaded."));
+		return;
+	}
+
+	TArray64<uint8> NetworkData;
+	FString FilePath = Directory.Path + FGenericPlatformMisc::GetDefaultPathSeparator() + Filename;
+	if (FFileHelper::LoadFileToArray(NetworkData, *FilePath))
+	{
+		NeuralNetwork->DeserializeFromBytes(NetworkData);
+		bNetworkLoaded = true;
+	}
+	else
+	{
+		UE_LOG(LogLearning, Warning, TEXT("Failed to load network. File not found: %s..."), *FilePath);
+	}
+}
+
+void ULearningAgentsType::EncodeObservations()
+{
+	UE_LEARNING_TRACE_CPUPROFILER_EVENT_SCOPE(ULearningAgentsType::EncodeObservations);
+
+	if (!bSetupPerformed)
+	{
+		UE_LOG(LogLearning, Error, TEXT("Setup must be run before observations can be encoded."));
+		return;
+	}
+
+	SetObservations(OccupiedAgentIds);
+
+	Observations->Encode(OccupiedAgentSet);
+
+#if ENABLE_VISUAL_LOG
+	for (const ULearningAgentsObservation* ObservationObject : ObservationObjects)
+	{
+		if (ObservationObject)
+		{
+			ObservationObject->VisualLog(OccupiedAgentSet);
+		}
+	}
+#endif
+}
+
+
+void ULearningAgentsType::EvaluatePolicy()
+{
+	UE_LEARNING_TRACE_CPUPROFILER_EVENT_SCOPE(ULearningAgentsType::EvaluatePolicy);
+
+	if (!bSetupPerformed)
+	{
+		UE_LOG(LogLearning, Error, TEXT("Setup must be run before the policy can be evaluated."));
+		return;
+	}
+
+	Policy->Evaluate(OccupiedAgentSet);
+}
+
+void ULearningAgentsType::DecodeActions()
+{
+	UE_LEARNING_TRACE_CPUPROFILER_EVENT_SCOPE(ULearningAgentsType::DecodeActions);
+
+	if (!bSetupPerformed)
+	{
+		UE_LOG(LogLearning, Error, TEXT("Setup must be run before actions can be decoded."));
+		return;
+	}
+
+	Actions->Decode(OccupiedAgentSet);
+
+	GetActions(OccupiedAgentIds);
+
+#if ENABLE_VISUAL_LOG
+	for (const ULearningAgentsAction* ActionObject : ActionObjects)
+	{
+		if (ActionObject)
+		{
+			ActionObject->VisualLog(OccupiedAgentSet);
+		}
+	}
+#endif
+}
+
+void ULearningAgentsType::BroadcastSetupComplete()
+{
+	OnSetupComplete.Broadcast();
+}
+
+void ULearningAgentsType::UpdateAgentSets()
+{
+	OccupiedAgentSet = OccupiedAgentIds;
+	OccupiedAgentSet.TryMakeSlice();
+	VacantAgentSet = VacantAgentIds;
+	VacantAgentSet.TryMakeSlice();
+}
\ No newline at end of file
diff --git a/Engine/Plugins/Experimental/LearningAgents/Source/LearningAgents/Public/LearningAgentsActions.h b/Engine/Plugins/Experimental/LearningAgents/Source/LearningAgents/Public/LearningAgentsActions.h
new file mode 100644
--- /dev/null
+++ b/Engine/Plugins/Experimental/LearningAgents/Source/LearningAgents/Public/LearningAgentsActions.h
@@ -0,0 +1,115 @@
+// Copyright Epic Games, Inc. All Rights Reserved.
+
+#pragma once
+
+#include "LearningArray.h"
+
+#include "Templates/SharedPointer.h"
+#include "UObject/Object.h"
+#include "EngineDefines.h"
+
+#include "LearningAgentsActions.generated.h"
+
+namespace UE::Learning
+{
+	struct FFeatureObject;
+	struct FFloatFeature;
+	struct FRotationVectorFeature;
+}
+
+class ULearningAgentsType;
+
+// For functions in this file, we are favoring having more verbose names such as "AddFloatAction" vs simply "Add" in 
+// order to keep it easy to find the correct function in blueprints.
+
+UCLASS(Abstract, BlueprintType)
+class LEARNINGAGENTS_API ULearningAgentsAction : public UObject
+{
+	GENERATED_BODY()
+
+public:
+
+#if ENABLE_VISUAL_LOG
+	/** Color used to draw this action in the visual log */
+	FLinearColor VisualLogColor = FColor::Blue;
+
+	virtual void VisualLog(const UE::Learning::FIndexSet Instances) const {}
+#endif
+};
+
+//------------------------------------------------------------------
+
+UCLASS()
+class LEARNINGAGENTS_API UFloatAction : public ULearningAgentsAction
+{
+	GENERATED_BODY()
+
+public:
+	UFUNCTION(BlueprintCallable, Category = "LearningAgents")
+	static UFloatAction* AddFloatAction(ULearningAgentsType* AgentType, FName Name = NAME_None, float Scale = 1.0f);
+
+	UFUNCTION(BlueprintPure, Category = "LearningAgents")
+	float GetFloatAction(int32 AgentId);
+
+	UFUNCTION(BlueprintCallable, Category = "LearningAgents")
+	void SetFloatAction(int32 AgentId, float Value);
+
+#if ENABLE_VISUAL_LOG
+	virtual void VisualLog(const UE::Learning::FIndexSet Instances) const override;
+#endif
+
+	TSharedPtr<UE::Learning::FFloatFeature> FeatureObject;
+};
+
+//------------------------------------------------------------------
+
+UCLASS()
+class LEARNINGAGENTS_API UVectorAction : public ULearningAgentsAction
+{
+	GENERATED_BODY()
+
+public:
+	UFUNCTION(BlueprintCallable, Category = "LearningAgents")
+	static UVectorAction* AddVectorAction(ULearningAgentsType* AgentType, FName Name = NAME_None, float Scale = 1.0f);
+
+	UFUNCTION(BlueprintPure, Category = "LearningAgents")
+	FVector GetVectorAction(int32 AgentId);
+
+	UFUNCTION(BlueprintCallable, Category = "LearningAgents")
+	void SetVectorAction(int32 AgentId, FVector Action);
+
+#if ENABLE_VISUAL_LOG
+	virtual void VisualLog(const UE::Learning::FIndexSet Instances) const override;
+#endif
+
+	TSharedPtr<UE::Learning::FFloatFeature> FeatureObject;
+};
+
+//------------------------------------------------------------------
+
+UCLASS()
+class LEARNINGAGENTS_API URotationVectorArrayAction : public ULearningAgentsAction
+{
+	GENERATED_BODY()
+
+public:
+	UFUNCTION(BlueprintCallable, Category = "LearningAgents")
+	static URotationVectorArrayAction* AddRotationVectorArrayAction(ULearningAgentsType* AgentType, FName Name = NAME_None, int32 RotationVectorNum = 0, float Scale = 180.0f);
+
+	UFUNCTION(BlueprintPure, Category = "LearningAgents")
+	void GetRotationVectorArrayAction(int32 AgentId, TArray<FVector>& OutRotationVectors);
+
+	UFUNCTION(BlueprintPure, Category = "LearningAgents")
+	void GetRotationVectorArrayActionAsQuats(int32 AgentId, TArray<FQuat>& OutRotations);
+
+	UFUNCTION(BlueprintPure, Category = "LearningAgents")
+	void GetRotationVectorArrayActionAsRotators(int32 AgentId, TArray<FRotator>& OutRotations);
+
+#if ENABLE_VISUAL_LOG
+	virtual void VisualLog(const UE::Learning::FIndexSet Instances) const override;
+#endif
+
+	TSharedPtr<UE::Learning::FRotationVectorFeature> FeatureObject;
+};
+
+//------------------------------------------------------------------
\ No newline at end of file
diff --git a/Engine/Plugins/Experimental/LearningAgents/Source/LearningAgents/Public/LearningAgentsHelpers.h b/Engine/Plugins/Experimental/LearningAgents/Source/LearningAgents/Public/LearningAgentsHelpers.h
new file mode 100644
--- /dev/null
+++ b/Engine/Plugins/Experimental/LearningAgents/Source/LearningAgents/Public/LearningAgentsHelpers.h
@@ -0,0 +1,52 @@
+// Copyright Epic Games, Inc. All Rights Reserved.
+
+#pragma once
+
+#include "Templates/SharedPointer.h"
+#include "UObject/Object.h"
+#include "Kismet/BlueprintFunctionLibrary.h"
+
+#include "Components/SplineComponent.h" // Required for ESplineCoordinateSpace::World
+
+#include "LearningAgentsHelpers.generated.h"
+
+class UMeshComponent;
+
+UCLASS()
+class LEARNINGAGENTS_API ULearningAgentsHelpers : public UBlueprintFunctionLibrary
+{
+	GENERATED_BODY()
+
+	//------------------------------------------------------------------
+	// General Math Helpers
+	//------------------------------------------------------------------
+
+	UFUNCTION(BlueprintPure, Category = "LearningAgents|Helpers")
+	static FTransform ProjectTransformOntoGroundPlane(const FTransform Transform, const FVector LocalForwardVector = FVector::ForwardVector);
+
+	UFUNCTION(BlueprintPure, Category = "LearningAgents|Helpers")
+	static void ProjectPositionRotationOntoGroundPlane(FVector& OutPosition, FRotator& OutRotation, const FVector InPosition, const FRotator InRotation, const FVector LocalForwardVector = FVector::ForwardVector);
+
+	//------------------------------------------------------------------
+	// MeshComponent Helpers 
+	//------------------------------------------------------------------
+
+	UFUNCTION(BlueprintPure, Category = "LearningAgents|Helpers")
+	static void GetMeshBonePositions(const UMeshComponent* MeshComponent, const TArray<FName>& BoneNames, TArray<FVector>& OutBonePositions);
+
+	//------------------------------------------------------------------
+	// SplineComponent Helpers 
+	//------------------------------------------------------------------
+
+	UFUNCTION(BlueprintPure, Category = "LearningAgents|Helpers")
+	static float GetProportionAlongSpline(const USplineComponent* SplineComponent, const float DistanceAlongSpline);
+
+	UFUNCTION(BlueprintPure, Category = "LearningAgents|Helpers")
+	static float GetProportionAlongSplineAsAngle(const USplineComponent* SplineComponent, const float DistanceAlongSpline);
+
+	UFUNCTION(BlueprintPure, Category = "LearningAgents|Helpers")
+	static void GetPositionsAlongSpline(TArray<FVector>& OutPositions, const USplineComponent* SplineComponent, const int32 PositionNum, const float StartDistanceAlongSpline, const float StopDistanceAlongSpline, const ESplineCoordinateSpace::Type CoordinateSpace = ESplineCoordinateSpace::World);
+	
+	UFUNCTION(BlueprintPure, Category = "LearningAgents|Helpers")
+	static float GetVelocityAlongSpline(const USplineComponent* SplineComponent, const FVector Position, const FVector Velocity, const float FiniteDifferenceDelta = 1.0f, const ESplineCoordinateSpace::Type CoordinateSpace = ESplineCoordinateSpace::World);
+};
diff --git a/Engine/Plugins/Experimental/LearningAgents/Source/LearningAgents/Public/LearningAgentsObservations.h b/Engine/Plugins/Experimental/LearningAgents/Source/LearningAgents/Public/LearningAgentsObservations.h
new file mode 100644
--- /dev/null
+++ b/Engine/Plugins/Experimental/LearningAgents/Source/LearningAgents/Public/LearningAgentsObservations.h
@@ -0,0 +1,267 @@
+// Copyright Epic Games, Inc. All Rights Reserved.
+
+#pragma once
+
+#include "LearningArray.h"
+
+#include "Templates/SharedPointer.h"
+#include "UObject/Object.h"
+#include "EngineDefines.h"
+
+#include "LearningAgentsObservations.generated.h"
+
+namespace UE::Learning
+{
+	struct FFeatureObject;
+	struct FFloatFeature;
+	struct FAngleFeature;
+	struct FPlanarDirectionFeature;
+	struct FDirectionFeature;
+	struct FPlanarPositionFeature;
+	struct FPositionFeature;
+	struct FPlanarVelocityFeature;
+	struct FVelocityFeature;
+}
+
+class ULearningAgentsType;
+
+// For functions in this file, we are favoring having more verbose names such as "AddFloatObservation" vs simply "Add" in 
+// order to keep it easy to find the correct function in blueprints.
+
+UCLASS(Abstract, BlueprintType)
+class LEARNINGAGENTS_API ULearningAgentsObservation : public UObject
+{
+	GENERATED_BODY()
+
+public:
+
+#if ENABLE_VISUAL_LOG
+	/** Color used to draw this observation in the visual log */
+	FLinearColor VisualLogColor = FColor::Red;
+
+	virtual void VisualLog(const UE::Learning::FIndexSet Instances) const {}
+#endif
+};
+
+//------------------------------------------------------------------
+
+UCLASS()
+class LEARNINGAGENTS_API UFloatObservation : public ULearningAgentsObservation
+{
+	GENERATED_BODY()
+
+public:
+	UFUNCTION(BlueprintCallable, Category = "LearningAgents")
+	static UFloatObservation* AddFloatObservation(ULearningAgentsType* AgentType, FName Name = NAME_None, float Scale = 1.0f);
+
+	UFUNCTION(BlueprintCallable, Category = "LearningAgents")
+	void SetFloatObservation(int32 AgentId, float Observation);
+
+#if ENABLE_VISUAL_LOG
+	virtual void VisualLog(const UE::Learning::FIndexSet Instances) const override;
+#endif
+
+	TSharedPtr<UE::Learning::FFloatFeature> FeatureObject;
+};
+
+//------------------------------------------------------------------
+
+UCLASS()
+class LEARNINGAGENTS_API UVectorObservation : public ULearningAgentsObservation
+{
+	GENERATED_BODY()
+
+public:
+	UFUNCTION(BlueprintCallable, Category = "LearningAgents")
+	static UVectorObservation* AddVectorObservation(ULearningAgentsType* AgentType, FName Name = NAME_None, float Scale = 1.0f);
+
+	UFUNCTION(BlueprintCallable, Category = "LearningAgents")
+	void SetVectorObservation(int32 AgentId, FVector Observation);
+
+#if ENABLE_VISUAL_LOG
+	virtual void VisualLog(const UE::Learning::FIndexSet Instances) const override;
+#endif
+
+	TSharedPtr<UE::Learning::FFloatFeature> FeatureObject;
+};
+
+//------------------------------------------------------------------
+
+UCLASS()
+class LEARNINGAGENTS_API UAngleObservation : public ULearningAgentsObservation
+{
+	GENERATED_BODY()
+
+public:
+	UFUNCTION(BlueprintCallable, Category = "LearningAgents")
+	static UAngleObservation* AddAngleObservation(ULearningAgentsType* AgentType, FName Name = NAME_None, float Scale = 1.0f);
+
+	UFUNCTION(BlueprintCallable, Category = "LearningAgents")
+	void SetAngleObservation(int32 AgentId, float Angle, float RelativeAngle = 0.0f);
+
+#if ENABLE_VISUAL_LOG
+	virtual void VisualLog(const UE::Learning::FIndexSet Instances) const override;
+#endif
+
+	TSharedPtr<UE::Learning::FAngleFeature> FeatureObject;
+};
+
+//------------------------------------------------------------------
+
+UCLASS()
+class LEARNINGAGENTS_API UPlanarDirectionObservation : public ULearningAgentsObservation
+{
+	GENERATED_BODY()
+
+public:
+	UFUNCTION(BlueprintCallable, Category = "LearningAgents")
+	static UPlanarDirectionObservation* AddPlanarDirectionObservation(ULearningAgentsType* AgentType, FName Name = NAME_None, float Scale = 1.0f, FVector Axis0 = FVector::ForwardVector, FVector Axis1 = FVector::RightVector);
+
+	UFUNCTION(BlueprintCallable, Category = "LearningAgents")
+	void SetPlanarDirectionObservation(int32 AgentId, FVector Direction, FRotator RelativeRotation = FRotator::ZeroRotator);
+
+#if ENABLE_VISUAL_LOG
+	virtual void VisualLog(const UE::Learning::FIndexSet Instances) const override;
+#endif
+
+	TSharedPtr<UE::Learning::FPlanarDirectionFeature> FeatureObject;
+};
+
+UCLASS()
+class LEARNINGAGENTS_API UDirectionObservation : public ULearningAgentsObservation
+{
+	GENERATED_BODY()
+
+public:
+	UFUNCTION(BlueprintCallable, Category = "LearningAgents")
+	static UDirectionObservation* AddDirectionObservation(ULearningAgentsType* AgentType, FName Name = NAME_None, float Scale = 100.0f);
+
+	UFUNCTION(BlueprintCallable, Category = "LearningAgents")
+	void SetDirectionObservation(int32 AgentId, FVector Direction, FRotator RelativeRotation = FRotator::ZeroRotator);
+
+#if ENABLE_VISUAL_LOG
+	virtual void VisualLog(const UE::Learning::FIndexSet Instances) const override;
+#endif
+
+	TSharedPtr<UE::Learning::FDirectionFeature> FeatureObject;
+};
+
+//------------------------------------------------------------------
+
+UCLASS()
+class LEARNINGAGENTS_API UPlanarPositionObservation : public ULearningAgentsObservation
+{
+	GENERATED_BODY()
+
+public:
+	UFUNCTION(BlueprintCallable, Category = "LearningAgents")
+	static UPlanarPositionObservation* AddPlanarPositionObservation(ULearningAgentsType* AgentType, FName Name = NAME_None, float Scale = 100.0f, FVector Axis0 = FVector::ForwardVector, FVector Axis1 = FVector::RightVector);
+
+	UFUNCTION(BlueprintCallable, Category = "LearningAgents")
+	void SetPlanarPositionObservation(int32 AgentId, FVector Position, FVector RelativePosition = FVector::ZeroVector, FRotator RelativeRotation = FRotator::ZeroRotator);
+
+#if ENABLE_VISUAL_LOG
+	virtual void VisualLog(const UE::Learning::FIndexSet Instances) const override;
+#endif
+
+	TSharedPtr<UE::Learning::FPlanarPositionFeature> FeatureObject;
+};
+
+UCLASS()
+class LEARNINGAGENTS_API UPositionObservation : public ULearningAgentsObservation
+{
+	GENERATED_BODY()
+
+public:
+	UFUNCTION(BlueprintCallable, Category = "LearningAgents")
+	static UPositionObservation* AddPositionObservation(ULearningAgentsType* AgentType, FName Name = NAME_None, float Scale = 100.0f);
+
+	UFUNCTION(BlueprintCallable, Category = "LearningAgents")
+	void SetPositionObservation(int32 AgentId, FVector Position, FVector RelativePosition = FVector::ZeroVector, FRotator RelativeRotation = FRotator::ZeroRotator);
+
+#if ENABLE_VISUAL_LOG
+	virtual void VisualLog(const UE::Learning::FIndexSet Instances) const override;
+#endif
+
+	TSharedPtr<UE::Learning::FPositionFeature> FeatureObject;
+};
+
+UCLASS()
+class LEARNINGAGENTS_API UPlanarPositionArrayObservation : public ULearningAgentsObservation
+{
+	GENERATED_BODY()
+
+public:
+	UFUNCTION(BlueprintCallable, Category = "LearningAgents")
+	static UPlanarPositionArrayObservation* AddPlanarPositionArrayObservation(ULearningAgentsType* AgentType, FName Name = NAME_None, int32 PositionNum = 0, float Scale = 100.0f, FVector Axis0 = FVector::ForwardVector, FVector Axis1 = FVector::RightVector);
+
+	UFUNCTION(BlueprintCallable, Category = "LearningAgents")
+	void SetPlanarPositionArrayObservation(int32 AgentId, const TArray<FVector>& Positions, FVector RelativePosition = FVector::ZeroVector, FRotator RelativeRotation = FRotator::ZeroRotator);
+
+#if ENABLE_VISUAL_LOG
+	virtual void VisualLog(const UE::Learning::FIndexSet Instances) const override;
+#endif
+
+	TSharedPtr<UE::Learning::FPlanarPositionFeature> FeatureObject;
+};
+
+UCLASS()
+class LEARNINGAGENTS_API UPositionArrayObservation : public ULearningAgentsObservation
+{
+	GENERATED_BODY()
+
+public:
+	UFUNCTION(BlueprintCallable, Category = "LearningAgents")
+	static UPositionArrayObservation* AddPositionArrayObservation(ULearningAgentsType* AgentType, FName Name = NAME_None, int32 PositionNum = 0, float Scale = 100.0f);
+
+	UFUNCTION(BlueprintCallable, Category = "LearningAgents")
+	void SetPositionArrayObservation(int32 AgentId, const TArray<FVector>& Positions, FVector RelativePosition = FVector::ZeroVector, FRotator RelativeRotation = FRotator::ZeroRotator);
+
+#if ENABLE_VISUAL_LOG
+	virtual void VisualLog(const UE::Learning::FIndexSet Instances) const override;
+#endif
+
+	TSharedPtr<UE::Learning::FPositionFeature> FeatureObject;
+};
+
+//------------------------------------------------------------------
+
+UCLASS()
+class LEARNINGAGENTS_API UPlanarVelocityObservation : public ULearningAgentsObservation
+{
+	GENERATED_BODY()
+
+public:
+	UFUNCTION(BlueprintCallable, Category = "LearningAgents")
+	static UPlanarVelocityObservation* AddPlanarVelocityObservation(ULearningAgentsType* AgentType, FName Name = NAME_None, float Scale = 200.0f, FVector Axis0 = FVector::ForwardVector, FVector Axis1 = FVector::RightVector);
+
+	UFUNCTION(BlueprintCallable, Category = "LearningAgents")
+	void SetPlanarVelocityObservation(int32 AgentId, FVector Velocity, FRotator RelativeRotation = FRotator::ZeroRotator);
+
+#if ENABLE_VISUAL_LOG
+	virtual void VisualLog(const UE::Learning::FIndexSet Instances) const override;
+#endif
+
+	TSharedPtr<UE::Learning::FPlanarVelocityFeature> FeatureObject;
+};
+
+UCLASS()
+class LEARNINGAGENTS_API UVelocityObservation : public ULearningAgentsObservation
+{
+	GENERATED_BODY()
+
+public:
+	UFUNCTION(BlueprintCallable, Category = "LearningAgents")
+	static UVelocityObservation* AddVelocityObservation(ULearningAgentsType* AgentType, FName Name = NAME_None, float Scale = 200.0f);
+
+	UFUNCTION(BlueprintCallable, Category = "LearningAgents")
+	void SetVelocityObservation(int32 AgentId, FVector Velocity, FRotator RelativeRotation = FRotator::ZeroRotator);
+
+#if ENABLE_VISUAL_LOG
+	virtual void VisualLog(const UE::Learning::FIndexSet Instances) const override;
+#endif
+
+	TSharedPtr<UE::Learning::FVelocityFeature> FeatureObject;
+};
+
+//------------------------------------------------------------------
\ No newline at end of file
diff --git a/Engine/Plugins/Experimental/LearningAgents/Source/LearningAgents/Public/LearningAgentsType.h b/Engine/Plugins/Experimental/LearningAgents/Source/LearningAgents/Public/LearningAgentsType.h
new file mode 100644
--- /dev/null
+++ b/Engine/Plugins/Experimental/LearningAgents/Source/LearningAgents/Public/LearningAgentsType.h
@@ -0,0 +1,242 @@
+// Copyright Epic Games, Inc. All Rights Reserved.
+
+#pragma once
+
+#include "Components/SceneComponent.h"
+#include "Containers/Array.h"
+#include "Containers/ArrayView.h"
+#include "Containers/ContainerAllocationPolicies.h"
+#include "Containers/ContainersFwd.h"
+#include "LearningArray.h"
+#include "Templates/SharedPointer.h"
+#include "UObject/Object.h"
+#include "UObject/ObjectPtr.h"
+
+#include "LearningAgentsType.generated.h"
+
+namespace UE::Learning
+{
+	struct FArrayMap;
+	struct FConcatenateFeature;
+	struct FFeatureObject;
+	struct FNeuralNetwork;
+	struct FNeuralNetworkPolicyFunction;
+}
+
+struct FDirectoryPath;
+class ULearningAgentsAction;
+class ULearningAgentsObservation;
+
+USTRUCT(BlueprintType, Category = "LearningAgents")
+struct FLearningAgentsTypeSettings
+{
+	GENERATED_BODY()
+
+public:
+
+	/** Seed for the action noise used by the policy  */
+	UPROPERTY(EditAnywhere, Category = "LearningAgents", meta = (ClampMin = "0", UIMin = "0"))
+	int32 ActionNoiseSeed = 1234;
+
+	/** Minimum action noise used by the policy */
+	UPROPERTY(EditAnywhere, Category = "LearningAgents", meta = (ClampMin = "0.0", UIMin = "0.0"))
+	float ActionNoiseMin = 0.25f;
+
+	/** Maximum action noise used by the policy */
+	UPROPERTY(EditAnywhere, Category = "LearningAgents", meta = (ClampMin = "0.0", UIMin = "0.0"))
+	float ActionNoiseMax = 0.25f;
+};
+
+UENUM(BlueprintType, Category = "LearningAgents")
+enum class ELearningAgentsActivationFunction : uint8
+{
+	/** ReLU Activation - Fast to train and evaluate but occasionally causes gradient collapse and untrainable networks. */
+	ReLU	UMETA(DisplayName = "ReLU"),
+
+	/** ELU Activation - Generally performs better than ReLU and is not prone to gradient collapse but slower to evaluate. */
+	ELU		UMETA(DisplayName = "ELU"),
+
+	/** TanH Activation - Smooth activation function that is slower to train and evaluate but sometimes more stable for certain tasks. */
+	TanH	UMETA(DisplayName = "TanH"),
+};
+
+USTRUCT(BlueprintType, Category = "LearningAgents")
+struct FLearningAgentsNetworkSettings
+{
+	GENERATED_BODY()
+
+public:
+
+	/** Total layers for neural network including input, hidden, and output layers */
+	UPROPERTY(EditAnywhere, Category = "LearningAgents", meta = (ClampMin = "2", UIMin = "2"))
+	int32 LayerNum = 3;
+
+	/** Number of neurons in each hidden layer */
+	UPROPERTY(EditAnywhere, Category = "LearningAgents", meta = (ClampMin = "1", UIMin = "1"))
+	int32 HiddenLayerSize = 128;
+
+	/** Activation function to use on hidden layers */
+	UPROPERTY(EditAnywhere, Category = "LearningAgents")
+	ELearningAgentsActivationFunction ActivationFunction = ELearningAgentsActivationFunction::ELU;
+};
+
+UCLASS(Abstract, BlueprintType, Blueprintable)
+class LEARNINGAGENTS_API ULearningAgentsType : public USceneComponent
+{
+	GENERATED_BODY()
+
+public:
+
+	DECLARE_MULTICAST_DELEGATE(FOnSetupComplete);
+	DECLARE_MULTICAST_DELEGATE_TwoParams(FOnAgentEvent, int32 /*AgentId*/, UObject* /*Agent*/);
+
+	// These constructors/destructors are needed to make forward declarations happy
+	ULearningAgentsType();
+	ULearningAgentsType(FVTableHelper& Helper);
+	virtual ~ULearningAgentsType();
+
+	/** Sets up the agent ids so that we can add/remove agents at any time. */
+	virtual void PostInitProperties() override;
+
+// ----- Setup -----
+public:
+
+	/**
+	 * Initializes this object and runs the setup functions for observations and actions
+	 */
+	UFUNCTION(BlueprintCallable, Category = "LearningAgents")
+	void SetupAgentType(
+		const FLearningAgentsTypeSettings& Settings = FLearningAgentsTypeSettings(),
+		const FLearningAgentsNetworkSettings& NetworkSettings = FLearningAgentsNetworkSettings());
+
+	UFUNCTION(BlueprintPure, Category = "LearningAgents")
+	const bool IsSetupPerformed() const;
+
+	UFUNCTION(BlueprintPure, Category = "LearningAgents")
+	const int32 GetMaxInstanceNum() const;
+
+// ----- Agent Management -----
+public:
+
+	UFUNCTION(BlueprintCallable, Category = "LearningAgents")
+	int32 AddAgent(UObject* Agent);
+	
+	UFUNCTION(BlueprintCallable, Category = "LearningAgents")
+	void RemoveAgentById(int32 AgentId);
+
+	UFUNCTION(BlueprintCallable, Category = "LearningAgents")
+	void RemoveAgent(UObject* Agent);
+
+	UFUNCTION(BlueprintPure, Category = "LearningAgents", meta = (DeterminesOutputType = "AgentClass"))
+	UObject* GetAgent(int32 AgentId, TSubclassOf<UObject> AgentClass);
+
+	const UObject* GetAgent(int32 AgentId) const;
+
+// ----- Observations -----
+public:
+
+	/**
+	 * Event where all observations should be set up for this agent type.
+	 * @param AgentType This AgentType. This is a convenience for blueprints so you don't have to find the Self pin.
+	 */
+	UFUNCTION(BlueprintNativeEvent, Category = "LearningAgents")
+	void SetupObservations(ULearningAgentsType* AgentType);
+
+	UFUNCTION(BlueprintNativeEvent, Category = "LearningAgents")
+	void SetObservations(const TArray<int32>& AgentIds);
+
+	void AddObservation(TObjectPtr<ULearningAgentsObservation> Object, const TSharedRef<UE::Learning::FFeatureObject>& Feature);
+
+// ----- Actions -----
+public:
+
+	/**
+	 * Event where all actions should be set up for this agent type.
+	 * @param AgentType This AgentType. This is a convenience for blueprints so you don't have to find the Self pin.
+	 */
+	UFUNCTION(BlueprintNativeEvent, Category = "LearningAgents")
+	void SetupActions(ULearningAgentsType* AgentType);
+
+	UFUNCTION(BlueprintNativeEvent, Category = "LearningAgents")
+	void GetActions(const TArray<int32>& AgentIds);
+
+	void AddAction(TObjectPtr<ULearningAgentsAction> Object, const TSharedRef<UE::Learning::FFeatureObject>& Feature);
+
+// ----- Inference Process -----
+public:
+
+	UFUNCTION(BlueprintCallable, Category = "LearningAgents")
+	void LoadNetwork(const FDirectoryPath& Directory, const FString Filename);
+
+	UFUNCTION(BlueprintCallable, Category = "LearningAgents")
+	void EncodeObservations();
+
+	UFUNCTION(BlueprintCallable, Category = "LearningAgents")
+	void EvaluatePolicy();
+
+	UFUNCTION(BlueprintCallable, Category = "LearningAgents")
+	void DecodeActions();
+
+	UFUNCTION(BlueprintCallable, Category = "LearningAgents")
+	void BroadcastSetupComplete();
+
+// ----- Non-blueprint public interface -----
+public:
+	const TSharedPtr<UE::Learning::FArrayMap>& GetInstanceData() const;
+	UE::Learning::FFeatureObject& GetObservationFeature() const;
+	UE::Learning::FFeatureObject& GetActionFeature() const;
+	UE::Learning::FNeuralNetwork& GetNeuralNetwork() const;
+	const UE::Learning::FNeuralNetworkPolicyFunction& GetPolicy() const;
+	const TConstArrayView<TObjectPtr<UObject>> GetAgents() const;
+	const UE::Learning::FIndexSet GetOccupiedAgentSet() const;
+	const UE::Learning::FIndexSet GetVacantAgentSet() const;
+	FOnSetupComplete& GetOnSetupComplete() { return OnSetupComplete; }
+	FOnAgentEvent& GetOnAgentAdded() { return OnAgentAdded; }
+	FOnAgentEvent& GetOnAgentRemoved() { return OnAgentRemoved; }
+	TConstArrayView<ULearningAgentsObservation*> GetObservationObjects() const { return ObservationObjects; }
+	TConstArrayView<ULearningAgentsAction*> GetActionObjects() const { return ActionObjects; }
+
+// ----- Private Data -----
+private:
+
+	void UpdateAgentSets();
+
+	TArray<int32> OccupiedAgentIds;
+	TArray<int32> VacantAgentIds;
+	UE::Learning::FIndexSet OccupiedAgentSet;
+	UE::Learning::FIndexSet VacantAgentSet;
+
+	TSharedPtr<UE::Learning::FArrayMap> InstanceData;
+
+	TArray<TSharedRef<UE::Learning::FFeatureObject>, TInlineAllocator<16>> ObservationFeatures;
+	TSharedPtr<UE::Learning::FConcatenateFeature> Observations;
+	
+	TArray<TSharedRef<UE::Learning::FFeatureObject>, TInlineAllocator<16>> ActionFeatures;
+	TSharedPtr<UE::Learning::FConcatenateFeature> Actions;
+
+	TSharedPtr<UE::Learning::FNeuralNetwork> NeuralNetwork;
+	TUniquePtr<UE::Learning::FNeuralNetworkPolicyFunction> Policy;
+
+	FOnSetupComplete OnSetupComplete;
+	FOnAgentEvent OnAgentAdded;
+	FOnAgentEvent OnAgentRemoved;
+
+	/** Maximum number of agent instances. Used to preallocate internal buffers. */
+	UPROPERTY(EditDefaultsOnly, Category = "LearningAgents", meta = (ClampMin = "1", UIMin = "1"))
+	int32 MaxInstanceNum = 1;
+
+	UPROPERTY(VisibleAnywhere, Category = "LearningAgents")
+	bool bSetupPerformed = false;
+
+	UPROPERTY(VisibleAnywhere, Category = "LearningAgents")
+	bool bNetworkLoaded = false;
+
+	UPROPERTY(VisibleAnywhere, Transient, Category = "LearningAgents")
+	TArray<TObjectPtr<UObject>> Agents;
+
+	UPROPERTY(VisibleAnywhere, Category = "LearningAgents")
+	TArray<TObjectPtr<ULearningAgentsObservation>> ObservationObjects;
+
+	UPROPERTY(VisibleAnywhere, Category = "LearningAgents")
+	TArray<TObjectPtr<ULearningAgentsAction>> ActionObjects;
+};
diff --git a/Engine/Plugins/Experimental/LearningAgents/Source/LearningAgentsEditor/LearningAgentsEditor.Build.cs b/Engine/Plugins/Experimental/LearningAgents/Source/LearningAgentsEditor/LearningAgentsEditor.Build.cs
new file mode 100644
--- /dev/null
+++ b/Engine/Plugins/Experimental/LearningAgents/Source/LearningAgentsEditor/LearningAgentsEditor.Build.cs
@@ -0,0 +1,45 @@
+// Copyright Epic Games, Inc. All Rights Reserved.
+
+using UnrealBuildTool;
+
+public class LearningAgentsEditor : ModuleRules
+{
+	public LearningAgentsEditor(ReadOnlyTargetRules Target) : base(Target)
+	{
+		PCHUsage = ModuleRules.PCHUsageMode.UseExplicitOrSharedPCHs;
+
+		PublicIncludePaths.AddRange(
+			new string[] {
+				// ... add public include paths required here ...
+			});
+		
+		PrivateIncludePaths.AddRange(
+			new string[] {
+				// ... add other private include paths required here ...
+			});
+		
+		PublicDependencyModuleNames.AddRange(
+			new string[]
+			{
+				"Core",
+				// ... add other public dependencies that you statically link with here ...
+			});
+		
+		PrivateDependencyModuleNames.AddRange(
+			new string[]
+			{
+				"CoreUObject",
+				"Engine",
+				"UnrealEd",
+				"LearningAgents",
+				"LearningAgentsTraining",
+				// ... add private dependencies that you statically link with here ...	
+			});
+		
+		DynamicallyLoadedModuleNames.AddRange(
+			new string[]
+			{
+				// ... add any modules that your module loads dynamically here ...
+			});
+	}
+}
diff --git a/Engine/Plugins/Experimental/LearningAgents/Source/LearningAgentsEditor/Private/LearningAgentsEditor.cpp b/Engine/Plugins/Experimental/LearningAgents/Source/LearningAgentsEditor/Private/LearningAgentsEditor.cpp
new file mode 100644
--- /dev/null
+++ b/Engine/Plugins/Experimental/LearningAgents/Source/LearningAgentsEditor/Private/LearningAgentsEditor.cpp
@@ -0,0 +1,34 @@
+// Copyright Epic Games, Inc. All Rights Reserved.
+
+#include "LearningAgentsEditor.h"
+
+#include "Kismet2/KismetEditorUtilities.h"
+#include "LearningAgentsType.h"
+#include "LearningAgentsTrainer.h"
+#include "LearningAgentsController.h"
+
+void FLearningAgentsEditorModule::StartupModule()
+{
+	// ULearningAgentsType Default Events
+	FKismetEditorUtilities::RegisterAutoGeneratedDefaultEvent(this, ULearningAgentsType::StaticClass(), GET_FUNCTION_NAME_CHECKED(ULearningAgentsType, SetupObservations));
+	FKismetEditorUtilities::RegisterAutoGeneratedDefaultEvent(this, ULearningAgentsType::StaticClass(), GET_FUNCTION_NAME_CHECKED(ULearningAgentsType, SetObservations));
+	FKismetEditorUtilities::RegisterAutoGeneratedDefaultEvent(this, ULearningAgentsType::StaticClass(), GET_FUNCTION_NAME_CHECKED(ULearningAgentsType, SetupActions));
+	FKismetEditorUtilities::RegisterAutoGeneratedDefaultEvent(this, ULearningAgentsType::StaticClass(), GET_FUNCTION_NAME_CHECKED(ULearningAgentsType, GetActions));
+	
+	// ULearningAgentsTrainer Default Events
+	FKismetEditorUtilities::RegisterAutoGeneratedDefaultEvent(this, ULearningAgentsTrainer::StaticClass(), GET_FUNCTION_NAME_CHECKED(ULearningAgentsTrainer, SetupRewards));
+	FKismetEditorUtilities::RegisterAutoGeneratedDefaultEvent(this, ULearningAgentsTrainer::StaticClass(), GET_FUNCTION_NAME_CHECKED(ULearningAgentsTrainer, SetRewards));
+	FKismetEditorUtilities::RegisterAutoGeneratedDefaultEvent(this, ULearningAgentsTrainer::StaticClass(), GET_FUNCTION_NAME_CHECKED(ULearningAgentsTrainer, SetupCompletions));
+	FKismetEditorUtilities::RegisterAutoGeneratedDefaultEvent(this, ULearningAgentsTrainer::StaticClass(), GET_FUNCTION_NAME_CHECKED(ULearningAgentsTrainer, SetCompletions));
+	FKismetEditorUtilities::RegisterAutoGeneratedDefaultEvent(this, ULearningAgentsTrainer::StaticClass(), GET_FUNCTION_NAME_CHECKED(ULearningAgentsTrainer, ResetInstance));
+
+	// ULearningAgentsController Default Events
+	FKismetEditorUtilities::RegisterAutoGeneratedDefaultEvent(this, ULearningAgentsController::StaticClass(), GET_FUNCTION_NAME_CHECKED(ULearningAgentsController, SetActions));
+}
+
+void FLearningAgentsEditorModule::ShutdownModule()
+{
+	FKismetEditorUtilities::UnregisterAutoBlueprintNodeCreation(this);
+}
+	
+IMPLEMENT_MODULE(FLearningAgentsEditorModule, LearningAgentsEditor)
diff --git a/Engine/Plugins/Experimental/LearningAgents/Source/LearningAgentsEditor/Public/LearningAgentsEditor.h b/Engine/Plugins/Experimental/LearningAgents/Source/LearningAgentsEditor/Public/LearningAgentsEditor.h
new file mode 100644
--- /dev/null
+++ b/Engine/Plugins/Experimental/LearningAgents/Source/LearningAgentsEditor/Public/LearningAgentsEditor.h
@@ -0,0 +1,12 @@
+// Copyright Epic Games, Inc. All Rights Reserved.
+
+#pragma once
+
+#include "Modules/ModuleManager.h"
+
+class FLearningAgentsEditorModule : public IModuleInterface
+{
+public:
+	virtual void StartupModule() override;
+	virtual void ShutdownModule() override;
+};
diff --git a/Engine/Plugins/Experimental/LearningAgents/Source/LearningAgentsTraining/LearningAgentsTraining.Build.cs b/Engine/Plugins/Experimental/LearningAgents/Source/LearningAgentsTraining/LearningAgentsTraining.Build.cs
new file mode 100644
--- /dev/null
+++ b/Engine/Plugins/Experimental/LearningAgents/Source/LearningAgentsTraining/LearningAgentsTraining.Build.cs
@@ -0,0 +1,45 @@
+// Copyright Epic Games, Inc. All Rights Reserved.
+
+using UnrealBuildTool;
+
+public class LearningAgentsTraining : ModuleRules
+{
+	public LearningAgentsTraining(ReadOnlyTargetRules Target) : base(Target)
+	{
+		PCHUsage = ModuleRules.PCHUsageMode.UseExplicitOrSharedPCHs;
+
+		PublicIncludePaths.AddRange(
+			new string[] {
+				// ... add public include paths required here ...
+			});
+		
+		PrivateIncludePaths.AddRange(
+			new string[] {
+				// ... add other private include paths required here ...
+			});
+		
+		PublicDependencyModuleNames.AddRange(
+			new string[]
+			{
+				"Core",
+				// ... add other public dependencies that you statically link with here ...
+			});
+		
+		PrivateDependencyModuleNames.AddRange(
+			new string[]
+			{
+				"CoreUObject",
+				"Engine",
+				"Learning",
+				"LearningAgents",
+				"LearningTraining",
+				// ... add private dependencies that you statically link with here ...	
+			});
+		
+		DynamicallyLoadedModuleNames.AddRange(
+			new string[]
+			{
+				// ... add any modules that your module loads dynamically here ...
+			});
+	}
+}
diff --git a/Engine/Plugins/Experimental/LearningAgents/Source/LearningAgentsTraining/Private/LearningAgentsCompletions.cpp b/Engine/Plugins/Experimental/LearningAgents/Source/LearningAgentsTraining/Private/LearningAgentsCompletions.cpp
new file mode 100644
--- /dev/null
+++ b/Engine/Plugins/Experimental/LearningAgents/Source/LearningAgentsTraining/Private/LearningAgentsCompletions.cpp
@@ -0,0 +1,221 @@
+// Copyright Epic Games, Inc. All Rights Reserved.
+
+#include "LearningAgentsCompletions.h"
+
+#include "LearningAgentsTrainer.h"
+#include "LearningAgentsType.h"
+#include "LearningArray.h"
+#include "LearningArrayMap.h"
+#include "LearningCompletion.h"
+#include "LearningCompletionObject.h"
+#include "LearningLog.h"
+
+#include "GameFramework/Actor.h"
+#include "VisualLogger/VisualLogger.h"
+
+#define UE_LEARNING_AGENTS_VLOG_STRING(Owner, Category, Verbosity, Location, Color, Format, ...) \
+	UE_VLOG_LOCATION(Owner, Category, Verbosity, Location, 0.0f, Color, Format, ##__VA_ARGS__)
+
+#define UE_LEARNING_AGENTS_VLOG_TRANSFORM(Owner, Category, Verbosity, Location, Rotation, Color, Format, ...) \
+	UE_VLOG_SEGMENT(Owner, Category, Verbosity, Location, Location + 15.0f * Rotation.RotateVector(FVector::ForwardVector), FColor::Red, TEXT("")); \
+	UE_VLOG_SEGMENT(Owner, Category, Verbosity, Location, Location + 15.0f * Rotation.RotateVector(FVector::RightVector), FColor::Green, TEXT("")); \
+	UE_VLOG_SEGMENT(Owner, Category, Verbosity, Location, Location + 15.0f * Rotation.RotateVector(FVector::UpVector), FColor::Blue, TEXT("")); \
+	UE_VLOG_OBOX(Owner, Category, Verbosity, FBox(10.0f * FVector(-1, -1, -1), 10.0f * FVector(1, 1, 1)), FTransform(Rotation, Location, FVector::OneVector).ToMatrixNoScale(), Color, TEXT("")); \
+	UE_LEARNING_AGENTS_VLOG_STRING(Owner, Category, Verbosity, Location + FVector(0.0f, 0.0f, 20.0f), Color, Format, ##__VA_ARGS__)
+
+#define UE_LEARNING_AGENTS_VLOG_PLANE(Owner, Category, Verbosity, Location, Rotation, Axis0, Axis1, Color, Format, ...) \
+	UE_VLOG_OBOX(Owner, Category, Verbosity, FBox(25.0f * FVector(-1, -1, 0), 25.0f * FVector(1, 1, 0)), FTransform(Rotation, Location, FVector::OneVector).ToMatrixNoScale(), Color, Format, ##__VA_ARGS__)
+
+namespace UE::Learning::Agents::Private
+{
+	template<typename CompletionUObject, typename CompletionFObject, typename... InArgTypes>
+	CompletionUObject* AddCompletion(ULearningAgentsTrainer* AgentTrainer, FName Name, InArgTypes&& ...Args)
+	{
+		if (!AgentTrainer)
+		{
+			UE_LOG(LogLearning, Error, TEXT("AgentTrainer is nullptr")); return nullptr;
+		}
+
+		CompletionUObject* Completion = NewObject<CompletionUObject>(AgentTrainer, Name);
+
+		Completion->CompletionObject = MakeShared<CompletionFObject>(
+			Completion->GetFName(),
+			AgentTrainer->GetAgentType().GetInstanceData().ToSharedRef(),
+			AgentTrainer->GetAgentType().GetMaxInstanceNum(),
+			Forward<InArgTypes>(Args)...);
+
+		AgentTrainer->AddCompletion(Completion, Completion->CompletionObject.ToSharedRef());
+
+		return Completion;
+	}
+}
+
+//------------------------------------------------------------------
+
+UConditionalCompletion* UConditionalCompletion::AddConditionalCompletion(ULearningAgentsTrainer* AgentTrainer, FName Name, ELearningAgentsCompletion InCompletionMode)
+{
+	return UE::Learning::Agents::Private::AddCompletion<UConditionalCompletion, UE::Learning::FConditionalCompletion>(
+		AgentTrainer,
+		Name,
+		InCompletionMode == ELearningAgentsCompletion::Termination ?
+		UE::Learning::ECompletionMode::Terminated :
+		UE::Learning::ECompletionMode::Truncated);
+}
+
+void UConditionalCompletion::SetConditionalCompletion(int32 AgentId, bool bIsComplete)
+{
+	TLearningArrayView<1, bool> View = CompletionObject->InstanceData->View(CompletionObject->ConditionHandle);
+
+	if (AgentId == INDEX_NONE)
+	{
+		UE_LOG(LogLearning, Error, TEXT("AgentId is invalid (INDEX_NONE)"));
+		return;
+	}
+
+	if (AgentId < 0 || AgentId >= View.Num<0>())
+	{
+		UE_LOG(LogLearning, Error, TEXT("AgentId %d is out of index. Valid range [0, %d]."), AgentId, View.Num<0>() - 1);
+		return;
+	}
+
+	View[AgentId] = bIsComplete;
+}
+
+#if ENABLE_VISUAL_LOG
+void UConditionalCompletion::VisualLog(const UE::Learning::FIndexSet Instances) const
+{
+	UE_LEARNING_TRACE_CPUPROFILER_EVENT_SCOPE(UConditionalCompletion::VisualLog);
+
+	const TLearningArrayView<1, const bool> ConditionView = CompletionObject->InstanceData->ConstView(CompletionObject->ConditionHandle);
+	const TLearningArrayView<1, const UE::Learning::ECompletionMode> CompletionView = CompletionObject->InstanceData->ConstView(CompletionObject->CompletionHandle);
+
+	if (const ULearningAgentsTrainer* AgentTrainer = Cast<ULearningAgentsTrainer>(GetOuter()))
+	{
+		for (const int32 Instance : Instances)
+		{
+			if (const AActor* Actor = Cast<AActor>(AgentTrainer->GetAgent(Instance)))
+			{
+				UE_LEARNING_AGENTS_VLOG_STRING(this, LogLearning, Display,
+					Actor->GetActorLocation(),
+					VisualLogColor.ToFColor(true),
+					TEXT("Agent %i\nCondition: %s\nCompletion: %s"),
+					Instance,
+					ConditionView[Instance] ? TEXT("true") : TEXT("false"),
+					UE::Learning::Completion::CompletionModeString(CompletionView[Instance]));
+			}
+		}
+	}
+}
+#endif
+
+//------------------------------------------------------------------
+
+UPlanarPositionDifferenceCompletion* UPlanarPositionDifferenceCompletion::AddPlanarPositionDifferenceCompletion(ULearningAgentsTrainer* AgentTrainer, FName Name, float Threshold, ELearningAgentsCompletion InCompletionMode, FVector Axis0, FVector Axis1)
+{
+	return UE::Learning::Agents::Private::AddCompletion<UPlanarPositionDifferenceCompletion, UE::Learning::FPlanarPositionDifferenceCompletion>(
+		AgentTrainer,
+		Name,
+		Threshold,
+		InCompletionMode == ELearningAgentsCompletion::Termination ?
+		UE::Learning::ECompletionMode::Terminated :
+		UE::Learning::ECompletionMode::Truncated,
+		Axis0,
+		Axis1);
+}
+
+void UPlanarPositionDifferenceCompletion::SetPlanarPositionDifferenceCompletion(int32 AgentId, FVector Position0, FVector Position1)
+{
+	TLearningArrayView<1, FVector> Position0View = CompletionObject->InstanceData->View(CompletionObject->Position0Handle);
+	TLearningArrayView<1, FVector> Position1View = CompletionObject->InstanceData->View(CompletionObject->Position1Handle);
+
+	if (AgentId == INDEX_NONE)
+	{
+		UE_LOG(LogLearning, Error, TEXT("AgentId is invalid (INDEX_NONE)"));
+		return;
+	}
+
+	if (AgentId < 0 || AgentId >= Position0View.Num<0>())
+	{
+		UE_LOG(LogLearning, Error, TEXT("AgentId %d is out of index. Valid range [0, %d]."), AgentId, Position0View.Num<0>() - 1);
+		return;
+	}
+
+	Position0View[AgentId] = Position0;
+	Position1View[AgentId] = Position1;
+}
+
+#if ENABLE_VISUAL_LOG
+void UPlanarPositionDifferenceCompletion::VisualLog(const UE::Learning::FIndexSet Instances) const
+{
+	UE_LEARNING_TRACE_CPUPROFILER_EVENT_SCOPE(UPlanarPositionDifferenceCompletion::VisualLog);
+
+	const TLearningArrayView<1, const FVector> Position0View = CompletionObject->InstanceData->ConstView(CompletionObject->Position0Handle);
+	const TLearningArrayView<1, const FVector> Position1View = CompletionObject->InstanceData->ConstView(CompletionObject->Position1Handle);
+	const TLearningArrayView<1, const float> ThresholdView = CompletionObject->InstanceData->ConstView(CompletionObject->ThresholdHandle);
+	const TLearningArrayView<1, const UE::Learning::ECompletionMode> CompletionView = CompletionObject->InstanceData->ConstView(CompletionObject->CompletionHandle);
+
+	if (const ULearningAgentsTrainer* AgentTrainer = Cast<ULearningAgentsTrainer>(GetOuter()))
+	{
+		for (const int32 Instance : Instances)
+		{
+			if (const AActor* Actor = Cast<AActor>(AgentTrainer->GetAgent(Instance)))
+			{
+				const FVector Position0 = Position0View[Instance];
+				const FVector Position1 = Position1View[Instance];
+
+				const FVector PlanarPosition0 = FVector(CompletionObject->Axis0.Dot(Position0), CompletionObject->Axis1.Dot(Position0), 0.0f);
+				const FVector PlanarPosition1 = FVector(CompletionObject->Axis0.Dot(Position1), CompletionObject->Axis1.Dot(Position1), 0.0f);
+
+				UE_VLOG_LOCATION(this, LogLearning, Display,
+					Position0,
+					10.0f,
+					VisualLogColor.ToFColor(true),
+					TEXT("Position0: [% 6.1f % 6.1f % 6.1f]\nPlanar Position0: [% 6.1f % 6.1f]"),
+					Position0.X, Position0.Y, Position0.Z,
+					PlanarPosition0.X, PlanarPosition0.Y);
+
+				UE_LEARNING_AGENTS_VLOG_PLANE(this, LogLearning, Display,
+					Position0,
+					FQuat::Identity,
+					FeatureObject->Axis0,
+					FeatureObject->Axis1,
+					VisualLogColor.ToFColor(true),
+					TEXT(""));
+
+				UE_VLOG_LOCATION(this, LogLearning, Display,
+					Position1,
+					10.0f,
+					VisualLogColor.ToFColor(true),
+					TEXT("Position1: [% 6.1f % 6.1f % 6.1f]\nPlanar Position1: [% 6.1f % 6.1f]"),
+					Position1.X, Position1.Y, Position1.Z,
+					PlanarPosition1.X, PlanarPosition1.Y);
+
+				UE_LEARNING_AGENTS_VLOG_PLANE(this, LogLearning, Display,
+					Position1,
+					FQuat::Identity,
+					FeatureObject->Axis0,
+					FeatureObject->Axis1,
+					VisualLogColor.ToFColor(true),
+					TEXT(""));
+
+				UE_VLOG_SEGMENT(this, LogLearning, Display,
+					Position0,
+					Position1,
+					VisualLogColor.ToFColor(true),
+					TEXT("Agent %i\nDistance: [% 6.3f]\nPlanar Distance: [% 6.3f]\nThreshold: [% 6.2f]\nCompletion: %s"),
+					Instance,
+					FVector::Distance(Position0, Position1),
+					FVector::Distance(PlanarPosition0, PlanarPosition1),
+					ThresholdView[Instance],
+					UE::Learning::Completion::CompletionModeString(CompletionView[Instance]));
+			}
+		}
+	}
+}
+#endif
+
+//------------------------------------------------------------------
+
+#undef UE_LEARNING_AGENTS_VLOG_STRING
+#undef UE_LEARNING_AGENTS_VLOG_TRANSFORM
+#undef UE_LEARNING_AGENTS_VLOG_PLANE
diff --git a/Engine/Plugins/Experimental/LearningAgents/Source/LearningAgentsTraining/Private/LearningAgentsController.cpp b/Engine/Plugins/Experimental/LearningAgents/Source/LearningAgentsTraining/Private/LearningAgentsController.cpp
new file mode 100644
--- /dev/null
+++ b/Engine/Plugins/Experimental/LearningAgents/Source/LearningAgentsTraining/Private/LearningAgentsController.cpp
@@ -0,0 +1,25 @@
+// Copyright Epic Games, Inc. All Rights Reserved.
+
+#include "LearningAgentsController.h"
+
+#include "LearningAgentsType.h"
+#include "LearningFeatureObject.h"
+#include "LearningLog.h"
+
+ULearningAgentsController::ULearningAgentsController() : ULearningAgentsTypeComponent() {}
+ULearningAgentsController::ULearningAgentsController(FVTableHelper& Helper) : ULearningAgentsController() {}
+ULearningAgentsController::~ULearningAgentsController() {}
+
+void ULearningAgentsController::SetActions_Implementation(const TArray<int32>& AgentIds)
+{
+	// Can be overridden to get actions without blueprints
+}
+
+void ULearningAgentsController::EncodeActions()
+{
+	UE_LEARNING_TRACE_CPUPROFILER_EVENT_SCOPE(ULearningAgentsController::EncodeActions);
+
+	SetActions(SelectedAgentIds);
+
+	AgentType->GetActionFeature().Encode(SelectedAgentsSet);
+}
