Learning Agents: Create initial experimental version

diff --git a/Engine/Plugins/Experimental/LearningAgents/Source/LearningAgentsTraining/Private/LearningAgentsDataStorage.cpp b/Engine/Plugins/Experimental/LearningAgents/Source/LearningAgentsTraining/Private/LearningAgentsDataStorage.cpp
new file mode 100644
--- /dev/null
+++ b/Engine/Plugins/Experimental/LearningAgents/Source/LearningAgentsTraining/Private/LearningAgentsDataStorage.cpp
@@ -0,0 +1,249 @@
+// Copyright Epic Games, Inc. All Rights Reserved.
+
+#include "LearningAgentsDataStorage.h"
+
+#include "Engine/EngineTypes.h"
+#include "HAL/FileManager.h"
+#include "LearningAgentsType.h"
+#include "LearningFeatureObject.h"
+#include "LearningLog.h"
+#include "Misc/FileHelper.h"
+#include "Misc/Paths.h"
+#include "Serialization/MemoryWriter.h"
+#include "Serialization/MemoryReader.h"
+
+ULearningAgentsRecord::ULearningAgentsRecord() {}
+ULearningAgentsRecord::ULearningAgentsRecord(FVTableHelper& Helper) {}
+ULearningAgentsRecord::~ULearningAgentsRecord() {}
+
+void ULearningAgentsRecord::Init(int32 ObsNum, int32 ActNum)
+{
+	ObservationNum = ObsNum;
+	ActionNum = ActNum;
+
+	AddChunk();
+}
+
+void ULearningAgentsRecord::AddExperience(TLearningArrayView<1, const float> NewObservations, TLearningArrayView<1, const float> NewActions)
+{
+	MetaData.bHasUnsavedChanges = true;
+
+	if (DataIndex >= ChunkSize)
+	{
+		AddChunk();
+	}
+
+	UE::Learning::Array::Copy(Observations.Last()[DataIndex], NewObservations);
+	UE::Learning::Array::Copy(Actions.Last()[DataIndex], NewActions);
+
+	DataIndex++;
+}
+
+const TConstArrayView<TLearningArray<2, float>> ULearningAgentsRecord::GetObservations() const
+{
+	return Observations;
+}
+
+const TConstArrayView<TLearningArray<2, float>> ULearningAgentsRecord::GetActions() const
+{
+	return Actions;
+}
+
+void ULearningAgentsRecord::Trim()
+{
+	Observations.Last().SetNumUninitialized({ DataIndex, ObservationNum });
+	Actions.Last().SetNumUninitialized({ DataIndex, ActionNum });
+}
+
+void ULearningAgentsRecord::AddChunk()
+{
+	TLearningArray<2, float>& Observation = Observations.AddDefaulted_GetRef();
+	Observation.SetNumUninitialized({ ChunkSize, ObservationNum });
+
+	TLearningArray<2, float>& Action = Actions.AddDefaulted_GetRef();
+	Action.SetNumUninitialized({ ChunkSize, ActionNum });
+
+	DataIndex = 0;
+}
+
+ULearningAgentsDataStorage::ULearningAgentsDataStorage() {}
+ULearningAgentsDataStorage::ULearningAgentsDataStorage(FVTableHelper& Helper) {}
+ULearningAgentsDataStorage::~ULearningAgentsDataStorage() {}
+
+ULearningAgentsRecord* ULearningAgentsDataStorage::CreateRecord(FName RecordName, ULearningAgentsType* AgentType)
+{
+	if (!AgentType)
+	{
+		UE_LOG(LogLearning, Warning, TEXT("CreateRecord: AgentType is nullptr. You must pass a valid agent type. Skipping creation."));
+		return nullptr;
+	}
+
+	ULearningAgentsRecord* NewRecord = NewObject<ULearningAgentsRecord>(this, RecordName);
+	NewRecord->MetaData.AgentType = AgentType;
+	NewRecord->MetaData.CreatedOn = FDateTime::Now();
+	NewRecord->MetaData.CreatedOnUtc = FDateTime::UtcNow();
+	NewRecord->MetaData.RecordName = RecordName;
+
+	NewRecord->Init(AgentType->GetObservationFeature().DimNum(), AgentType->GetActionFeature().DimNum());
+	Records.Add(NewRecord);
+
+	return NewRecord;
+}
+
+const TArray<ULearningAgentsRecord*>& ULearningAgentsDataStorage::GetAllRecords() const
+{
+	return Records;
+}
+
+int32 ULearningAgentsDataStorage::LoadAllRecords(ULearningAgentsType* AgentType, const FDirectoryPath& Directory)
+{
+	if (!AgentType)
+	{
+		UE_LOG(LogLearning, Warning, TEXT("LoadAllRecords: AgentType is nullptr. You must pass a valid agent type. Skipping loading."));
+		return 0;
+	}
+
+	TArray<FString> DataFiles;
+	IFileManager::Get().FindFiles(DataFiles, *Directory.Path, *FileExtension);
+
+	if (DataFiles.Num() == 0)
+	{
+		UE_LOG(LogLearning, Warning, TEXT("LoadAllRecords: Found no .%s files in directory %s"), *FileExtension, *Directory.Path);
+		return 0;
+	}
+
+	for (const FString& DataFile : DataFiles)
+	{
+		LoadRecord(AgentType, Directory, DataFile);
+	}
+
+	return DataFiles.Num();
+}
+
+ULearningAgentsRecord* ULearningAgentsDataStorage::LoadRecord(ULearningAgentsType* AgentType, const FDirectoryPath& Directory, const FString& Filename)
+{
+	if (!AgentType)
+	{
+		UE_LOG(LogLearning, Warning, TEXT("LoadRecord: AgentType is nullptr. You must pass a valid agent type. Skipping loading."));
+		return nullptr;
+	}
+
+	TArray<uint8> Bytes;
+
+	FString InputFilePath = Directory.Path + FGenericPlatformMisc::GetDefaultPathSeparator() + Filename;
+
+	if (!FPaths::FileExists(InputFilePath))
+	{
+		UE_LOG(LogLearning, Error, TEXT("LoadRecord: FilePath %s does not exist."), *InputFilePath);
+		return nullptr;
+	}
+
+	const bool bSuccess = FFileHelper::LoadFileToArray(Bytes, *InputFilePath);
+	if (!bSuccess)
+	{
+		UE_LOG(LogLearning, Error, TEXT("Failed to load data from file %s"), *InputFilePath);
+		return nullptr;
+	}
+
+	FMemoryReader Reader(Bytes);
+
+	int32 ChunksNum = -1;
+	Reader.Serialize(&ChunksNum, sizeof(int32));
+	if (ChunksNum <= 0)
+	{
+		UE_LOG(LogLearning, Error, TEXT("LoadRecord: file ChunksNum header was %d. Must be greater than 0."), ChunksNum);
+		return nullptr;
+	}
+
+	int32 ObsNum = -1;
+	Reader.Serialize(&ObsNum, sizeof(int32));
+	int32 AgentObsNum = AgentType->GetObservationFeature().DimNum();
+	if (AgentObsNum != ObsNum)
+	{
+		UE_LOG(LogLearning, Error, TEXT("LoadRecord: Observation data dimension size %d incompatible with agent's observation dimension size %d"), ObsNum, AgentObsNum);
+		return nullptr;
+	}
+
+	int32 ActNum = -1;
+	Reader.Serialize(&ActNum, sizeof(int32));
+	int32 AgentActionNum = AgentType->GetActionFeature().DimNum();
+	if (AgentActionNum != ActNum)
+	{
+		UE_LOG(LogLearning, Error, TEXT("LoadRecord: Action data dimension size %d incompatible with agent's action dimension size %d"), ObsNum, AgentObsNum);
+		return nullptr;
+	}
+
+	// Load all the data into the first chunk
+	ULearningAgentsRecord* Record = CreateRecord(FName(Filename), AgentType);
+	Reader.Serialize(&Record->MetaData.CreatedOn, sizeof(FDateTime));
+	Reader.Serialize(&Record->MetaData.CreatedOnUtc, sizeof(FDateTime));
+	Record->MetaData.bWasLoadedFromFile = true;
+
+	for (int32 ChunkIndex = 0; ChunkIndex < ChunksNum; ChunkIndex++)
+	{
+		if (ChunkIndex > 0)
+		{
+			Record->AddChunk();
+		}
+
+		UE::Learning::Array::Serialize(Reader, Record->Observations[ChunkIndex]);
+		UE::Learning::Array::Serialize(Reader, Record->Actions[ChunkIndex]);
+	}
+
+	return Record;
+}
+
+void ULearningAgentsDataStorage::SaveAllRecords(const FDirectoryPath& Directory) const
+{
+	for (TObjectPtr<ULearningAgentsRecord> Record : Records)
+	{
+		if (Record->MetaData.bWasLoadedFromFile && !Record->MetaData.bHasUnsavedChanges)
+		{
+			continue;
+		}
+
+		SaveRecord(Directory, Record);
+	}
+}
+
+void ULearningAgentsDataStorage::SaveRecord(const FDirectoryPath& Directory, ULearningAgentsRecord* Record) const
+{
+	if (!Record)
+	{
+		UE_LOG(LogLearning, Warning, TEXT("SaveRecord: Record is nullptr. You must pass a valid record. Skipping save."));
+		return;
+	}
+
+	Record->Trim();
+
+	TArray<uint8> Bytes;
+	FMemoryWriter Writer(Bytes);
+
+	// Write some info we would need to extract the data again
+	int32 ChunksNum = Record->Observations.Num();
+	int32 ObsNum = Record->Observations[0].Num<1>();
+	int32 ActNum = Record->Actions[0].Num<1>();
+
+	Writer << ChunksNum;
+	Writer << ObsNum;
+	Writer << ActNum;
+	Writer << Record->MetaData.CreatedOn;
+	Writer << Record->MetaData.CreatedOnUtc;
+
+	for (int32 ChunkIndex = 0; ChunkIndex < ChunksNum; ChunkIndex++)
+	{
+		UE::Learning::Array::Serialize(Writer, Record->Observations[ChunkIndex]);
+		UE::Learning::Array::Serialize(Writer, Record->Actions[ChunkIndex]);
+	}
+
+	const FString FileName = Record->MetaData.RecordName.IsNone() ? Record->MetaData.AgentType->GetName() : Record->MetaData.RecordName.ToString();
+	const FString OutputFilePath = Directory.Path + FGenericPlatformMisc::GetDefaultPathSeparator() +
+		(!Record->MetaData.bWasLoadedFromFile && bPrependUtcTimeStamp ? FDateTime::UtcNow().ToString() + "_" : "") + FileName + "." + FileExtension;
+	const bool bSuccess = FFileHelper::SaveArrayToFile(Bytes, *OutputFilePath, &IFileManager::Get(), EFileWrite::FILEWRITE_Append);
+	if (!bSuccess)
+	{
+		UE_LOG(LogLearning, Error, TEXT("Failed to write data to output file for %s"), *FileName);
+	}
+
+	UE_LOG(LogLearning, Display, TEXT("SaveRecord: Saved data for agent to %s"), *OutputFilePath);
+}
diff --git a/Engine/Plugins/Experimental/LearningAgents/Source/LearningAgentsTraining/Private/LearningAgentsImitationTrainer.cpp b/Engine/Plugins/Experimental/LearningAgents/Source/LearningAgentsTraining/Private/LearningAgentsImitationTrainer.cpp
new file mode 100644
--- /dev/null
+++ b/Engine/Plugins/Experimental/LearningAgents/Source/LearningAgentsTraining/Private/LearningAgentsImitationTrainer.cpp
@@ -0,0 +1,176 @@
+// Copyright Epic Games, Inc. All Rights Reserved.
+
+#include "LearningAgentsImitationTrainer.h"
+
+#include "GenericPlatform/GenericPlatformMisc.h"
+#include "HAL/FileManager.h"
+#include "LearningAgentsType.h"
+#include "LearningArrayMap.h"
+#include "LearningExperience.h"
+#include "LearningFeatureObject.h"
+#include "LearningLog.h"
+#include "LearningNeuralNetwork.h"
+#include "LearningImitationTrainer.h"
+#include "LearningAgentsDataStorage.h"
+#include "LearningPolicyObject.h"
+#include "Misc/Paths.h"
+
+ULearningAgentsImitationTrainer::ULearningAgentsImitationTrainer() : ULearningAgentsTypeComponent() {}
+ULearningAgentsImitationTrainer::ULearningAgentsImitationTrainer(FVTableHelper& Helper) : ULearningAgentsImitationTrainer() {}
+ULearningAgentsImitationTrainer::~ULearningAgentsImitationTrainer() {}
+
+void ULearningAgentsImitationTrainer::EndPlay(const EEndPlayReason::Type EndPlayReason)
+{
+	if (bIsTraining)
+	{
+		EndTraining();
+	}
+}
+
+const bool ULearningAgentsImitationTrainer::IsTraining() const
+{
+	return bIsTraining;
+}
+
+void ULearningAgentsImitationTrainer::BeginTraining(TArray<UObject*> Records)
+{
+	if (bIsTraining)
+	{
+		UE_LOG(LogLearning, Warning, TEXT("Already Training!"));
+		return;
+	}
+
+	if (Records.IsEmpty())
+	{
+		UE_LOG(LogLearning, Error, TEXT("Records list is empty! Imitation Trainer needs at least one record to work correctly."));
+		return;
+	}
+
+	UE_LOG(LogLearning, Display, TEXT("Imitation Training Started"));
+
+	bIsTraining = true;
+	bIsTrainingComplete = false;
+	bRequestImitationTrainingStop = false;
+
+	// Start Python Training Process (this must be done on main thread)
+
+#if WITH_EDITOR
+	const FString PythonExecutablePath = UE::Learning::Trainer::DefaultEditorPythonExecutablePath();
+	const FString SitePackagesPath = UE::Learning::Trainer::DefaultEditorSitePackagesPath();
+	const FString PythonContentPath = UE::Learning::Trainer::DefaultEditorPythonContentPath();
+	const FString IntermediatePath = UE::Learning::Trainer::DefaultEditorIntermediatePath();
+#else
+	UE_LEARNING_CHECKF(PLATFORM_WINDOWS || PLATFORM_MAC || PLATFORM_LINUX, TEXT("Python only supported on Windows, Mac, and Linux."));
+	const FString PythonExe = PLATFORM_WINDOWS ? TEXT("python.exe") : TEXT("bin/python");
+	const FString EnginePath = IFileManager::Get().ConvertToAbsolutePathForExternalAppForRead(*(FPaths::RootDir() / TEXT("../../../../../../Engine")));
+	const FString PythonExecutablePath = EnginePath / TEXT("Binaries/ThirdParty/Python3") / FPlatformMisc::GetUBTPlatform() / PythonExe;
+	const FString SitePackagesPath = EnginePath / TEXT("Plugins/Experimental/PythonFoundationPackages/Content/Python/Lib/") / FPlatformMisc::GetUBTPlatform() / TEXT("site-packages");
+	const FString PythonContentPath = EnginePath / TEXT("Plugins/Experimental/LearningAgents/Content/Python/");
+	const FString IntermediatePath = EnginePath / TEXT("Plugins/Experimental/LearningAgents/Intermediate");
+#endif
+
+	UE::Learning::FImitationTrainerSettings Settings;
+	Settings.bReinitializeNetwork = true;
+	Settings.Device = UE::Learning::ETrainerDevice::CPU;
+	Settings.bUseTensorboard = false;
+
+	// Get length
+	int32 StepNum = 0;
+	for (UObject* Object : Records)
+	{
+		ULearningAgentsRecord* Record = Cast<ULearningAgentsRecord>(Object);
+		if (!Record)
+		{
+			UE_LOG(LogLearning, Warning, TEXT("BeginTraining: input contains object (%s) that can't be cast to ULearningAgentsRecord."), *(Object ? Object->GetName() : FString(TEXT("nullptr"))));
+			continue;
+		}
+
+		TConstArrayView<TLearningArray<2, float>> RecordObsChunks = Record->GetObservations();
+
+		// Assuming everyone has been trimmed already
+		for (TLearningArray<2, float> Chunk : RecordObsChunks)
+		{
+			StepNum += Chunk.Num<0>();
+		}
+	}
+
+	if (StepNum == 0)
+	{
+		UE_LOG(LogLearning, Warning, TEXT("BeginTraining: input contains no valid training data. Returning early."));
+		return;
+	}
+
+	RecordedObservations.SetNumUninitialized({ StepNum, AgentType->GetObservationFeature().DimNum() });
+	RecordedActions.SetNumUninitialized({ StepNum, AgentType->GetActionFeature().DimNum() });
+
+	int32 DataIndex = 0;
+	for (UObject* Object : Records)
+	{
+		ULearningAgentsRecord* Record = Cast<ULearningAgentsRecord>(Object);
+		if (!Record)
+		{
+			UE_LOG(LogLearning, Warning, TEXT("BeginTraining: input contains object (%s) that can't be cast to ULearningAgentsRecord."), *(Object ? Object->GetName() : FString(TEXT("nullptr"))));
+			continue;
+		}
+
+		TConstArrayView<TLearningArray<2, float>> RecordObsChunks = Record->GetObservations();
+		TConstArrayView<TLearningArray<2, float>> RecordActionChunks = Record->GetActions();
+
+		// Assuming everyone has been trimmed already
+		for (int32 i = 0; i < RecordObsChunks.Num(); i++)
+		{
+			int32 ChunkLength = RecordObsChunks[i].Num<0>();
+			UE::Learning::Array::Copy(RecordedObservations.Slice(DataIndex, ChunkLength), RecordObsChunks[i]);
+			UE::Learning::Array::Copy(RecordedActions.Slice(DataIndex, ChunkLength), RecordActionChunks[i]);
+			DataIndex += ChunkLength;
+		}
+	}
+
+	ImitationTrainer = MakeUnique<UE::Learning::FSharedMemoryImitationTrainer>(
+		TEXT("ImitationLearner"),
+		PythonExecutablePath,
+		SitePackagesPath,
+		PythonContentPath,
+		IntermediatePath,
+		AgentType->GetNeuralNetwork(),
+		AgentType->GetPolicy().Settings.ActionNoiseMin,
+		AgentType->GetPolicy().Settings.ActionNoiseMax,
+		RecordedObservations.Num<0>(),
+		Settings);
+
+	ImitationTrainingTask = UE::Tasks::Launch(UE_SOURCE_LOCATION, [&]()
+		{
+			UE::Learning::ImitationTrainer::Train(
+				*ImitationTrainer,
+				AgentType->GetNeuralNetwork(),
+				RecordedObservations,
+				RecordedActions,
+				true,
+				&bRequestImitationTrainingStop,
+				&NetworkLock);
+
+			bIsTrainingComplete = true;
+		});
+}
+
+void ULearningAgentsImitationTrainer::EndTraining()
+{
+	if (bIsTraining)
+	{
+		UE_LOG(LogLearning, Display, TEXT("Imitation Training Ended"));
+
+		bRequestImitationTrainingStop = true;
+		ImitationTrainingTask.Wait(FTimespan(0, 0, 30));
+
+		bIsTraining = false;
+	}
+	else
+	{
+		UE_LOG(LogLearning, Warning, TEXT("Not Training!"));
+	}
+}
+
+const bool ULearningAgentsImitationTrainer::IsTrainingComplete() const
+{
+	return bIsTrainingComplete;
+}
diff --git a/Engine/Plugins/Experimental/LearningAgents/Source/LearningAgentsTraining/Private/LearningAgentsRecorder.cpp b/Engine/Plugins/Experimental/LearningAgents/Source/LearningAgentsTraining/Private/LearningAgentsRecorder.cpp
new file mode 100644
--- /dev/null
+++ b/Engine/Plugins/Experimental/LearningAgents/Source/LearningAgentsTraining/Private/LearningAgentsRecorder.cpp
@@ -0,0 +1,86 @@
+// Copyright Epic Games, Inc. All Rights Reserved.
+
+#include "LearningAgentsRecorder.h"
+
+#include "Engine/World.h"
+#include "LearningAgentsDataStorage.h"
+#include "LearningAgentsType.h"
+#include "LearningFeatureObject.h"
+#include "LearningLog.h"
+#include "LearningTrainer.h"
+
+ULearningAgentsRecorder::ULearningAgentsRecorder() : ULearningAgentsTypeComponent() {}
+ULearningAgentsRecorder::ULearningAgentsRecorder(FVTableHelper& Helper) : ULearningAgentsRecorder() {}
+ULearningAgentsRecorder::~ULearningAgentsRecorder() {}
+
+void ULearningAgentsRecorder::OnRegister()
+{
+	Super::OnRegister();
+
+	if (!GetWorld()->IsGameWorld())
+	{
+		// We're not in a game yet so we don't need the storage yet
+		return;
+	}
+
+	DataStorage = NewObject<ULearningAgentsDataStorage>(this);
+
+	if (DataDirectory.Path.IsEmpty())
+	{
+#if WITH_EDITOR
+		UE_LOG(LogLearning, Display, TEXT("ULearningAgentsRecorder writing data to default editor intermediate directory since DataDirectory was not set."));
+		DataDirectory.Path = UE::Learning::Trainer::DefaultEditorIntermediatePath();
+#else
+		UE_LOG(LogLearning, Error, TEXT("ULearningAgentsRecorder: DataDirectory was not set. This is non-editor build so need a directory setting."));
+#endif
+	}
+}
+
+void ULearningAgentsRecorder::EndPlay(const EEndPlayReason::Type EndPlayReason)
+{
+	StopRecording();
+
+	if (bSaveDataOnEndPlay)
+	{
+		DataStorage->SaveAllRecords(DataDirectory);
+	}
+
+	Super::EndPlay(EndPlayReason);
+}
+
+void ULearningAgentsRecorder::AddExperience()
+{
+	if (!IsRecording())
+	{
+		UE_LOG(LogLearning, Warning, TEXT("Trying to add experience but we aren't currently recording. Call StartRecording() before AddExperience()."));
+		return;
+	}
+
+	for (const int32 AgentId : SelectedAgentsSet)
+	{
+		CurrentRecords[AgentId]->AddExperience(AgentType->GetObservationFeature().FeatureBuffer()[AgentId], AgentType->GetActionFeature().FeatureBuffer()[AgentId]);
+	}
+}
+
+const bool ULearningAgentsRecorder::IsRecording() const
+{
+	return !CurrentRecords.IsEmpty();
+}
+
+void ULearningAgentsRecorder::StartRecording()
+{
+	for (const int32 AgentId : SelectedAgentsSet)
+	{
+		CurrentRecords.Add(AgentId, DataStorage->CreateRecord(FName(AgentType->GetName() + "_id" + FString::FromInt(AgentId)), AgentType));
+	}
+}
+
+void ULearningAgentsRecorder::StopRecording()
+{
+	for (TPair<int32, ULearningAgentsRecord*> Record : CurrentRecords)
+	{
+		Record.Value->Trim();
+	}
+
+	CurrentRecords.Empty();
+}
diff --git a/Engine/Plugins/Experimental/LearningAgents/Source/LearningAgentsTraining/Private/LearningAgentsRewards.cpp b/Engine/Plugins/Experimental/LearningAgents/Source/LearningAgentsTraining/Private/LearningAgentsRewards.cpp
new file mode 100644
--- /dev/null
+++ b/Engine/Plugins/Experimental/LearningAgents/Source/LearningAgentsTraining/Private/LearningAgentsRewards.cpp
@@ -0,0 +1,501 @@
+// Copyright Epic Games, Inc. All Rights Reserved.
+
+#include "LearningAgentsRewards.h"
+
+#include "LearningAgentsTrainer.h"
+#include "LearningAgentsType.h"
+#include "LearningArray.h"
+#include "LearningArrayMap.h"
+#include "LearningRewardObject.h"
+#include "LearningLog.h"
+
+#include "GameFramework/Actor.h"
+#include "VisualLogger/VisualLogger.h"
+
+#define UE_LEARNING_AGENTS_VLOG_STRING(Owner, Category, Verbosity, Location, Color, Format, ...) \
+	UE_VLOG_LOCATION(Owner, Category, Verbosity, Location, 0.0f, Color, Format, ##__VA_ARGS__)
+
+#define UE_LEARNING_AGENTS_VLOG_TRANSFORM(Owner, Category, Verbosity, Location, Rotation, Color, Format, ...) \
+	UE_VLOG_SEGMENT(Owner, Category, Verbosity, Location, Location + 15.0f * Rotation.RotateVector(FVector::ForwardVector), FColor::Red, TEXT("")); \
+	UE_VLOG_SEGMENT(Owner, Category, Verbosity, Location, Location + 15.0f * Rotation.RotateVector(FVector::RightVector), FColor::Green, TEXT("")); \
+	UE_VLOG_SEGMENT(Owner, Category, Verbosity, Location, Location + 15.0f * Rotation.RotateVector(FVector::UpVector), FColor::Blue, TEXT("")); \
+	UE_VLOG_OBOX(Owner, Category, Verbosity, FBox(10.0f * FVector(-1, -1, -1), 10.0f * FVector(1, 1, 1)), FTransform(Rotation, Location, FVector::OneVector).ToMatrixNoScale(), Color, TEXT("")); \
+	UE_LEARNING_AGENTS_VLOG_STRING(Owner, Category, Verbosity, Location + FVector(0.0f, 0.0f, 20.0f), Color, Format, ##__VA_ARGS__)
+
+#define UE_LEARNING_AGENTS_VLOG_PLANE(Owner, Category, Verbosity, Location, Rotation, Axis0, Axis1, Color, Format, ...) \
+	UE_VLOG_OBOX(Owner, Category, Verbosity, FBox(25.0f * FVector(-1, -1, 0), 25.0f * FVector(1, 1, 0)), FTransform(Rotation, Location, FVector::OneVector).ToMatrixNoScale(), Color, Format, ##__VA_ARGS__)
+
+namespace UE::Learning::Agents::Private
+{
+	template<typename RewardUObject, typename RewardFObject, typename... InArgTypes>
+	RewardUObject* AddReward(ULearningAgentsTrainer* AgentTrainer, FName Name, InArgTypes&& ...Args)
+	{
+		if (!AgentTrainer)
+		{
+			UE_LOG(LogLearning, Error, TEXT("AgentTrainer is nullptr"));
+			return nullptr;
+		}
+
+		RewardUObject* Reward = NewObject<RewardUObject>(AgentTrainer, Name);
+
+		Reward->RewardObject = MakeShared<RewardFObject>(
+			Reward->GetFName(),
+			AgentTrainer->GetAgentType().GetInstanceData().ToSharedRef(),
+			AgentTrainer->GetAgentType().GetMaxInstanceNum(),
+			Forward<InArgTypes>(Args)...);
+
+		AgentTrainer->AddReward(Reward, Reward->RewardObject.ToSharedRef());
+
+		return Reward;
+	}
+}
+
+//------------------------------------------------------------------
+
+UFloatReward* UFloatReward::AddFloatReward(ULearningAgentsTrainer* AgentTrainer, FName Name, float Weight)
+{
+	return UE::Learning::Agents::Private::AddReward<UFloatReward, UE::Learning::FFloatReward>(AgentTrainer, Name, Weight);
+}
+
+void UFloatReward::SetFloatReward(int32 AgentId, float Reward)
+{
+	TLearningArrayView<1, float> View = RewardObject->InstanceData->View(RewardObject->ValueHandle);
+
+	if (AgentId == INDEX_NONE)
+	{
+		UE_LOG(LogLearning, Error, TEXT("AgentId is invalid (INDEX_NONE)"));
+		return;
+	}
+
+	if (AgentId < 0 || AgentId >= View.Num<0>())
+	{
+		UE_LOG(LogLearning, Error, TEXT("AgentId %d is out of index. Valid range [0, %d]."), AgentId, View.Num<0>() - 1);
+		return;
+	}
+
+	View[AgentId] = Reward;
+}
+
+#if ENABLE_VISUAL_LOG
+void UFloatReward::VisualLog(const UE::Learning::FIndexSet Instances) const
+{
+	UE_LEARNING_TRACE_CPUPROFILER_EVENT_SCOPE(UFloatReward::VisualLog);
+
+	const TLearningArrayView<1, const float> ValueView = RewardObject->InstanceData->ConstView(RewardObject->ValueHandle);
+	const TLearningArrayView<1, const float> WeightView = RewardObject->InstanceData->ConstView(RewardObject->WeightHandle);
+	const TLearningArrayView<1, const float> RewardView = RewardObject->InstanceData->ConstView(RewardObject->RewardHandle);
+
+	if (const ULearningAgentsTrainer* AgentTrainer = Cast<ULearningAgentsTrainer>(GetOuter()))
+	{
+		for (const int32 Instance : Instances)
+		{
+			if (const AActor* Actor = Cast<AActor>(AgentTrainer->GetAgent(Instance)))
+			{
+				UE_LEARNING_AGENTS_VLOG_STRING(this, LogLearning, Display,
+					Actor->GetActorLocation(),
+					VisualLogColor.ToFColor(true),
+					TEXT("Agent %i\nWeight: [% 6.2f]\nValue: [% 6.2f]\nReward: [% 6.3f]"),
+					Instance,
+					WeightView[Instance],
+					ValueView[Instance],
+					RewardView[Instance]);
+			}
+		}
+	}
+}
+#endif
+
+//------------------------------------------------------------------
+
+UScalarVelocityReward* UScalarVelocityReward::AddScalarVelocityReward(ULearningAgentsTrainer* AgentTrainer, FName Name, float Weight, float Scale)
+{
+	return UE::Learning::Agents::Private::AddReward<UScalarVelocityReward, UE::Learning::FScalarVelocityReward>(AgentTrainer, Name, Weight, Scale);
+}
+
+void UScalarVelocityReward::SetScalarVelocityReward(int32 AgentId, float Velocity)
+{
+	TLearningArrayView<1, float> VelocityView = RewardObject->InstanceData->View(RewardObject->VelocityHandle);
+
+	if (AgentId == INDEX_NONE)
+	{
+		UE_LOG(LogLearning, Error, TEXT("AgentId is invalid (INDEX_NONE)"));
+		return;
+	}
+
+	if (AgentId < 0 || AgentId >= VelocityView.Num<0>())
+	{
+		UE_LOG(LogLearning, Error, TEXT("AgentId %d is out of index. Valid range [0, %d]."), AgentId, VelocityView.Num<0>() - 1);
+		return;
+	}
+
+	VelocityView[AgentId] = Velocity;
+}
+
+#if ENABLE_VISUAL_LOG
+void UScalarVelocityReward::VisualLog(const UE::Learning::FIndexSet Instances) const
+{
+	UE_LEARNING_TRACE_CPUPROFILER_EVENT_SCOPE(UScalarVelocityReward::VisualLog);
+
+	const TLearningArrayView<1, const float> VelocityView = RewardObject->InstanceData->ConstView(RewardObject->VelocityHandle);
+	const TLearningArrayView<1, const float> WeightView = RewardObject->InstanceData->ConstView(RewardObject->WeightHandle);
+	const TLearningArrayView<1, const float> ScaleView = RewardObject->InstanceData->ConstView(RewardObject->ScaleHandle);
+	const TLearningArrayView<1, const float> RewardView = RewardObject->InstanceData->ConstView(RewardObject->RewardHandle);
+
+	if (const ULearningAgentsTrainer* AgentTrainer = Cast<ULearningAgentsTrainer>(GetOuter()))
+	{
+		for (const int32 Instance : Instances)
+		{
+			if (const AActor* Actor = Cast<const AActor>(AgentTrainer->GetAgent(Instance)))
+			{
+				UE_LEARNING_AGENTS_VLOG_STRING(this, LogLearning, Display,
+					Actor->GetActorLocation(),
+					VisualLogColor.ToFColor(true),
+					TEXT("Agent %i\nWeight: [% 6.2f]\nScale: [% 6.2f]\nVelocity: [% 6.2f]\nReward: [% 6.3f]"),
+					Instance,
+					WeightView[Instance],
+					ScaleView[Instance],
+					VelocityView[Instance],
+					RewardView[Instance]);
+			}
+		}
+	}
+}
+#endif
+
+ULocalDirectionalVelocityReward* ULocalDirectionalVelocityReward::AddLocalDirectionalVelocityReward(ULearningAgentsTrainer* AgentTrainer, FName Name, float Weight, float Scale, FVector Axis)
+{
+	return UE::Learning::Agents::Private::AddReward<ULocalDirectionalVelocityReward, UE::Learning::FLocalDirectionalVelocityReward>(AgentTrainer, Name, Weight, Scale, Axis);
+}
+
+void ULocalDirectionalVelocityReward::SetLocalDirectionalVelocityReward(int32 AgentId, FVector Velocity, FRotator RelativeRotation)
+{
+	TLearningArrayView<1, FVector> VelocityView = RewardObject->InstanceData->View(RewardObject->VelocityHandle);
+	TLearningArrayView<1, FQuat> RelativeRotationView = RewardObject->InstanceData->View(RewardObject->RelativeRotationHandle);
+
+	if (AgentId == INDEX_NONE)
+	{
+		UE_LOG(LogLearning, Error, TEXT("AgentId is invalid (INDEX_NONE)"));
+		return;
+	}
+
+	if (AgentId < 0 || AgentId >= VelocityView.Num<0>())
+	{
+		UE_LOG(LogLearning, Error, TEXT("AgentId %d is out of index. Valid range [0, %d]."), AgentId, VelocityView.Num<0>() - 1);
+		return;
+	}
+
+	VelocityView[AgentId] = Velocity;
+	RelativeRotationView[AgentId] = FQuat::MakeFromRotator(RelativeRotation);
+}
+
+#if ENABLE_VISUAL_LOG
+void ULocalDirectionalVelocityReward::VisualLog(const UE::Learning::FIndexSet Instances) const
+{
+	UE_LEARNING_TRACE_CPUPROFILER_EVENT_SCOPE(ULocalDirectionalVelocityReward::VisualLog);
+
+	const TLearningArrayView<1, const FVector> VelocityView = RewardObject->InstanceData->View(RewardObject->VelocityHandle);
+	const TLearningArrayView<1, const FQuat> RelativeRotationView = RewardObject->InstanceData->View(RewardObject->RelativeRotationHandle);
+	const TLearningArrayView<1, const float> WeightView = RewardObject->InstanceData->ConstView(RewardObject->WeightHandle);
+	const TLearningArrayView<1, const float> ScaleView = RewardObject->InstanceData->ConstView(RewardObject->ScaleHandle);
+	const TLearningArrayView<1, const float> RewardView = RewardObject->InstanceData->ConstView(RewardObject->RewardHandle);
+
+	if (const ULearningAgentsTrainer* AgentTrainer = Cast<ULearningAgentsTrainer>(GetOuter()))
+	{
+		for (const int32 Instance : Instances)
+		{
+			if (const AActor* Actor = Cast<AActor>(AgentTrainer->GetAgent(Instance)))
+			{
+				const FVector Velocity = VelocityView[Instance];
+				const FQuat RelativeRotation = RelativeRotationView[Instance];
+				const FVector LocalVelocity = RelativeRotation.UnrotateVector(Velocity);
+				const FVector Direction = RelativeRotation.RotateVector(RewardObject->Axis);
+
+				UE_VLOG_ARROW(this, LogLearning, Display,
+					Actor->GetActorLocation(),
+					Actor->GetActorLocation() + Velocity,
+					VisualLogColor.ToFColor(true),
+					TEXT(""));
+
+				UE_LEARNING_AGENTS_VLOG_STRING(this, LogLearning, Display,
+					Actor->GetActorLocation() + Velocity,
+					VisualLogColor.ToFColor(true),
+					TEXT("Velocity: [% 6.3f % 6.3f % 6.3f]\nLocal Velocity: [% 6.3f % 6.3f % 6.3f]"),
+					Velocity.X, Velocity.Y, Velocity.Z,
+					LocalVelocity.X, LocalVelocity.Y, LocalVelocity.Z);
+
+				UE_VLOG_ARROW(this, LogLearning, Display,
+					Actor->GetActorLocation(),
+					Actor->GetActorLocation() + 100.0f * Direction,
+					VisualLogColor.ToFColor(true),
+					TEXT(""));
+
+				UE_LEARNING_AGENTS_VLOG_STRING(this, LogLearning, Display,
+					Actor->GetActorLocation() + 100.0f * Direction,
+					VisualLogColor.ToFColor(true),
+					TEXT("Direction: [% 6.3f % 6.3f % 6.3f]\nLocal Direction: [% 6.3f % 6.3f % 6.3f]"),
+					Direction.X, Direction.Y, Direction.Z,
+					RewardObject->Axis.X, RewardObject->Axis.Y, RewardObject->Axis.Z);
+
+				UE_LEARNING_AGENTS_VLOG_TRANSFORM(this, LogLearning, Display,
+					Actor->GetActorLocation(),
+					RelativeRotation,
+					VisualLogColor.ToFColor(true),
+					TEXT("Agent %i\nDot Product: [% 6.3f]\nWeight: [% 6.2f]\nScale: [% 6.2f]\nReward: [% 6.3f]"),
+					Instance,
+					LocalVelocity.Dot(RewardObject->Axis),
+					WeightView[Instance],
+					ScaleView[Instance],
+					RewardView[Instance]);
+			}
+		}
+	}
+}
+#endif
+
+//------------------------------------------------------------------
+
+UPlanarPositionDifferencePenalty* UPlanarPositionDifferencePenalty::AddPlanarPositionDifferencePenalty(ULearningAgentsTrainer* AgentTrainer, FName Name, float Weight, float Scale, float Threshold, FVector Axis0, FVector Axis1)
+{
+	return UE::Learning::Agents::Private::AddReward<UPlanarPositionDifferencePenalty, UE::Learning::FPlanarPositionDifferencePenalty>(AgentTrainer, Name, Weight, Scale, Threshold, Axis0, Axis1);
+}
+
+void UPlanarPositionDifferencePenalty::SetPlanarPositionDifferencePenalty(int32 AgentId, FVector Position0, FVector Position1)
+{
+	TLearningArrayView<1, FVector> Position0View = RewardObject->InstanceData->View(RewardObject->Position0Handle);
+	TLearningArrayView<1, FVector> Position1View = RewardObject->InstanceData->View(RewardObject->Position1Handle);
+
+	if (AgentId == INDEX_NONE)
+	{
+		UE_LOG(LogLearning, Error, TEXT("AgentId is invalid (INDEX_NONE)"));
+		return;
+	}
+
+	if (AgentId < 0 || AgentId >= Position0View.Num<0>())
+	{
+		UE_LOG(LogLearning, Error, TEXT("AgentId %d is out of index. Valid range [0, %d]."), AgentId, Position0View.Num<0>() - 1);
+		return;
+	}
+
+	Position0View[AgentId] = Position0;
+	Position1View[AgentId] = Position1;
+}
+
+#if ENABLE_VISUAL_LOG
+void UPlanarPositionDifferencePenalty::VisualLog(const UE::Learning::FIndexSet Instances) const
+{
+	UE_LEARNING_TRACE_CPUPROFILER_EVENT_SCOPE(UPlanarPositionDifferencePenalty::VisualLog);
+
+	const TLearningArrayView<1, const FVector> Position0View = RewardObject->InstanceData->ConstView(RewardObject->Position0Handle);
+	const TLearningArrayView<1, const FVector> Position1View = RewardObject->InstanceData->ConstView(RewardObject->Position1Handle);
+	const TLearningArrayView<1, const float> WeightView = RewardObject->InstanceData->ConstView(RewardObject->WeightHandle);
+	const TLearningArrayView<1, const float> ScaleView = RewardObject->InstanceData->ConstView(RewardObject->ScaleHandle);
+	const TLearningArrayView<1, const float> ThresholdView = RewardObject->InstanceData->ConstView(RewardObject->ThresholdHandle);
+	const TLearningArrayView<1, const float> RewardView = RewardObject->InstanceData->ConstView(RewardObject->RewardHandle);
+
+	if (const ULearningAgentsTrainer* AgentTrainer = Cast<ULearningAgentsTrainer>(GetOuter()))
+	{
+		for (const int32 Instance : Instances)
+		{
+			if (const AActor* Actor = Cast<AActor>(AgentTrainer->GetAgent(Instance)))
+			{
+				const FVector Position0 = Position0View[Instance];
+				const FVector Position1 = Position1View[Instance];
+
+				const FVector PlanarPosition0 = FVector(RewardObject->Axis0.Dot(Position0), RewardObject->Axis1.Dot(Position0), 0.0f);
+				const FVector PlanarPosition1 = FVector(RewardObject->Axis0.Dot(Position1), RewardObject->Axis1.Dot(Position1), 0.0f);
+
+				UE_VLOG_LOCATION(this, LogLearning, Display,
+					Position0,
+					10.0f,
+					VisualLogColor.ToFColor(true),
+					TEXT("Position0: [% 6.1f % 6.1f % 6.1f]\nPlanar Position0: [% 6.1f % 6.1f]"),
+					Position0.X, Position0.Y, Position0.Z,
+					PlanarPosition0.X, PlanarPosition0.Y);
+
+				UE_LEARNING_AGENTS_VLOG_PLANE(this, LogLearning, Display,
+					Position0,
+					FQuat::Identity,
+					FeatureObject->Axis0,
+					FeatureObject->Axis1,
+					VisualLogColor.ToFColor(true),
+					TEXT(""));
+
+				UE_VLOG_LOCATION(this, LogLearning, Display,
+					Position1,
+					10.0f,
+					VisualLogColor.ToFColor(true),
+					TEXT("Position1: [% 6.1f % 6.1f % 6.1f]\nPlanar Position1: [% 6.1f % 6.1f]"),
+					Position1.X, Position1.Y, Position1.Z,
+					PlanarPosition1.X, PlanarPosition1.Y);
+
+				UE_LEARNING_AGENTS_VLOG_PLANE(this, LogLearning, Display,
+					Position1,
+					FQuat::Identity,
+					FeatureObject->Axis0,
+					FeatureObject->Axis1,
+					VisualLogColor.ToFColor(true),
+					TEXT(""));
+
+				UE_VLOG_SEGMENT(this, LogLearning, Display,
+					Position0,
+					Position1,
+					VisualLogColor.ToFColor(true),
+					TEXT("Agent %i\nDistance: [% 6.3f]\nPlanar Distance: [% 6.3f]\nWeight: [% 6.2f]\nScale: [% 6.2f]\nThreshold: [% 6.2f]\nReward: [% 6.3f]"),
+					Instance,
+					FVector::Distance(Position0, Position1),
+					FVector::Distance(PlanarPosition0, PlanarPosition1),
+					WeightView[Instance],
+					ScaleView[Instance],
+					ThresholdView[Instance],
+					RewardView[Instance]);
+			}
+		}
+	}
+}
+#endif
+
+//------------------------------------------------------------------
+
+UPositionArraySimilarityReward* UPositionArraySimilarityReward::AddPositionArraySimilarityReward(ULearningAgentsTrainer* AgentTrainer, FName Name, int32 PositionNum, float Weight, float Scale)
+{
+	return UE::Learning::Agents::Private::AddReward<UPositionArraySimilarityReward, UE::Learning::FPositionArraySimilarityReward>(AgentTrainer, Name, PositionNum, Weight, Scale);
+}
+
+void UPositionArraySimilarityReward::SetPositionArraySimilarityReward(
+	int32 AgentId,
+	const TArray<FVector>& Positions0,
+	const TArray<FVector>& Positions1,
+	FVector RelativePosition0,
+	FVector RelativePosition1,
+	FRotator RelativeRotation0,
+	FRotator RelativeRotation1)
+{
+	TLearningArrayView<2, FVector> Positions0View = RewardObject->InstanceData->View(RewardObject->Positions0Handle);
+	TLearningArrayView<2, FVector> Positions1View = RewardObject->InstanceData->View(RewardObject->Positions1Handle);
+	TLearningArrayView<1, FVector> RelativePosition0View = RewardObject->InstanceData->View(RewardObject->RelativePosition0Handle);
+	TLearningArrayView<1, FVector> RelativePosition1View = RewardObject->InstanceData->View(RewardObject->RelativePosition1Handle);
+	TLearningArrayView<1, FQuat> RelativeRotation0View = RewardObject->InstanceData->View(RewardObject->RelativeRotation0Handle);
+	TLearningArrayView<1, FQuat> RelativeRotation1View = RewardObject->InstanceData->View(RewardObject->RelativeRotation1Handle);
+
+	if (AgentId == INDEX_NONE)
+	{
+		UE_LOG(LogLearning, Error, TEXT("AgentId is invalid (INDEX_NONE)"));
+		return;
+	}
+
+	if (AgentId < 0 || AgentId >= Positions0View.Num<0>())
+	{
+		UE_LOG(LogLearning, Error, TEXT("AgentId %d is out of index. Valid range [0, %d]."), AgentId, Positions0View.Num<0>() - 1);
+		return;
+	}
+
+	const int32 PositionNum = Positions0View.Num<1>();
+
+	if (Positions0.Num() != PositionNum || Positions1.Num() != PositionNum)
+	{
+		UE_LOG(LogLearning, Error, TEXT("Incorrect number of positions in array. Got %i and %i, expected %i."), Positions0.Num(), Positions1.Num(), PositionNum);
+		return;
+	}
+
+	RelativePosition0View[AgentId] = RelativePosition0;
+	RelativePosition1View[AgentId] = RelativePosition1;
+	RelativeRotation0View[AgentId] = FQuat::MakeFromRotator(RelativeRotation0);
+	RelativeRotation1View[AgentId] = FQuat::MakeFromRotator(RelativeRotation1);
+	UE::Learning::Array::Copy(Positions0View[AgentId], TLearningArrayView<1, const FVector>(Positions0));
+	UE::Learning::Array::Copy(Positions1View[AgentId], TLearningArrayView<1, const FVector>(Positions1));
+}
+
+#if ENABLE_VISUAL_LOG
+void UPositionArraySimilarityReward::VisualLog(const UE::Learning::FIndexSet Instances) const
+{
+	UE_LEARNING_TRACE_CPUPROFILER_EVENT_SCOPE(UPositionArraySimilarityReward::VisualLog);
+
+	const TLearningArrayView<2, const FVector> Positions0View = RewardObject->InstanceData->ConstView(RewardObject->Positions0Handle);
+	const TLearningArrayView<2, const FVector> Positions1View = RewardObject->InstanceData->ConstView(RewardObject->Positions1Handle);
+	const TLearningArrayView<1, const FVector> RelativePosition0View = RewardObject->InstanceData->ConstView(RewardObject->RelativePosition0Handle);
+	const TLearningArrayView<1, const FVector> RelativePosition1View = RewardObject->InstanceData->ConstView(RewardObject->RelativePosition1Handle);
+	const TLearningArrayView<1, const FQuat> RelativeRotation0View = RewardObject->InstanceData->ConstView(RewardObject->RelativeRotation0Handle);
+	const TLearningArrayView<1, const FQuat> RelativeRotation1View = RewardObject->InstanceData->ConstView(RewardObject->RelativeRotation1Handle);
+	const TLearningArrayView<1, const float> WeightView = RewardObject->InstanceData->ConstView(RewardObject->WeightHandle);
+	const TLearningArrayView<1, const float> ScaleView = RewardObject->InstanceData->ConstView(RewardObject->ScaleHandle);
+	const TLearningArrayView<1, const float> ThresholdView = RewardObject->InstanceData->ConstView(RewardObject->ThresholdHandle);
+	const TLearningArrayView<1, const float> RewardView = RewardObject->InstanceData->ConstView(RewardObject->RewardHandle);
+
+	const int32 PositionNum = Positions0View.Num<1>();
+		
+	if (const ULearningAgentsTrainer* AgentTrainer = Cast<ULearningAgentsTrainer>(GetOuter()))
+	{
+		for (const int32 Instance : Instances)
+		{
+			if (const AActor* Actor = Cast<AActor>(AgentTrainer->GetAgent(Instance)))
+			{
+				const FVector RelativePosition0 = RelativePosition0View[Instance];
+				const FVector RelativePosition1 = RelativePosition1View[Instance];
+				const FQuat RelativeRotation0 = RelativeRotation0View[Instance];
+				const FQuat RelativeRotation1 = RelativeRotation1View[Instance];
+
+				for (int32 PositionIdx = 0; PositionIdx < PositionNum; PositionIdx++)
+				{
+					const FVector Position0 = Positions0View[Instance][PositionIdx];
+					const FVector Position1 = Positions1View[Instance][PositionIdx];
+
+					const FVector LocalPosition0 = RelativeRotation0.UnrotateVector(Positions0View[Instance][PositionIdx] - RelativePosition0);
+					const FVector LocalPosition1 = RelativeRotation1.UnrotateVector(Positions1View[Instance][PositionIdx] - RelativePosition1);
+
+					UE_VLOG_LOCATION(this, LogLearning, Display,
+						Position0,
+						10.0f,
+						VisualLogColor.ToFColor(true),
+						TEXT("Position0: [% 6.1f % 6.1f % 6.1f]\nLocal Position0: [% 6.1f % 6.1f % 6.1f]"),
+						Position0.X, Position0.Y, Position0.Z,
+						LocalPosition0.X, LocalPosition0.Y, LocalPosition0.Z);
+
+					UE_VLOG_LOCATION(this, LogLearning, Display,
+						Position1,
+						10.0f,
+						VisualLogColor.ToFColor(true),
+						TEXT("Position1: [% 6.1f % 6.1f % 6.1f]\nLocal Position1: [% 6.1f % 6.1f % 6.1f]"),
+						Position1.X, Position1.Y, Position1.Z,
+						LocalPosition1.X, LocalPosition1.Y, LocalPosition1.Z);
+
+					UE_VLOG_SEGMENT(this, LogLearning, Display,
+						Position0,
+						Position1,
+						VisualLogColor.ToFColor(true),
+						TEXT("Distance: [% 6.1f]\nLocal Distance: [% 6.1f]"),
+						FVector::Distance(Position0, Position1),
+						FVector::Distance(LocalPosition0, LocalPosition1));
+				}
+
+				UE_LEARNING_AGENTS_VLOG_TRANSFORM(this, LogLearning, Display,
+					RelativePosition0,
+					RelativeRotation0,
+					VisualLogColor.ToFColor(true),
+					TEXT("Relative Transform 0"));
+
+				UE_LEARNING_AGENTS_VLOG_TRANSFORM(this, LogLearning, Display,
+					RelativePosition1,
+					RelativeRotation1,
+					VisualLogColor.ToFColor(true),
+					TEXT("Relative Transform 1"));
+
+				UE_LEARNING_AGENTS_VLOG_STRING(this, LogLearning, Display,
+					Actor->GetActorLocation(),
+					VisualLogColor.ToFColor(true),
+					TEXT("Agent %i\nWeight: [% 6.2f]\nScale: [% 6.2f]\nThreshold: [% 6.2f]\nReward: [% 6.3f]"),
+					Instance,
+					WeightView[Instance],
+					ScaleView[Instance],
+					ThresholdView[Instance],
+					RewardView[Instance]);
+			}
+		}
+	}
+}
+#endif
+
+//------------------------------------------------------------------
+
+#undef UE_LEARNING_AGENTS_VLOG_STRING
+#undef UE_LEARNING_AGENTS_VLOG_TRANSFORM
+#undef UE_LEARNING_AGENTS_VLOG_PLANE
diff --git a/Engine/Plugins/Experimental/LearningAgents/Source/LearningAgentsTraining/Private/LearningAgentsTrainer.cpp b/Engine/Plugins/Experimental/LearningAgents/Source/LearningAgentsTraining/Private/LearningAgentsTrainer.cpp
new file mode 100644
--- /dev/null
+++ b/Engine/Plugins/Experimental/LearningAgents/Source/LearningAgentsTraining/Private/LearningAgentsTrainer.cpp
@@ -0,0 +1,552 @@
+// Copyright Epic Games, Inc. All Rights Reserved.
+
+#include "LearningAgentsTrainer.h"
+
+#include "GenericPlatform/GenericPlatformMisc.h"
+#include "HAL/FileManager.h"
+#include "LearningAgentsType.h"
+#include "LearningAgentsRewards.h"
+#include "LearningAgentsCompletions.h"
+#include "LearningAgentsObservations.h"
+#include "LearningArray.h"
+#include "LearningArrayMap.h"
+#include "LearningExperience.h"
+#include "LearningFeatureObject.h"
+#include "LearningLog.h"
+#include "LearningNeuralNetwork.h"
+#include "LearningPolicyObject.h"
+#include "LearningPPOTrainer.h"
+#include "LearningRewardObject.h"
+#include "LearningCompletion.h"
+#include "LearningCompletionObject.h"
+#include "Misc/App.h"
+#include "Misc/Paths.h"
+#include "GameFramework/GameUserSettings.h"
+#include "PhysicsEngine/PhysicsSettings.h"
+#include "Engine/GameViewportClient.h"
+#include "EngineDefines.h"
+
+ULearningAgentsTrainer::ULearningAgentsTrainer() : ULearningAgentsTypeComponent() {}
+ULearningAgentsTrainer::ULearningAgentsTrainer(FVTableHelper& Helper) : ULearningAgentsTrainer() {}
+ULearningAgentsTrainer::~ULearningAgentsTrainer() {}
+
+void ULearningAgentsTrainer::OnAgentAdded_Implementation(int32 AgentId, UObject* Agent)
+{
+	Super::OnAgentAdded_Implementation(AgentId, Agent);
+
+	if (bSetupPerformed)
+	{
+		// Reset the instance and the buffer in case we have stale data from another agent that was using this id previously
+		ResetInstance({AgentId});
+		EpisodeBuffer->Reset(AgentId);
+	}
+}
+
+void ULearningAgentsTrainer::SetupTrainer(ULearningAgentsType* AgentTypeOverride, const FLearningAgentsTrainerSettings& Settings)
+{
+	if (AgentTypeOverride)
+	{
+		AgentType = AgentTypeOverride;
+	}
+	else if (!AgentType)
+	{
+		UE_LOG(LogLearning, Error, TEXT("LearningAgentsTrainer setup called but AgentType is nullptr."));
+		return;
+	}
+
+	if (!AgentType->IsSetupPerformed())
+	{
+		UE_LOG(LogLearning, Error, TEXT("ULearningAgentsTrainer: AgentType not setup. "
+			"Consider using ULearningAgentsTrainer::OnAgentTypeSetupComplete when calling SetupTrainer."));
+		return;
+	}
+
+	// Reset Setup Flag
+	bSetupPerformed = false;
+
+	// Setup Rewards
+	RewardObjects.Empty();
+	RewardFeatures.Empty();
+	SetupRewards(this);
+	Rewards = MakeShared<UE::Learning::FSumReward>(TEXT("Rewards"),
+		TLearningArrayView<1, const TSharedRef<UE::Learning::FRewardObject>>(RewardFeatures),
+		AgentType->GetInstanceData().ToSharedRef(),
+		AgentType->GetMaxInstanceNum());
+
+	// Setup Completions
+	CompletionObjects.Empty();
+	CompletionFeatures.Empty();
+	SetupCompletions(this);
+	Completions = MakeShared<UE::Learning::FAnyCompletion>(TEXT("Completions"),
+		TLearningArrayView<1, const TSharedRef<UE::Learning::FCompletionObject>>(CompletionFeatures),
+		AgentType->GetInstanceData().ToSharedRef(),
+		AgentType->GetMaxInstanceNum());
+
+	// Create Episode Buffer
+	EpisodeBuffer = MakeUnique<UE::Learning::FEpisodeBuffer>();
+	EpisodeBuffer->Resize(
+		AgentType->GetMaxInstanceNum(),
+		Settings.MaxStepNum,
+		AgentType->GetObservationFeature().DimNum(),
+		AgentType->GetActionFeature().DimNum());
+
+	MaxStepsCompletion = Settings.MaxStepsCompletion;
+
+	// Create Replay Buffer
+	ReplayBuffer = MakeUnique<UE::Learning::FReplayBuffer>();
+	ReplayBuffer->Resize(
+		AgentType->GetObservationFeature().DimNum(),
+		AgentType->GetActionFeature().DimNum(),
+		Settings.MaximumRecordedEpisodesPerIteration,
+		Settings.MaximumRecordedStepsPerIteration);
+
+	// Create Reset Buffer
+	ResetBuffer = MakeUnique<UE::Learning::FResetInstanceBuffer>();
+	ResetBuffer->Resize(AgentType->GetMaxInstanceNum());
+
+	// Done!
+	bSetupPerformed = true;
+}
+
+const bool ULearningAgentsTrainer::IsSetupPerformed() const
+{
+	return bSetupPerformed;
+}
+
+void ULearningAgentsTrainer::EndPlay(const EEndPlayReason::Type EndPlayReason)
+{
+	if (bIsTraining)
+	{
+		StopTraining();
+	}
+
+	Super::EndPlay(EndPlayReason);
+}
+
+void ULearningAgentsTrainer::SetupRewards_Implementation(ULearningAgentsTrainer* AgentTrainer)
+{
+	// Can be overridden to setup rewards without blueprints
+}
+
+void ULearningAgentsTrainer::SetRewards_Implementation(const TArray<int32>& AgentIds)
+{
+	// Can be overridden to set rewards without blueprints
+}
+
+void ULearningAgentsTrainer::AddReward(TObjectPtr<ULearningAgentsReward> Object, const TSharedRef<UE::Learning::FRewardObject>& Reward)
+{
+	check(!bSetupPerformed);
+	RewardObjects.Add(Object);
+	RewardFeatures.Add(Reward);
+}
+
+void ULearningAgentsTrainer::SetupCompletions_Implementation(ULearningAgentsTrainer* AgentTrainer)
+{
+	// Can be overridden to setup completions without blueprints
+}
+
+void ULearningAgentsTrainer::SetCompletions_Implementation(const TArray<int32>& AgentIds)
+{
+	// Can be overridden to evaluate completions without blueprints
+}
+
+void ULearningAgentsTrainer::AddCompletion(TObjectPtr<ULearningAgentsCompletion> Object, const TSharedRef<UE::Learning::FCompletionObject>& Completion)
+{
+	check(!bSetupPerformed);
+	CompletionObjects.Add(Object);
+	CompletionFeatures.Add(Completion);
+}
+
+const bool ULearningAgentsTrainer::IsTraining() const
+{
+	return bIsTraining;
+}
+
+void ULearningAgentsTrainer::StartTraining(
+	const FLearningAgentsTrainerTrainingSettings& TrainingSettings, 
+	const FLearningAgentsTrainerGameSettings& Settings)
+{
+	if (!bSetupPerformed)
+	{
+		UE_LOG(LogLearning, Error, TEXT("Setup must be run before training can start."));
+		return;
+	}
+
+	// Record GameState Settings
+
+	bFixedTimestepUsed = FApp::UseFixedTimeStep();
+	FixedTimeStepDeltaTime = FApp::GetFixedDeltaTime();
+
+	UGameUserSettings* GameSettings = UGameUserSettings::GetGameUserSettings();
+	if (GameSettings)
+	{
+		bVSyncEnabled = GameSettings->IsVSyncEnabled();
+	}
+
+	UPhysicsSettings* PhysicsSettings = UPhysicsSettings::Get();
+	if (PhysicsSettings)
+	{
+		MaxPhysicsStep = PhysicsSettings->MaxPhysicsDeltaTime;
+	}
+
+	UGameViewportClient* ViewportClient = GetWorld() ? GetWorld()->GetGameViewport() : nullptr;
+	if (ViewportClient)
+	{
+		ViewModeIndex = ViewportClient->ViewModeIndex;
+	}
+	// Apply Training GameState Settings
+
+	FApp::SetUseFixedTimeStep(Settings.bUseFixedTimeStep);
+
+	if (Settings.FixedTimeStepFrequency > UE_SMALL_NUMBER)
+	{
+		FApp::SetFixedDeltaTime(1.0f / Settings.FixedTimeStepFrequency);
+		if (Settings.bSetMaxPhysicsStepToFixedTimeStep && PhysicsSettings)
+		{
+			PhysicsSettings->MaxPhysicsDeltaTime = 1.0f / Settings.FixedTimeStepFrequency;
+		}
+	}
+	else
+	{
+		UE_LOG(LogLearning, Warning, TEXT("Provided invalid FixedTimeStepFrequency: %0.5f"), Settings.FixedTimeStepFrequency);
+	}
+
+	if (Settings.bDisableVSync && GameSettings)
+	{
+		GameSettings->SetVSyncEnabled(false);
+		GameSettings->ApplySettings(false);
+	}
+
+	if (Settings.bUseUnlitViewportRendering && ViewportClient)
+	{
+		ViewportClient->ViewModeIndex = EViewModeIndex::VMI_Unlit;
+	}
+
+
+	// Start Trainer
+
+#if WITH_EDITOR
+	const FString PythonExecutablePath = UE::Learning::Trainer::DefaultEditorPythonExecutablePath();
+	const FString SitePackagesPath = UE::Learning::Trainer::DefaultEditorSitePackagesPath();
+	const FString PythonContentPath = UE::Learning::Trainer::DefaultEditorPythonContentPath();
+	const FString IntermediatePath = UE::Learning::Trainer::DefaultEditorIntermediatePath();
+#else
+	// If we want to run training in a cooked, non-editor build, then by default we wont have access to python or the 
+	// learning training scripts - these are editor-only and will be stripped during the cooking process.
+	//
+	// However, running training in non-editor builds can be very important - we probably want to disable rendering 
+	// and sound while we are training to make experience gathering as fast as possible - and for any non-trivial game 
+	// is simply may not be realistic to run it for a long time in play-in-editor mode.
+	//
+	// For this reason even in non-editor builds we let you provide paths the `python.exe` provided by the editor, 
+	// as well as PythonFoundationPackages site-packages, and the Learning training scripts. This allows you to 
+	// run training when these things actually exist somewhere on your machine, which will usually be the case on 
+	// a normal development machine.
+
+	// todo bmulcahy all these paths should be configurable by the dev
+	UE_LEARNING_CHECKF(PLATFORM_WINDOWS || PLATFORM_MAC || PLATFORM_LINUX, TEXT("Python only supported on Windows, Mac, and Linux."));
+	const FString PythonExe = PLATFORM_WINDOWS ? TEXT("python.exe") : TEXT("bin/python");
+	const FString EnginePath = IFileManager::Get().ConvertToAbsolutePathForExternalAppForRead(*(FPaths::RootDir() / TEXT("../../../../../../Engine")));
+	const FString PythonExecutablePath = EnginePath / TEXT("Binaries/ThirdParty/Python3") / FPlatformMisc::GetUBTPlatform() / PythonExe;
+	const FString SitePackagesPath = EnginePath / TEXT("Plugins/Experimental/PythonFoundationPackages/Content/Python/Lib/") / FPlatformMisc::GetUBTPlatform() / TEXT("site-packages");
+	const FString PythonContentPath = EnginePath / TEXT("Plugins/Experimental/LearningAgents/Content/Python/");
+	const FString IntermediatePath = EnginePath / TEXT("Plugins/Experimental/LearningAgents/Intermediate");
+#endif
+
+	UE::Learning::FPPOTrainerSettings PPOSettings;
+	PPOSettings.IterationNum = TrainingSettings.NumberOfIterations;
+	PPOSettings.bReinitializeNetwork = TrainingSettings.bReinitializeNetwork;
+	PPOSettings.bUseTensorboard = TrainingSettings.bUseTensorboard;
+	PPOSettings.InitialActionScale = TrainingSettings.InitialActionScale;
+	PPOSettings.DiscountFactor = TrainingSettings.DiscountFactor;
+	PPOSettings.Seed = TrainingSettings.RandomSeed;
+	PPOSettings.TrimEpisodeStartStepNum = TrainingSettings.NumberOfStepsToTrimAtStartOfEpisode;
+	PPOSettings.TrimEpisodeEndStepNum = TrainingSettings.NumberOfStepsToTrimAtEndOfEpisode;
+	PPOSettings.Device = TrainingSettings.Device == ELearningAgentsTrainerDevice::CPU ? UE::Learning::ETrainerDevice::CPU : UE::Learning::ETrainerDevice::GPU;
+
+	UE::Learning::FNeuralNetwork& NeuralNetwork = AgentType->GetNeuralNetwork();
+	const UE::Learning::FNeuralNetworkPolicyFunction& Policy = AgentType->GetPolicy();
+
+	// Start Python Training Process (this must be done on main thread)
+	Trainer = MakeUnique<UE::Learning::FSharedMemoryPPOTrainer>(
+		TEXT("Training"),
+		PythonExecutablePath,
+		SitePackagesPath,
+		PythonContentPath,
+		IntermediatePath,
+		NeuralNetwork,
+		Policy.Settings.ActionNoiseMin,
+		Policy.Settings.ActionNoiseMax,
+		*ReplayBuffer,
+		PPOSettings);
+
+	TrainerTimeout = TrainingSettings.TrainerCommunicationTimeout;
+
+	UE::Learning::ETrainerResponse Response;
+
+	if (TrainingSettings.bReinitializeNetwork)
+	{
+		UE_LOG(LogLearning, Display, TEXT("Receiving initial policy..."));
+		Response = Trainer->RecvPolicy(NeuralNetwork, TrainerTimeout);
+	}
+	else
+	{
+		UE_LOG(LogLearning, Display, TEXT("Sending initial policy..."));
+		Response = Trainer->SendPolicy(NeuralNetwork, TrainerTimeout);
+	}
+
+	if (Response != UE::Learning::ETrainerResponse::Success)
+	{
+		UE_LOG(LogLearning, Error, TEXT("Error sending or receiving policy from trainer. Check log for errors."));
+		Trainer->Terminate();
+		UE_DEBUG_BREAK();
+		return;
+	}
+
+	// Reset Agents, Episode Buffer, and Replay Buffer
+	ResetInstance(SelectedAgentIds);
+
+	EpisodeBuffer->Reset(SelectedAgentsSet);
+	ReplayBuffer->Reset();
+
+	// Done!
+	bIsTraining = true;
+}
+
+void ULearningAgentsTrainer::DoneTraining()
+{
+	if (bIsTraining)
+	{
+		// Wait for Trainer to finish
+		Trainer->Wait(1.0f);
+
+		// If not finished in time, terminate
+		Trainer->Terminate();
+
+		// Apply back previous game settings
+		FApp::SetUseFixedTimeStep(bFixedTimestepUsed);
+		FApp::SetFixedDeltaTime(FixedTimeStepDeltaTime);
+		UGameUserSettings* GameSettings = UGameUserSettings::GetGameUserSettings();
+		if (GameSettings)
+		{
+			GameSettings->SetVSyncEnabled(bVSyncEnabled);
+			GameSettings->ApplySettings(true);
+		}
+
+		UPhysicsSettings* PhysicsSettings = UPhysicsSettings::Get();
+		if (PhysicsSettings)
+		{
+			PhysicsSettings->MaxPhysicsDeltaTime = MaxPhysicsStep;
+		}
+
+		UGameViewportClient* ViewportClient = GetWorld() ? GetWorld()->GetGameViewport() : nullptr;
+		if (ViewportClient)
+		{
+			ViewportClient->ViewModeIndex = ViewModeIndex;
+		}
+
+		// Done
+		bIsTraining = false;
+	}
+}
+
+void ULearningAgentsTrainer::StopTraining()
+{
+	if (bIsTraining)
+	{
+		UE_LOG(LogLearning, Display, TEXT("Stopping training..."));
+		Trainer->SendStop();
+		DoneTraining();
+	}
+}
+
+void ULearningAgentsTrainer::EvaluateRewards()
+{
+	UE_LEARNING_TRACE_CPUPROFILER_EVENT_SCOPE(ULearningAgentsTrainer::EvaluateRewards);
+
+	if (!bSetupPerformed)
+	{
+		UE_LOG(LogLearning, Error, TEXT("Setup must be run before rewards can be evaluated."));
+		return;
+	}
+
+	SetRewards(SelectedAgentIds);
+
+	Rewards->Evaluate(SelectedAgentsSet);
+
+#if ENABLE_VISUAL_LOG
+	for (const ULearningAgentsReward* RewardObject : RewardObjects)
+	{
+		if (RewardObject)
+		{
+			RewardObject->VisualLog(SelectedAgentsSet);
+		}
+	}
+#endif
+}
+
+void ULearningAgentsTrainer::EvaluateCompletions()
+{
+	UE_LEARNING_TRACE_CPUPROFILER_EVENT_SCOPE(ULearningAgentsTrainer::EvaluateCompletions);
+
+	if (!bSetupPerformed)
+	{
+		UE_LOG(LogLearning, Error, TEXT("Setup must be run before completions can be evaluated."));
+		return;
+	}
+
+	SetCompletions(SelectedAgentIds);
+
+	Completions->Evaluate(SelectedAgentsSet);
+	
+#if ENABLE_VISUAL_LOG
+	for (const ULearningAgentsCompletion* CompletionObject : CompletionObjects)
+	{
+		if (CompletionObject)
+		{
+			CompletionObject->VisualLog(SelectedAgentsSet);
+		}
+	}
+#endif
+}
+
+void ULearningAgentsTrainer::IterateTraining()
+{
+	if (!bSetupPerformed)
+	{
+		UE_LOG(LogLearning, Error, TEXT("Setup must be run before training can be performed."));
+		return;
+	}
+
+	if (!bIsTraining)
+	{
+		if (!bPreviouslyLogged)
+		{
+			UE_LOG(LogLearning, Error, TEXT("Attempted to iterate training but the trainer is not running. Ensure StartTraining() is called before IterateTraining() and no errors occurred."));
+			UE_DEBUG_BREAK();
+			bPreviouslyLogged = true;
+		}
+
+		return;
+	}
+
+	UE::Learning::FFeatureObject& Observations = AgentType->GetObservationFeature();
+	UE::Learning::FFeatureObject& Actions = AgentType->GetActionFeature();
+
+	// Add Experience to Episode Buffer
+
+	EpisodeBuffer->Push(
+		Observations.FeatureBuffer(),
+		Actions.FeatureBuffer(),
+		Rewards->RewardBuffer(),
+		SelectedAgentsSet);
+
+	// Check for completion based on reaching the maximum episode length
+
+	UE::Learning::Completion::EvaluateEndOfEpisodeCompletions(
+		Completions->CompletionBuffer(),
+		EpisodeBuffer->GetEpisodeStepNums(),
+		EpisodeBuffer->GetMaxStepNum(),
+		MaxStepsCompletion == ELearningAgentsCompletion::Truncation 
+			? UE::Learning::ECompletionMode::Truncated 
+			: UE::Learning::ECompletionMode::Terminated,
+		SelectedAgentsSet);
+
+	// Find the set of Instances that need to be reset
+
+	ResetBuffer->SetResetInstancesFromCompletions(Completions->CompletionBuffer(), SelectedAgentsSet);
+
+	if (ResetBuffer->GetResetInstanceNum() > 0)
+	{
+		// Encode Observations for completed Instances
+
+		AgentType->SetObservations(ResetBuffer->GetResetInstances().ToArray());
+
+		Observations.Encode(ResetBuffer->GetResetInstances());
+		
+#if ENABLE_VISUAL_LOG
+		for (const ULearningAgentsObservation* ObservationObject : AgentType->GetObservationObjects())
+		{
+			if (ObservationObject)
+			{
+				ObservationObject->VisualLog(ResetBuffer->GetResetInstances());
+			}
+		}
+#endif
+
+		// Push Experience to Replay Buffer
+
+		bool bReplayBufferFull = ReplayBuffer->AddEpisodes(
+			Completions->CompletionBuffer(),
+			Observations.FeatureBuffer(),
+			*EpisodeBuffer,
+			ResetBuffer->GetResetInstances());
+
+		if (bReplayBufferFull)
+		{
+			// Send Experience
+
+			UE::Learning::ETrainerResponse Response = UE::Learning::ETrainerResponse::Success;
+
+			Response = Trainer->SendExperience(*ReplayBuffer, TrainerTimeout);
+
+			if (Response != UE::Learning::ETrainerResponse::Success)
+			{
+				UE_LOG(LogLearning, Error, TEXT("Error waiting to push experience to trainer. Check log for errors."));
+				StopTraining();
+				UE_DEBUG_BREAK();
+				return;
+			}
+
+			// Reset Replay Buffer
+
+			ReplayBuffer->Reset();
+
+			// Get Updated Policy
+
+			Response = Trainer->RecvPolicy(AgentType->GetNeuralNetwork(), TrainerTimeout);
+
+			if (Response == UE::Learning::ETrainerResponse::Completed)
+			{
+				UE_LOG(LogLearning, Display, TEXT("Trainer completed training."));
+				DoneTraining();
+				return;
+			}
+			else if (Response != UE::Learning::ETrainerResponse::Success)
+			{
+				UE_LOG(LogLearning, Error, TEXT("Error waiting for policy from trainer. Check log for errors."));
+				StopTraining();
+				UE_DEBUG_BREAK();
+				return;
+			}
+
+			// Mark all instances for reset since we have a new policy
+			ResetBuffer->SetResetInstances(SelectedAgentsSet);
+		}
+
+		// Reset Instances
+		ResetInstance(ResetBuffer->GetResetInstances().ToArray());
+
+		EpisodeBuffer->Reset(ResetBuffer->GetResetInstances());
+	}
+}
+
+void ULearningAgentsTrainer::ResetAllInstances()
+{
+	UE_LEARNING_TRACE_CPUPROFILER_EVENT_SCOPE(ULearningAgentsTrainer::ResetAllInstances);
+
+	if (!bSetupPerformed)
+	{
+		UE_LOG(LogLearning, Error, TEXT("Setup must be run before agents can be reset."));
+		return;
+	}
+
+	ResetInstance(SelectedAgentIds);
+
+	EpisodeBuffer->Reset(SelectedAgentsSet);
+}
+
+void ULearningAgentsTrainer::ResetInstance_Implementation(const TArray<int32>& AgentId)
+{
+	// Can be overridden to reset agent without blueprints
+}
diff --git a/Engine/Plugins/Experimental/LearningAgents/Source/LearningAgentsTraining/Private/LearningAgentsTypeComponent.cpp b/Engine/Plugins/Experimental/LearningAgents/Source/LearningAgentsTraining/Private/LearningAgentsTypeComponent.cpp
new file mode 100644
--- /dev/null
+++ b/Engine/Plugins/Experimental/LearningAgents/Source/LearningAgentsTraining/Private/LearningAgentsTypeComponent.cpp
@@ -0,0 +1,120 @@
+// Copyright Epic Games, Inc. All Rights Reserved.
+
+#include "LearningAgentsTypeComponent.h"
+
+#include "Engine/World.h"
+#include "LearningAgentsType.h"
+#include "LearningLog.h"
+#include "UObject/NameTypes.h"
+
+ULearningAgentsTypeComponent::ULearningAgentsTypeComponent()
+{
+	PrimaryComponentTick.bCanEverTick = true;
+	PrimaryComponentTick.bStartWithTickEnabled = false;
+}
+ULearningAgentsTypeComponent::ULearningAgentsTypeComponent(FVTableHelper& Helper) : ULearningAgentsTypeComponent() {}
+ULearningAgentsTypeComponent::~ULearningAgentsTypeComponent() {}
+
+void ULearningAgentsTypeComponent::OnRegister()
+{
+	Super::OnRegister();
+
+	if (!GetWorld()->IsGameWorld())
+	{
+		// We're not in a game yet so we don't need to register all the callbacks
+		return;
+	}
+
+	USceneComponent* Parent = GetAttachParent();
+	if (Parent)
+	{
+		AgentType = Cast<ULearningAgentsType>(Parent);
+		if (AgentType)
+		{
+			AgentType->GetOnSetupComplete().AddUObject(this, &ULearningAgentsTypeComponent::OnAgentTypeSetupComplete);
+			AgentType->GetOnAgentAdded().AddUObject(this, &ULearningAgentsTypeComponent::OnAgentAdded);
+			AgentType->GetOnAgentRemoved().AddUObject(this, &ULearningAgentsTypeComponent::OnAgentRemoved);
+			return;
+		}
+	}
+
+	UE_LOG(LogLearning, Warning,
+		TEXT("%s: Not attached to ULearningAgentType. OnAgentTypeSetupComplete/OnAgentRemoved will not be called. " 
+			"If you wish to use these events, make sure to attach this to an ULearningAgentType component."),
+		*GetName());
+}
+
+void ULearningAgentsTypeComponent::OnAgentTypeSetupComplete_Implementation()
+{
+	// Can be overridden in child class
+}
+
+void ULearningAgentsTypeComponent::OnAgentAdded_Implementation(int32 AgentId, UObject* Agent)
+{
+	AddAgent(AgentId);
+}
+
+void ULearningAgentsTypeComponent::AddAgent(int32 AgentId)
+{
+	if (AgentId < 0)
+	{
+		UE_LOG(LogLearning, Error, TEXT("Unable to add: AgentId must be a valid index. AgentId was %d"), AgentId);
+		return;
+	}
+
+	const UE::Learning::FIndexSet ExistingAgents = AgentType->GetOccupiedAgentSet();
+	if (!ExistingAgents.Contains(AgentId))
+	{
+		UE_LOG(LogLearning, Error,
+			TEXT("Unable to add: AgentId %d not found on AgentType. Make sure to add agents to the agent type before adding."),
+			AgentId);
+		return;
+	}
+
+	SelectedAgentIds.Add(AgentId);
+	UpdateAgentSet();
+}
+
+void ULearningAgentsTypeComponent::OnAgentRemoved_Implementation(int32 AgentId, UObject* Agent)
+{
+	RemoveAgent(AgentId);
+}
+
+void ULearningAgentsTypeComponent::RemoveAgent(int32 AgentId)
+{
+	if (AgentId < 0)
+	{
+		UE_LOG(LogLearning, Warning, TEXT("Unable to remove: AgentId must be a valid index. AgentId was %d"), AgentId);
+		return;
+	}
+
+	int32 RemovedCount = SelectedAgentIds.RemoveSingleSwap(AgentId, false);
+	if (RemovedCount == 0)
+	{
+		UE_LOG(LogLearning, Warning, TEXT("Unable to remove: AgentId %d not found in the added agents set."), AgentId);
+		return;
+	}
+
+	UpdateAgentSet();
+}
+
+const ULearningAgentsType& ULearningAgentsTypeComponent::GetAgentType()
+{
+	return *AgentType;
+}
+
+UObject* ULearningAgentsTypeComponent::GetAgent(int32 AgentId, TSubclassOf<UObject> AgentClass)
+{
+	return AgentType->GetAgent(AgentId, AgentClass);
+}
+
+const UObject* ULearningAgentsTypeComponent::GetAgent(int32 AgentId) const
+{
+	return AgentType->GetAgent(AgentId);
+}
+
+void ULearningAgentsTypeComponent::UpdateAgentSet()
+{
+	SelectedAgentsSet = SelectedAgentIds;
+	SelectedAgentsSet.TryMakeSlice();
+}
diff --git a/Engine/Plugins/Experimental/LearningAgents/Source/LearningAgentsTraining/Public/LearningAgentsCompletions.h b/Engine/Plugins/Experimental/LearningAgents/Source/LearningAgentsTraining/Public/LearningAgentsCompletions.h
new file mode 100644
--- /dev/null
+++ b/Engine/Plugins/Experimental/LearningAgents/Source/LearningAgentsTraining/Public/LearningAgentsCompletions.h
@@ -0,0 +1,85 @@
+// Copyright Epic Games, Inc. All Rights Reserved.
+
+#pragma once
+
+#include "LearningArray.h"
+#include "LearningAgentsTrainer.h" // Required for ELearningAgentsCompletion::Termination
+
+#include "Templates/SharedPointer.h"
+#include "UObject/Object.h"
+#include "EngineDefines.h"
+
+#include "LearningAgentsCompletions.generated.h"
+
+namespace UE::Learning
+{
+	struct FCompletionObject;
+	struct FConditionalCompletion;
+	struct FPlanarPositionDifferenceCompletion;
+}
+
+// For functions in this file, we are favoring having more verbose names such as "AddConditionalCompletion" vs simply "Add" in 
+// order to keep it easy to find the correct function in blueprints.
+
+UCLASS(Abstract, BlueprintType)
+class LEARNINGAGENTSTRAINING_API ULearningAgentsCompletion : public UObject
+{
+	GENERATED_BODY()
+
+public:
+
+#if ENABLE_VISUAL_LOG
+	/** Color used to draw this completion in the visual log */
+	FLinearColor VisualLogColor = FColor::Yellow;
+
+	virtual void VisualLog(const UE::Learning::FIndexSet Instances) const {}
+#endif
+};
+
+//------------------------------------------------------------------
+
+UCLASS()
+class LEARNINGAGENTSTRAINING_API UConditionalCompletion : public ULearningAgentsCompletion
+{
+	GENERATED_BODY()
+
+public:
+	UFUNCTION(BlueprintCallable, Category = "LearningAgents")
+	static UConditionalCompletion* AddConditionalCompletion(ULearningAgentsTrainer* AgentTrainer, FName Name = NAME_None, ELearningAgentsCompletion InCompletionMode = ELearningAgentsCompletion::Termination);
+
+	UFUNCTION(BlueprintCallable, Category = "LearningAgents")
+	void SetConditionalCompletion(int32 AgentId, bool bIsCompleted);
+
+#if ENABLE_VISUAL_LOG
+	virtual void VisualLog(const UE::Learning::FIndexSet Instances) const override;
+#endif
+
+	TSharedPtr<UE::Learning::FConditionalCompletion> CompletionObject;
+};
+
+//------------------------------------------------------------------
+
+UCLASS()
+class LEARNINGAGENTSTRAINING_API UPlanarPositionDifferenceCompletion : public ULearningAgentsCompletion
+{
+	GENERATED_BODY()
+
+public:
+	UFUNCTION(BlueprintCallable, Category = "LearningAgents")
+	static UPlanarPositionDifferenceCompletion* AddPlanarPositionDifferenceCompletion(
+		ULearningAgentsTrainer* AgentTrainer, 
+		FName Name = NAME_None,
+		float Threshold = 100.0f, 
+		ELearningAgentsCompletion InCompletionMode = ELearningAgentsCompletion::Termination,
+		FVector Axis0 = FVector::ForwardVector,
+		FVector Axis1 = FVector::RightVector);
+
+	UFUNCTION(BlueprintCallable, Category = "LearningAgents")
+	void SetPlanarPositionDifferenceCompletion(int32 AgentId, FVector Position0, FVector Position1);
+
+#if ENABLE_VISUAL_LOG
+	virtual void VisualLog(const UE::Learning::FIndexSet Instances) const override;
+#endif
+
+	TSharedPtr<UE::Learning::FPlanarPositionDifferenceCompletion> CompletionObject;
+};
diff --git a/Engine/Plugins/Experimental/LearningAgents/Source/LearningAgentsTraining/Public/LearningAgentsController.h b/Engine/Plugins/Experimental/LearningAgents/Source/LearningAgentsTraining/Public/LearningAgentsController.h
new file mode 100644
--- /dev/null
+++ b/Engine/Plugins/Experimental/LearningAgents/Source/LearningAgentsTraining/Public/LearningAgentsController.h
@@ -0,0 +1,36 @@
+// Copyright Epic Games, Inc. All Rights Reserved.
+
+#pragma once
+
+#include "LearningAgentsTypeComponent.h"
+
+#include "Containers/Array.h"
+#include "Engine/EngineTypes.h"
+#include "LearningArray.h"
+#include "UObject/Object.h"
+#include "UObject/ObjectPtr.h"
+
+#include "LearningAgentsController.generated.h"
+
+class ULearningAgentsType;
+
+UCLASS(Abstract, BlueprintType, Blueprintable)
+class LEARNINGAGENTSTRAINING_API ULearningAgentsController : public ULearningAgentsTypeComponent
+{
+	GENERATED_BODY()
+
+public:
+	// These constructors/destructors are needed to make forward declarations happy
+	ULearningAgentsController();
+	ULearningAgentsController(FVTableHelper& Helper);
+	virtual ~ULearningAgentsController();
+
+public:
+
+	UFUNCTION(BlueprintNativeEvent, Category = "LearningAgents")
+	void SetActions(const TArray<int32>& AgentIds);
+
+	UFUNCTION(BlueprintCallable, Category = "LearningAgents")
+	void EncodeActions();
+
+};
diff --git a/Engine/Plugins/Experimental/LearningAgents/Source/LearningAgentsTraining/Public/LearningAgentsDataStorage.h b/Engine/Plugins/Experimental/LearningAgents/Source/LearningAgentsTraining/Public/LearningAgentsDataStorage.h
new file mode 100644
--- /dev/null
+++ b/Engine/Plugins/Experimental/LearningAgents/Source/LearningAgentsTraining/Public/LearningAgentsDataStorage.h
@@ -0,0 +1,118 @@
+// Copyright Epic Games, Inc. All Rights Reserved.
+
+#pragma once
+
+#include "Components/ActorComponent.h"
+#include "Containers/Array.h"
+#include "LearningArray.h"
+#include "Misc/DateTime.h"
+#include "UObject/NameTypes.h"
+#include "UObject/Object.h"
+#include "UObject/ObjectPtr.h"
+
+#include "LearningAgentsDataStorage.generated.h"
+
+struct FDirectoryPath;
+class ULearningAgentsType;
+
+USTRUCT(BlueprintType)
+struct FRecordMetaData
+{
+	GENERATED_BODY()
+
+	UPROPERTY(BlueprintReadOnly, VisibleAnywhere, Category = "LearningAgents")
+	TObjectPtr<ULearningAgentsType> AgentType;
+
+	UPROPERTY(BlueprintReadOnly, VisibleAnywhere, Category = "LearningAgents")
+	FDateTime CreatedOn;
+
+	UPROPERTY(BlueprintReadOnly, VisibleAnywhere, Category = "LearningAgents")
+	FDateTime CreatedOnUtc;
+
+	UPROPERTY(BlueprintReadOnly, VisibleAnywhere, Category = "LearningAgents")
+	FName RecordName;
+
+	UPROPERTY(BlueprintReadOnly, VisibleAnywhere, Category = "LearningAgents")
+	bool bWasLoadedFromFile = false;
+
+	UPROPERTY(BlueprintReadOnly, VisibleAnywhere, Category = "LearningAgents")
+	bool bHasUnsavedChanges = false;
+};
+
+UCLASS(BlueprintType)
+class LEARNINGAGENTSTRAINING_API ULearningAgentsRecord : public UObject
+{
+	friend class ULearningAgentsDataStorage;
+
+	GENERATED_BODY()
+
+public:
+	// These constructors/destructors are needed to make forward declarations happy
+	ULearningAgentsRecord();
+	ULearningAgentsRecord(FVTableHelper& Helper);
+	virtual ~ULearningAgentsRecord();
+
+	void Init(int32 ObsNum, int32 ActNum);
+
+	void AddExperience(TLearningArrayView<1, const float> Observations, TLearningArrayView<1, const float> Actions);
+
+	 const TConstArrayView<TLearningArray<2, float>> GetObservations() const;
+	const TConstArrayView<TLearningArray<2, float>> GetActions() const;
+
+	void Trim();
+
+	UPROPERTY(BlueprintReadOnly, VisibleAnywhere, Category = "LearningAgents")
+	FRecordMetaData MetaData;
+
+private:
+	void AddChunk();
+
+	int32 DataIndex = 0;
+	int32 ChunkSize = 1024;
+
+	int32 ObservationNum = 0;
+	int32 ActionNum = 0;
+
+	TArray<TLearningArray<2, float>> Observations;
+	TArray<TLearningArray<2, float>> Actions;
+};
+
+UCLASS(BlueprintType, Blueprintable)
+class LEARNINGAGENTSTRAINING_API ULearningAgentsDataStorage : public UObject
+{
+	GENERATED_BODY()
+
+public:
+	// These constructors/destructors are needed to make forward declarations happy
+	ULearningAgentsDataStorage();
+	ULearningAgentsDataStorage(FVTableHelper& Helper);
+	virtual ~ULearningAgentsDataStorage();
+
+	UFUNCTION(BlueprintCallable, Category = "LearningAgents")
+	ULearningAgentsRecord* CreateRecord(FName RecordName, ULearningAgentsType* AgentType);
+
+	UFUNCTION(BlueprintPure, Category = "LearningAgents")
+	const TArray<ULearningAgentsRecord*>& GetAllRecords() const;
+
+	UFUNCTION(BlueprintCallable, Category = "LearningAgents")
+	int32 LoadAllRecords(ULearningAgentsType* AgentType, const FDirectoryPath& Directory);
+
+	UFUNCTION(BlueprintCallable, Category = "LearningAgents")
+	ULearningAgentsRecord* LoadRecord(ULearningAgentsType* AgentType, const FDirectoryPath& Directory, const FString& Filename);
+
+	UFUNCTION(BlueprintCallable, Category = "LearningAgents")
+	void SaveAllRecords(const FDirectoryPath& Directory) const;
+
+	UFUNCTION(BlueprintCallable, Category = "LearningAgents")
+	void SaveRecord(const FDirectoryPath& Directory, ULearningAgentsRecord* Record) const;
+
+private:
+	UPROPERTY(VisibleAnywhere, Category = "LearningAgents")
+	TArray<TObjectPtr<ULearningAgentsRecord>> Records;
+
+	UPROPERTY(EditDefaultsOnly, Category = "LearningAgents")
+	bool bPrependUtcTimeStamp = true;
+
+	UPROPERTY(EditDefaultsOnly, Category = "LearningAgents")
+	FString FileExtension = TEXT("bin");
+};
diff --git a/Engine/Plugins/Experimental/LearningAgents/Source/LearningAgentsTraining/Public/LearningAgentsImitationTrainer.h b/Engine/Plugins/Experimental/LearningAgents/Source/LearningAgentsTraining/Public/LearningAgentsImitationTrainer.h
new file mode 100644
--- /dev/null
+++ b/Engine/Plugins/Experimental/LearningAgents/Source/LearningAgentsTraining/Public/LearningAgentsImitationTrainer.h
@@ -0,0 +1,74 @@
+// Copyright Epic Games, Inc. All Rights Reserved.
+
+#pragma once
+
+#include "Engine/EngineTypes.h"
+#include "LearningAgentsTypeComponent.h"
+#include "Templates/SharedPointer.h"
+#include "UObject/Object.h"
+#include "UObject/ObjectPtr.h"
+#include "UObject/NameTypes.h"
+#include "Tasks/Task.h"
+#include "LearningArray.h"
+
+#include "LearningAgentsImitationTrainer.generated.h"
+
+namespace UE::Learning
+{
+	struct FArrayMap;
+	struct FNeuralNetwork;
+	struct FNeuralNetworkPolicyFunction;
+	struct FSharedMemoryImitationTrainer;
+}
+
+class ULearningAgentsType;
+
+UCLASS(Abstract, BlueprintType, Blueprintable)
+class LEARNINGAGENTSTRAINING_API ULearningAgentsImitationTrainer : public ULearningAgentsTypeComponent
+{
+	GENERATED_BODY()
+
+public:
+
+	// These constructors/destructors are needed to make forward declarations happy
+	ULearningAgentsImitationTrainer();
+	ULearningAgentsImitationTrainer(FVTableHelper& Helper);
+	virtual ~ULearningAgentsImitationTrainer();
+
+	virtual void EndPlay(const EEndPlayReason::Type EndPlayReason) override;
+
+public:
+
+	UFUNCTION(BlueprintPure, Category = "LearningAgents")
+	const bool IsTraining() const;
+
+	UFUNCTION(BlueprintCallable, Category = "LearningAgents")
+	void BeginTraining(TArray<UObject*> Records);
+
+	UFUNCTION(BlueprintCallable, Category = "LearningAgents")
+	void EndTraining();
+
+	UFUNCTION(BlueprintPure, Category = "LearningAgents")
+	const bool IsTrainingComplete() const;
+
+//** ----- Private ----- */
+private:
+
+	/** True if training is currently in-progress. Otherwise, false. */
+	UPROPERTY(VisibleAnywhere, Category = "LearningAgents")
+	bool bIsTraining = false;
+
+	UPROPERTY(VisibleAnywhere, Category = "LearningAgents")
+	bool bIsTrainingComplete = false;
+
+	TLearningArray<2, float> RecordedObservations;
+	TLearningArray<2, float> RecordedActions;
+
+	TUniquePtr<UE::Learning::FSharedMemoryImitationTrainer> ImitationTrainer;
+	UE::Tasks::FTask ImitationTrainingTask;
+	FRWLock NetworkLock;
+
+	TAtomic<bool> bRequestImitationTrainingStop = false;
+
+	float TrainerTimeout = 10.0f;
+};
diff --git a/Engine/Plugins/Experimental/LearningAgents/Source/LearningAgentsTraining/Public/LearningAgentsRecorder.h b/Engine/Plugins/Experimental/LearningAgents/Source/LearningAgentsTraining/Public/LearningAgentsRecorder.h
new file mode 100644
--- /dev/null
+++ b/Engine/Plugins/Experimental/LearningAgents/Source/LearningAgentsTraining/Public/LearningAgentsRecorder.h
@@ -0,0 +1,70 @@
+// Copyright Epic Games, Inc. All Rights Reserved.
+
+#pragma once
+
+#include "LearningAgentsTypeComponent.h"
+
+#include "Containers/Map.h"
+#include "Engine/EngineTypes.h"
+#include "UObject/Object.h"
+#include "UObject/ObjectPtr.h"
+
+#include "LearningAgentsRecorder.generated.h"
+
+class ULearningAgentsType;
+class ULearningAgentsDataStorage;
+class ULearningAgentsRecord;
+
+UCLASS(Abstract, BlueprintType, Blueprintable)
+class LEARNINGAGENTSTRAINING_API ULearningAgentsRecorder : public ULearningAgentsTypeComponent
+{
+	GENERATED_BODY()
+
+public:
+	// These constructors/destructors are needed to make forward declarations happy
+	ULearningAgentsRecorder();
+	ULearningAgentsRecorder(FVTableHelper& Helper);
+	virtual ~ULearningAgentsRecorder();
+
+	virtual void OnRegister() override;
+
+	virtual void EndPlay(const EEndPlayReason::Type EndPlayReason) override;
+
+//** ----- Recording Process ----- */
+public:
+
+	/** Adds experience to tracked agents' records. Each agent id whose record is full will be returned, e.g. so you can reset them. */
+	UFUNCTION(BlueprintCallable, Category = "LearningAgents")
+	void AddExperience();
+
+	UFUNCTION(BlueprintCallable, Category = "LearningAgents")
+	const bool IsRecording() const;
+
+	UFUNCTION(BlueprintCallable, Category = "LearningAgents")
+	void StartRecording();
+
+	UFUNCTION(BlueprintCallable, Category = "LearningAgents")
+	void StopRecording();
+
+//** ----- Recorder Configuration ----- */
+private:
+
+	/** Directory where records will be saved. If not set, BeginPlay will automatically set this to the editor's default intermediate folder. */
+	UPROPERTY(EditDefaultsOnly, Category = "LearningAgents")
+	FDirectoryPath DataDirectory;
+
+	/** If true, recorder will automatically save all records on EndPlay. Set this to false if you want to manually save records. */
+	UPROPERTY(EditAnywhere, Category = "LearningAgents")
+	bool bSaveDataOnEndPlay = true;
+
+//** ----- Recorder State ----- */
+private:
+
+	/** The data storage manager. It can be used to save/load agent records. */
+	UPROPERTY(VisibleAnywhere, Category = "LearningAgents")
+	TObjectPtr<ULearningAgentsDataStorage> DataStorage;
+
+	/** All records which are currently being written */
+	UPROPERTY(VisibleAnywhere, Category = "LearningAgents")
+	TMap<int32, TObjectPtr<ULearningAgentsRecord>> CurrentRecords;
+};
diff --git a/Engine/Plugins/Experimental/LearningAgents/Source/LearningAgentsTraining/Public/LearningAgentsRewards.h b/Engine/Plugins/Experimental/LearningAgents/Source/LearningAgentsTraining/Public/LearningAgentsRewards.h
new file mode 100644
--- /dev/null
+++ b/Engine/Plugins/Experimental/LearningAgents/Source/LearningAgentsTraining/Public/LearningAgentsRewards.h
@@ -0,0 +1,162 @@
+// Copyright Epic Games, Inc. All Rights Reserved.
+
+#pragma once
+
+#include "LearningArray.h"
+
+#include "Templates/SharedPointer.h"
+#include "UObject/Object.h"
+#include "EngineDefines.h"
+
+#include "LearningAgentsRewards.generated.h"
+
+namespace UE::Learning
+{
+	struct FRewardObject;
+	struct FFloatReward;
+	struct FScalarVelocityReward;
+	struct FLocalDirectionalVelocityReward;
+	struct FPlanarPositionDifferencePenalty;
+	struct FPositionDifferencePenalty;
+	struct FPositionArraySimilarityReward;
+}
+
+class ULearningAgentsTrainer;
+
+// For functions in this file, we are favoring having more verbose names such as "AddFloatReward" vs simply "Add" in 
+// order to keep it easy to find the correct function in blueprints.
+
+/** Base class for all rewards */
+UCLASS(Abstract, BlueprintType)
+class LEARNINGAGENTSTRAINING_API ULearningAgentsReward : public UObject
+{
+	GENERATED_BODY()
+
+public:
+
+#if ENABLE_VISUAL_LOG
+	/** Color used to draw this reward in the visual log */
+	FLinearColor VisualLogColor = FColor::Green;
+
+	virtual void VisualLog(const UE::Learning::FIndexSet Instances) const {}
+#endif
+};
+
+//------------------------------------------------------------------
+
+UCLASS()
+class LEARNINGAGENTSTRAINING_API UFloatReward : public ULearningAgentsReward
+{
+	GENERATED_BODY()
+
+public:
+	UFUNCTION(BlueprintCallable, Category = "LearningAgents")
+	static UFloatReward* AddFloatReward(ULearningAgentsTrainer* AgentTrainer, FName Name = NAME_None, float Weight = 1.0f);
+
+	UFUNCTION(BlueprintCallable, Category = "LearningAgents")
+	void SetFloatReward(int32 AgentId, float Reward);
+
+#if ENABLE_VISUAL_LOG
+	virtual void VisualLog(const UE::Learning::FIndexSet Instances) const override;
+#endif
+
+	TSharedPtr<UE::Learning::FFloatReward> RewardObject;
+};
+
+//------------------------------------------------------------------
+
+UCLASS()
+class LEARNINGAGENTSTRAINING_API UScalarVelocityReward : public ULearningAgentsReward
+{
+	GENERATED_BODY()
+
+public:
+	UFUNCTION(BlueprintCallable, Category = "LearningAgents")
+	static UScalarVelocityReward* AddScalarVelocityReward(ULearningAgentsTrainer* AgentTrainer, FName Name = NAME_None, float Weight = 1.0f, float Scale = 200.0f);
+
+	UFUNCTION(BlueprintCallable, Category = "LearningAgents")
+	void SetScalarVelocityReward(int32 AgentId, float Velocity);
+
+#if ENABLE_VISUAL_LOG
+	virtual void VisualLog(const UE::Learning::FIndexSet Instances) const override;
+#endif
+
+	TSharedPtr<UE::Learning::FScalarVelocityReward> RewardObject;
+};
+
+UCLASS()
+class LEARNINGAGENTSTRAINING_API ULocalDirectionalVelocityReward : public ULearningAgentsReward
+{
+	GENERATED_BODY()
+
+public:
+	UFUNCTION(BlueprintCallable, Category = "LearningAgents")
+	static ULocalDirectionalVelocityReward* AddLocalDirectionalVelocityReward(ULearningAgentsTrainer* AgentTrainer, FName Name = NAME_None, float Weight = 1.0f, float Scale = 200.0f, FVector Axis = FVector::ForwardVector);
+
+	UFUNCTION(BlueprintCallable, Category = "LearningAgents")
+	void SetLocalDirectionalVelocityReward(int32 AgentId, FVector Velocity, FRotator RelativeRotation = FRotator::ZeroRotator);
+
+#if ENABLE_VISUAL_LOG
+	virtual void VisualLog(const UE::Learning::FIndexSet Instances) const override;
+#endif
+
+	TSharedPtr<UE::Learning::FLocalDirectionalVelocityReward> RewardObject;
+};
+
+//------------------------------------------------------------------
+
+UCLASS()
+class LEARNINGAGENTSTRAINING_API UPlanarPositionDifferencePenalty : public ULearningAgentsReward
+{
+	GENERATED_BODY()
+
+public:
+	UFUNCTION(BlueprintCallable, Category = "LearningAgents")
+	static UPlanarPositionDifferencePenalty* AddPlanarPositionDifferencePenalty(
+		ULearningAgentsTrainer* AgentTrainer, 
+		FName Name = NAME_None,
+		float Weight = 1.0f, 
+		float Scale = 100.0f, 
+		float Threshold = 0.0f, 
+		FVector Axis0 = FVector::ForwardVector, 
+		FVector Axis1 = FVector::RightVector);
+
+	UFUNCTION(BlueprintCallable, Category = "LearningAgents")
+	void SetPlanarPositionDifferencePenalty(int32 AgentId, FVector Position0, FVector Position1);
+
+#if ENABLE_VISUAL_LOG
+	virtual void VisualLog(const UE::Learning::FIndexSet Instances) const override;
+#endif
+
+	TSharedPtr<UE::Learning::FPlanarPositionDifferencePenalty> RewardObject;
+};
+
+//------------------------------------------------------------------
+
+UCLASS()
+class LEARNINGAGENTSTRAINING_API UPositionArraySimilarityReward : public ULearningAgentsReward
+{
+	GENERATED_BODY()
+
+public:
+	UFUNCTION(BlueprintCallable, Category = "LearningAgents")
+	static UPositionArraySimilarityReward* AddPositionArraySimilarityReward(ULearningAgentsTrainer* AgentTrainer, FName Name = NAME_None, int32 PositionNum = 0, float Scale = 100.0f, float Weight = 1.0f);
+
+	UFUNCTION(BlueprintCallable, Category = "LearningAgents")
+	void SetPositionArraySimilarityReward(
+		int32 AgentId, 
+		const TArray<FVector>& Positions0, 
+		const TArray<FVector>& Positions1, 
+		FVector RelativePosition0 = FVector::ZeroVector,
+		FVector RelativePosition1 = FVector::ZeroVector,
+		FRotator RelativeRotation0 = FRotator::ZeroRotator,
+		FRotator RelativeRotation1 = FRotator::ZeroRotator);
+
+#if ENABLE_VISUAL_LOG
+	virtual void VisualLog(const UE::Learning::FIndexSet Instances) const override;
+#endif
+
+	TSharedPtr<UE::Learning::FPositionArraySimilarityReward> RewardObject;
+};
+
+//------------------------------------------------------------------
diff --git a/Engine/Plugins/Experimental/LearningAgents/Source/LearningAgentsTraining/Public/LearningAgentsTrainer.h b/Engine/Plugins/Experimental/LearningAgents/Source/LearningAgentsTraining/Public/LearningAgentsTrainer.h
new file mode 100644
--- /dev/null
+++ b/Engine/Plugins/Experimental/LearningAgents/Source/LearningAgentsTraining/Public/LearningAgentsTrainer.h
@@ -0,0 +1,290 @@
+// Copyright Epic Games, Inc. All Rights Reserved.
+
+#pragma once
+
+#include "Engine/EngineTypes.h"
+#include "LearningAgentsTypeComponent.h"
+#include "Templates/SharedPointer.h"
+#include "UObject/Object.h"
+#include "UObject/ObjectPtr.h"
+#include "UObject/NameTypes.h"
+
+#include "LearningAgentsTrainer.generated.h"
+
+namespace UE::Learning
+{
+	struct FAnyCompletion;
+	struct FArrayMap;
+	struct FEpisodeBuffer;
+	struct FNeuralNetwork;
+	struct FNeuralNetworkPolicyFunction;
+	struct FReplayBuffer;
+	struct FResetInstanceBuffer;
+	struct FRewardObject;
+	struct FSharedMemoryPPOTrainer;
+	struct FSumReward;
+	struct FCompletionObject;
+}
+
+class ULearningAgentsCompletion;
+class ULearningAgentsReward;
+class ULearningAgentsType;
+
+UENUM(BlueprintType, Category = "LearningAgents")
+enum class ELearningAgentsCompletion : uint8
+{
+	/** Episode ended early but was still in progress. Value function will be used to estimate final return. */
+	Truncation	UMETA(DisplayName = "Truncation"),
+
+	/** Episode ended early and zero reward was expected for all future steps. */
+	Termination	UMETA(DisplayName = "Termination"),
+};
+
+USTRUCT(BlueprintType, Category = "LearningAgents")
+struct FLearningAgentsTrainerSettings
+{
+	GENERATED_BODY()
+
+public:
+
+	/** Completion type to use when the maximum number of steps for an episode is reached. */
+	UPROPERTY(EditAnywhere, Category = "LearningAgents")
+	ELearningAgentsCompletion MaxStepsCompletion = ELearningAgentsCompletion::Truncation;
+
+	/** Max number of steps to take while training before episode automatically terminates */
+	UPROPERTY(EditAnywhere, Category = "LearningAgents", meta = (ClampMin = "1", UIMin = "1"))
+	int32 MaxStepNum = 300;
+
+	/** Maximum number of episodes to record before running a training iteration */
+	UPROPERTY(EditAnywhere, Category = "LearningAgents", meta = (ClampMin = "1", UIMin = "1"))
+	int32 MaximumRecordedEpisodesPerIteration = 1000;
+
+	/** Maximum number of steps to record before running a training iteration */
+	UPROPERTY(EditAnywhere, Category = "LearningAgents", meta = (ClampMin = "1", UIMin = "1"))
+	int32 MaximumRecordedStepsPerIteration = 10000;
+};
+
+USTRUCT(BlueprintType, Category="LearningAgents")
+struct FLearningAgentsTrainerGameSettings
+{
+	GENERATED_BODY()
+
+public:
+
+	UPROPERTY(EditAnywhere, Category = "LearningAgents")
+	bool bUseFixedTimeStep = true;
+
+	UPROPERTY(EditAnywhere, Category = "LearningAgents", meta = (DisplayName = "Fixed Time Step Frequency (Hz)"), meta = (ClampMin = "0.0", UIMin = "0.0"))
+	float FixedTimeStepFrequency = 60.0f;
+
+	UPROPERTY(EditAnywhere, Category = "LearningAgents")
+	bool bSetMaxPhysicsStepToFixedTimeStep = true;
+
+	UPROPERTY(EditAnywhere, Category = "LearningAgents")
+	bool bDisableVSync = true;
+
+	UPROPERTY(EditAnywhere, Category = "LearningAgents")
+	bool bUseUnlitViewportRendering = false;
+};
+
+UENUM(BlueprintType, Category = "LearningAgents")
+enum class ELearningAgentsTrainerDevice : uint8
+{
+	CPU,
+	GPU,
+};
+
+USTRUCT(BlueprintType, Category = "LearningAgents")
+struct FLearningAgentsTrainerTrainingSettings
+{
+	GENERATED_BODY()
+
+public:
+
+	UPROPERTY(EditAnywhere, Category = "LearningAgents", meta = (ClampMin = "0", UIMin = "0"))
+	int32 NumberOfIterations = 1000000;
+
+	UPROPERTY(EditAnywhere, Category = "LearningAgents")
+	bool bReinitializeNetwork = true;
+
+	UPROPERTY(EditAnywhere, Category = "LearningAgents")
+	bool bUseTensorboard = false;
+
+	UPROPERTY(EditAnywhere, Category = "LearningAgents", meta = (ClampMin = "0.0", UIMin = "0.0"))
+	float InitialActionScale = 0.1f;
+
+	UPROPERTY(EditAnywhere, Category = "LearningAgents", meta = (ClampMin = "0.0", UIMin = "0.0"))
+	float DiscountFactor = 0.99f;
+
+	UPROPERTY(EditAnywhere, Category = "LearningAgents", meta = (ClampMin = "0", UIMin = "0"))
+	int32 RandomSeed = 1234;
+
+	UPROPERTY(EditAnywhere, Category = "LearningAgents", meta = (ClampMin = "0", UIMin = "0"))
+	int32 NumberOfStepsToTrimAtStartOfEpisode = 0;
+
+	UPROPERTY(EditAnywhere, Category = "LearningAgents", meta = (ClampMin = "0", UIMin = "0"))
+	int32 NumberOfStepsToTrimAtEndOfEpisode = 0;
+
+	UPROPERTY(EditAnywhere, Category = "LearningAgents")
+	ELearningAgentsTrainerDevice Device = ELearningAgentsTrainerDevice::GPU;
+
+	UPROPERTY(EditAnywhere, Category = "LearningAgents", meta = (ClampMin = "0.0", UIMin = "0.0"))
+	float TrainerCommunicationTimeout = 20.0f;
+};
+
+
+UCLASS(Abstract, BlueprintType, Blueprintable)
+class LEARNINGAGENTSTRAINING_API ULearningAgentsTrainer : public ULearningAgentsTypeComponent
+{
+	GENERATED_BODY()
+
+public:
+
+	// These constructors/destructors are needed to make forward declarations happy
+	ULearningAgentsTrainer();
+	ULearningAgentsTrainer(FVTableHelper& Helper);
+	virtual ~ULearningAgentsTrainer();
+
+	virtual void OnAgentAdded_Implementation(int32 AgentId, UObject* Agent) override;
+
+	virtual void EndPlay(const EEndPlayReason::Type EndPlayReason) override;
+
+// ----- Setup -----
+public:
+
+	/**
+	* Initializes this object and runs the setup functions for observations and actions.
+	* If the AgentType has not yet been setup, this call will fail, so consider using the
+	* OnAgentTypeSetupComplete event to ensure the timing is correct.
+	* 
+	* @param AgentTypeOverride This agent type will override any parent agent type we found. Must be provided if this trainer component is not attached to a parent agent type component.
+	* @param Settings The settings used to setup this trainer
+	* @see OnAgentTypeSetupComplete
+	*/
+	UFUNCTION(BlueprintCallable, Category = "LearningAgents")
+	void SetupTrainer(
+		ULearningAgentsType* AgentTypeOverride,
+		const FLearningAgentsTrainerSettings& Settings = FLearningAgentsTrainerSettings());
+
+	UFUNCTION(BlueprintPure, Category = "LearningAgents")
+	const bool IsSetupPerformed() const;
+
+//** ----- Rewards ----- */
+public:
+	
+	/**
+	 * Event where all rewards should be set up for this agent trainer.
+	 * 
+	 * @param AgentTrainer This AgentTrainer. This is a convenience for blueprints so you don't have to find the Self pin.
+	 */
+	UFUNCTION(BlueprintNativeEvent, Category = "LearningAgents")
+	void SetupRewards(ULearningAgentsTrainer* AgentTrainer);
+
+	UFUNCTION(BlueprintNativeEvent, Category = "LearningAgents")
+	void SetRewards(const TArray<int32>& AgentIds);
+
+	void AddReward(TObjectPtr<ULearningAgentsReward> Object, const TSharedRef<UE::Learning::FRewardObject>& Reward);
+
+//** ----- Completions ----- */
+public:
+
+	/**
+	 * Event where all completions should be set up for this agent trainer.
+	 * 
+	 * @param AgentTrainer This AgentTrainer. This is a convenience for blueprints so you don't have to find the Self pin.
+	 */
+	UFUNCTION(BlueprintNativeEvent, Category = "LearningAgents")
+	void SetupCompletions(ULearningAgentsTrainer* AgentTrainer);
+
+	UFUNCTION(BlueprintNativeEvent, Category = "LearningAgents")
+	void SetCompletions(const TArray<int32>& AgentIds);
+
+	void AddCompletion(TObjectPtr<ULearningAgentsCompletion> Object, const TSharedRef<UE::Learning::FCompletionObject>& Completion);
+
+//** ----- Resets ----- */
+public:
+
+	/**
+	 * Event where both the agent and its training environment should be reset so that a new training episode could
+	 * commence.
+	 *
+	 * @param AgentId The id of the current agent to be reset
+	 * @param Agent A pointer to the current agent to be reset
+	 */
+	UFUNCTION(BlueprintNativeEvent, Category = "LearningAgents")
+	void ResetInstance(const TArray<int32>& AgentIds);
+
+//** ----- Training Process ----- */
+public:
+
+	UFUNCTION(BlueprintPure, Category = "LearningAgents")
+	const bool IsTraining() const;
+
+	UFUNCTION(BlueprintCallable, Category = "LearningAgents")
+	void StartTraining(
+		const FLearningAgentsTrainerTrainingSettings& TrainingSettings = FLearningAgentsTrainerTrainingSettings(),
+		const FLearningAgentsTrainerGameSettings& GameSettings = FLearningAgentsTrainerGameSettings());
+
+	UFUNCTION(BlueprintCallable, Category = "LearningAgents")
+	void StopTraining();
+
+	UFUNCTION(BlueprintCallable, Category = "LearningAgents")
+	void EvaluateRewards();
+
+	UFUNCTION(BlueprintCallable, Category = "LearningAgents")
+	void EvaluateCompletions();
+
+	UFUNCTION(BlueprintCallable, Category = "LearningAgents")
+	void IterateTraining();
+
+	/**
+	* Manually reset all agents. Does not record the experience gathered up to this point by each agent.
+	*/
+	UFUNCTION(BlueprintCallable, Category = "LearningAgents")
+	void ResetAllInstances();
+
+//** ----- Private ----- */
+private:
+
+	UPROPERTY(VisibleAnywhere, Category = "LearningAgents")
+	bool bSetupPerformed = false;
+
+	/** True if training is currently in-progress. Otherwise, false. */
+	UPROPERTY(VisibleAnywhere, Category = "LearningAgents")
+	bool bIsTraining = false;
+
+	UPROPERTY(VisibleAnywhere, Category = "LearningAgents")
+	TArray<TObjectPtr<ULearningAgentsReward>> RewardObjects;
+
+	UPROPERTY(VisibleAnywhere, Category = "LearningAgents")
+	TArray<TObjectPtr<ULearningAgentsCompletion>> CompletionObjects;
+
+	TArray<TSharedRef<UE::Learning::FRewardObject>, TInlineAllocator<16>> RewardFeatures;
+	TSharedPtr<UE::Learning::FSumReward> Rewards;
+
+	TArray<TSharedRef<UE::Learning::FCompletionObject>, TInlineAllocator<16>> CompletionFeatures;
+	TSharedPtr<UE::Learning::FAnyCompletion> Completions;
+
+	TUniquePtr<UE::Learning::FEpisodeBuffer> EpisodeBuffer;
+	TUniquePtr<UE::Learning::FReplayBuffer> ReplayBuffer;
+	TUniquePtr<UE::Learning::FResetInstanceBuffer> ResetBuffer;
+	TUniquePtr<UE::Learning::FSharedMemoryPPOTrainer> Trainer;
+
+	/** One-time Error Logging State */
+	bool bPreviouslyLogged = false;
+
+	ELearningAgentsCompletion MaxStepsCompletion = ELearningAgentsCompletion::Truncation;
+
+	float TrainerTimeout = 10.0f;
+
+	void DoneTraining();
+
+//** ----- Private Recording of GameSettings ----- */
+private:
+
+	bool bFixedTimestepUsed = false;
+	float FixedTimeStepDeltaTime = -1.0f;
+	bool bVSyncEnabled = true;
+	float MaxPhysicsStep = -1.0f;
+	int32 ViewModeIndex = -1;
+};
diff --git a/Engine/Plugins/Experimental/LearningAgents/Source/LearningAgentsTraining/Public/LearningAgentsTypeComponent.h b/Engine/Plugins/Experimental/LearningAgents/Source/LearningAgentsTraining/Public/LearningAgentsTypeComponent.h
new file mode 100644
--- /dev/null
+++ b/Engine/Plugins/Experimental/LearningAgents/Source/LearningAgentsTraining/Public/LearningAgentsTypeComponent.h
@@ -0,0 +1,71 @@
+// Copyright Epic Games, Inc. All Rights Reserved.
+
+#pragma once
+
+#include "Components/SceneComponent.h"
+#include "Containers/Array.h"
+#include "LearningArray.h"
+#include "UObject/Object.h"
+#include "UObject/ObjectPtr.h"
+
+#include "LearningAgentsTypeComponent.generated.h"
+
+class ULearningAgentsType;
+
+UCLASS(Abstract, BlueprintType, Blueprintable)
+class LEARNINGAGENTSTRAINING_API ULearningAgentsTypeComponent : public USceneComponent
+{
+	GENERATED_BODY()
+
+public:
+
+	// These constructors/destructors are needed to make forward declarations happy
+	ULearningAgentsTypeComponent();
+	ULearningAgentsTypeComponent(FVTableHelper& Helper);
+	virtual ~ULearningAgentsTypeComponent();
+
+	virtual void OnRegister() override;
+
+	const ULearningAgentsType& GetAgentType();
+
+// ----- Blueprint Accessible Interface -----
+public:
+
+	UFUNCTION(BlueprintNativeEvent, Category = "LearningAgents")
+	void OnAgentTypeSetupComplete();
+
+	UFUNCTION(BlueprintNativeEvent, Category = "LearningAgents")
+	void OnAgentAdded(int32 AgentId, UObject* Agent);
+
+	UFUNCTION(BlueprintCallable, Category = "LearningAgents")
+	void AddAgent(int32 AgentId);
+
+	UFUNCTION(BlueprintNativeEvent, Category = "LearningAgents")
+	void OnAgentRemoved(int32 AgentId, UObject* Agent);
+
+	UFUNCTION(BlueprintCallable, Category = "LearningAgents")
+	void RemoveAgent(int32 AgentId);
+
+	UFUNCTION(BlueprintPure, Category = "LearningAgents", meta = (DeterminesOutputType = "AgentClass"))
+	UObject* GetAgent(int32 AgentId, TSubclassOf<UObject> AgentClass);
+
+	const UObject* GetAgent(int32 AgentId) const;
+
+// ----- Protected Blueprint Members -----
+protected:
+
+	/** The agent type this component is attached to. */
+	UPROPERTY(BlueprintReadOnly, VisibleAnywhere, Category = "LearningAgents")
+	TObjectPtr<ULearningAgentsType> AgentType;
+
+	/** The valid agent ids that this component has selected. */
+	UPROPERTY(BlueprintReadOnly, VisibleAnywhere, Category = "LearningAgents")
+	TArray<int32> SelectedAgentIds;
+
+// ----- Protected Non-blueprint Members -----
+protected:
+
+	void UpdateAgentSet();
+
+	UE::Learning::FIndexSet SelectedAgentsSet;
+};
