Learning Agents: Create initial experimental version

diff --git a/Engine/Plugins/Experimental/LearningAgents/Source/LearningTraining/Private/LearningExperience.cpp b/Engine/Plugins/Experimental/LearningAgents/Source/LearningTraining/Private/LearningExperience.cpp
new file mode 100644
--- /dev/null
+++ b/Engine/Plugins/Experimental/LearningAgents/Source/LearningTraining/Private/LearningExperience.cpp
@@ -0,0 +1,367 @@
+// Copyright Epic Games, Inc. All Rights Reserved.
+
+#include "LearningExperience.h"
+
+#include "LearningCompletion.h"
+#include "LearningLog.h"
+
+namespace UE::Learning
+{
+	void FEpisodeBuffer::Resize(
+		const int32 InMaxInstanceNum,
+		const int32 InMaxStepNum,
+		const int32 InObservationVectorDimNum,
+		const int32 InActionVectorDimNum)
+	{
+		MaxInstanceNum = InMaxInstanceNum;
+		MaxStepNum = InMaxStepNum;
+		EpisodeStepNums.SetNumUninitialized({ InMaxInstanceNum });
+		Observations.SetNumUninitialized({ InMaxInstanceNum, InMaxStepNum, InObservationVectorDimNum });
+		Actions.SetNumUninitialized({ InMaxInstanceNum, InMaxStepNum, InActionVectorDimNum });
+		Rewards.SetNumUninitialized({ InMaxInstanceNum, InMaxStepNum });
+
+		Array::Zero(EpisodeStepNums);
+	}
+
+	void FEpisodeBuffer::Push(
+		const TLearningArrayView<2, const float> InObservations,
+		const TLearningArrayView<2, const float> InActions,
+		const TLearningArrayView<1, const float> InRewards,
+		const FIndexSet Instances)
+	{
+		UE_LEARNING_TRACE_CPUPROFILER_EVENT_SCOPE(Learning::FEpisodeBuffer::Push);
+
+		for (const int32 InstanceIdx : Instances)
+		{
+			UE_LEARNING_CHECKF(EpisodeStepNums[InstanceIdx] < MaxStepNum, TEXT("Episode Buffer full!"));
+
+			Array::Copy(Observations[InstanceIdx][EpisodeStepNums[InstanceIdx]], InObservations[InstanceIdx]);
+			Array::Copy(Actions[InstanceIdx][EpisodeStepNums[InstanceIdx]], InActions[InstanceIdx]);
+			Rewards[InstanceIdx][EpisodeStepNums[InstanceIdx]] = InRewards[InstanceIdx];
+
+			EpisodeStepNums[InstanceIdx]++;
+		}
+	}
+
+	void FEpisodeBuffer::Reset(const FIndexSet Instances)
+	{
+		Array::Zero(EpisodeStepNums, Instances);
+	}
+
+	const int32 FEpisodeBuffer::GetMaxInstanceNum() const
+	{
+		return MaxInstanceNum;
+	}
+
+	const int32 FEpisodeBuffer::GetMaxStepNum() const
+	{
+		return MaxStepNum;
+	}
+
+	const TLearningArrayView<1, const int32> FEpisodeBuffer::GetEpisodeStepNums() const
+	{
+		return EpisodeStepNums;
+	}
+
+	const TLearningArrayView<2, const float> FEpisodeBuffer::GetObservations(const int32 InstanceIdx) const
+	{
+		return Observations[InstanceIdx].Slice(0, EpisodeStepNums[InstanceIdx]);
+	}
+
+	const TLearningArrayView<2, const float> FEpisodeBuffer::GetActions(const int32 InstanceIdx) const
+	{
+		return Actions[InstanceIdx].Slice(0, EpisodeStepNums[InstanceIdx]);
+	}
+
+	const TLearningArrayView<1, const float> FEpisodeBuffer::GetRewards(const int32 InstanceIdx) const
+	{
+		return Rewards[InstanceIdx].Slice(0, EpisodeStepNums[InstanceIdx]);
+	}
+
+
+	void FReplayBuffer::Resize(
+		const int32 InObservationVectorDimensionNum,
+		const int32 InActionVectorDimensionNum,
+		const int32 InMaxEpisodeNum,
+		const int32 InMaxStepNum)
+	{
+		MaxEpisodeNum = InMaxEpisodeNum;
+		MaxStepNum = InMaxStepNum;
+		EpisodeNum = 0;
+		StepNum = 0;
+		EpisodeStarts.SetNumUninitialized({ InMaxEpisodeNum });
+		EpisodeLengths.SetNumUninitialized({ InMaxEpisodeNum });
+		EpisodeCompletionModes.SetNumUninitialized({ InMaxEpisodeNum });
+		EpisodeFinalObservations.SetNumUninitialized({ InMaxEpisodeNum, InObservationVectorDimensionNum });
+		Observations.SetNumUninitialized({ InMaxStepNum, InObservationVectorDimensionNum });
+		Actions.SetNumUninitialized({ InMaxStepNum, InActionVectorDimensionNum });
+		Rewards.SetNumUninitialized({ InMaxStepNum });
+	}
+
+	void FReplayBuffer::Reset()
+	{
+		EpisodeNum = 0;
+		StepNum = 0;
+	}
+
+	bool FReplayBuffer::AddEpisodes(
+		const TLearningArrayView<1, const ECompletionMode> InEpisodeCompletionModes,
+		const TLearningArrayView<2, const float> InEpisodeFinalObservations,
+		const FEpisodeBuffer& EpisodeBuffer,
+		const FIndexSet Instances)
+	{
+		UE_LEARNING_TRACE_CPUPROFILER_EVENT_SCOPE(Learning::FReplayBuffer::AddEpisodes);
+
+		for (const int32 InstanceIdx : Instances)
+		{
+			const ECompletionMode EpisodeCompletion = InEpisodeCompletionModes[InstanceIdx];
+			const int32 EpisodeStepNum = EpisodeBuffer.GetEpisodeStepNums()[InstanceIdx];
+
+			UE_LEARNING_CHECKF(EpisodeCompletion != ECompletionMode::Running,
+				TEXT("Tried to add experience from episode that is still running..."));
+
+			// If there is still space in the buffer...
+			if (EpisodeNum + 1 <= MaxEpisodeNum && StepNum + EpisodeStepNum <= MaxStepNum)
+			{
+				// Copy the data into the replay buffer
+				Array::Copy(Observations.Slice(StepNum, EpisodeStepNum), EpisodeBuffer.GetObservations(InstanceIdx));
+				Array::Copy(Actions.Slice(StepNum, EpisodeStepNum), EpisodeBuffer.GetActions(InstanceIdx));
+				Array::Copy(Rewards.Slice(StepNum, EpisodeStepNum), EpisodeBuffer.GetRewards(InstanceIdx));
+
+				// Write the Episode start, length, completion, and final observation
+				EpisodeStarts[EpisodeNum] = StepNum;
+				EpisodeLengths[EpisodeNum] = EpisodeStepNum;
+				EpisodeCompletionModes[EpisodeNum] = EpisodeCompletion;
+				Array::Copy(EpisodeFinalObservations[EpisodeNum], InEpisodeFinalObservations[InstanceIdx]);
+
+				// Increment the Counts
+				EpisodeNum++;
+				StepNum += EpisodeStepNum;
+			}
+			else
+			{
+				// Buffer is full
+				return true;
+			}
+		}
+
+		// Buffer is not full
+		return false;
+	}
+
+	const int32 FReplayBuffer::GetMaxEpisodeNum() const
+	{
+		return MaxEpisodeNum;
+	}
+
+	const int32 FReplayBuffer::GetMaxStepNum() const
+	{
+		return MaxStepNum;
+	}
+
+	const int32 FReplayBuffer::GetEpisodeNum() const
+	{
+		return EpisodeNum;
+	}
+
+	const int32 FReplayBuffer::GetStepNum() const
+	{
+		return StepNum;
+	}
+
+	const TLearningArrayView<1, const int32> FReplayBuffer::GetEpisodeStarts() const
+	{
+		return EpisodeStarts.Slice(0, EpisodeNum);
+	}
+
+	const TLearningArrayView<1, const int32> FReplayBuffer::GetEpisodeLengths() const
+	{
+		return EpisodeLengths.Slice(0, EpisodeNum);
+	}
+
+	const TLearningArrayView<1, const ECompletionMode> FReplayBuffer::GetEpisodeCompletionModes() const
+	{
+		return EpisodeCompletionModes.Slice(0, EpisodeNum);
+	}
+
+	const TLearningArrayView<2, const float> FReplayBuffer::GetEpisodeFinalObservations() const
+	{
+		return EpisodeFinalObservations.Slice(0, EpisodeNum);
+	}
+
+	const TLearningArrayView<2, const float> FReplayBuffer::GetObservations() const
+	{
+		return Observations.Slice(0, StepNum);
+	}
+
+	const TLearningArrayView<2, const float> FReplayBuffer::GetActions() const
+	{
+		return Actions.Slice(0, StepNum);
+	}
+
+	const TLearningArrayView<1, const float> FReplayBuffer::GetRewards() const
+	{
+		return Rewards.Slice(0, StepNum);
+	}
+
+	namespace Experience
+	{
+		void GatherExperienceReset(
+			FReplayBuffer& ReplayBuffer,
+			FEpisodeBuffer& EpisodeBuffer,
+			FResetInstanceBuffer& ResetBuffer,
+			const FIndexSet Instances)
+		{
+			ReplayBuffer.Reset();
+			EpisodeBuffer.Reset(Instances);
+			ResetBuffer.SetResetInstances(Instances);
+		}
+
+		bool GatherExperienceIteration(
+			FReplayBuffer& ReplayBuffer,
+			FEpisodeBuffer& EpisodeBuffer,
+			FResetInstanceBuffer& ResetBuffer,
+			TLearningArrayView<2, float> ObservationVectorBuffer,
+			TLearningArrayView<2, float> ActionVectorBuffer,
+			TLearningArrayView<1, float> RewardBuffer,
+			TLearningArrayView<1, ECompletionMode> CompletionBuffer,
+			const ECompletionMode EpisodeEndCompletionMode,
+			const TFunctionRef<void(const FIndexSet Instances)> ResetFunction,
+			const TFunctionRef<void(const FIndexSet Instances)> ObservationFunction,
+			const TFunctionRef<void(const FIndexSet Instances)> PolicyFunction,
+			const TFunctionRef<void(const FIndexSet Instances)> ActionFunction,
+			const TFunctionRef<void(const FIndexSet Instances)> UpdateFunction,
+			const TFunctionRef<void(const FIndexSet Instances)> RewardFunction,
+			const TFunctionRef<void(const FIndexSet Instances)> CompletionFunction,
+			const FIndexSet Instances)
+		{
+			UE_LEARNING_TRACE_CPUPROFILER_EVENT_SCOPE(Learning::Experience::GatherExperienceIteration);
+
+			if (ResetBuffer.GetResetInstanceNum() > 0)
+			{
+				// Reset Environment
+
+				ResetFunction(ResetBuffer.GetResetInstances());
+			}
+
+			// Encode Observations
+
+			ObservationFunction(Instances);
+			Array::Check(ObservationVectorBuffer, Instances);
+
+			// Evaluate Policy
+
+			PolicyFunction(Instances);
+
+			// Decode Actions
+
+			Array::Check(ActionVectorBuffer, Instances);
+			ActionFunction(Instances);
+
+			// Update Environment
+
+			UpdateFunction(Instances);
+
+			// Compute Rewards
+
+			RewardFunction(Instances);
+			Array::Check(RewardBuffer, Instances);
+
+			// Push to Experience Buffer
+
+			EpisodeBuffer.Push(
+				ObservationVectorBuffer,
+				ActionVectorBuffer,
+				RewardBuffer,
+				Instances);
+
+			// Evaluate Completions
+
+			CompletionFunction(Instances);
+
+			Completion::EvaluateEndOfEpisodeCompletions(
+				CompletionBuffer,
+				EpisodeBuffer.GetEpisodeStepNums(),
+				EpisodeBuffer.GetMaxStepNum(),
+				EpisodeEndCompletionMode,
+				Instances);
+
+			ResetBuffer.SetResetInstancesFromCompletions(CompletionBuffer, Instances);
+
+			// Evaluate Observations again for instances that are completed
+
+			if (ResetBuffer.GetResetInstanceNum() > 0)
+			{
+				ObservationFunction(ResetBuffer.GetResetInstances());
+
+				// Push to Replay Buffer
+
+				const bool bReplayBufferFull = ReplayBuffer.AddEpisodes(
+					CompletionBuffer,
+					ObservationVectorBuffer,
+					EpisodeBuffer,
+					ResetBuffer.GetResetInstances());
+
+				EpisodeBuffer.Reset(ResetBuffer.GetResetInstances());
+
+				return bReplayBufferFull;
+			}
+			else
+			{
+				return false;
+			}
+		}
+
+		void GatherExperienceUntilReplayBufferFull(
+			FReplayBuffer& ReplayBuffer,
+			FEpisodeBuffer& EpisodeBuffer,
+			FResetInstanceBuffer& ResetBuffer,
+			TLearningArrayView<2, float> ObservationVectorBuffer,
+			TLearningArrayView<2, float> ActionVectorBuffer,
+			TLearningArrayView<1, float> RewardBuffer,
+			TLearningArrayView<1, ECompletionMode> CompletionBuffer,
+			const ECompletionMode EpisodeEndCompletionMode,
+			const TFunctionRef<void(const FIndexSet Instances)> ResetFunction,
+			const TFunctionRef<void(const FIndexSet Instances)> ObservationFunction,
+			const TFunctionRef<void(const FIndexSet Instances)> PolicyFunction,
+			const TFunctionRef<void(const FIndexSet Instances)> ActionFunction,
+			const TFunctionRef<void(const FIndexSet Instances)> UpdateFunction,
+			const TFunctionRef<void(const FIndexSet Instances)> RewardFunction,
+			const TFunctionRef<void(const FIndexSet Instances)> CompletionFunction,
+			const FIndexSet Instances)
+		{
+			UE_LEARNING_TRACE_CPUPROFILER_EVENT_SCOPE(Learning::Experience::GatherExperienceUntilReplayBufferFull);
+
+			bool bReplayBufferFull = false;
+
+			GatherExperienceReset(ReplayBuffer, EpisodeBuffer, ResetBuffer, Instances);
+
+			while (true)
+			{
+				bReplayBufferFull = GatherExperienceIteration(
+					ReplayBuffer,
+					EpisodeBuffer,
+					ResetBuffer,
+					ObservationVectorBuffer,
+					ActionVectorBuffer,
+					RewardBuffer,
+					CompletionBuffer,
+					EpisodeEndCompletionMode,
+					ResetFunction,
+					ObservationFunction,
+					PolicyFunction,
+					ActionFunction,
+					UpdateFunction,
+					RewardFunction,
+					CompletionFunction,
+					Instances);
+
+				if (bReplayBufferFull)
+				{
+					break;
+				}
+			}
+		}
+
+	}
+}
diff --git a/Engine/Plugins/Experimental/LearningAgents/Source/LearningTraining/Private/LearningImitationTrainer.cpp b/Engine/Plugins/Experimental/LearningAgents/Source/LearningTraining/Private/LearningImitationTrainer.cpp
new file mode 100644
--- /dev/null
+++ b/Engine/Plugins/Experimental/LearningAgents/Source/LearningTraining/Private/LearningImitationTrainer.cpp
@@ -0,0 +1,729 @@
+// Copyright Epic Games, Inc. All Rights Reserved.
+
+#include "LearningImitationTrainer.h"
+
+#include "LearningArray.h"
+#include "LearningLog.h"
+#include "LearningNeuralNetwork.h"
+#include "LearningProgress.h"
+#include "LearningSharedMemory.h"
+#include "LearningSharedMemoryTraining.h"
+#include "LearningSocketTraining.h"
+
+#include "Misc/MonitoredProcess.h"
+#include "Misc/Guid.h"
+#include "Misc/FileHelper.h"
+
+#include "Dom/JsonObject.h"
+#include "Serialization/JsonSerializer.h"
+#include "Serialization/JsonWriter.h"
+
+#include "Sockets.h"
+#include "Common/TcpSocketBuilder.h"
+#include "SocketSubsystem.h"
+
+ULearningSocketImitationTrainerServerCommandlet::ULearningSocketImitationTrainerServerCommandlet(const FObjectInitializer& ObjectInitializer)
+	: Super(ObjectInitializer)
+{}
+
+int32 ULearningSocketImitationTrainerServerCommandlet::Main(const FString& Commandline)
+{
+	UE_LOG(LogLearning, Display, TEXT("Running Imitation Training Server Commandlet..."));
+
+	TArray<FString> Tokens;
+	TArray<FString> Switches;
+	TMap<FString, FString> Params;
+
+	UCommandlet::ParseCommandLine(*Commandline, Tokens, Switches, Params);
+
+	const FString* PythonExecutiblePathParam = Params.Find(TEXT("PythonExecutiblePath"));
+	const FString* SitePackagesPathParam = Params.Find(TEXT("SitePackagesPath"));
+	const FString* PythonContentPathParam = Params.Find(TEXT("PythonContentPath"));
+	const FString* IntermediatePathParam = Params.Find(TEXT("IntermediatePath"));
+	const FString* IpAddressParam = Params.Find(TEXT("IpAddress"));
+	const FString* PortParam = Params.Find(TEXT("Port"));
+	const FString* LogSettingsParam = Params.Find(TEXT("LogSettings"));
+
+#if WITH_EDITOR
+	const FString PythonExecutiblePath = PythonExecutiblePathParam ? *PythonExecutiblePathParam : UE::Learning::Trainer::DefaultEditorPythonExecutablePath();
+	const FString SitePackagesPath = SitePackagesPathParam ? *SitePackagesPathParam : UE::Learning::Trainer::DefaultEditorSitePackagesPath();
+	const FString PythonContentPath = PythonContentPathParam ? *PythonContentPathParam : UE::Learning::Trainer::DefaultEditorPythonContentPath();
+	const FString IntermediatePath = IntermediatePathParam ? *IntermediatePathParam : UE::Learning::Trainer::DefaultEditorIntermediatePath();
+#else
+	UE_LEARNING_NOT_IMPLEMENTED();
+	const FString PythonExecutiblePath = TEXT("");
+	const FString SitePackagesPath = TEXT("");
+	const FString PythonContentPath = TEXT("");
+	const FString IntermediatePath = TEXT("");
+	return 0;
+#endif
+	const TCHAR* IpAddress = IpAddressParam ? *(*IpAddressParam) : UE::Learning::Trainer::DefaultIp;
+	const uint32 Port = PortParam ? FCString::Atoi(*(*PortParam)) : UE::Learning::Trainer::DefaultPort;
+
+	UE::Learning::ELogSetting LogSettings = UE::Learning::ELogSetting::Normal;
+	if (LogSettingsParam)
+	{
+		if (*LogSettingsParam == TEXT("Normal"))
+		{
+			LogSettings = UE::Learning::ELogSetting::Normal;
+		}
+		else if (*LogSettingsParam == TEXT("Silent"))
+		{
+			LogSettings = UE::Learning::ELogSetting::Silent;
+		}
+		else
+		{
+			UE_LEARNING_NOT_IMPLEMENTED();
+			return 1;
+		}
+	}
+
+	UE_LOG(LogLearning, Display, TEXT("---  Imitation Training Server Arguments ---"));
+	UE_LOG(LogLearning, Display, TEXT("PythonExecutiblePath: %s"), *PythonExecutiblePath);
+	UE_LOG(LogLearning, Display, TEXT("SitePackagesPath: %s"), *SitePackagesPath);
+	UE_LOG(LogLearning, Display, TEXT("PythonContentPath: %s"), *PythonContentPath);
+	UE_LOG(LogLearning, Display, TEXT("IntermediatePath: %s"), *IntermediatePath);
+	UE_LOG(LogLearning, Display, TEXT("IpAddress: %s"), IpAddress);
+	UE_LOG(LogLearning, Display, TEXT("Port: %i"), Port);
+	UE_LOG(LogLearning, Display, TEXT("LogSettings: %s"), LogSettings == UE::Learning::ELogSetting::Normal ? TEXT("Normal") : TEXT("Silent"));
+
+	UE::Learning::FSocketImitationTrainerServerProcess ServerProcess(
+		PythonExecutiblePath,
+		SitePackagesPath,
+		PythonContentPath,
+		IntermediatePath,
+		IpAddress,
+		Port,
+		LogSettings);
+
+	while (ServerProcess.IsRunning())
+	{
+		FPlatformProcess::Sleep(0.01f);
+	}
+
+	return 0;
+}
+
+namespace UE::Learning
+{
+	FSharedMemoryImitationTrainer::FSharedMemoryImitationTrainer(
+		const FString& TaskName,
+		const FString& PythonExecutiblePath,
+		const FString& SitePackagesPath,
+		const FString& PythonContentPath,
+		const FString& IntermediatePath,
+		const FNeuralNetwork& Network,
+		const float ActionNoiseMin,
+		const float ActionNoiseMax,
+		const int32 MaxSampleNum,
+		const FImitationTrainerSettings& Settings,
+		const ELogSetting LogSettings)
+	{
+		UE_LEARNING_CHECK(Network.GetOutputNum() % 2 == 0);
+
+		const int32 ObservationVectorDimensionNum = Network.GetInputNum();
+		const int32 ActionVectorDimensionNum = Network.GetOutputNum() / 2;
+
+		// Allocate Shared Memory
+
+		Policy = SharedMemory::Allocate<1, uint8>({ Network.GetTotalByteNum() });
+		Controls = SharedMemory::Allocate<1, volatile int32>({ SharedMemoryTraining::GetControlNum() });
+		Observations = SharedMemory::Allocate<2, float>({ MaxSampleNum, ObservationVectorDimensionNum });
+		Actions = SharedMemory::Allocate<2, float>({ MaxSampleNum, ActionVectorDimensionNum });
+
+		// We need to zero the control memory before we start
+		// the training sub-process since it may contain uninitialized 
+		// values or those left over from previous runs.
+
+		Array::Zero(Controls.View);
+
+		// Write Config
+
+		const FString TimeStamp = FDateTime::Now().ToFormattedString(TEXT("%Y-%m-%d_%H-%M-%S"));
+		const FString TrainerMethod = TEXT("Imitation");
+		const FString TrainerType = TEXT("SharedMemory");
+		const FString ConfigPath = IntermediatePath / TEXT("Configs") / FString::Printf(TEXT("%s_%s_%s_%s.json"), *TaskName, *TrainerMethod, *TrainerType, *TimeStamp);
+
+		TSharedRef<FJsonObject> ConfigObject = MakeShared<FJsonObject>();
+		ConfigObject->SetStringField(TEXT("TaskName"), *TaskName);
+		ConfigObject->SetStringField(TEXT("TrainerMethod"), TrainerMethod);
+		ConfigObject->SetStringField(TEXT("TrainerType"), TrainerType);
+		ConfigObject->SetStringField(TEXT("TimeStamp"), *TimeStamp);
+
+		ConfigObject->SetStringField(TEXT("SitePackagesPath"), *SitePackagesPath);
+		ConfigObject->SetStringField(TEXT("IntermediatePath"), *IntermediatePath);
+
+		ConfigObject->SetStringField(TEXT("PolicyGuid"), *Policy.Guid.ToString(EGuidFormats::DigitsWithHyphensInBraces));
+		ConfigObject->SetStringField(TEXT("ControlsGuid"), *Controls.Guid.ToString(EGuidFormats::DigitsWithHyphensInBraces));
+		ConfigObject->SetStringField(TEXT("ObservationsGuid"), *Observations.Guid.ToString(EGuidFormats::DigitsWithHyphensInBraces));
+		ConfigObject->SetStringField(TEXT("ActionsGuid"), *Actions.Guid.ToString(EGuidFormats::DigitsWithHyphensInBraces));
+
+		ConfigObject->SetNumberField(TEXT("ObservationVectorDimensionNum"), ObservationVectorDimensionNum);
+		ConfigObject->SetNumberField(TEXT("ActionVectorDimensionNum"), ActionVectorDimensionNum);
+		ConfigObject->SetNumberField(TEXT("NetworkTotalByteNum"), Network.GetTotalByteNum());
+		ConfigObject->SetNumberField(TEXT("MaxSampleNum"), MaxSampleNum);
+		ConfigObject->SetNumberField(TEXT("HiddenUnitNum"), Network.GetHiddenNum());
+		ConfigObject->SetNumberField(TEXT("LayerNum"), Network.GetLayerNum());
+		ConfigObject->SetStringField(TEXT("ActivationFunction"), ActivationFunctionString(Network.ActivationFunction));
+
+		ConfigObject->SetNumberField(TEXT("ActionNoiseMin"), ActionNoiseMin);
+		ConfigObject->SetNumberField(TEXT("ActionNoiseMax"), ActionNoiseMax);
+
+		ConfigObject->SetNumberField(TEXT("IterationNum"), Settings.IterationNum);
+		ConfigObject->SetNumberField(TEXT("LearningRateActor"), Settings.LearningRateActor);
+		ConfigObject->SetNumberField(TEXT("LearningRateDecay"), Settings.LearningRateDecay);
+		ConfigObject->SetNumberField(TEXT("WeightDecay"), Settings.WeightDecay);
+		ConfigObject->SetNumberField(TEXT("BatchSize"), Settings.BatchSize);
+		ConfigObject->SetNumberField(TEXT("Seed"), Settings.Seed);
+		ConfigObject->SetBoolField(TEXT("ReinitializeNetwork"), Settings.bReinitializeNetwork);
+		ConfigObject->SetStringField(TEXT("Device"), Trainer::GetDeviceString(Settings.Device));
+		ConfigObject->SetBoolField(TEXT("UseTensorBoard"), Settings.bUseTensorboard);
+		ConfigObject->SetBoolField(TEXT("LoggingEnabled"), LogSettings == ELogSetting::Silent ? false : true);
+
+		FString JsonString;
+		TSharedRef<TJsonWriter<>> JsonWriter = TJsonWriterFactory<>::Create(&JsonString, 0);
+		FJsonSerializer::Serialize(ConfigObject, JsonWriter, true);
+
+		FFileHelper::SaveStringToFile(JsonString, *ConfigPath);
+
+		// Start Python Sub-process
+
+		const FString CommandLineArguments = FString::Printf(TEXT("\"%s\" SharedMemory \"%s\""), *(PythonContentPath / TEXT("train_imitation.py")), *ConfigPath);
+
+		TrainingProcess = MakeShared<FMonitoredProcess>(PythonExecutiblePath, CommandLineArguments, true, true);
+		TrainingProcess->OnCanceled().BindRaw(this, &FSharedMemoryImitationTrainer::HandleTrainingProcessCanceled);
+		TrainingProcess->OnCompleted().BindRaw(this, &FSharedMemoryImitationTrainer::HandleTrainingProcessCompleted);
+		TrainingProcess->OnOutput().BindStatic(&FSharedMemoryImitationTrainer::HandleTrainingProcessOutput);
+		TrainingProcess->Launch();
+	}
+
+	FSharedMemoryImitationTrainer::~FSharedMemoryImitationTrainer()
+	{
+		Terminate();
+	}
+
+	ETrainerResponse FSharedMemoryImitationTrainer::SendStop(const float Timeout)
+	{
+		return SharedMemoryTraining::SendStop(Controls.View);
+	}
+
+	ETrainerResponse FSharedMemoryImitationTrainer::SendContinue(const float Timeout)
+	{
+		return SharedMemoryTraining::SendContinue(Controls.View);
+	}
+
+	ETrainerResponse FSharedMemoryImitationTrainer::RecvPolicy(
+		FNeuralNetwork& OutNetwork,
+		const float Timeout,
+		FRWLock* NetworkLock,
+		const ELogSetting LogSettings)
+	{
+		return SharedMemoryTraining::RecvPolicy(
+			Controls.View,
+			OutNetwork,
+			Policy.View,
+			Timeout,
+			NetworkLock,
+			LogSettings);
+	}
+
+	ETrainerResponse FSharedMemoryImitationTrainer::SendPolicy(
+		const FNeuralNetwork& Network,
+		const float Timeout,
+		FRWLock* NetworkLock,
+		const ELogSetting LogSettings)
+	{
+		return SharedMemoryTraining::SendPolicy(
+			Controls.View,
+			Policy.View,
+			Network,
+			Timeout,
+			NetworkLock,
+			LogSettings);
+	}
+
+	ETrainerResponse FSharedMemoryImitationTrainer::SendExperience(
+		const TLearningArrayView<2, const float> ObservationVectors,
+		const TLearningArrayView<2, const float> ActionVectors,
+		const float Timeout,
+		const ELogSetting LogSettings)
+	{
+		return SharedMemoryTraining::SendExperience(
+			Observations.View,
+			Actions.View,
+			Controls.View,
+			ObservationVectors,
+			ActionVectors,
+			Timeout,
+			LogSettings);
+	}
+
+	void FSharedMemoryImitationTrainer::Deallocate()
+	{
+		SharedMemory::Deallocate(Policy);
+		SharedMemory::Deallocate(Controls);
+		SharedMemory::Deallocate(Observations);
+		SharedMemory::Deallocate(Actions);
+	}
+
+	ETrainerResponse FSharedMemoryImitationTrainer::Wait(float Timeout)
+	{
+		const float SleepTime = 0.001f;
+		float WaitTime = 0.0f;
+
+		while (TrainingProcess.IsValid())
+		{
+			FPlatformProcess::Sleep(SleepTime);
+			WaitTime += SleepTime;
+
+			if (WaitTime > Timeout)
+			{
+				return ETrainerResponse::Timeout;
+			}
+		}
+
+		return ETrainerResponse::Success;
+	}
+
+	void FSharedMemoryImitationTrainer::Terminate()
+	{
+		if (TrainingProcess.IsValid())
+		{
+			TrainingProcess->Cancel(true);
+		}
+
+		TrainingProcess.Reset();
+
+		if (Policy.Region)
+		{
+			Deallocate();
+		}
+	}
+
+	void FSharedMemoryImitationTrainer::HandleTrainingProcessCanceled()
+	{
+		UE_LOG(LogLearning, Warning, TEXT("Training process canceled"));
+
+		TrainingProcess.Reset();
+	}
+
+	void FSharedMemoryImitationTrainer::HandleTrainingProcessCompleted(int32 ReturnCode)
+	{
+		if (ReturnCode != 0)
+		{
+			UE_LOG(LogLearning, Warning, TEXT("Training Process finished with warnings or errors"));
+		}
+
+		TrainingProcess.Reset();
+	}
+
+	void FSharedMemoryImitationTrainer::HandleTrainingProcessOutput(FString Output)
+	{
+		if (!Output.IsEmpty())
+		{
+			UE_LOG(LogLearning, Display, TEXT("Training Process: %s"), *Output);
+		}
+	}
+
+
+	FSocketImitationTrainerServerProcess::FSocketImitationTrainerServerProcess(
+		const FString& PythonExecutiblePath,
+		const FString& SitePackagesPath,
+		const FString& PythonContentPath,
+		const FString& IntermediatePath,
+		const TCHAR* IpAddress,
+		const uint32 Port,
+		const ELogSetting LogSettings)
+	{
+		const FString CommandLineArguments = FString::Printf(TEXT("\"%s\" Socket \"%s:%i\" \"%s\" \"%s\" %i"),
+			*(PythonContentPath / TEXT("train_imitation.py")), IpAddress, Port, *SitePackagesPath, *IntermediatePath, LogSettings == ELogSetting::Normal ? 1 : 0);
+
+		TrainingProcess = MakeShared<FMonitoredProcess>(PythonExecutiblePath, CommandLineArguments, true, true);
+		TrainingProcess->OnCanceled().BindRaw(this, &FSocketImitationTrainerServerProcess::HandleTrainingProcessCanceled);
+		TrainingProcess->OnCompleted().BindRaw(this, &FSocketImitationTrainerServerProcess::HandleTrainingProcessCompleted);
+		TrainingProcess->OnOutput().BindStatic(&FSocketImitationTrainerServerProcess::HandleTrainingProcessOutput);
+		TrainingProcess->Launch();
+	}
+
+	FSocketImitationTrainerServerProcess::~FSocketImitationTrainerServerProcess()
+	{
+		Terminate();
+	}
+
+	bool FSocketImitationTrainerServerProcess::IsRunning() const
+	{
+		return TrainingProcess.IsValid();
+	}
+
+	bool FSocketImitationTrainerServerProcess::Wait(float Timeout)
+	{
+		const float SleepTime = 0.001f;
+		float WaitTime = 0.0f;
+
+		while (TrainingProcess.IsValid())
+		{
+			FPlatformProcess::Sleep(SleepTime);
+			WaitTime += SleepTime;
+
+			if (WaitTime > Timeout)
+			{
+				return false;
+			}
+		}
+
+		return true;
+	}
+
+	void FSocketImitationTrainerServerProcess::Terminate()
+	{
+		if (TrainingProcess.IsValid())
+		{
+			TrainingProcess->Cancel(true);
+		}
+
+		TrainingProcess.Reset();
+	}
+
+	void FSocketImitationTrainerServerProcess::HandleTrainingProcessCanceled()
+	{
+		UE_LOG(LogLearning, Warning, TEXT("Training process canceled"));
+
+		TrainingProcess.Reset();
+	}
+
+	void FSocketImitationTrainerServerProcess::HandleTrainingProcessCompleted(int32 ReturnCode)
+	{
+		if (ReturnCode != 0)
+		{
+			UE_LOG(LogLearning, Warning, TEXT("Training Process finished with warnings or errors"));
+		}
+
+		TrainingProcess.Reset();
+	}
+
+	void FSocketImitationTrainerServerProcess::HandleTrainingProcessOutput(FString Output)
+	{
+		if (!Output.IsEmpty())
+		{
+			UE_LOG(LogLearning, Display, TEXT("Training Process: %s"), *Output);
+		}
+	}
+
+	FSocketImitationTrainer::FSocketImitationTrainer(
+		ETrainerResponse& OutResponse,
+		const FString& TaskName,
+		const FNeuralNetwork& Network,
+		const float ActionNoiseMin,
+		const float ActionNoiseMax,
+		const int32 MaxSampleNum,
+		const TCHAR* IpAddress,
+		const uint32 Port,
+		const float Timeout,
+		const FImitationTrainerSettings& Settings)
+	{
+		UE_LEARNING_CHECK(Network.GetOutputNum() % 2 == 0);
+
+		const int32 ObservationVectorDimensionNum = Network.GetInputNum();
+		const int32 ActionVectorDimensionNum = Network.GetOutputNum() / 2;
+
+		// Write Config
+
+		const FString TimeStamp = FDateTime::Now().ToFormattedString(TEXT("%Y-%m-%d_%H-%M-%S"));
+		const FString TrainerMethod = TEXT("RL");
+		const FString TrainerType = TEXT("Network");
+
+		TSharedRef<FJsonObject> ConfigObject = MakeShared<FJsonObject>();
+		ConfigObject->SetStringField(TEXT("TaskName"), *TaskName);
+		ConfigObject->SetStringField(TEXT("TrainerMethod"), TrainerMethod);
+		ConfigObject->SetStringField(TEXT("TrainerType"), TrainerType);
+		ConfigObject->SetStringField(TEXT("TimeStamp"), *TimeStamp);
+
+		ConfigObject->SetNumberField(TEXT("ObservationVectorDimensionNum"), ObservationVectorDimensionNum);
+		ConfigObject->SetNumberField(TEXT("ActionVectorDimensionNum"), ActionVectorDimensionNum);
+		ConfigObject->SetNumberField(TEXT("NetworkTotalByteNum"), Network.GetTotalByteNum());
+		ConfigObject->SetNumberField(TEXT("MaxSampleNum"), MaxSampleNum);
+		ConfigObject->SetNumberField(TEXT("HiddenUnitNum"), Network.GetHiddenNum());
+		ConfigObject->SetNumberField(TEXT("LayerNum"), Network.GetLayerNum());
+		ConfigObject->SetStringField(TEXT("ActivationFunction"), ActivationFunctionString(Network.ActivationFunction));
+
+		ConfigObject->SetNumberField(TEXT("ActionNoiseMin"), ActionNoiseMin);
+		ConfigObject->SetNumberField(TEXT("ActionNoiseMax"), ActionNoiseMax);
+
+		ConfigObject->SetNumberField(TEXT("IterationNum"), Settings.IterationNum);
+		ConfigObject->SetNumberField(TEXT("LearningRateActor"), Settings.LearningRateActor);
+		ConfigObject->SetNumberField(TEXT("LearningRateDecay"), Settings.LearningRateDecay);
+		ConfigObject->SetNumberField(TEXT("WeightDecay"), Settings.WeightDecay);
+		ConfigObject->SetNumberField(TEXT("BatchSize"), Settings.BatchSize);
+		ConfigObject->SetNumberField(TEXT("Seed"), Settings.Seed);
+		ConfigObject->SetBoolField(TEXT("ReinitializeNetwork"), Settings.bReinitializeNetwork);
+		ConfigObject->SetStringField(TEXT("Device"), Trainer::GetDeviceString(Settings.Device));
+		ConfigObject->SetBoolField(TEXT("UseTensorBoard"), Settings.bUseTensorboard);
+
+		FString JsonString;
+		TSharedRef<TJsonWriter<>> JsonWriter = TJsonWriterFactory<>::Create(&JsonString, 0);
+		FJsonSerializer::Serialize(ConfigObject, JsonWriter, true);
+
+		// Allocate buffer to receive network data in
+
+		NetworkBuffer.SetNumUninitialized({ Network.GetTotalByteNum() });
+
+		// Create Socket
+
+		ISocketSubsystem* SocketSubsystem = ISocketSubsystem::Get(PLATFORM_SOCKETSUBSYSTEM);
+		if (!SocketSubsystem)
+		{
+			UE_LOG(LogLearning, Error, TEXT("Could not get socket subsystem"));
+			OutResponse = ETrainerResponse::Unexpected;
+			return;
+		}
+
+		bool bIsValid = false;
+		TSharedRef<FInternetAddr> Address = SocketSubsystem->CreateInternetAddr();
+		Address->SetIp(IpAddress, bIsValid);
+		Address->SetPort(Port);
+
+		if (!bIsValid)
+		{
+			UE_LOG(LogLearning, Error, TEXT("Invalid Ip Address \"%s\"..."), IpAddress);
+			OutResponse = ETrainerResponse::Unexpected;
+			return;
+		}
+
+		// Connect  to Server
+
+		Socket = FTcpSocketBuilder(TEXT("LearningNetworkRLTrainerSocket")).AsNonBlocking().Build();
+		Socket->Connect(*Address);
+
+		OutResponse = SocketTraining::WaitForConnection(*Socket, Timeout);
+		if (OutResponse != ETrainerResponse::Success) { return; }
+
+
+		// Send Config
+
+		OutResponse = SocketTraining::SendConfig(*Socket, JsonString, Timeout);
+		return;
+	}
+
+	FSocketImitationTrainer::~FSocketImitationTrainer()
+	{
+		Terminate();
+	}
+
+	ETrainerResponse FSocketImitationTrainer::Wait(const float Timeout)
+	{
+		return ETrainerResponse::Success;
+	}
+
+	void FSocketImitationTrainer::Terminate()
+	{
+		if (Socket)
+		{
+			Socket->Close();
+			Socket = nullptr;
+		}
+	}
+
+	ETrainerResponse FSocketImitationTrainer::SendStop(const float Timeout)
+	{
+		return SocketTraining::SendStop(*Socket, Timeout);
+	}
+
+	ETrainerResponse FSocketImitationTrainer::SendContinue(const float Timeout)
+	{
+		return SocketTraining::SendContinue(*Socket, Timeout);
+	}
+
+	ETrainerResponse FSocketImitationTrainer::RecvPolicy(
+		FNeuralNetwork& OutNetwork,
+		const float Timeout,
+		FRWLock* NetworkLock,
+		const ELogSetting LogSettings)
+	{
+		return SocketTraining::RecvPolicy(*Socket, OutNetwork, NetworkBuffer, Timeout, NetworkLock, LogSettings);
+	}
+
+	ETrainerResponse FSocketImitationTrainer::SendPolicy(
+		const FNeuralNetwork& Network,
+		const float Timeout,
+		FRWLock* NetworkLock,
+		const ELogSetting LogSettings)
+	{
+		return SocketTraining::SendPolicy(*Socket, NetworkBuffer, Network, Timeout, NetworkLock, LogSettings);
+	}
+
+	ETrainerResponse FSocketImitationTrainer::SendExperience(
+		const TLearningArrayView<2, const float> ObservationVectors,
+		const TLearningArrayView<2, const float> ActionVectors,
+		const float Timeout,
+		const ELogSetting LogSettings)
+	{
+		return SocketTraining::SendExperience(*Socket, ObservationVectors, ActionVectors, Timeout, LogSettings);
+	}
+
+	namespace ImitationTrainer
+	{
+		ETrainerResponse Train(
+			IImitationTrainer& Trainer,
+			FNeuralNetwork& Network,
+			const TLearningArrayView<2, const float> ObservationVectors,
+			const TLearningArrayView<2, const float> ActionVectors,
+			const bool bReinitializeNetwork,
+			TAtomic<bool>* bRequestTrainingStopSignal,
+			FRWLock* NetworkLock,
+			TAtomic<bool>* bNetworkUpdatedSignal,
+			const ELogSetting LogSettings)
+		{
+			ETrainerResponse Response = ETrainerResponse::Success;
+
+			if (bReinitializeNetwork)
+			{
+				// Receive initial Policy
+
+				if (LogSettings != ELogSetting::Silent)
+				{
+					UE_LOG(LogLearning, Display, TEXT("Receiving initial Policy..."));
+				}
+
+				Response = Trainer.RecvPolicy(Network, 20.0f, NetworkLock);
+
+				if (Response != ETrainerResponse::Success)
+				{
+					if (LogSettings != ELogSetting::Silent)
+					{
+						UE_LOG(LogLearning, Error, TEXT("Error receiving initial policy from trainer: %s. Check log for errors."), Trainer::ResponseString(Response));
+					}
+
+					Trainer.Terminate();
+					return Response;
+				}
+			}
+			else
+			{
+				// Send initial Policy
+
+				if (LogSettings != ELogSetting::Silent)
+				{
+					UE_LOG(LogLearning, Display, TEXT("Sending initial Policy..."));
+				}
+
+				Response = Trainer.SendPolicy(Network, 20.0f, NetworkLock);
+
+				if (Response != ETrainerResponse::Success)
+				{
+					if (LogSettings != ELogSetting::Silent)
+					{
+						UE_LOG(LogLearning, Error, TEXT("Error sending initial policy from trainer: %s. Check log for errors."), Trainer::ResponseString(Response));
+					}
+
+					Trainer.Terminate();
+					return Response;
+				}
+			}
+
+			// Send Experience
+
+			Response = Trainer.SendExperience(ObservationVectors, ActionVectors, 10.0f);
+
+			if (Response != ETrainerResponse::Success)
+			{
+				if (LogSettings != ELogSetting::Silent)
+				{
+					UE_LOG(LogLearning, Error, TEXT("Error sending experience to trainer: %s. Check log for errors."), Trainer::ResponseString(Response));
+				}
+
+				Trainer.Terminate();
+				return Response;
+			}
+
+			// Start Training Loop
+
+			while (true)
+			{
+				if (bRequestTrainingStopSignal && (*bRequestTrainingStopSignal))
+				{
+					*bRequestTrainingStopSignal = false;
+
+					if (LogSettings != ELogSetting::Silent)
+					{
+						UE_LOG(LogLearning, Display, TEXT("Stopping Training..."));
+					}
+
+					Response = Trainer.SendStop();
+
+					if (Response != ETrainerResponse::Success)
+					{
+						if (LogSettings != ELogSetting::Silent)
+						{
+							UE_LOG(LogLearning, Error, TEXT("Error sending stop signal to trainer: %s. Check log for errors."), Trainer::ResponseString(Response));
+						}
+
+						Trainer.Terminate();
+						return Response;
+					}
+
+					break;
+				}
+				else
+				{
+					Response = Trainer.SendContinue();
+
+					if (Response != ETrainerResponse::Success)
+					{
+						if (LogSettings != ELogSetting::Silent)
+						{
+							UE_LOG(LogLearning, Error, TEXT("Error sending continue signal to trainer: %s. Check log for errors."), Trainer::ResponseString(Response));
+						}
+
+						Trainer.Terminate();
+						return Response;
+					}
+				}
+
+				Response = Trainer.RecvPolicy(Network, 10.0f, NetworkLock);
+
+				if (Response == ETrainerResponse::Completed)
+				{
+					if (LogSettings != ELogSetting::Silent)
+					{
+						UE_LOG(LogLearning, Display, TEXT("Trainer completed training."));
+					}
+					break;
+				}
+				else if (Response != ETrainerResponse::Success)
+				{
+					if (LogSettings != ELogSetting::Silent)
+					{
+						UE_LOG(LogLearning, Error, TEXT("Error receiving policy from trainer: %s. Check log for errors."), Trainer::ResponseString(Response));
+					}
+					break;
+				}
+
+				if (bNetworkUpdatedSignal)
+				{
+					*bNetworkUpdatedSignal = true;
+				}
+			}
+
+			// Allow some time for trainer to shut down gracefully before we kill it...
+
+			Response = Trainer.Wait(5.0f);
+
+			if (Response != ETrainerResponse::Success)
+			{
+				if (LogSettings != ELogSetting::Silent)
+				{
+					UE_LOG(LogLearning, Error, TEXT("Error waiting for trainer to exit: %s. Check log for errors."), Trainer::ResponseString(Response));
+				}
+			}
+
+			Trainer.Terminate();
+
+			if (LogSettings != ELogSetting::Silent)
+			{
+				UE_LOG(LogLearning, Display, TEXT("Training Task Done!"));
+			}
+
+			return ETrainerResponse::Success;
+		}
+	}
+}
diff --git a/Engine/Plugins/Experimental/LearningAgents/Source/LearningTraining/Private/LearningPPOTrainer.cpp b/Engine/Plugins/Experimental/LearningAgents/Source/LearningTraining/Private/LearningPPOTrainer.cpp
new file mode 100644
--- /dev/null
+++ b/Engine/Plugins/Experimental/LearningAgents/Source/LearningTraining/Private/LearningPPOTrainer.cpp
@@ -0,0 +1,897 @@
+// Copyright Epic Games, Inc. All Rights Reserved.
+
+#include "LearningPPOTrainer.h"
+
+#include "LearningArray.h"
+#include "LearningLog.h"
+#include "LearningNeuralNetwork.h"
+#include "LearningExperience.h"
+#include "LearningProgress.h"
+#include "LearningSharedMemory.h"
+#include "LearningSharedMemoryTraining.h"
+#include "LearningSocketTraining.h"
+
+#include "Misc/MonitoredProcess.h"
+#include "Misc/Guid.h"
+#include "Misc/FileHelper.h"
+#include "Misc/CommandLine.h"
+#include "Misc/Parse.h"
+
+#include "Dom/JsonObject.h"
+#include "Serialization/JsonSerializer.h"
+#include "Serialization/JsonWriter.h"
+
+#include "Sockets.h"
+#include "Common/TcpSocketBuilder.h"
+#include "SocketSubsystem.h"
+
+ULearningSocketPPOTrainerServerCommandlet::ULearningSocketPPOTrainerServerCommandlet(const FObjectInitializer& ObjectInitializer)
+	: Super(ObjectInitializer)
+{}
+
+int32 ULearningSocketPPOTrainerServerCommandlet::Main(const FString& Commandline)
+{
+	UE_LOG(LogLearning, Display, TEXT("Running PPO Training Server Commandlet..."));
+
+	TArray<FString> Tokens;
+	TArray<FString> Switches;
+	TMap<FString, FString> Params;
+
+	UCommandlet::ParseCommandLine(*Commandline, Tokens, Switches, Params);
+
+	const FString* PythonExecutiblePathParam = Params.Find(TEXT("PythonExecutiblePath"));
+	const FString* SitePackagesPathParam = Params.Find(TEXT("SitePackagesPath"));
+	const FString* PythonContentPathParam = Params.Find(TEXT("PythonContentPath"));
+	const FString* IntermediatePathParam = Params.Find(TEXT("IntermediatePath"));
+	const FString* IpAddressParam = Params.Find(TEXT("IpAddress"));
+	const FString* PortParam = Params.Find(TEXT("Port"));
+	const FString* LogSettingsParam = Params.Find(TEXT("LogSettings"));
+
+#if WITH_EDITOR
+	const FString PythonExecutiblePath = PythonExecutiblePathParam ? *PythonExecutiblePathParam : UE::Learning::Trainer::DefaultEditorPythonExecutablePath();
+	const FString SitePackagesPath = SitePackagesPathParam ? *SitePackagesPathParam : UE::Learning::Trainer::DefaultEditorSitePackagesPath();
+	const FString PythonContentPath = PythonContentPathParam ? *PythonContentPathParam : UE::Learning::Trainer::DefaultEditorPythonContentPath();
+	const FString IntermediatePath = IntermediatePathParam ? *IntermediatePathParam : UE::Learning::Trainer::DefaultEditorIntermediatePath();
+#else
+	UE_LEARNING_NOT_IMPLEMENTED();
+	const FString PythonExecutiblePath = TEXT("");
+	const FString SitePackagesPath = TEXT("");
+	const FString PythonContentPath = TEXT("");
+	const FString IntermediatePath = TEXT("");
+	return 0;
+#endif
+	const TCHAR* IpAddress = IpAddressParam ? *(*IpAddressParam) : UE::Learning::Trainer::DefaultIp;
+	const uint32 Port = PortParam ? FCString::Atoi(*(*PortParam)) : UE::Learning::Trainer::DefaultPort;
+	
+	UE::Learning::ELogSetting LogSettings = UE::Learning::ELogSetting::Normal;
+	if (LogSettingsParam)
+	{
+		if (*LogSettingsParam == TEXT("Normal"))
+		{
+			LogSettings = UE::Learning::ELogSetting::Normal;
+		}
+		else if (*LogSettingsParam == TEXT("Silent"))
+		{
+			LogSettings = UE::Learning::ELogSetting::Silent;
+		}
+		else
+		{
+			UE_LEARNING_NOT_IMPLEMENTED();
+			return 1;
+		}
+	}
+	
+	UE_LOG(LogLearning, Display, TEXT("---  PPO Training Server Arguments ---"));
+	UE_LOG(LogLearning, Display, TEXT("PythonExecutiblePath: %s"), *PythonExecutiblePath);
+	UE_LOG(LogLearning, Display, TEXT("SitePackagesPath: %s"), *SitePackagesPath);
+	UE_LOG(LogLearning, Display, TEXT("PythonContentPath: %s"), *PythonContentPath);
+	UE_LOG(LogLearning, Display, TEXT("IntermediatePath: %s"), *IntermediatePath);
+	UE_LOG(LogLearning, Display, TEXT("IpAddress: %s"), IpAddress);
+	UE_LOG(LogLearning, Display, TEXT("Port: %i"), Port);
+	UE_LOG(LogLearning, Display, TEXT("LogSettings: %s"), LogSettings == UE::Learning::ELogSetting::Normal ? TEXT("Normal") : TEXT("Silent"));
+
+	UE::Learning::FSocketPPOTrainerServerProcess ServerProcess(
+		PythonExecutiblePath,
+		SitePackagesPath,
+		PythonContentPath,
+		IntermediatePath,
+		IpAddress,
+		Port,
+		LogSettings);
+
+	while (ServerProcess.IsRunning())
+	{
+		FPlatformProcess::Sleep(0.01f);
+	}
+
+	return 0;
+}
+
+namespace UE::Learning
+{
+	FSharedMemoryPPOTrainer::FSharedMemoryPPOTrainer(
+		const FString& TaskName,
+		const FString& PythonExecutiblePath,
+		const FString& SitePackagesPath,
+		const FString& PythonContentPath,
+		const FString& IntermediatePath,
+		const FNeuralNetwork& Network,
+		const float ActionNoiseMin,
+		const float ActionNoiseMax,
+		const FReplayBuffer& ReplayBuffer,
+		const FPPOTrainerSettings& Settings,
+		const ELogSetting LogSettings,
+		const bool bHideTrainingWindow,
+		const bool bRedirectTrainingOutput)
+	{
+		UE_LEARNING_CHECK(Network.GetOutputNum() % 2 == 0);
+
+		const int32 ObservationVectorDimensionNum = Network.GetInputNum();
+		const int32 ActionVectorDimensionNum = Network.GetOutputNum() / 2;
+
+		if (!ensure(Policy.Region == nullptr))
+		{
+			UE_LOG(LogLearning, Warning, TEXT("Training already started!"));
+			return;
+		}
+
+		// Allocate shared memory
+
+		ProcessIdx = 0;
+		FParse::Value(FCommandLine::Get(), TEXT("LearningProcessIdx"), ProcessIdx);
+
+		if (ProcessIdx == 0)
+		{
+			// Allocate Shared Memory
+
+			Policy = SharedMemory::Allocate<1, uint8>({ Network.GetTotalByteNum() });
+			Controls = SharedMemory::Allocate<2, volatile int32>({ Settings.ProcessNum, SharedMemoryTraining::GetControlNum() });
+			EpisodeStarts = SharedMemory::Allocate<2, int32>({ Settings.ProcessNum, ReplayBuffer.GetMaxEpisodeNum() });
+			EpisodeLengths = SharedMemory::Allocate<2, int32>({ Settings.ProcessNum, ReplayBuffer.GetMaxEpisodeNum() });
+			EpisodeCompletionModes = SharedMemory::Allocate<2, ECompletionMode>({ Settings.ProcessNum, ReplayBuffer.GetMaxEpisodeNum() });
+			EpisodeFinalObservations = SharedMemory::Allocate<3, float>({ Settings.ProcessNum, ReplayBuffer.GetMaxEpisodeNum(), ObservationVectorDimensionNum });
+			Observations = SharedMemory::Allocate<3, float>({ Settings.ProcessNum, ReplayBuffer.GetMaxStepNum(), ObservationVectorDimensionNum });
+			Actions = SharedMemory::Allocate<3, float>({ Settings.ProcessNum, ReplayBuffer.GetMaxStepNum(), ActionVectorDimensionNum });
+			Rewards = SharedMemory::Allocate<2, float>({ Settings.ProcessNum, ReplayBuffer.GetMaxStepNum() });
+
+			// We need to zero the control memory before we start
+			// the training sub-process since it may contain uninitialized 
+			// values or those left over from previous runs.
+
+			Array::Zero(Controls.View);
+
+			// Create Experience Gathering Sub-processes
+
+			for (uint16 SubprocessIdx = 1; SubprocessIdx < Settings.ProcessNum; SubprocessIdx++)
+			{
+				ensureMsgf(!WITH_EDITOR || IsRunningCommandlet(), TEXT("Multi-processing generally does not work in-editor as it requires a standalone executable."));
+
+				FString SubprocessCommandLine = FCommandLine::GetOriginal();
+
+				SubprocessCommandLine += FString::Printf(TEXT(" -LearningProcessIdx %i"), SubprocessIdx);
+				SubprocessCommandLine += FString(TEXT(" -LearningPolicyGuid ")) + Policy.Guid.ToString();
+				SubprocessCommandLine += FString(TEXT(" -LearningControlsGuid ")) + Controls.Guid.ToString();
+				SubprocessCommandLine += FString(TEXT(" -LearningEpisodeStartsGuid ")) + EpisodeStarts.Guid.ToString();
+				SubprocessCommandLine += FString(TEXT(" -LearningEpisodeLengthsGuid ")) + EpisodeLengths.Guid.ToString();
+				SubprocessCommandLine += FString(TEXT(" -LearningEpisodeCompletionModesGuid ")) + EpisodeCompletionModes.Guid.ToString();
+				SubprocessCommandLine += FString(TEXT(" -LearningEpisodeFinalObservationsGuid ")) + EpisodeFinalObservations.Guid.ToString();
+				SubprocessCommandLine += FString(TEXT(" -LearningObservationsGuid ")) + Observations.Guid.ToString();
+				SubprocessCommandLine += FString(TEXT(" -LearningActionsGuid ")) + Actions.Guid.ToString();
+				SubprocessCommandLine += FString(TEXT(" -LearningRewardsGuid ")) + Rewards.Guid.ToString();
+
+				const TSharedPtr<FMonitoredProcess> Subprocess = MakeShared<FMonitoredProcess>(
+					FPlatformProcess::ExecutablePath(), 
+					SubprocessCommandLine, 
+					Settings.bMultiProcessHideTrainingWindow, 
+					Settings.bMultiProcessRedirectTrainingOutput);
+
+				if (Settings.bMultiProcessRedirectTrainingOutput)
+				{
+					Subprocess->OnCanceled().BindRaw(this, &FSharedMemoryPPOTrainer::HandleSubprocessCanceled);
+					Subprocess->OnCompleted().BindRaw(this, &FSharedMemoryPPOTrainer::HandleSubprocessCompleted);
+					Subprocess->OnOutput().BindStatic(&FSharedMemoryPPOTrainer::HandleSubprocessOutput);
+				}
+
+				Subprocess->Launch();
+
+				UE_LOG(LogLearning, Display, TEXT("Subprocess Command: %s %s"), FPlatformProcess::ExecutablePath(), *SubprocessCommandLine);
+
+				ExperienceGatheringSubprocesses.Emplace(Subprocess);
+			}
+
+			// Write Config
+
+			const FString TimeStamp = FDateTime::Now().ToFormattedString(TEXT("%Y-%m-%d_%H-%M-%S"));
+			const FString TrainerMethod = TEXT("PPO");
+			const FString TrainerType = TEXT("SharedMemory");
+			const FString ConfigPath = IntermediatePath / TEXT("Configs") / FString::Printf(TEXT("%s_%s_%s_%s.json"), *TaskName, *TrainerMethod, *TrainerType, *TimeStamp);
+
+			TSharedRef<FJsonObject> ConfigObject = MakeShared<FJsonObject>();
+			ConfigObject->SetStringField(TEXT("TaskName"), *TaskName);
+			ConfigObject->SetStringField(TEXT("TrainerMethod"), TrainerMethod);
+			ConfigObject->SetStringField(TEXT("TrainerType"), TrainerType);
+			ConfigObject->SetStringField(TEXT("TimeStamp"), *TimeStamp);
+
+			ConfigObject->SetStringField(TEXT("SitePackagesPath"), *SitePackagesPath);
+			ConfigObject->SetStringField(TEXT("IntermediatePath"), *IntermediatePath);
+
+			ConfigObject->SetStringField(TEXT("PolicyGuid"), *Policy.Guid.ToString(EGuidFormats::DigitsWithHyphensInBraces));
+			ConfigObject->SetStringField(TEXT("ControlsGuid"), *Controls.Guid.ToString(EGuidFormats::DigitsWithHyphensInBraces));
+			ConfigObject->SetStringField(TEXT("EpisodeStartsGuid"), *EpisodeStarts.Guid.ToString(EGuidFormats::DigitsWithHyphensInBraces));
+			ConfigObject->SetStringField(TEXT("EpisodeLengthsGuid"), *EpisodeLengths.Guid.ToString(EGuidFormats::DigitsWithHyphensInBraces));
+			ConfigObject->SetStringField(TEXT("EpisodeCompletionModesGuid"), *EpisodeCompletionModes.Guid.ToString(EGuidFormats::DigitsWithHyphensInBraces));
+			ConfigObject->SetStringField(TEXT("EpisodeFinalObservationsGuid"), *EpisodeFinalObservations.Guid.ToString(EGuidFormats::DigitsWithHyphensInBraces));
+			ConfigObject->SetStringField(TEXT("ObservationsGuid"), *Observations.Guid.ToString(EGuidFormats::DigitsWithHyphensInBraces));
+			ConfigObject->SetStringField(TEXT("ActionsGuid"), *Actions.Guid.ToString(EGuidFormats::DigitsWithHyphensInBraces));
+			ConfigObject->SetStringField(TEXT("RewardsGuid"), *Rewards.Guid.ToString(EGuidFormats::DigitsWithHyphensInBraces));
+
+			ConfigObject->SetNumberField(TEXT("ObservationVectorDimensionNum"), ObservationVectorDimensionNum);
+			ConfigObject->SetNumberField(TEXT("ActionVectorDimensionNum"), ActionVectorDimensionNum);
+			ConfigObject->SetNumberField(TEXT("NetworkTotalByteNum"), Network.GetTotalByteNum());
+			ConfigObject->SetNumberField(TEXT("MaxEpisodeNum"), ReplayBuffer.GetMaxEpisodeNum());
+			ConfigObject->SetNumberField(TEXT("MaxStepNum"), ReplayBuffer.GetMaxStepNum());
+			ConfigObject->SetNumberField(TEXT("HiddenUnitNum"), Network.GetHiddenNum());
+			ConfigObject->SetNumberField(TEXT("LayerNum"), Network.GetLayerNum());
+			ConfigObject->SetStringField(TEXT("ActivationFunction"), ActivationFunctionString(Network.ActivationFunction));
+
+			ConfigObject->SetNumberField(TEXT("ActionNoiseMin"), ActionNoiseMin);
+			ConfigObject->SetNumberField(TEXT("ActionNoiseMax"), ActionNoiseMax);
+
+			ConfigObject->SetNumberField(TEXT("ProcessNum"), Settings.ProcessNum);
+			ConfigObject->SetNumberField(TEXT("IterationNum"), Settings.IterationNum);
+			ConfigObject->SetNumberField(TEXT("LearningRateActor"), Settings.LearningRateActor);
+			ConfigObject->SetNumberField(TEXT("LearningRateCritic"), Settings.LearningRateCritic);
+			ConfigObject->SetNumberField(TEXT("LearningRateDecay"), Settings.LearningRateDecay);
+			ConfigObject->SetNumberField(TEXT("WeightDecay"), Settings.WeightDecay);
+			ConfigObject->SetNumberField(TEXT("InitialActionScale"), Settings.InitialActionScale);
+			ConfigObject->SetNumberField(TEXT("BatchSize"), Settings.BatchSize);
+			ConfigObject->SetNumberField(TEXT("EpsilonClip"), Settings.EpsilonClip);
+			ConfigObject->SetNumberField(TEXT("ActionRegularizationWeight"), Settings.ActionRegularizationWeight);
+			ConfigObject->SetNumberField(TEXT("EntropyWeight"), Settings.EntropyWeight);
+			ConfigObject->SetNumberField(TEXT("GaeLambda"), Settings.GaeLambda);
+			ConfigObject->SetBoolField(TEXT("ClipAdvantages"), Settings.bClipAdvantages);
+			ConfigObject->SetBoolField(TEXT("AdvantageNormalization"), Settings.bAdvantageNormalization);
+			ConfigObject->SetNumberField(TEXT("TrimEpisodeStartStepNum"), Settings.TrimEpisodeStartStepNum);
+			ConfigObject->SetNumberField(TEXT("TrimEpisodeEndStepNum"), Settings.TrimEpisodeEndStepNum);
+			ConfigObject->SetNumberField(TEXT("Seed"), Settings.Seed);
+			ConfigObject->SetNumberField(TEXT("DiscountFactor"), Settings.DiscountFactor);
+			ConfigObject->SetBoolField(TEXT("ReinitializeNetwork"), Settings.bReinitializeNetwork);
+			ConfigObject->SetStringField(TEXT("Device"), Trainer::GetDeviceString(Settings.Device));
+			ConfigObject->SetBoolField(TEXT("UseTensorBoard"), Settings.bUseTensorboard);
+			ConfigObject->SetBoolField(TEXT("LoggingEnabled"), LogSettings == ELogSetting::Silent ? false : true);
+
+			FString JsonString;
+			TSharedRef<TJsonWriter<>> JsonWriter = TJsonWriterFactory<>::Create(&JsonString, 0);
+			FJsonSerializer::Serialize(ConfigObject, JsonWriter, true);
+
+			FFileHelper::SaveStringToFile(JsonString, *ConfigPath);
+
+			// Start Python Training Sub-process
+
+			const FString CommandLineArguments = FString::Printf(TEXT("\"%s\" SharedMemory \"%s\""), *(PythonContentPath / TEXT("train_ppo.py")), *ConfigPath);
+
+			TrainingProcess = MakeShared<FMonitoredProcess>(PythonExecutiblePath, CommandLineArguments, bHideTrainingWindow, bRedirectTrainingOutput);
+
+			if (bRedirectTrainingOutput)
+			{
+				TrainingProcess->OnCanceled().BindRaw(this, &FSharedMemoryPPOTrainer::HandleTrainingProcessCanceled);
+				TrainingProcess->OnCompleted().BindRaw(this, &FSharedMemoryPPOTrainer::HandleTrainingProcessCompleted);
+				TrainingProcess->OnOutput().BindStatic(&FSharedMemoryPPOTrainer::HandleTrainingProcessOutput);
+			}
+
+			TrainingProcess->Launch();
+		}
+		else
+		{
+			// Parse Guids from command line args
+
+			FGuid PolicyGuid; ensure(FParse::Value(FCommandLine::Get(), TEXT("LearningPolicyGuid"), PolicyGuid));
+			FGuid ControlsGuid; ensure(FParse::Value(FCommandLine::Get(), TEXT("LearningControlsGuid"), ControlsGuid));
+			FGuid EpisodeStartsGuid; ensure(FParse::Value(FCommandLine::Get(), TEXT("LearningEpisodeStartsGuid"), EpisodeStartsGuid));
+			FGuid EpisodeLengthsGuid; ensure(FParse::Value(FCommandLine::Get(), TEXT("LearningEpisodeLengthsGuid"), EpisodeLengthsGuid));
+			FGuid EpisodeCompletionModesGuid; ensure(FParse::Value(FCommandLine::Get(), TEXT("LearningEpisodeCompletionModesGuid"), EpisodeCompletionModesGuid));
+			FGuid EpisodeFinalObservationsGuid; ensure(FParse::Value(FCommandLine::Get(), TEXT("LearningEpisodeFinalObservationsGuid"), EpisodeFinalObservationsGuid));
+			FGuid ObservationsGuid; ensure(FParse::Value(FCommandLine::Get(), TEXT("LearningObservationsGuid"), ObservationsGuid));
+			FGuid ActionsGuid; ensure(FParse::Value(FCommandLine::Get(), TEXT("LearningActionsGuid"), ActionsGuid));
+			FGuid RewardsGuid; ensure(FParse::Value(FCommandLine::Get(), TEXT("LearningRewardsGuid"), RewardsGuid));
+
+			// Map shared memory
+
+			Policy = SharedMemory::Map<1, uint8>(PolicyGuid, { Network.GetTotalByteNum() });
+			Controls = SharedMemory::Map<2, volatile int32>(ControlsGuid, { Settings.ProcessNum, SharedMemoryTraining::GetControlNum() });
+			EpisodeStarts = SharedMemory::Map<2, int32>(EpisodeStartsGuid, { Settings.ProcessNum, ReplayBuffer.GetMaxEpisodeNum() });
+			EpisodeLengths = SharedMemory::Map<2, int32>(EpisodeLengthsGuid, { Settings.ProcessNum, ReplayBuffer.GetMaxEpisodeNum() });
+			EpisodeCompletionModes = SharedMemory::Map<2, ECompletionMode>(EpisodeCompletionModesGuid, { Settings.ProcessNum, ReplayBuffer.GetMaxEpisodeNum() });
+			EpisodeFinalObservations = SharedMemory::Map<3, float>(EpisodeFinalObservationsGuid, { Settings.ProcessNum, ReplayBuffer.GetMaxEpisodeNum(), ObservationVectorDimensionNum });
+			Observations = SharedMemory::Map<3, float>(ObservationsGuid, { Settings.ProcessNum, ReplayBuffer.GetMaxStepNum(), ObservationVectorDimensionNum });
+			Actions = SharedMemory::Map<3, float>(ActionsGuid, { Settings.ProcessNum, ReplayBuffer.GetMaxStepNum(), ActionVectorDimensionNum });
+			Rewards = SharedMemory::Map<2, float>(RewardsGuid, { Settings.ProcessNum, ReplayBuffer.GetMaxStepNum() });
+		}
+	}
+
+	FSharedMemoryPPOTrainer::~FSharedMemoryPPOTrainer()
+	{
+		Terminate();
+	}
+
+	ETrainerResponse FSharedMemoryPPOTrainer::Wait(float Timeout)
+	{
+		const float SleepTime = 0.001f;
+		float WaitTime = 0.0f;
+
+		while (TrainingProcess.IsValid())
+		{
+			FPlatformProcess::Sleep(SleepTime);
+			WaitTime += SleepTime;
+
+			if (WaitTime > Timeout)
+			{
+				return ETrainerResponse::Timeout;
+			}
+		}
+
+		return ETrainerResponse::Success;
+	}
+
+	void FSharedMemoryPPOTrainer::Terminate()
+	{
+		if (TrainingProcess.IsValid())
+		{
+			TrainingProcess->Cancel(true);
+		}
+
+		TrainingProcess.Reset();
+
+		if (Policy.Region)
+		{
+			Deallocate();
+		}
+	}
+
+	ETrainerResponse FSharedMemoryPPOTrainer::SendStop(const float Timeout)
+	{
+		return SharedMemoryTraining::SendStop(Controls.View[ProcessIdx]);
+	}
+
+	ETrainerResponse FSharedMemoryPPOTrainer::RecvPolicy(
+		FNeuralNetwork& OutNetwork,
+		const float Timeout,
+		FRWLock* NetworkLock,
+		const ELogSetting LogSettings)
+	{
+		return SharedMemoryTraining::RecvPolicy(
+			Controls.View[ProcessIdx],
+			OutNetwork,
+			Policy.View,
+			Timeout,
+			NetworkLock,
+			LogSettings);
+	}
+
+	ETrainerResponse FSharedMemoryPPOTrainer::SendPolicy(
+		const FNeuralNetwork& Network,
+		const float Timeout,
+		FRWLock* NetworkLock,
+		const ELogSetting LogSettings)
+	{
+		return SharedMemoryTraining::SendPolicy(
+			Controls.View[ProcessIdx],
+			Policy.View,
+			Network,
+			Timeout,
+			NetworkLock,
+			LogSettings);
+	}
+
+	ETrainerResponse FSharedMemoryPPOTrainer::SendExperience(
+		const FReplayBuffer& ReplayBuffer,
+		const float Timeout,
+		const ELogSetting LogSettings)
+	{
+		return SharedMemoryTraining::SendExperience(
+			EpisodeStarts.View[ProcessIdx],
+			EpisodeLengths.View[ProcessIdx],
+			EpisodeCompletionModes.View[ProcessIdx],
+			EpisodeFinalObservations.View[ProcessIdx],
+			Observations.View[ProcessIdx],
+			Actions.View[ProcessIdx],
+			Rewards.View[ProcessIdx],
+			Controls.View[ProcessIdx],
+			ReplayBuffer,
+			Timeout,
+			LogSettings);
+	}
+
+	void FSharedMemoryPPOTrainer::Deallocate()
+	{
+		if (Policy.Region != nullptr)
+		{
+			SharedMemory::Deallocate(Policy);
+			SharedMemory::Deallocate(Controls);
+			SharedMemory::Deallocate(EpisodeStarts);
+			SharedMemory::Deallocate(EpisodeLengths);
+			SharedMemory::Deallocate(EpisodeCompletionModes);
+			SharedMemory::Deallocate(EpisodeFinalObservations);
+			SharedMemory::Deallocate(Observations);
+			SharedMemory::Deallocate(Actions);
+			SharedMemory::Deallocate(Rewards);
+		}
+	}
+
+	void FSharedMemoryPPOTrainer::HandleSubprocessCanceled()
+	{
+		UE_LOG(LogLearning, Warning, TEXT("Subprocess canceled"));
+	}
+
+	void FSharedMemoryPPOTrainer::HandleSubprocessCompleted(int32 ReturnCode)
+	{
+		if (ReturnCode != 0)
+		{
+			UE_LOG(LogLearning, Warning, TEXT("Subprocess finished with warnings or errors"));
+		}
+	}
+
+	void FSharedMemoryPPOTrainer::HandleSubprocessOutput(FString Output)
+	{
+		if (!Output.IsEmpty())
+		{
+			UE_LOG(LogLearning, Display, TEXT("Subprocess: %s"), *Output);
+		}
+	}
+
+
+	void FSharedMemoryPPOTrainer::HandleTrainingProcessCanceled()
+	{
+		UE_LOG(LogLearning, Warning, TEXT("Training process canceled"));
+
+		TrainingProcess.Reset();
+	}
+
+	void FSharedMemoryPPOTrainer::HandleTrainingProcessCompleted(int32 ReturnCode)
+	{
+		if (ReturnCode != 0)
+		{
+			UE_LOG(LogLearning, Warning, TEXT("Training Process finished with warnings or errors"));
+		}
+
+		TrainingProcess.Reset();
+	}
+
+	void FSharedMemoryPPOTrainer::HandleTrainingProcessOutput(FString Output)
+	{
+		if (!Output.IsEmpty())
+		{
+			UE_LOG(LogLearning, Display, TEXT("Training Process: %s"), *Output);
+		}
+	}
+
+	FSocketPPOTrainerServerProcess::FSocketPPOTrainerServerProcess(
+		const FString& PythonExecutiblePath,
+		const FString& SitePackagesPath,
+		const FString& PythonContentPath,
+		const FString& IntermediatePath,
+		const TCHAR* IpAddress,
+		const uint32 Port,
+		const ELogSetting LogSettings)
+	{
+		const FString CommandLineArguments = FString::Printf(TEXT("\"%s\" Socket \"%s:%i\" \"%s\" \"%s\" %i"), 
+			*(PythonContentPath / TEXT("train_ppo.py")), IpAddress, Port, *SitePackagesPath, *IntermediatePath, LogSettings == ELogSetting::Normal ? 1 : 0);
+
+		TrainingProcess = MakeShared<FMonitoredProcess>(PythonExecutiblePath, CommandLineArguments, true, true);
+		TrainingProcess->OnCanceled().BindRaw(this, &FSocketPPOTrainerServerProcess::HandleTrainingProcessCanceled);
+		TrainingProcess->OnCompleted().BindRaw(this, &FSocketPPOTrainerServerProcess::HandleTrainingProcessCompleted);
+		TrainingProcess->OnOutput().BindStatic(&FSocketPPOTrainerServerProcess::HandleTrainingProcessOutput);
+		TrainingProcess->Launch();
+	}
+
+	FSocketPPOTrainerServerProcess::~FSocketPPOTrainerServerProcess()
+	{
+		Terminate();
+	}
+
+	bool FSocketPPOTrainerServerProcess::IsRunning() const
+	{
+		return TrainingProcess.IsValid();
+	}
+
+	bool FSocketPPOTrainerServerProcess::Wait(float Timeout)
+	{
+		const float SleepTime = 0.001f;
+		float WaitTime = 0.0f;
+
+		while (TrainingProcess.IsValid())
+		{
+			FPlatformProcess::Sleep(SleepTime);
+			WaitTime += SleepTime;
+
+			if (WaitTime > Timeout)
+			{
+				return false;
+			}
+		}
+
+		return true;
+	}
+
+	void FSocketPPOTrainerServerProcess::Terminate()
+	{
+		if (TrainingProcess.IsValid())
+		{
+			TrainingProcess->Cancel(true);
+		}
+
+		TrainingProcess.Reset();
+	}
+
+	void FSocketPPOTrainerServerProcess::HandleTrainingProcessCanceled()
+	{
+		UE_LOG(LogLearning, Warning, TEXT("Training process canceled"));
+
+		TrainingProcess.Reset();
+	}
+
+	void FSocketPPOTrainerServerProcess::HandleTrainingProcessCompleted(int32 ReturnCode)
+	{
+		if (ReturnCode != 0)
+		{
+			UE_LOG(LogLearning, Warning, TEXT("Training Process finished with warnings or errors"));
+		}
+
+		TrainingProcess.Reset();
+	}
+
+	void FSocketPPOTrainerServerProcess::HandleTrainingProcessOutput(FString Output)
+	{
+		if (!Output.IsEmpty())
+		{
+			UE_LOG(LogLearning, Display, TEXT("Training Process: %s"), *Output);
+		}
+	}
+
+	FSocketPPOTrainer::FSocketPPOTrainer(
+		ETrainerResponse& OutResponse,
+		const FString& TaskName,
+		const FNeuralNetwork& Network,
+		const float ActionNoiseMin,
+		const float ActionNoiseMax,
+		const FReplayBuffer& ReplayBuffer,
+		const TCHAR* IpAddress,
+		const uint32 Port,
+		const float Timeout,
+		const FPPOTrainerSettings& Settings)
+	{
+		UE_LEARNING_CHECK(Network.GetOutputNum() % 2 == 0);
+
+		const int32 ObservationVectorDimensionNum = Network.GetInputNum();
+		const int32 ActionVectorDimensionNum = Network.GetOutputNum() / 2;
+
+		// Write Config
+
+		const FString TimeStamp = FDateTime::Now().ToFormattedString(TEXT("%Y-%m-%d_%H-%M-%S"));
+		const FString TrainerMethod = TEXT("PPO");
+		const FString TrainerType = TEXT("Network");
+
+		TSharedRef<FJsonObject> ConfigObject = MakeShared<FJsonObject>();
+		ConfigObject->SetStringField(TEXT("TaskName"), *TaskName);
+		ConfigObject->SetStringField(TEXT("TrainerMethod"), TrainerMethod);
+		ConfigObject->SetStringField(TEXT("TrainerType"), TrainerType);
+		ConfigObject->SetStringField(TEXT("TimeStamp"), *TimeStamp);
+
+		ConfigObject->SetNumberField(TEXT("ObservationVectorDimensionNum"), ObservationVectorDimensionNum);
+		ConfigObject->SetNumberField(TEXT("ActionVectorDimensionNum"), ActionVectorDimensionNum);
+		ConfigObject->SetNumberField(TEXT("NetworkTotalByteNum"), Network.GetTotalByteNum());
+		ConfigObject->SetNumberField(TEXT("MaxEpisodeNum"), ReplayBuffer.GetMaxEpisodeNum());
+		ConfigObject->SetNumberField(TEXT("MaxStepNum"), ReplayBuffer.GetMaxStepNum());
+		ConfigObject->SetNumberField(TEXT("HiddenUnitNum"), Network.GetHiddenNum());
+		ConfigObject->SetNumberField(TEXT("LayerNum"), Network.GetLayerNum());
+		ConfigObject->SetStringField(TEXT("ActivationFunction"), ActivationFunctionString(Network.ActivationFunction));
+
+		ConfigObject->SetNumberField(TEXT("ActionNoiseMin"), ActionNoiseMin);
+		ConfigObject->SetNumberField(TEXT("ActionNoiseMax"), ActionNoiseMax);
+
+		ConfigObject->SetNumberField(TEXT("IterationNum"), Settings.IterationNum);
+		ConfigObject->SetNumberField(TEXT("LearningRateActor"), Settings.LearningRateActor);
+		ConfigObject->SetNumberField(TEXT("LearningRateCritic"), Settings.LearningRateCritic);
+		ConfigObject->SetNumberField(TEXT("LearningRateDecay"), Settings.LearningRateDecay);
+		ConfigObject->SetNumberField(TEXT("WeightDecay"), Settings.WeightDecay);
+		ConfigObject->SetNumberField(TEXT("InitialActionScale"), Settings.InitialActionScale);
+		ConfigObject->SetNumberField(TEXT("BatchSize"), Settings.BatchSize);
+		ConfigObject->SetNumberField(TEXT("EpsilonClip"), Settings.EpsilonClip);
+		ConfigObject->SetNumberField(TEXT("ActionRegularizationWeight"), Settings.ActionRegularizationWeight);
+		ConfigObject->SetNumberField(TEXT("EntropyWeight"), Settings.EntropyWeight);
+		ConfigObject->SetNumberField(TEXT("GaeLambda"), Settings.GaeLambda);
+		ConfigObject->SetBoolField(TEXT("ClipAdvantages"), Settings.bClipAdvantages);
+		ConfigObject->SetBoolField(TEXT("AdvantageNormalization"), Settings.bAdvantageNormalization);
+		ConfigObject->SetNumberField(TEXT("TrimEpisodeStartStepNum"), Settings.TrimEpisodeStartStepNum);
+		ConfigObject->SetNumberField(TEXT("TrimEpisodeEndStepNum"), Settings.TrimEpisodeEndStepNum);
+		ConfigObject->SetNumberField(TEXT("Seed"), Settings.Seed);
+		ConfigObject->SetNumberField(TEXT("DiscountFactor"), Settings.DiscountFactor);
+		ConfigObject->SetBoolField(TEXT("ReinitializeNetwork"), Settings.bReinitializeNetwork);
+		ConfigObject->SetStringField(TEXT("Device"), Trainer::GetDeviceString(Settings.Device));
+		ConfigObject->SetBoolField(TEXT("UseTensorBoard"), Settings.bUseTensorboard);
+
+		FString JsonString;
+		TSharedRef<TJsonWriter<>> JsonWriter = TJsonWriterFactory<>::Create(&JsonString, 0);
+		FJsonSerializer::Serialize(ConfigObject, JsonWriter, true);
+
+		// Allocate buffer to receive network data in
+
+		NetworkBuffer.SetNumUninitialized({ Network.GetTotalByteNum() });
+
+		// Create Socket
+
+		ISocketSubsystem* SocketSubsystem = ISocketSubsystem::Get(PLATFORM_SOCKETSUBSYSTEM);
+		if (!SocketSubsystem)
+		{
+			UE_LOG(LogLearning, Error, TEXT("Could not get socket subsystem"));
+			OutResponse = ETrainerResponse::Unexpected;
+			return;
+		}
+
+		bool bIsValid = false;
+		TSharedRef<FInternetAddr> Address = SocketSubsystem->CreateInternetAddr();
+		Address->SetIp(IpAddress, bIsValid);
+		Address->SetPort(Port);
+
+		if (!bIsValid)
+		{
+			UE_LOG(LogLearning, Error, TEXT("Invalid Ip Address \"%s\"..."), IpAddress);
+			OutResponse = ETrainerResponse::Unexpected;
+			return;
+		}
+
+		// Connect  to Server
+
+		Socket = FTcpSocketBuilder(TEXT("LearningNetworkPPOTrainerSocket")).AsNonBlocking().Build();
+		Socket->Connect(*Address);
+
+		OutResponse = SocketTraining::WaitForConnection(*Socket, Timeout);
+		if (OutResponse != ETrainerResponse::Success) { return; }
+
+		// Send Config
+
+		OutResponse = SocketTraining::SendConfig(*Socket, JsonString, Timeout);
+		return;
+	}
+
+	FSocketPPOTrainer::~FSocketPPOTrainer()
+	{
+		Terminate();
+	}
+
+	ETrainerResponse FSocketPPOTrainer::Wait(const float Timeout)
+	{
+		return ETrainerResponse::Success;
+	}
+
+	void FSocketPPOTrainer::Terminate()
+	{
+		if (Socket)
+		{
+			Socket->Close();
+			Socket = nullptr;
+		}
+	}
+
+	ETrainerResponse FSocketPPOTrainer::SendStop(const float Timeout)
+	{
+		return SocketTraining::SendStop(*Socket, Timeout);
+	}
+
+	ETrainerResponse FSocketPPOTrainer::RecvPolicy(
+		FNeuralNetwork& OutNetwork,
+		const float Timeout,
+		FRWLock* NetworkLock,
+		const ELogSetting LogSettings)
+	{
+		return SocketTraining::RecvPolicy(*Socket, OutNetwork, NetworkBuffer, Timeout, NetworkLock, LogSettings);
+	}
+
+	ETrainerResponse FSocketPPOTrainer::SendPolicy(
+		const FNeuralNetwork& Network,
+		const float Timeout,
+		FRWLock* NetworkLock,
+		const ELogSetting LogSettings)
+	{
+		return SocketTraining::SendPolicy(*Socket, NetworkBuffer, Network, Timeout, NetworkLock, LogSettings);
+	}
+
+	ETrainerResponse FSocketPPOTrainer::SendExperience(
+		const FReplayBuffer& ReplayBuffer,
+		const float Timeout,
+		const ELogSetting LogSettings)
+	{
+		return SocketTraining::SendExperience(*Socket, ReplayBuffer, Timeout, LogSettings);
+	}
+
+	namespace PPOTrainer
+	{
+		ETrainerResponse Train(
+			IPPOTrainer& Trainer,
+			FReplayBuffer& ReplayBuffer,
+			FEpisodeBuffer& EpisodeBuffer,
+			FResetInstanceBuffer& ResetBuffer,
+			FNeuralNetwork& Network,
+			TLearningArrayView<2, float> ObservationVectorBuffer,
+			TLearningArrayView<2, float> ActionVectorBuffer,
+			TLearningArrayView<1, float> RewardBuffer,
+			TLearningArrayView<1, ECompletionMode> CompletionBuffer,
+			const ECompletionMode EpisodeEndCompletionMode,
+			const TFunctionRef<void(const FIndexSet Instances)> ResetFunction,
+			const TFunctionRef<void(const FIndexSet Instances)> ObservationFunction,
+			const TFunctionRef<void(const FIndexSet Instances)> PolicyFunction,
+			const TFunctionRef<void(const FIndexSet Instances)> ActionFunction,
+			const TFunctionRef<void(const FIndexSet Instances)> UpdateFunction,
+			const TFunctionRef<void(const FIndexSet Instances)> RewardFunction,
+			const TFunctionRef<void(const FIndexSet Instances)> CompletionFunction,
+			const FIndexSet Instances,
+			const bool bReinitializeNetwork,
+			TAtomic<bool>* bRequestTrainingStopSignal,
+			FRWLock* NetworkLock,
+			TAtomic<bool>* bNetworkUpdatedSignal,
+			const ELogSetting LogSettings)
+		{
+			TRACE_CPUPROFILER_EVENT_SCOPE(Learning::PPOTrainer::Train);
+
+			ETrainerResponse Response = ETrainerResponse::Success;
+
+			if (bReinitializeNetwork)
+			{
+				// Receive initial Policy
+
+				if (LogSettings != ELogSetting::Silent)
+				{
+					UE_LOG(LogLearning, Display, TEXT("Receiving initial Policy..."));
+				}
+
+				Response = Trainer.RecvPolicy(Network, 20.0f, NetworkLock);
+
+				if (Response != ETrainerResponse::Success)
+				{
+					if (LogSettings != ELogSetting::Silent)
+					{
+						UE_LOG(LogLearning, Error, TEXT("Error receiving initial policy from trainer: %s. Check log for errors."), Trainer::ResponseString(Response));
+					}
+
+					Trainer.Terminate();
+					return Response;
+				}
+
+				if (bNetworkUpdatedSignal)
+				{
+					*bNetworkUpdatedSignal = true;
+				}
+			}
+			else
+			{
+				// Send initial Policy
+
+				if (LogSettings != ELogSetting::Silent)
+				{
+					UE_LOG(LogLearning, Display, TEXT("Sending initial Policy..."));
+				}
+
+				Response = Trainer.SendPolicy(Network, 20.0f, NetworkLock);
+
+				if (Response != ETrainerResponse::Success)
+				{
+					if (LogSettings != ELogSetting::Silent)
+					{
+						UE_LOG(LogLearning, Error, TEXT("Error sending initial policy to trainer: %s. Check log for errors."), Trainer::ResponseString(Response));
+					}
+
+					Trainer.Terminate();
+					return Response;
+				}
+			}
+
+			// Start Training Loop
+
+			while (true)
+			{
+				if (bRequestTrainingStopSignal && (*bRequestTrainingStopSignal))
+				{
+					*bRequestTrainingStopSignal = false;
+
+					if (LogSettings != ELogSetting::Silent)
+					{
+						UE_LOG(LogLearning, Display, TEXT("Stopping Training..."));
+					}
+
+					Response = Trainer.SendStop();
+
+					if (Response != ETrainerResponse::Success)
+					{
+						if (LogSettings != ELogSetting::Silent)
+						{
+							UE_LOG(LogLearning, Error, TEXT("Error sending stop signal to trainer: %s. Check log for errors."), Trainer::ResponseString(Response));
+						}
+
+						Trainer.Terminate();
+						return Response;
+					}
+
+					break;
+				}
+				else
+				{
+					Experience::GatherExperienceUntilReplayBufferFull(
+						ReplayBuffer,
+						EpisodeBuffer,
+						ResetBuffer,
+						ObservationVectorBuffer,
+						ActionVectorBuffer,
+						RewardBuffer,
+						CompletionBuffer,
+						EpisodeEndCompletionMode,
+						ResetFunction,
+						ObservationFunction,
+						PolicyFunction,
+						ActionFunction,
+						UpdateFunction,
+						RewardFunction,
+						CompletionFunction,
+						Instances);
+
+					Response = Trainer.SendExperience(ReplayBuffer, 10.0f);
+
+					if (Response != ETrainerResponse::Success)
+					{
+						if (LogSettings != ELogSetting::Silent)
+						{
+							UE_LOG(LogLearning, Error, TEXT("Error sending experience to trainer: %s. Check log for errors."), Trainer::ResponseString(Response));
+						}
+
+						Trainer.Terminate();
+						return Response;
+					}
+				}
+
+				Response = Trainer.RecvPolicy(Network, 10.0f, NetworkLock);
+
+				if (Response == ETrainerResponse::Completed)
+				{
+					if (LogSettings != ELogSetting::Silent)
+					{
+						UE_LOG(LogLearning, Display, TEXT("Trainer completed training."));
+					}
+					break;
+				}
+				else if (Response != ETrainerResponse::Success)
+				{
+					if (LogSettings != ELogSetting::Silent)
+					{
+						UE_LOG(LogLearning, Error, TEXT("Error receiving policy from trainer: %s. Check log for errors."), Trainer::ResponseString(Response));
+					}
+					break;
+				}
+
+				if (bNetworkUpdatedSignal)
+				{
+					*bNetworkUpdatedSignal = true;
+				}
+			}
+
+			// Allow some time for trainer to shut down gracefully before we kill it...
+			
+			Response = Trainer.Wait(5.0f);
+
+			if (Response != ETrainerResponse::Success)
+			{
+				if (LogSettings != ELogSetting::Silent)
+				{
+					UE_LOG(LogLearning, Error, TEXT("Error waiting for trainer to exit: %s. Check log for errors."), Trainer::ResponseString(Response));
+				}
+			}
+
+			Trainer.Terminate();
+
+			if (LogSettings != ELogSetting::Silent)
+			{
+				UE_LOG(LogLearning, Display, TEXT("Training Task Done!"));
+			}
+
+			return ETrainerResponse::Success;
+		}
+	}
+}
diff --git a/Engine/Plugins/Experimental/LearningAgents/Source/LearningTraining/Private/LearningSharedMemoryTraining.cpp b/Engine/Plugins/Experimental/LearningAgents/Source/LearningTraining/Private/LearningSharedMemoryTraining.cpp
new file mode 100644
--- /dev/null
+++ b/Engine/Plugins/Experimental/LearningAgents/Source/LearningTraining/Private/LearningSharedMemoryTraining.cpp
@@ -0,0 +1,235 @@
+// Copyright Epic Games, Inc. All Rights Reserved.
+
+#include "LearningSharedMemoryTraining.h"
+
+#include "LearningProgress.h"
+#include "LearningNeuralNetwork.h"
+#include "LearningExperience.h"
+#include "LearningCompletion.h"
+
+#include "HAL/PlatformProcess.h"
+
+namespace UE::Learning::SharedMemoryTraining
+{
+	enum class EControls : uint8
+	{
+		ExperienceEpisodeNum = 0,
+		ExperienceStepNum = 1,
+		ExperienceSignal = 2,
+		PolicySignal = 3,
+		CompleteSignal = 4,
+		StopSignal = 5,
+		ContinueSignal = 6,
+
+		ControlNum = 7,
+	};
+
+	uint8 GetControlNum()
+	{
+		return (uint8)EControls::ControlNum;
+	}
+
+	ETrainerResponse SendStop(TLearningArrayView<1, volatile int32> Controls)
+	{
+		Controls[(uint8)EControls::StopSignal] = true;
+		return ETrainerResponse::Success;
+	}
+
+	ETrainerResponse SendContinue(TLearningArrayView<1, volatile int32> Controls)
+	{
+		Controls[(uint8)EControls::ContinueSignal] = true;
+		return ETrainerResponse::Success;
+	}
+
+	ETrainerResponse RecvPolicy(
+		TLearningArrayView<1, volatile int32> Controls,
+		FNeuralNetwork& OutNetwork,
+		const TLearningArrayView<1, const uint8> Policy,
+		const float Timeout,
+		FRWLock* NetworkLock,
+		const ELogSetting LogSettings)
+	{
+		const float SleepTime = 0.001f;
+		float WaitTime = 0.0f;
+
+		// Wait until the Policy is done being written by the sub-process
+		while (!Controls[(uint8)EControls::PolicySignal])
+		{
+			// Check if Completed Signal has been raised
+			if (Controls[(uint8)EControls::CompleteSignal])
+			{
+				// If so set low to confirm we have read it
+				Controls[(uint8)EControls::CompleteSignal] = false;
+				return ETrainerResponse::Completed;
+			}
+
+			FPlatformProcess::Sleep(SleepTime);
+			WaitTime += SleepTime;
+
+			if (WaitTime > Timeout)
+			{
+				return ETrainerResponse::Timeout;
+			}
+		}
+
+		if (LogSettings != ELogSetting::Silent)
+		{
+			UE_LOG(LogLearning, Display, TEXT("Pulling Policy..."));
+		}
+
+		// Read the policy
+		{
+			FScopeNullableWriteLock ScopeLock(NetworkLock);
+			OutNetwork.DeserializeFromBytes(Policy);
+		}
+
+		// Confirm we have read the policy
+		Controls[(uint8)EControls::PolicySignal] = false;
+
+		return ETrainerResponse::Success;
+	}
+
+	ETrainerResponse SendPolicy(
+		TLearningArrayView<1, volatile int32> Controls,
+		TLearningArrayView<1, uint8> Policy,
+		const FNeuralNetwork& Network,
+		const float Timeout,
+		FRWLock* NetworkLock,
+		const ELogSetting LogSettings)
+	{
+		const float SleepTime = 0.001f;
+		float WaitTime = 0.0f;
+
+		// Wait until the Policy is requested by the sub-process
+		while (!Controls[(uint8)EControls::PolicySignal])
+		{
+			FPlatformProcess::Sleep(SleepTime);
+			WaitTime += SleepTime;
+
+			if (WaitTime > Timeout)
+			{
+				return ETrainerResponse::Timeout;
+			}
+		}
+
+		if (LogSettings != ELogSetting::Silent)
+		{
+			UE_LOG(LogLearning, Display, TEXT("Pushing Policy..."));
+		}
+
+		// Write the policy
+		{
+			FScopeNullableReadLock ScopeLock(NetworkLock);
+			Network.SerializeToBytes(Policy);
+		}
+
+		// Confirm we have written the policy
+		Controls[(uint8)EControls::PolicySignal] = false;
+
+		return ETrainerResponse::Success;
+	}
+
+	ETrainerResponse SendExperience(
+		TLearningArrayView<1, int32> EpisodeStarts,
+		TLearningArrayView<1, int32> EpisodeLengths,
+		TLearningArrayView<1, ECompletionMode> EpisodeCompletionModes,
+		TLearningArrayView<2, float> EpisodeFinalObservations,
+		TLearningArrayView<2, float> Observations,
+		TLearningArrayView<2, float> Actions,
+		TLearningArrayView<1, float> Rewards,
+		TLearningArrayView<1, volatile int32> Controls,
+		const FReplayBuffer& ReplayBuffer,
+		const float Timeout,
+		const ELogSetting LogSettings)
+	{
+		const float SleepTime = 0.001f;
+		float WaitTime = 0.0f;
+
+		// Wait until the sub-process is done reading any experience
+		while (Controls[(uint8)EControls::ExperienceSignal])
+		{
+			FPlatformProcess::Sleep(SleepTime);
+			WaitTime += SleepTime;
+
+			if (WaitTime > Timeout)
+			{
+				return ETrainerResponse::Timeout;
+			}
+		}
+
+		if (LogSettings != ELogSetting::Silent)
+		{
+			UE_LOG(LogLearning, Display, TEXT("Pushing Experience..."));
+		}
+
+		const int32 EpisodeNum = ReplayBuffer.GetEpisodeNum();
+		const int32 StepNum = ReplayBuffer.GetStepNum();
+
+		// Write experience to the shared memory
+		{
+			TRACE_CPUPROFILER_EVENT_SCOPE(Learning::SharedMemoryTraining::SendExperience::Copy);
+
+			Array::Copy(EpisodeStarts.Slice(0, EpisodeNum), ReplayBuffer.GetEpisodeStarts());
+			Array::Copy(EpisodeLengths.Slice(0, EpisodeNum), ReplayBuffer.GetEpisodeLengths());
+			Array::Copy(EpisodeCompletionModes.Slice(0, EpisodeNum), ReplayBuffer.GetEpisodeCompletionModes());
+			Array::Copy(EpisodeFinalObservations.Slice(0, EpisodeNum), ReplayBuffer.GetEpisodeFinalObservations());
+			Array::Copy(Observations.Slice(0, StepNum), ReplayBuffer.GetObservations());
+			Array::Copy(Actions.Slice(0, StepNum), ReplayBuffer.GetActions());
+			Array::Copy(Rewards.Slice(0, StepNum), ReplayBuffer.GetRewards());
+		}
+
+		// Indicate that experience is written
+		Controls[(uint8)EControls::ExperienceEpisodeNum] = EpisodeNum;
+		Controls[(uint8)EControls::ExperienceStepNum] = StepNum;
+		Controls[(uint8)EControls::ExperienceSignal] = true;
+
+		return ETrainerResponse::Success;
+	}
+
+	ETrainerResponse SendExperience(
+		TLearningArrayView<2, float> Observations,
+		TLearningArrayView<2, float> Actions,
+		TLearningArrayView<1, volatile int32> Controls,
+		const TLearningArrayView<2, const float> ObservationExperience,
+		const TLearningArrayView<2, const float> ActionExperience,
+		const float Timeout,
+		const ELogSetting LogSettings)
+	{
+		const float SleepTime = 0.001f;
+		float WaitTime = 0.0f;
+
+		// Wait until the sub-process is done reading any experience
+		while (Controls[(uint8)EControls::ExperienceSignal])
+		{
+			FPlatformProcess::Sleep(SleepTime);
+			WaitTime += SleepTime;
+
+			if (WaitTime > Timeout)
+			{
+				return ETrainerResponse::Timeout;
+			}
+		}
+
+		if (LogSettings != ELogSetting::Silent)
+		{
+			UE_LOG(LogLearning, Display, TEXT("Pushing Experience..."));
+		}
+
+		const int32 StepNum = ObservationExperience.Num(0);
+
+		// Write experience to the shared memory
+		{
+			TRACE_CPUPROFILER_EVENT_SCOPE(Learning::SharedMemoryTraining::SendExperience::Copy);
+
+			Array::Copy(Observations.Slice(0, StepNum), ObservationExperience);
+			Array::Copy(Actions.Slice(0, StepNum), ActionExperience);
+		}
+
+		// Confirm that experience is written
+		Controls[(uint8)EControls::ExperienceStepNum] = StepNum;
+		Controls[(uint8)EControls::ExperienceSignal] = true;
+
+		return ETrainerResponse::Success;
+	}
+
+}
\ No newline at end of file
diff --git a/Engine/Plugins/Experimental/LearningAgents/Source/LearningTraining/Private/LearningSharedMemoryTraining.h b/Engine/Plugins/Experimental/LearningAgents/Source/LearningTraining/Private/LearningSharedMemoryTraining.h
new file mode 100644
--- /dev/null
+++ b/Engine/Plugins/Experimental/LearningAgents/Source/LearningTraining/Private/LearningSharedMemoryTraining.h
@@ -0,0 +1,62 @@
+// Copyright Epic Games, Inc. All Rights Reserved.
+
+#pragma once
+
+#include "LearningTrainer.h"
+
+namespace UE::Learning
+{
+	enum class ECompletionMode : uint8;
+	struct FNeuralNetwork;
+	struct FReplayBuffer;
+
+	namespace SharedMemoryTraining
+	{
+		LEARNINGTRAINING_API uint8 GetControlNum();
+
+		LEARNINGTRAINING_API ETrainerResponse RecvPolicy(
+			TLearningArrayView<1, volatile int32> Controls,
+			FNeuralNetwork& OutNetwork,
+			const TLearningArrayView<1, const uint8> Policy,
+			const float Timeout = Trainer::DefaultTimeout,
+			FRWLock* NetworkLock = nullptr,
+			const ELogSetting LogSettings = Trainer::DefaultLogSettings);
+
+		LEARNINGTRAINING_API ETrainerResponse SendStop(
+			TLearningArrayView<1, volatile int32> Controls);
+
+		LEARNINGTRAINING_API ETrainerResponse SendContinue(
+			TLearningArrayView<1, volatile int32> Controls);
+
+		LEARNINGTRAINING_API ETrainerResponse SendPolicy(
+			TLearningArrayView<1, volatile int32> Controls,
+			TLearningArrayView<1, uint8> Policy,
+			const FNeuralNetwork& Network,
+			const float Timeout = Trainer::DefaultTimeout,
+			FRWLock* NetworkLock = nullptr,
+			const ELogSetting LogSettings = Trainer::DefaultLogSettings);
+
+		LEARNINGTRAINING_API ETrainerResponse SendExperience(
+			TLearningArrayView<1, int32> EpisodeStarts,
+			TLearningArrayView<1, int32> EpisodeLengths,
+			TLearningArrayView<1, ECompletionMode> EpisodeCompletionModes,
+			TLearningArrayView<2, float> EpisodeFinalObservations,
+			TLearningArrayView<2, float> Observations,
+			TLearningArrayView<2, float> Actions,
+			TLearningArrayView<1, float> Rewards,
+			TLearningArrayView<1, volatile int32> Controls,
+			const FReplayBuffer& ReplayBuffer,
+			const float Timeout = Trainer::DefaultTimeout,
+			const ELogSetting LogSettings = Trainer::DefaultLogSettings);
+
+		LEARNINGTRAINING_API ETrainerResponse SendExperience(
+			TLearningArrayView<2, float> Observations,
+			TLearningArrayView<2, float> Actions,
+			TLearningArrayView<1, volatile int32> Controls,
+			const TLearningArrayView<2, const float> ObservationExperience,
+			const TLearningArrayView<2, const float> ActionExperience,
+			const float Timeout = Trainer::DefaultTimeout,
+			const ELogSetting LogSettings = Trainer::DefaultLogSettings);
+
+	}
+}
diff --git a/Engine/Plugins/Experimental/LearningAgents/Source/LearningTraining/Private/LearningSocketTraining.cpp b/Engine/Plugins/Experimental/LearningAgents/Source/LearningTraining/Private/LearningSocketTraining.cpp
new file mode 100644
--- /dev/null
+++ b/Engine/Plugins/Experimental/LearningAgents/Source/LearningTraining/Private/LearningSocketTraining.cpp
@@ -0,0 +1,295 @@
+// Copyright Epic Games, Inc. All Rights Reserved.
+
+#include "LearningSocketTraining.h"
+
+#include "LearningProgress.h"
+#include "LearningNeuralNetwork.h"
+#include "LearningExperience.h"
+
+#include "Sockets.h"
+#include "HAL/PlatformProcess.h"
+
+namespace UE::Learning::SocketTraining
+{
+	enum class ESignal : uint8
+	{
+		Invalid = 0,
+		SendConfig = 1,
+		SendExperience = 2,
+		RecvPolicy = 3,
+		SendPolicy = 4,
+		RecvComplete = 5,
+		SendStop = 6,
+		SendContinue = 7,
+	};
+
+	ETrainerResponse WaitForConnection(FSocket& Socket, const float Timeout)
+	{
+		float WaitTime = 0.0f;
+		const float SleepTime = 0.001f;
+
+		while (true)
+		{
+			if (Socket.GetConnectionState() == ESocketConnectionState::SCS_ConnectionError)
+			{
+				UE_LOG(LogLearning, Error, TEXT("Socket refused connection..."));
+				return ETrainerResponse::Unexpected;
+			}
+
+			if (Socket.GetConnectionState() == ESocketConnectionState::SCS_Connected)
+			{
+				break;
+			}
+			else
+			{
+				FPlatformProcess::Sleep(SleepTime);
+				WaitTime += SleepTime;
+
+				if (WaitTime > Timeout)
+				{
+					return  ETrainerResponse::Timeout;
+				}
+			}
+		}
+
+		return ETrainerResponse::Success;
+	}
+
+	ETrainerResponse RecvWithTimeout(FSocket& Socket, uint8* Bytes, const int32 ByteNum, const float Timeout)
+	{
+		float WaitTime = 0.0f;
+		const float SleepTime = 0.001f;
+
+		int32 BytesRead = 0;
+		int32 TotalBytesRead = 0;
+
+		while (true)
+		{
+			if (Socket.Recv(Bytes + TotalBytesRead, ByteNum - TotalBytesRead, BytesRead))
+			{
+				TotalBytesRead += BytesRead;
+
+				if (TotalBytesRead == ByteNum)
+				{
+					return ETrainerResponse::Success;
+				}
+			}
+
+			FPlatformProcess::Sleep(SleepTime);
+			WaitTime += SleepTime;
+
+			if (WaitTime > Timeout)
+			{
+				return ETrainerResponse::Timeout;
+			}
+		}
+	}
+
+	ETrainerResponse RecvPolicy(
+		FSocket& Socket,
+		FNeuralNetwork& OutNetwork,
+		TLearningArrayView<1, uint8> OutNetworkBuffer,
+		const float Timeout,
+		FRWLock* NetworkLock,
+		const ELogSetting LogSettings)
+	{
+		if (LogSettings != ELogSetting::Silent)
+		{
+			UE_LOG(LogLearning, Display, TEXT("Pulling Policy..."));
+		}
+
+		ETrainerResponse Response = ETrainerResponse::Success;
+		uint8 Signal = (uint8)ESignal::Invalid;
+		Response = RecvWithTimeout(Socket, &Signal, 1, Timeout);
+		if (Response != ETrainerResponse::Success) { return Response; }
+
+		if (Signal == (uint8)ESignal::RecvComplete)
+		{
+			return ETrainerResponse::Completed;
+		}
+
+		if (Signal != (uint8)ESignal::RecvPolicy)
+		{
+			return ETrainerResponse::Unexpected;
+		}
+
+		Response = RecvWithTimeout(Socket, OutNetworkBuffer.GetData(), OutNetworkBuffer.Num(), Timeout);
+		if (Response != ETrainerResponse::Success) { return Response; }
+
+		{
+			FScopeNullableWriteLock ScopeLock(NetworkLock);
+			OutNetwork.DeserializeFromBytes(OutNetworkBuffer);
+		}
+
+		return ETrainerResponse::Success;
+	}
+
+	ETrainerResponse SendWithTimeout(FSocket& Socket, const uint8* Bytes, const int32 ByteNum, const float Timeout)
+	{
+		float WaitTime = 0.0f;
+		const float SleepTime = 0.001f;
+
+		int32 BytesSent = 0;
+		int32 TotalBytesSent = 0;
+
+		while (true)
+		{
+			if (Socket.Send(Bytes + TotalBytesSent, ByteNum - TotalBytesSent, BytesSent))
+			{
+				TotalBytesSent += BytesSent;
+
+				if (TotalBytesSent == ByteNum)
+				{
+					return ETrainerResponse::Success;
+				}
+			}
+
+			FPlatformProcess::Sleep(SleepTime);
+			WaitTime += SleepTime;
+
+			if (WaitTime > Timeout)
+			{
+				return ETrainerResponse::Timeout;
+			}
+		}
+	}
+
+	ETrainerResponse SendConfig(FSocket& Socket, const FString& ConfigString, const float Timeout)
+	{
+		const uint8 Signal = (uint8)ESignal::SendConfig;
+		ETrainerResponse Response = SendWithTimeout(Socket, &Signal, 1, Timeout);
+		if (Response != ETrainerResponse::Success) { return Response; }
+
+		const FTCHARToUTF8 UTF8String(*ConfigString);
+		const int32 ConfigLength = UTF8String.Length();
+		
+		Response = SendWithTimeout(Socket, (uint8*)&ConfigLength, sizeof(int32), Timeout);
+		if (Response != ETrainerResponse::Success) { return Response; }
+		
+		Response = SendWithTimeout(Socket, (uint8*)UTF8String.Get(), UTF8String.Length(), Timeout);
+		if (Response != ETrainerResponse::Success) { return Response; }
+
+		return ETrainerResponse::Success;
+	}
+
+	ETrainerResponse SendStop(FSocket& Socket, const float Timeout)
+	{
+		const uint8 Signal = (uint8)ESignal::SendStop;
+		return SendWithTimeout(Socket, &Signal, 1, Timeout);
+	}
+
+	ETrainerResponse SendContinue(
+		FSocket& Socket,
+		const float Timeout)
+	{
+		const uint8 Signal = (uint8)ESignal::SendContinue;
+		return SendWithTimeout(Socket, &Signal, 1, Timeout);
+	}
+
+	ETrainerResponse SendPolicy(
+		FSocket& Socket,
+		TLearningArrayView<1, uint8> NetworkBuffer,
+		const FNeuralNetwork& Network,
+		const float Timeout,
+		FRWLock* NetworkLock,
+		const ELogSetting LogSettings)
+	{
+		if (LogSettings != ELogSetting::Silent)
+		{
+			UE_LOG(LogLearning, Display, TEXT("Pushing Policy..."));
+		}
+
+		{
+			FScopeNullableReadLock ScopeLock(NetworkLock);
+			Network.SerializeToBytes(NetworkBuffer);
+		}
+
+		const uint8 Signal = (uint8)ESignal::SendPolicy;
+		ETrainerResponse Response = SendWithTimeout(Socket, &Signal, 1, Timeout);
+		if (Response != ETrainerResponse::Success) { return Response; }
+
+		Response = SendWithTimeout(Socket, NetworkBuffer.GetData(), NetworkBuffer.Num(), Timeout);
+		if (Response != ETrainerResponse::Success) { return Response; }
+
+		return ETrainerResponse::Success;
+	}
+
+	ETrainerResponse SendExperience(
+		FSocket& Socket,
+		const FReplayBuffer& ReplayBuffer,
+		const float Timeout,
+		const ELogSetting LogSettings)
+	{
+		if (LogSettings != ELogSetting::Silent)
+		{
+			UE_LOG(LogLearning, Display, TEXT("Pushing Experience..."));
+		}
+
+		const int32 EpisodeNum = ReplayBuffer.GetEpisodeNum();
+		const int32 StepNum = ReplayBuffer.GetStepNum();
+
+		const uint8 Signal = (uint8)ESignal::SendExperience;
+		ETrainerResponse Response = SendWithTimeout(Socket, &Signal, 1, Timeout);
+		if (Response != ETrainerResponse::Success) { return Response; }
+
+		Response = SendWithTimeout(Socket, (const uint8*)&EpisodeNum, sizeof(int32), Timeout);
+		if (Response != ETrainerResponse::Success) { return Response; }
+		
+		Response = SendWithTimeout(Socket, (const uint8*)&StepNum, sizeof(int32), Timeout);
+		if (Response != ETrainerResponse::Success) { return Response; }
+		
+		Response = SendWithTimeout(Socket, (const uint8*)ReplayBuffer.GetEpisodeStarts().GetData(), ReplayBuffer.GetEpisodeStarts().Num() * sizeof(int32), Timeout);
+		if (Response != ETrainerResponse::Success) { return Response; }
+		
+		Response = SendWithTimeout(Socket, (const uint8*)ReplayBuffer.GetEpisodeLengths().GetData(), ReplayBuffer.GetEpisodeLengths().Num() * sizeof(int32), Timeout);
+		if (Response != ETrainerResponse::Success) { return Response; }
+		
+		Response = SendWithTimeout(Socket, (const uint8*)ReplayBuffer.GetEpisodeCompletionModes().GetData(), ReplayBuffer.GetEpisodeCompletionModes().Num() * sizeof(ECompletionMode), Timeout);
+		if (Response != ETrainerResponse::Success) { return Response; }
+		
+		Response = SendWithTimeout(Socket, (const uint8*)ReplayBuffer.GetEpisodeFinalObservations().GetData(), ReplayBuffer.GetEpisodeFinalObservations().Num() * sizeof(float), Timeout);
+		if (Response != ETrainerResponse::Success) { return Response; }
+		
+		Response = SendWithTimeout(Socket, (const uint8*)ReplayBuffer.GetObservations().GetData(), ReplayBuffer.GetObservations().Num() * sizeof(float), Timeout);
+		if (Response != ETrainerResponse::Success) { return Response; }
+		
+		Response = SendWithTimeout(Socket, (const uint8*)ReplayBuffer.GetActions().GetData(), ReplayBuffer.GetActions().Num() * sizeof(float), Timeout);
+		if (Response != ETrainerResponse::Success) { return Response; }
+		
+		Response = SendWithTimeout(Socket, (const uint8*)ReplayBuffer.GetRewards().GetData(), ReplayBuffer.GetRewards().Num() * sizeof(float), Timeout);
+		if (Response != ETrainerResponse::Success) { return Response; }
+
+		return ETrainerResponse::Success;
+	}
+
+	ETrainerResponse SendExperience(
+		FSocket& Socket,
+		const TLearningArrayView<2, const float> ObservationExperience,
+		const TLearningArrayView<2, const float> ActionExperience,
+		const float Timeout,
+		const ELogSetting LogSettings)
+	{
+		if (LogSettings != ELogSetting::Silent)
+		{
+			UE_LOG(LogLearning, Display, TEXT("Pushing Experience..."));
+		}
+
+		const int32 SampleNum = ObservationExperience.Num<0>();
+
+		const uint8 Signal = (uint8)ESignal::SendExperience;
+		ETrainerResponse Response = SendWithTimeout(Socket, &Signal, 1, Timeout);
+		if (Response != ETrainerResponse::Success) { return Response; }
+
+		Response = SendWithTimeout(Socket, (const uint8*)&SampleNum, sizeof(int32), Timeout);
+		if (Response != ETrainerResponse::Success) { return Response; }
+
+		Response = SendWithTimeout(Socket, (const uint8*)ObservationExperience.GetData(), ObservationExperience.Num() * sizeof(float), Timeout);
+		if (Response != ETrainerResponse::Success) { return Response; }
+
+		Response = SendWithTimeout(Socket, (const uint8*)ActionExperience.GetData(), ActionExperience.Num() * sizeof(float), Timeout);
+		if (Response != ETrainerResponse::Success) { return Response; }
+
+		return ETrainerResponse::Success;
+	}
+
+}
\ No newline at end of file
diff --git a/Engine/Plugins/Experimental/LearningAgents/Source/LearningTraining/Private/LearningTrainer.cpp b/Engine/Plugins/Experimental/LearningAgents/Source/LearningTraining/Private/LearningTrainer.cpp
new file mode 100644
--- /dev/null
+++ b/Engine/Plugins/Experimental/LearningAgents/Source/LearningTraining/Private/LearningTrainer.cpp
@@ -0,0 +1,76 @@
+// Copyright Epic Games, Inc. All Rights Reserved.
+
+#include "LearningTrainer.h"
+
+#include "Misc/Paths.h"
+#include "Interfaces/IPluginManager.h"
+#include "HAL/FileManager.h"
+
+namespace UE::Learning::Trainer
+{
+	const TCHAR* GetDeviceString(const ETrainerDevice Device)
+	{
+		switch (Device)
+		{
+		case ETrainerDevice::GPU: return TEXT("GPU");
+		case ETrainerDevice::CPU: return TEXT("CPU");
+		default: UE_LEARNING_NOT_IMPLEMENTED(); return TEXT("Unknown");
+		}
+	}
+
+	const TCHAR* ResponseString(const ETrainerResponse Response)
+	{
+		switch (Response)
+		{
+		case ETrainerResponse::Success: return TEXT("Success");
+		case ETrainerResponse::Unexpected: return TEXT("Unexpected communication received");
+		case ETrainerResponse::Completed: return TEXT("Training completed");
+		case ETrainerResponse::Stopped: return TEXT("Training stopped");
+		case ETrainerResponse::Timeout: return TEXT("Communication timeout");
+		default: UE_LEARNING_NOT_IMPLEMENTED(); return TEXT("Unknown");
+		}
+	}
+
+	float DiscountFactorFromHalfLife(const float HalfLife, const float DeltaTime)
+	{
+		return FMath::Pow(0.5f, DeltaTime / FMath::Max(HalfLife, UE_SMALL_NUMBER));
+	}
+
+	float DiscountFactorFromHalfLifeSteps(const int32 HalfLifeSteps)
+	{
+		UE_LEARNING_CHECKF(HalfLifeSteps >= 1, TEXT("Number of HalfLifeSteps should be at least 1 but got %i"), HalfLifeSteps);
+
+		return FMath::Pow(0.5f, 1.0f / FMath::Max(HalfLifeSteps, 1));
+	}
+
+#if WITH_EDITOR
+
+	FString DefaultEditorPythonExecutablePath()
+	{
+		UE_LEARNING_CHECKF(PLATFORM_WINDOWS || PLATFORM_MAC || PLATFORM_LINUX, TEXT("Python only supported on Windows, Mac, and Linux."));
+
+		return IFileManager::Get().ConvertToAbsolutePathForExternalAppForRead(
+			*(FPaths::EngineDir() / TEXT("Binaries/ThirdParty/Python3") / FPlatformMisc::GetUBTPlatform() / (PLATFORM_WINDOWS ? TEXT("python.exe") : TEXT("bin/python"))));
+	}
+
+	FString DefaultEditorSitePackagesPath()
+	{
+		return IFileManager::Get().ConvertToAbsolutePathForExternalAppForRead(
+			*(IPluginManager::Get().FindPlugin(TEXT("PythonFoundationPackages"))->GetBaseDir() / TEXT("Content/Python/Lib") / FPlatformMisc::GetUBTPlatform() / TEXT("site-packages")));
+	}
+
+	FString DefaultEditorPythonContentPath()
+	{
+		return  IFileManager::Get().ConvertToAbsolutePathForExternalAppForRead(
+			*(IPluginManager::Get().FindPlugin(TEXT("LearningAgents"))->GetBaseDir() / TEXT("Content/Python")));
+	}
+
+	FString DefaultEditorIntermediatePath()
+	{
+		return IFileManager::Get().ConvertToAbsolutePathForExternalAppForRead(
+			*(IPluginManager::Get().FindPlugin(TEXT("LearningAgents"))->GetBaseDir() / TEXT("Intermediate")));
+	}
+
+#endif
+
+}
\ No newline at end of file
diff --git a/Engine/Plugins/Experimental/LearningAgents/Source/LearningTraining/Private/LearningTraining.cpp b/Engine/Plugins/Experimental/LearningAgents/Source/LearningTraining/Private/LearningTraining.cpp
new file mode 100644
--- /dev/null
+++ b/Engine/Plugins/Experimental/LearningAgents/Source/LearningTraining/Private/LearningTraining.cpp
@@ -0,0 +1,16 @@
+// Copyright Epic Games, Inc. All Rights Reserved.
+
+#include "LearningTraining.h"
+
+void FLearningTrainingModule::StartupModule()
+{
+	// This code will execute after your module is loaded into memory; the exact timing is specified in the .uplugin file per-module
+}
+
+void FLearningTrainingModule::ShutdownModule()
+{
+	// This function may be called during shutdown to clean up your module.  For modules that support dynamic reloading,
+	// we call this function before unloading the module.
+}
+	
+IMPLEMENT_MODULE(FLearningTrainingModule, LearningTraining)
\ No newline at end of file
