Learning Agents: Create initial experimental version

diff --git a/Engine/Plugins/Experimental/LearningAgents/Source/LearningTraining/Public/LearningExperience.h b/Engine/Plugins/Experimental/LearningAgents/Source/LearningTraining/Public/LearningExperience.h
new file mode 100644
--- /dev/null
+++ b/Engine/Plugins/Experimental/LearningAgents/Source/LearningTraining/Public/LearningExperience.h
@@ -0,0 +1,227 @@
+// Copyright Epic Games, Inc. All Rights Reserved.
+
+#pragma once
+
+#include "LearningArray.h"
+
+namespace UE::Learning
+{
+	enum class ECompletionMode : uint8;
+	struct FResetInstanceBuffer;
+
+	/**
+	* Buffer storing the observations, actions, and rewards of multiple instances over an episode
+	*/
+	struct LEARNINGTRAINING_API FEpisodeBuffer
+	{
+		/**
+		* Resize the experience buffer
+		*
+		* @param InMaxInstanceNum				Maximum number of instances
+		* @param InMaxStepNum					Maximum number of steps in an episode
+		* @param InObservationVectorDimNum		Number of dimensions of the observation vector
+		* @param InActionVectorDimNum			Number of dimensions of the action vector
+		*/
+		void Resize(
+			const int32 InMaxInstanceNum,
+			const int32 InMaxStepNum,
+			const int32 InObservationVectorDimNum,
+			const int32 InActionVectorDimNum);
+
+		/**
+		* Reset the buffer for the given set of instances
+		*/
+		void Reset(const FIndexSet Instances);
+
+		/**
+		* Push new experience to the buffer
+		*
+		* @param InObservations					Observation vectors of shape (MaxInstanceNum, ObservationVectorDimNum)
+		* @param InActions						Action vectors of shape (MaxInstanceNum, ActionVectorDimNum)
+		* @param InRewards						Rewards of shape (MaxInstanceNum)
+		* @param Instances						Instances to push to buffer
+		*/
+		void Push(
+			const TLearningArrayView<2, const float> InObservations,
+			const TLearningArrayView<2, const float> InActions,
+			const TLearningArrayView<1, const float> InRewards,
+			const FIndexSet Instances);
+
+		const int32 GetMaxInstanceNum() const;
+		const int32 GetMaxStepNum() const;
+		const TLearningArrayView<1, const int32> GetEpisodeStepNums() const;
+		const TLearningArrayView<2, const float> GetEpisodeFinalObservations() const;
+		const TLearningArrayView<2, const float> GetObservations(const int32 InstanceIdx) const;
+		const TLearningArrayView<2, const float> GetActions(const int32 InstanceIdx) const;
+		const TLearningArrayView<1, const float> GetRewards(const int32 InstanceIdx) const;
+
+	private:
+
+		int32 MaxInstanceNum = 0;
+		int32 MaxStepNum = 0;
+		TLearningArray<1, int32> EpisodeStepNums;
+		TLearningArray<3, float> Observations;
+		TLearningArray<3, float> Actions;
+		TLearningArray<2, float> Rewards;
+	};
+
+
+	/**
+	* Large buffer that sequentially concatenates a series of episodes in a large
+	* flat array. Used to collate episodic data together from multiple instances.
+	*/
+	struct LEARNINGTRAINING_API FReplayBuffer
+	{
+		/**
+		* Resizes the replay buffer.
+		*
+		* @param ObservationVectorDimNum		Dimensionality of observation vector
+		* @param ActionVectorDimNum				Dimensionality of action vector
+		* @param MaxEpisodeNum					Maximum number of episodes to be stored in the buffer
+		* @param MaxStepNum						Maximum number of steps to be stored in the buffer
+		*/
+		void Resize(
+			const int32 InObservationVectorDimNum,
+			const int32 InActionVectorDimNum,
+			const int32 InMaxEpisodeNum = 2048,
+			const int32 InMaxStepNum = 16384);
+
+		/**
+		* Reset the replay buffer. Does not free memory - just resets episode and sample num to zero.
+		*/
+		void Reset();
+
+		/**
+		* Add a set of episodes to the replay buffer
+		*
+		* @param InEpisodeCompletionModes		Array of completion modes for each instance of shape (InstanceNum)
+		* @param InEpisodeFinalObservations		Array of final observations for each instance of shape (InstanceNum, ObservationVectorDimNum)
+		* @param EpisodeBuffer					Episode buffer to add experience from
+		* @param Instances						Instances to add
+		*
+		* @returns								True when the replay buffer is full
+		*/
+		bool AddEpisodes(
+			const TLearningArrayView<1, const ECompletionMode> InEpisodeCompletionModes,
+			const TLearningArrayView<2, const float> InEpisodeFinalObservations,
+			const FEpisodeBuffer& EpisodeBuffer,
+			const FIndexSet Instances);
+
+		const int32 GetMaxEpisodeNum() const;
+		const int32 GetMaxStepNum() const;
+		const int32 GetEpisodeNum() const;
+		const int32 GetStepNum() const;
+		const TLearningArrayView<1, const int32> GetEpisodeStarts() const;
+		const TLearningArrayView<1, const int32> GetEpisodeLengths() const;
+		const TLearningArrayView<1, const ECompletionMode> GetEpisodeCompletionModes() const;
+		const TLearningArrayView<2, const float> GetEpisodeFinalObservations() const;
+		const TLearningArrayView<2, const float> GetObservations() const;
+		const TLearningArrayView<2, const float> GetActions() const;
+		const TLearningArrayView<1, const float> GetRewards() const;
+
+	private:
+
+		int32 MaxEpisodeNum = 0;
+		int32 MaxStepNum = 0;
+
+		int32 EpisodeNum = 0;
+		int32 StepNum = 0;
+
+		TLearningArray<1, int32> EpisodeStarts;
+		TLearningArray<1, int32> EpisodeLengths;
+		TLearningArray<1, ECompletionMode> EpisodeCompletionModes;
+		TLearningArray<2, float> EpisodeFinalObservations;
+		TLearningArray<2, float> Observations;
+		TLearningArray<2, float> Actions;
+		TLearningArray<1, float> Rewards;
+	};
+
+	namespace Experience
+	{
+		/**
+		* Resets an environment, replay buffer, and experience buffer
+		*/
+		LEARNINGTRAINING_API void GatherExperienceReset(
+			FReplayBuffer& ReplayBuffer,
+			FEpisodeBuffer& EpisodeBuffer,
+			FResetInstanceBuffer& ResetBuffer,
+			const FIndexSet Instances);
+
+		/**
+		* Runs a single step of experience gathering
+		*
+		* @param ReplayBuffer					Replay Buffer
+		* @param EpisodeBuffer					Episode Buffer
+		* @param ResetBuffer					Reset Buffer
+		* @param ObservationVectorBuffer		Buffer to read/write observation vectors into
+		* @param ActionVectorBuffer				Buffer to read/write action vectors into
+		* @param RewardBuffer					Buffer to read/write rewards into
+		* @param CompletionBuffer				Buffer to read/write completions into
+		* @param EpisodeEndCompletionMode		Completion mode to use for episodes that reach the max length
+		* @param ResetFunction					Function to run for resetting the environment
+		* @param ObservationFunction			Function to run for evaluating observations
+		* @param PolicyFunction					Function to run generating actions from observations
+		* @param ActionFunction					Function to run for evaluating actions
+		* @param UpdateFunction					Function to run for updating the environment
+		* @param RewardFunction					Function to run for evaluating rewards
+		* @param CompletionFunction				Function to run for evaluating completions
+		* @param Instances						Set of instances to gather experience for
+		*/
+		LEARNINGTRAINING_API bool GatherExperienceIteration(
+			FReplayBuffer& ReplayBuffer,
+			FEpisodeBuffer& EpisodeBuffer,
+			FResetInstanceBuffer& ResetBuffer,
+			TLearningArrayView<2, float> ObservationVectorBuffer,
+			TLearningArrayView<2, float> ActionVectorBuffer,
+			TLearningArrayView<1, float> RewardBuffer,
+			TLearningArrayView<1, ECompletionMode> CompletionBuffer,
+			const ECompletionMode EpisodeEndCompletionMode,
+			const TFunctionRef<void(const FIndexSet Instances)> ResetFunction,
+			const TFunctionRef<void(const FIndexSet Instances)> ObservationFunction,
+			const TFunctionRef<void(const FIndexSet Instances)> PolicyFunction,
+			const TFunctionRef<void(const FIndexSet Instances)> ActionFunction,
+			const TFunctionRef<void(const FIndexSet Instances)> UpdateFunction,
+			const TFunctionRef<void(const FIndexSet Instances)> RewardFunction,
+			const TFunctionRef<void(const FIndexSet Instances)> CompletionFunction,
+			const FIndexSet Instances);
+
+		/**
+		* Resets, and then runs experience gathering until the provided replay buffer is full
+		*
+		* @param ReplayBuffer					Replay Buffer
+		* @param EpisodeBuffer					Episode Buffer
+		* @param ResetBuffer					Reset Buffer
+		* @param ObservationVectorBuffer		Buffer to read/write observation vectors into
+		* @param ActionVectorBuffer				Buffer to read/write action vectors into
+		* @param RewardBuffer					Buffer to read/write rewards into
+		* @param CompletionBuffer				Buffer to read/write completions into
+		* @param EpisodeEndCompletionMode		Completion mode to use for episodes that reach the max length
+		* @param ResetFunction					Function to run for resetting the environment
+		* @param ObservationFunction			Function to run for evaluating observations
+		* @param PolicyFunction					Function to run generating actions from observations
+		* @param ActionFunction					Function to run for evaluating actions
+		* @param UpdateFunction					Function to run for updating the environment
+		* @param RewardFunction					Function to run for evaluating rewards
+		* @param CompletionFunction				Function to run for evaluating completions
+		* @param Instances						Set of instances to gather experience for
+		*/
+		LEARNINGTRAINING_API void GatherExperienceUntilReplayBufferFull(
+			FReplayBuffer& ReplayBuffer,
+			FEpisodeBuffer& EpisodeBuffer,
+			FResetInstanceBuffer& ResetBuffer,
+			TLearningArrayView<2, float> ObservationVectorBuffer,
+			TLearningArrayView<2, float> ActionVectorBuffer,
+			TLearningArrayView<1, float> RewardBuffer,
+			TLearningArrayView<1, ECompletionMode> CompletionBuffer,
+			const ECompletionMode EpisodeEndCompletionMode,
+			const TFunctionRef<void(const FIndexSet Instances)> ResetFunction,
+			const TFunctionRef<void(const FIndexSet Instances)> ObservationFunction,
+			const TFunctionRef<void(const FIndexSet Instances)> PolicyFunction,
+			const TFunctionRef<void(const FIndexSet Instances)> ActionFunction,
+			const TFunctionRef<void(const FIndexSet Instances)> UpdateFunction,
+			const TFunctionRef<void(const FIndexSet Instances)> RewardFunction,
+			const TFunctionRef<void(const FIndexSet Instances)> CompletionFunction,
+			const FIndexSet Instances);
+
+	}
+};
\ No newline at end of file
diff --git a/Engine/Plugins/Experimental/LearningAgents/Source/LearningTraining/Public/LearningImitationTrainer.h b/Engine/Plugins/Experimental/LearningAgents/Source/LearningTraining/Public/LearningImitationTrainer.h
new file mode 100644
--- /dev/null
+++ b/Engine/Plugins/Experimental/LearningAgents/Source/LearningTraining/Public/LearningImitationTrainer.h
@@ -0,0 +1,384 @@
+// Copyright Epic Games, Inc. All Rights Reserved.
+
+#pragma once
+
+#include "LearningArray.h"
+#include "LearningLog.h"
+#include "LearningTrainer.h"
+#include "LearningSharedMemory.h"
+
+#include "Commandlets/Commandlet.h"
+#include "Templates/SharedPointer.h"
+
+#include "LearningImitationTrainer.generated.h"
+
+class FSocket;
+class FMonitoredProcess;
+
+UCLASS()
+class LEARNINGTRAINING_API ULearningSocketImitationTrainerServerCommandlet : public UCommandlet
+{
+	GENERATED_UCLASS_BODY()
+
+	/** Runs the commandlet */
+	virtual int32 Main(const FString& Params) override;
+};
+
+namespace UE::Learning
+{
+	struct FNeuralNetwork;
+
+	struct FImitationTrainerSettings
+	{
+		// Number of iterations to train the network for. Controls the overall training time.
+		// Training for about 100000 iterations should give you well trained network, but
+		// closer to 1000000 iterations or more is required for an exhaustively trained network.
+		uint32 IterationNum = 1000000;
+
+		// Learning rate of the actor network. Typical values are between 0.001f and 0.0001f
+		float LearningRateActor = 0.0001f;
+
+		// Ratio by which to decay the learning rate every 1000 iterations.
+		float LearningRateDecay = 0.99f;
+
+		// Amount of weight decay to apply to the network. Larger values encourage network 
+		// weights to be smaller.
+		float WeightDecay = 0.001f;
+
+		// Batch size to use for training. Smaller values tend to produce better results 
+		// at the cost of slowing down training.
+		uint32 BatchSize = 128;
+
+		// Random seed to use for training
+		uint32 Seed = 1234;
+
+		// If true, will reinitialize the network to random weights at the start of training. Set this
+		// to false if you are starting from a network which has already been trained. Make sure 
+		// to use the same settings here for ImitationTrainer::Train.
+		bool bReinitializeNetwork = true;
+
+		// Which device to use for training
+		ETrainerDevice Device = ETrainerDevice::GPU;
+
+		// If to use TensorBoard for logging and tracking the training progress.
+		// 
+		// Even when enabled, TensorBoard will only work if it is installed in your Unreal Editor
+		// bundled version of Python, which is not the case by default. TensorBoard can be installed 
+		// for this version of Python by going to your Unreal Editor Python Binaries directory 
+		// (e.g. "\Engine\Binaries\ThirdParty\Python3\Win64") and running `.\python -m pip install tensorboard`. 
+		bool bUseTensorboard = false;
+	};
+
+	/**
+	* Interface for an object which can train a policy from experience using imitation learning.
+	*/
+	struct IImitationTrainer
+	{
+		virtual ~IImitationTrainer() {}
+
+		/**
+		* Terminate the trainer immediately.
+		*/
+		virtual void Terminate() = 0;
+
+		/**
+		* Wait for the trainer to finish.
+		* 
+		* @param Timeout		Timeout to wait in seconds
+		* @returns				Trainer response
+		*/
+		virtual ETrainerResponse Wait(const float Timeout = Trainer::DefaultTimeout) = 0;
+
+		/*
+		* Signal for the trainer to stop.
+		* 
+		* @param Timeout		Timeout to wait in seconds
+		* @returns				Trainer response
+		*/
+		virtual ETrainerResponse SendStop(const float Timeout = Trainer::DefaultTimeout) = 0;
+
+		/*
+		* Signal for the trainer to continue.
+		* 
+		* @param Timeout		Timeout to wait in seconds
+		* @returns				Trainer response
+		*/
+		virtual ETrainerResponse SendContinue(const float Timeout = Trainer::DefaultTimeout) = 0;
+
+		/**
+		* Wait for the trainer and pull an updated policy.
+		*
+		* @param OutNetwork		Network to update
+		* @param Timeout		Timeout to wait in seconds
+		* @param NetworkLock	Lock to use when updating network
+		* @param LogSettings	Log settings
+		* @returns				Trainer response
+		*/
+		virtual ETrainerResponse RecvPolicy(
+			FNeuralNetwork& OutNetwork,
+			const float Timeout = Trainer::DefaultTimeout,
+			FRWLock* NetworkLock = nullptr,
+			const ELogSetting LogSettings = Trainer::DefaultLogSettings) = 0;
+
+		/**
+		* Wait for the training process to be ready and push an updated policy to the shared memory.
+		*
+		* @param Network		Network to push
+		* @param Timeout		Timeout to wait in seconds
+		* @param NetworkLock	Lock to use when pushing network
+		* @param LogSettings	Log settings
+		* @returns				Trainer response
+		*/
+		virtual ETrainerResponse SendPolicy(
+			const FNeuralNetwork& Network,
+			const float Timeout = Trainer::DefaultTimeout,
+			FRWLock* NetworkLock = nullptr,
+			const ELogSetting LogSettings = Trainer::DefaultLogSettings) = 0;
+
+		/**
+		* Push experience to the trainer.
+		*
+		* @param ObservationVectors		Set of observation vectors
+		* @param ActionVectors			Set of action vectors
+		* @param Timeout				Timeout to wait in seconds
+		* @param LogSettings			Log settings
+		* @returns						Trainer response
+		*/
+		virtual ETrainerResponse SendExperience(
+			const TLearningArrayView<2, const float> ObservationVectors,
+			const TLearningArrayView<2, const float> ActionVectors,
+			const float Timeout = Trainer::DefaultTimeout,
+			const ELogSetting LogSettings = Trainer::DefaultLogSettings) = 0;
+	};
+
+	/**
+	* Object used to create and communicate with the Python training sub-process
+	*/
+	struct LEARNINGTRAINING_API FSharedMemoryImitationTrainer : public IImitationTrainer
+	{
+		/**
+		* Create a new imitation trainer sub-process
+		*
+		* @param TaskName				Name of the training task - used to help identify the logs, snapshots, and other files generated by training
+		* @param PythonExecutablePath	Path to the python executable used for training
+		* @param SitePackagesPath		Path to the site-packages shipped with the PythonFoundationPackages plugin
+		* @param PythonContentPath		Path to the Python Content folder provided by the Learning plugin
+		* @param IntermediatePath		Path to the intermediate folder to write temporary files, logs, and snapshots to
+		* @param Network				Neural Network object
+		* @param ActionNoiseMin			Minimum amount of action noise being used during episode generation
+		* @param ActionNoiseMax			Maximum amount of action noise being used during episode generation
+		* @param MaxSampleNum			Maximum number of samples in the training data
+		* @param Settings				Trainer settings
+		* @param LogSettings			Logging settings to use
+		*/
+		FSharedMemoryImitationTrainer(
+			const FString& TaskName,
+			const FString& PythonExecutablePath,
+			const FString& SitePackagesPath,
+			const FString& PythonContentPath,
+			const FString& IntermediatePath,
+			const FNeuralNetwork& Network,
+			const float ActionNoiseMin,
+			const float ActionNoiseMax,
+			const int32 MaxSampleNum,
+			const FImitationTrainerSettings& Settings = FImitationTrainerSettings(),
+			const ELogSetting LogSettings = ELogSetting::Normal);
+
+		~FSharedMemoryImitationTrainer();
+
+		virtual void Terminate() override final;
+
+		virtual ETrainerResponse Wait(const float Timeout) override final;
+
+		virtual ETrainerResponse SendStop(const float Timeout) override final;
+
+		virtual ETrainerResponse SendContinue(const float Timeout) override final;
+
+		virtual ETrainerResponse RecvPolicy(
+			FNeuralNetwork& OutNetwork,
+			const float Timeout = Trainer::DefaultTimeout,
+			FRWLock* NetworkLock = nullptr,
+			const ELogSetting LogSettings = Trainer::DefaultLogSettings) override final;
+
+		virtual ETrainerResponse SendPolicy(
+			const FNeuralNetwork& Network,
+			const float Timeout = Trainer::DefaultTimeout,
+			FRWLock* NetworkLock = nullptr,
+			const ELogSetting LogSettings = Trainer::DefaultLogSettings) override final;
+
+		virtual ETrainerResponse SendExperience(
+			const TLearningArrayView<2, const float> ObservationVectors,
+			const TLearningArrayView<2, const float> ActionVectors,
+			const float Timeout = Trainer::DefaultTimeout,
+			const ELogSetting LogSettings = Trainer::DefaultLogSettings) override final;
+
+	private:
+
+		/**
+		* Free and deallocate all shared memory
+		*/
+		void Deallocate();
+
+		// Shared memory
+
+		UE::Learning::TSharedMemoryArrayView<1, uint8> Policy;
+		UE::Learning::TSharedMemoryArrayView<1, volatile int32> Controls; // Mark as volatile to avoid compiler optimizing away reads without writes etc.
+		UE::Learning::TSharedMemoryArrayView<2, float> Observations;
+		UE::Learning::TSharedMemoryArrayView<2, float> Actions;
+
+		// Training Process
+
+		TSharedPtr<FMonitoredProcess> TrainingProcess;
+
+		void HandleTrainingProcessCanceled();
+		void HandleTrainingProcessCompleted(int32 ReturnCode);
+		static void HandleTrainingProcessOutput(FString Output);
+	};
+
+	/**
+	* This object allows you to launch the FSocketImitationTrainer server as a subprocess,
+	* which is convenient when you want to train using it locally.
+	*/
+	struct LEARNINGTRAINING_API FSocketImitationTrainerServerProcess
+	{
+		/**
+		* Creates a training server as a subprocess
+		*
+		* @param PythonExecutablePath		Path to the python executable used for training. In general should be the python shipped with Unreal Editor.
+		* @param SitePackagesPath			Path to the site-packages shipped with the PythonFoundationPackages plugin
+		* @param PythonContentPath			Path to the Python Content folder provided by the Learning plugin
+		* @param IntermediatePath			Path to the intermediate folder to write temporary files, logs, and snapshots to
+		* @param IpAddress					Ip address to bind the listening socket to. For a local server you will want to use 127.0.0.1
+		* @param Port						Port to use for the listening socket.
+		* @param LogSettings				Logging settings to use
+		*/
+		FSocketImitationTrainerServerProcess(
+			const FString& PythonExecutablePath,
+			const FString& SitePackagesPath,
+			const FString& PythonContentPath,
+			const FString& IntermediatePath,
+			const TCHAR* IpAddress = Trainer::DefaultIp,
+			const uint32 Port = Trainer::DefaultPort,
+			const ELogSetting LogSettings = ELogSetting::Normal);
+
+		~FSocketImitationTrainerServerProcess();
+
+		/**
+		* Check if the server process is still running
+		*/
+		bool IsRunning() const;
+
+		/**
+		* Wait for the server process to end
+		*
+		* @param Timeout		Timeout to wait in seconds
+		* @returns				true if successful, otherwise false if it times out
+		*/
+		bool Wait(float Timeout);
+
+		/**
+		* Terminate the server process
+		*/
+		void Terminate();
+
+	private:
+
+		TSharedPtr<FMonitoredProcess> TrainingProcess;
+
+		void HandleTrainingProcessCanceled();
+		void HandleTrainingProcessCompleted(int32 ReturnCode);
+		static void HandleTrainingProcessOutput(FString Output);
+	};
+
+	struct LEARNINGTRAINING_API FSocketImitationTrainer : public IImitationTrainer
+	{
+		/**
+		* Create a new socket-based imitation trainer 
+		*
+		* @param OutResponse				Response to the initial connection
+		* @param TaskName					Name of the training task - used to help identify the logs, snapshots, and other files generated by training
+		* @param Network					Neural Network object
+		* @param ActionNoiseMin				Minimum amount of action noise being used during episode generation
+		* @param ActionNoiseMax				Maximum amount of action noise being used during episode generation
+		* @param MaxSampleNum				Maximum number of samples in the training data
+		* @param IpAddress					Server Ip address
+		* @param Port						Server Port
+		* @param Timeout					Timeout to wait in seconds for connection and initial data transfer
+		* @param Settings					Training settings
+		* @param LogSettings				Logging settings to use
+		*/
+		FSocketImitationTrainer(
+			ETrainerResponse& OutResponse,
+			const FString& TaskName,
+			const FNeuralNetwork& Network,
+			const float ActionNoiseMin,
+			const float ActionNoiseMax,
+			const int32 MaxSampleNum,
+			const TCHAR* IpAddress = Trainer::DefaultIp,
+			const uint32 Port = Trainer::DefaultPort,
+			const float Timeout = Trainer::DefaultTimeout,
+			const FImitationTrainerSettings& Settings = FImitationTrainerSettings());
+
+		~FSocketImitationTrainer();
+
+		virtual void Terminate() override final;
+
+		virtual ETrainerResponse Wait(const float Timeout) override final;
+
+		virtual ETrainerResponse SendStop(const float Timeout) override final;
+
+		virtual ETrainerResponse SendContinue(const float Timeout) override final;
+
+		virtual ETrainerResponse RecvPolicy(
+			FNeuralNetwork& OutNetwork,
+			const float Timeout = Trainer::DefaultTimeout,
+			FRWLock* NetworkLock = nullptr,
+			const ELogSetting LogSettings = Trainer::DefaultLogSettings) override final;
+
+		virtual ETrainerResponse SendPolicy(
+			const FNeuralNetwork& Network,
+			const float Timeout = Trainer::DefaultTimeout,
+			FRWLock* NetworkLock = nullptr,
+			const ELogSetting LogSettings = Trainer::DefaultLogSettings) override final;
+
+		virtual ETrainerResponse SendExperience(
+			const TLearningArrayView<2, const float> ObservationVectors,
+			const TLearningArrayView<2, const float> ActionVectors,
+			const float Timeout = Trainer::DefaultTimeout,
+			const ELogSetting LogSettings = Trainer::DefaultLogSettings) override final;
+
+	private:
+
+		TLearningArray<1, uint8> NetworkBuffer;
+		FSocket* Socket = nullptr;
+	};
+
+	namespace ImitationTrainer
+	{
+		/**
+		* Train a policy using experience already gathered from example episodes
+		*
+		* @param Trainer							Trainer
+		* @param Network							Policy network
+		* @param ObservationVectors					Observation Data
+		* @param ActionVectors						Action Data
+		* @param bReinitializeNetwork				Was the network reinitialized to random weights or used as-is
+		* @param bRequestTrainingStopSignal			Optional signal that can be raised to indicate training should be stopped
+		* @param NetworkLock						Optional Lock to use when updating the policy network
+		* @param bNetworkUpdatedSignal				Optional signal that will be raised when the policy network is updated
+		* @param LogSettings						Logging settings
+		* @returns									Trainer response in case of errors during communication otherwise Success
+		*/
+		LEARNINGTRAINING_API ETrainerResponse Train(
+			IImitationTrainer& Trainer,
+			FNeuralNetwork& Network,
+			const TLearningArrayView<2, const float> ObservationVectors,
+			const TLearningArrayView<2, const float> ActionVectors,
+			const bool bReinitializeNetwork = true,
+			TAtomic<bool>* bRequestTrainingStopSignal = nullptr,
+			FRWLock* NetworkLock = nullptr,
+			TAtomic<bool>* bNetworkUpdatedSignal = nullptr,
+			const ELogSetting LogSettings = ELogSetting::Normal);
+	}
+
+}
diff --git a/Engine/Plugins/Experimental/LearningAgents/Source/LearningTraining/Public/LearningPPOTrainer.h b/Engine/Plugins/Experimental/LearningAgents/Source/LearningTraining/Public/LearningPPOTrainer.h
new file mode 100644
--- /dev/null
+++ b/Engine/Plugins/Experimental/LearningAgents/Source/LearningTraining/Public/LearningPPOTrainer.h
@@ -0,0 +1,497 @@
+// Copyright Epic Games, Inc. All Rights Reserved.
+
+#pragma once
+
+#include "LearningArray.h"
+#include "LearningLog.h"
+#include "LearningTrainer.h"
+#include "LearningSharedMemory.h"
+
+#include "Commandlets/Commandlet.h"
+#include "Templates/SharedPointer.h"
+
+#include "LearningPPOTrainer.generated.h"
+
+class FSocket;
+class FMonitoredProcess;
+
+UCLASS()
+class LEARNINGTRAINING_API ULearningSocketPPOTrainerServerCommandlet : public UCommandlet
+{
+	GENERATED_UCLASS_BODY()
+
+	/** Runs the commandlet */
+	virtual int32 Main(const FString& Params) override;
+};
+
+namespace UE::Learning
+{
+	struct FNeuralNetwork;
+	struct FReplayBuffer;
+	struct FResetInstanceBuffer;
+	struct FEpisodeBuffer;
+	enum class ECompletionMode : uint8;
+
+	/**
+	* Settings to use when training using PPO
+	*/
+	struct FPPOTrainerSettings
+	{
+		// Number of iterations to train the network for. Controls the overall training time.
+		// Training for about 100000 iterations should give you well trained network, but
+		// closer to 1000000 iterations or more is required for an exhaustively trained network.
+		uint32 IterationNum = 1000000;
+
+		// Learning rate of the actor network. Typical values are between 0.001f and 0.0001f
+		float LearningRateActor = 0.0001f;
+
+		// Learning rate of the critic network. To avoid instability generally the critic 
+		// should have a larger learning rate then the actor.
+		float LearningRateCritic = 0.001f;
+
+		// Ratio by which to decay the learning rate every 1000 iterations.
+		float LearningRateDecay = 0.99f;
+
+		// Amount of weight decay to apply to the network. Larger values encourage network 
+		// weights to be smaller but too large a value can cause the network weights to collapse to all zeros.
+		float WeightDecay = 0.001f;
+
+		// Initial scale to apply to actions before noise is added to them. The smaller this is, 
+		// the less likely you are to have spurious correlations at the beginning of training which 
+		// can make things slow or unstable. Too small and the network may become difficult to train.
+		float InitialActionScale = 0.1f;
+
+		// Batch size to use for training. Smaller values tend to produce better results 
+		// at the cost of slowing down training.
+		uint32 BatchSize = 128;
+
+		// Clipping ratio to apply to policy updates. Keeps the training "on-policy". 
+		// Larger values may speed up training at the cost of stability. Conversely, too small 
+		// values will keep the policy from being unable to learn an optimal policy.
+		float EpsilonClip = 0.2f;
+
+		// Weight used to regularize actions. Larger values will encourage smaller actions but too large
+		// will cause actions to become always zero.
+		float ActionRegularizationWeight = 0.001f;
+
+		// Weighting used for the entropy bonus. Larger values encourage larger action 
+		// noise and therefore greater exploration but can make actions very noisy.
+		float EntropyWeight = 0.01f;
+
+		// This is used in the Generalized Advantage Estimation as what is essentially 
+		// an exponential smoothing/decay. Typical values should be between 0.9 and 1.0.
+		float GaeLambda = 0.9f;
+
+		// If to clip very large or small advantages. This has few downsides and helps 
+		// with numerical stability.
+		bool bClipAdvantages = true;
+
+		// If to normalize advantages. This tends to makes training more robust to adjustments 
+		// of the scale of rewards.
+		bool bAdvantageNormalization = true;
+
+		// Number of steps to trim from the start of each episode during training. This can
+		// be useful if some reset process is taking several steps or you know your starting
+		// states are not entirely valid for example.
+		int32 TrimEpisodeStartStepNum = 0;
+
+		// Number of steps to trim from the end of each episode during training. This can be
+		// useful if you know the last few steps of an episode are not valid or contain incorrect
+		// information.
+		int32 TrimEpisodeEndStepNum = 0;
+
+		// Random Seed to use for training
+		uint32 Seed = 1234;
+
+		// The discount factor causes future rewards to be scaled down so that the policy will 
+		// favor near-term rewards over potentially uncertain long-term rewards. Larger values 
+		// encourage the system to "look-ahead" but make training more difficult.
+		float DiscountFactor = 0.99f;
+		
+		// If to reinitialize the network to random weights at the start of training. Set this
+		// to false if you are start from a network which has already been trained such as via 
+		// Imitation Training. Make sure to use the same settings here for PPOTrainer::Train.
+		bool bReinitializeNetwork = true;
+
+		// Which device to use for training
+		ETrainerDevice Device = ETrainerDevice::GPU;
+
+		// The number of processes to use for multi-processed training.
+		// 
+		// It is important to know how this multi-process training works so that it can be used
+		// correctly when you set this >1:
+		// 
+		// When called with this argument set >1, the process will spawn additional processes running 
+		// the same command as is currently being run but with the additional command line argument 
+		// `LearningProcessIdx`. Once the training starts on these subprocesses process, this command 
+		// line argument will be used to indicate that this subprocess should only be used for gathering 
+		// experience and should not start a new training process. 
+		// 
+		// This means that this will generally not work in editor (or PIE)
+		// 
+		// IMPORTANT: If you are seeding your experience generation process to make it deterministic
+		// you should use this `LearningProcessIdx` command line argument to change the seed of your
+		// experience gathering - otherwise each process will gather identical experience.
+		uint16 ProcessNum = 1;
+
+		// If to hide the window for each multi-processing subprocess
+		bool bMultiProcessHideTrainingWindow = false;
+
+		// If to redirect the output of each multi-processing subprocess to the main process
+		bool bMultiProcessRedirectTrainingOutput = false;
+
+		// If to use TensorBoard for logging and tracking the training progress.
+		// 
+		// Even when enabled, TensorBoard will only work if it is installed in your Unreal Editor
+		// bundled version of Python, which is not the case by default. TensorBoard can be installed 
+		// for this version of Python by going to your Unreal Editor Python Binaries directory 
+		// (e.g. "\Engine\Binaries\ThirdParty\Python3\Win64") and running `.\python -m pip install tensorboard`. 
+		bool bUseTensorboard = false;
+	};
+
+	/**
+	* Interface for an object which can train a policy using PPO.
+	*/
+	struct IPPOTrainer
+	{
+		virtual ~IPPOTrainer() {}
+
+		/**
+		* Terminate the trainer immediately.
+		*/
+		virtual void Terminate() = 0;
+
+		/**
+		* Wait for the trainer to finish.
+		* 
+		* @param Timeout		Timeout to wait in seconds
+		* @returns				Trainer response
+		*/
+		virtual ETrainerResponse Wait(const float Timeout = Trainer::DefaultTimeout) = 0;
+
+		/**
+		* Wait for the trainer to push an updated policy network.
+		*
+		* @param OutNetwork		Network to update
+		* @param Timeout		Timeout to wait in seconds
+		* @param NetworkLock	Lock to use when updating network
+		* @param LogSettings	Log settings
+		* @returns				Trainer response
+		*/
+		virtual ETrainerResponse RecvPolicy(
+			FNeuralNetwork& OutNetwork,
+			const float Timeout = Trainer::DefaultTimeout,
+			FRWLock* NetworkLock = nullptr,
+			const ELogSetting LogSettings = Trainer::DefaultLogSettings) = 0;
+
+		/*
+		* Signal for the trainer to stop.
+		*/
+		virtual ETrainerResponse SendStop(const float Timeout = Trainer::DefaultTimeout) = 0;
+
+		/**
+		* Wait for the trainer to be ready and push the current policy network.
+		*
+		* @param Network		Network to push
+		* @param Timeout		Timeout to wait in seconds
+		* @param NetworkLock	Lock to use when pushing network
+		* @param LogSettings	Log settings
+		* @returns				Trainer response
+		*/
+		virtual ETrainerResponse SendPolicy(
+			const FNeuralNetwork& Network,
+			const float Timeout = Trainer::DefaultTimeout,
+			FRWLock* NetworkLock = nullptr,
+			const ELogSetting LogSettings = Trainer::DefaultLogSettings) = 0;
+
+		/**
+		* Wait for the trainer to be ready and push new experience.
+		*
+		* @param ReplayBuffer	Replay buffer of experience to push
+		* @param Timeout		Timeout to wait in seconds
+		* @param LogSettings	Log settings
+		* @returns				Trainer response
+		*/
+		virtual ETrainerResponse SendExperience(
+			const FReplayBuffer& ReplayBuffer,
+			const float Timeout = Trainer::DefaultTimeout,
+			const ELogSetting LogSettings = Trainer::DefaultLogSettings) = 0;
+	};
+
+	/**
+	* Trainer that uses shared memory and a Python sub-process to perform training
+	* 
+	* This trainer is the most simple and efficient when training the policy on the 
+	* same computer that experience is being gathered on.
+	*/
+	struct LEARNINGTRAINING_API FSharedMemoryPPOTrainer : public IPPOTrainer
+	{
+		/**
+		* Creates a new Shared Memory PPO trainer
+		*
+		* @param TaskName					Name of the training task - used to help identify the logs, snapshots, and other files generated by training
+		* @param PythonExecutablePath		Path to the python executable used for training. In general should be the python shipped with Unreal Editor.
+		* @param SitePackagesPath			Path to the site-packages shipped with the PythonFoundationPackages plugin
+		* @param PythonContentPath			Path to the Python Content folder provided by the Learning plugin
+		* @param IntermediatePath			Path to the intermediate folder to write temporary files, logs, and snapshots to
+		* @param Network					Neural Network object
+		* @param ActionNoiseMin				Minimum Action Noise used by Policy during experience gathering
+		* @param ActionNoiseMax				Maximum Action Noise used by Policy during experience gathering
+		* @param ReplayBuffer				Replay buffer used to collect experience
+		* @param Settings					Trainer settings
+		* @param LogSettings				Logging settings to use
+		* @param bHideTrainingWindow		If to hide the training sub-process console window
+		* @param bRedirectTrainingOutput	If to redirect the training console window output to the log
+		*/
+		FSharedMemoryPPOTrainer(
+			const FString& TaskName,
+			const FString& PythonExecutablePath,
+			const FString& SitePackagesPath,
+			const FString& PythonContentPath,
+			const FString& IntermediatePath,
+			const FNeuralNetwork& Network,
+			const float ActionNoiseMin,
+			const float ActionNoiseMax,
+			const FReplayBuffer& ReplayBuffer,
+			const FPPOTrainerSettings& Settings = FPPOTrainerSettings(),
+			const ELogSetting LogSettings = ELogSetting::Normal,
+			const bool bHideTrainingWindow = true,
+			const bool bRedirectTrainingOutput = true);
+
+		~FSharedMemoryPPOTrainer();
+
+		virtual void Terminate() override final;
+
+		virtual ETrainerResponse Wait(const float Timeout) override final;
+
+		virtual ETrainerResponse RecvPolicy(
+			FNeuralNetwork& OutNetwork,
+			const float Timeout = Trainer::DefaultTimeout,
+			FRWLock* NetworkLock = nullptr,
+			const ELogSetting LogSettings = Trainer::DefaultLogSettings) override final;
+
+		virtual ETrainerResponse SendStop(const float Timeout = Trainer::DefaultTimeout) override final;
+
+		virtual ETrainerResponse SendPolicy(
+			const FNeuralNetwork& Network,
+			const float Timeout = Trainer::DefaultTimeout,
+			FRWLock* NetworkLock = nullptr,
+			const ELogSetting LogSettings = Trainer::DefaultLogSettings) override final;
+
+		virtual ETrainerResponse SendExperience(
+			const FReplayBuffer& ReplayBuffer,
+			const float Timeout = Trainer::DefaultTimeout,
+			const ELogSetting LogSettings = Trainer::DefaultLogSettings) override final;
+
+	private:
+
+		/**
+		* Free and deallocate all shared memory
+		*/
+		void Deallocate();
+
+		// Shared memory
+
+		UE::Learning::TSharedMemoryArrayView<1, uint8> Policy;
+		UE::Learning::TSharedMemoryArrayView<2, volatile int32> Controls; // Mark as volatile to avoid compiler optimizing away reads without writes etc.
+		UE::Learning::TSharedMemoryArrayView<2, int32> EpisodeStarts;
+		UE::Learning::TSharedMemoryArrayView<2, int32> EpisodeLengths;
+		UE::Learning::TSharedMemoryArrayView<2, ECompletionMode> EpisodeCompletionModes;
+		UE::Learning::TSharedMemoryArrayView<3, float> EpisodeFinalObservations;
+		UE::Learning::TSharedMemoryArrayView<3, float> Observations;
+		UE::Learning::TSharedMemoryArrayView<3, float> Actions;
+		UE::Learning::TSharedMemoryArrayView<2, float> Rewards;
+
+		// Training Process
+
+		uint16 ProcessIdx = INDEX_NONE;
+		TSharedPtr<FMonitoredProcess> TrainingProcess;
+		TArray<TSharedPtr<FMonitoredProcess>, TInlineAllocator<128>> ExperienceGatheringSubprocesses;
+
+		void HandleSubprocessCanceled();
+		void HandleSubprocessCompleted(int32 ReturnCode);
+		static void HandleSubprocessOutput(FString Output);
+
+		void HandleTrainingProcessCanceled();
+		void HandleTrainingProcessCompleted(int32 ReturnCode);
+		static void HandleTrainingProcessOutput(FString Output);
+	};
+
+	/**
+	* This object allows you to launch the FSocketPPOTrainer server as a subprocess, 
+	* which is convenient when you want to train using it locally.
+	*/
+	struct LEARNINGTRAINING_API FSocketPPOTrainerServerProcess
+	{
+		/**
+		* Creates a training server as a subprocess
+		*
+		* @param PythonExecutablePath		Path to the python executable used for training. In general should be the python shipped with Unreal Editor.
+		* @param SitePackagesPath			Path to the site-packages shipped with the PythonFoundationPackages plugin
+		* @param PythonContentPath			Path to the Python Content folder provided by the Learning plugin
+		* @param IntermediatePath			Path to the intermediate folder to write temporary files, logs, and snapshots to
+		* @param IpAddress					Ip address to bind the listening socket to. For a local server you will want to use 127.0.0.1
+		* @param Port						Port to use for the listening socket.
+		* @param LogSettings				Logging settings to use
+		*/
+		FSocketPPOTrainerServerProcess(
+			const FString& PythonExecutablePath,
+			const FString& SitePackagesPath,
+			const FString& PythonContentPath,
+			const FString& IntermediatePath,
+			const TCHAR* IpAddress = Trainer::DefaultIp,
+			const uint32 Port = Trainer::DefaultPort,
+			const ELogSetting LogSettings = ELogSetting::Normal);
+
+		~FSocketPPOTrainerServerProcess();
+
+		/**
+		* Check if the server process is still running
+		*/
+		bool IsRunning() const;
+
+		/**
+		* Wait for the server process to end
+		*
+		* @param Timeout		Timeout to wait in seconds
+		* @returns				true if successful, otherwise false if it times out
+		*/
+		bool Wait(float Timeout);
+
+		/**
+		* Terminate the server process
+		*/
+		void Terminate();
+
+	private:
+
+		TSharedPtr<FMonitoredProcess> TrainingProcess;
+
+		void HandleTrainingProcessCanceled();
+		void HandleTrainingProcessCompleted(int32 ReturnCode);
+		static void HandleTrainingProcessOutput(FString Output);
+	};
+
+	/**
+	* Trainer that connects to an external training server to perform training
+	*
+	* This trainer can be used to allow the python training process the run
+	* on a different machine to the experience gathering process.
+	*/
+	struct LEARNINGTRAINING_API FSocketPPOTrainer : public IPPOTrainer
+	{
+		/**
+		* Creates a new Socket PPO trainer
+		*
+		* @param OutResponse				Response to the initial connection
+		* @param TaskName					Name of the training task - used to help identify the logs, snapshots, and other files generated by training
+		* @param Network					Neural Network object
+		* @param ActionNoiseMin				Minimum Action Noise used by Policy during experience gathering
+		* @param ActionNoiseMax				Maximum Action Noise used by Policy during experience gathering
+		* @param ReplayBuffer				Replay buffer used to collect experience
+		* @param IpAddress					Server Ip address
+		* @param Port						Server Port
+		* @param Timeout					Timeout to wait in seconds for connection and initial data transfer
+		* @param Settings					Trainer settings
+		*/
+		FSocketPPOTrainer(
+			ETrainerResponse& OutResponse,
+			const FString& TaskName,
+			const FNeuralNetwork& Network,
+			const float ActionNoiseMin,
+			const float ActionNoiseMax,
+			const FReplayBuffer& ReplayBuffer,
+			const TCHAR* IpAddress = Trainer::DefaultIp,
+			const uint32 Port = Trainer::DefaultPort,
+			const float Timeout = Trainer::DefaultTimeout,
+			const FPPOTrainerSettings& Settings = FPPOTrainerSettings());
+
+		~FSocketPPOTrainer();
+
+		virtual void Terminate() override final;
+
+		virtual ETrainerResponse Wait(const float Timeout) override final;
+
+		virtual ETrainerResponse RecvPolicy(
+			FNeuralNetwork& OutNetwork,
+			const float Timeout = Trainer::DefaultTimeout,
+			FRWLock* NetworkLock = nullptr,
+			const ELogSetting LogSettings = Trainer::DefaultLogSettings) override final;
+
+		virtual ETrainerResponse SendStop(const float Timeout = Trainer::DefaultTimeout) override final;
+
+		virtual ETrainerResponse SendPolicy(
+			const FNeuralNetwork& Network,
+			const float Timeout = Trainer::DefaultTimeout,
+			FRWLock* NetworkLock = nullptr,
+			const ELogSetting LogSettings = Trainer::DefaultLogSettings) override final;
+
+		virtual ETrainerResponse SendExperience(
+			const FReplayBuffer& ReplayBuffer,
+			const float Timeout = Trainer::DefaultTimeout,
+			const ELogSetting LogSettings = Trainer::DefaultLogSettings) override final;
+
+	private:
+
+		TLearningArray<1, uint8> NetworkBuffer;
+		FSocket* Socket = nullptr;
+	};
+
+
+	namespace PPOTrainer
+	{
+		/**
+		* Train a policy while gathering experience
+		*
+		* @param Trainer						Trainer
+		* @param ReplayBuffer					Replay Buffer
+		* @param EpisodeBuffer					Episode Buffer
+		* @param ResetBuffer					Reset Buffer
+		* @param Network						Network to use
+		* @param ObservationVectorBuffer		Buffer to read/write observation vectors into
+		* @param ActionVectorBuffer				Buffer to read/write action vectors into
+		* @param RewardBuffer					Buffer to read/write rewards into
+		* @param CompletionBuffer				Buffer to read/write completions into
+		* @param EpisodeEndCompletionMode		Completion mode to use for episodes that reach the max length
+		* @param ResetFunction					Function to run for resetting the environment
+		* @param ObservationFunction			Function to run for evaluating observations
+		* @param PolicyFunction					Function to run for evaluating the policy
+		* @param ActionFunction					Function to run for evaluating actions
+		* @param UpdateFunction					Function to run for updating the environment
+		* @param RewardFunction					Function to run for evaluating rewards
+		* @param CompletionFunction				Function to run for evaluating completions
+		* @param Instances						Set of instances to run training for
+		* @param bReinitializeNetwork			Was the network reinitialized to random weights or used as-is
+		* @param bRequestTrainingStopSignal		Optional signal that can be set to indicate training should be stopped
+		* @param NetworkLock					Optional Lock to use when updating the policy network
+		* @param bNetworkUpdatedSignal			Optional signal that will be set when the policy network is updated
+		* @param LogSettings					Logging settings
+		* @returns								Trainer response in case of errors during communication otherwise Success
+		*/
+		LEARNINGTRAINING_API ETrainerResponse Train(
+			IPPOTrainer& Trainer,
+			FReplayBuffer& ReplayBuffer,
+			FEpisodeBuffer& EpisodeBuffer,
+			FResetInstanceBuffer& ResetBuffer,
+			FNeuralNetwork& Network,
+			TLearningArrayView<2, float> ObservationVectorBuffer,
+			TLearningArrayView<2, float> ActionVectorBuffer,
+			TLearningArrayView<1, float> RewardBuffer,
+			TLearningArrayView<1, ECompletionMode> CompletionBuffer,
+			const ECompletionMode EpisodeEndCompletionMode,
+			const TFunctionRef<void(const FIndexSet Instances)> ResetFunction,
+			const TFunctionRef<void(const FIndexSet Instances)> ObservationFunction,
+			const TFunctionRef<void(const FIndexSet Instances)> PolicyFunction,
+			const TFunctionRef<void(const FIndexSet Instances)> ActionFunction,
+			const TFunctionRef<void(const FIndexSet Instances)> UpdateFunction,
+			const TFunctionRef<void(const FIndexSet Instances)> RewardFunction,
+			const TFunctionRef<void(const FIndexSet Instances)> CompletionFunction,
+			const FIndexSet Instances,
+			const bool bReinitializeNetwork = true,
+			TAtomic<bool>* bRequestTrainingStopSignal = nullptr,
+			FRWLock* NetworkLock = nullptr,
+			TAtomic<bool>* bNetworkUpdatedSignal = nullptr,
+			const ELogSetting LogSettings = ELogSetting::Normal);
+	}
+
+}
diff --git a/Engine/Plugins/Experimental/LearningAgents/Source/LearningTraining/Public/LearningTrainer.h b/Engine/Plugins/Experimental/LearningAgents/Source/LearningTraining/Public/LearningTrainer.h
new file mode 100644
--- /dev/null
+++ b/Engine/Plugins/Experimental/LearningAgents/Source/LearningTraining/Public/LearningTrainer.h
@@ -0,0 +1,97 @@
+// Copyright Epic Games, Inc. All Rights Reserved.
+
+#pragma once
+
+#include "LearningArray.h"
+#include "LearningLog.h"
+
+namespace UE::Learning
+{
+	/**
+	* Training Device
+	*/
+	enum class ETrainerDevice : uint8
+	{
+		CPU = 0,
+		GPU = 1,
+	};
+
+	/**
+	* Type of response from a Trainer
+	*/
+	enum class ETrainerResponse : uint8
+	{
+		// The communication was successful
+		Success = 0,
+
+		// The communication send or received was unexpected
+		Unexpected = 1,
+
+		// Training is complete
+		Completed = 2,
+
+		// Training is stopped
+		Stopped = 3,
+
+		// The communication timed-out
+		Timeout = 4,
+	};
+
+	namespace Trainer
+	{
+		/**
+		* Default Timeout to use during communication.
+		*/
+		static constexpr float DefaultTimeout = 10.0f;
+
+		/**
+		* Default Log Settings to use during communication.
+		*/
+		static constexpr ELogSetting DefaultLogSettings = ELogSetting::Silent;
+
+		/**
+		* Default IP to use for networked training
+		*/
+		static constexpr const TCHAR* DefaultIp = TEXT("127.0.0.1");
+
+		/**
+		* Default Port to use for networked training
+		*/
+		static constexpr uint32 DefaultPort = 48491;
+
+		/**
+		* Converts a ETrainerDevice into a string.
+		*/
+		LEARNINGTRAINING_API const TCHAR* GetDeviceString(const ETrainerDevice Device);
+
+		/**
+		* Converts a ETrainerResponse into a string for use in logging and error messages.
+		*/
+		LEARNINGTRAINING_API const TCHAR* ResponseString(const ETrainerResponse Response);
+
+		/**
+		* Compute the discount factor that corresponds to a particular HalfLife and DeltaTime.
+		*
+		* @param HalfLife		Time by which the reward should be discounted by half
+		* @param DeltaTime		DeltaTime taken upon each step of the environment
+		* @returns				Corresponding discount factor
+		*/
+		LEARNINGTRAINING_API float DiscountFactorFromHalfLife(const float HalfLife, const float DeltaTime);
+
+		/**
+		* Compute the discount factor that corresponds to a particular HalfLife provided in terms of number of steps
+		*
+		* @param HalfLifeSteps	Number of steps taken at which the reward should be discounted by half
+		* @returns				Corresponding discount factor
+		*/
+		LEARNINGTRAINING_API float DiscountFactorFromHalfLifeSteps(const int32 HalfLifeSteps);
+
+#if WITH_EDITOR
+		LEARNINGTRAINING_API FString DefaultEditorPythonExecutablePath();
+		LEARNINGTRAINING_API FString DefaultEditorSitePackagesPath();
+		LEARNINGTRAINING_API FString DefaultEditorPythonContentPath();
+		LEARNINGTRAINING_API FString DefaultEditorIntermediatePath();
+#endif
+	}
+
+}
\ No newline at end of file
