LearningAgents: Large refactoring to allow critic to be loaded and used at runtime. Fixed some other misc bugs and fallout from switching features to all be arrays by default. Some refactoring of different LearningAgents components such as the Policy and Critic and simplified things by removing some of the callbacks that were used to structure setup code and agent management.

diff --git a/Engine/Plugins/Experimental/LearningAgents/Source/Learning/Private/Learning.ispc b/Engine/Plugins/Experimental/LearningAgents/Source/Learning/Private/Learning.ispc
--- a/Engine/Plugins/Experimental/LearningAgents/Source/Learning/Private/Learning.ispc
+++ b/Engine/Plugins/Experimental/LearningAgents/Source/Learning/Private/Learning.ispc
@@ -1,9 +1,5 @@
 // Copyright Epic Games, Inc. All Rights Reserved.
 
-#ifndef MIN_flt
-#define MIN_flt			(1.175494351e-38F)
-#endif
-
 #ifndef MAX_flt
 #define MAX_flt			(3.402823466e+38F)
 #endif
@@ -225,6 +221,8 @@ inline float LearningSigmoid(const float X)
 
 export void LearningLayerActionNoise(
 	uniform float Output[],
+	uniform float OutputMean[],
+	uniform float OutputStd[],
 	const uniform float Input[],
 	const uniform uint32 Seed[],
 	const uniform float ActionNoiseScale[],
@@ -236,15 +234,14 @@ export void LearningLayerActionNoise(
 {
 	foreach (BatchIdx = 0 ... BatchNum, OutputIdx = 0 ... OutputNum)
 	{
-		const float ActionMean = Input[BatchIdx * InputNum + OutputIdx];
-		const float LogActionStd = Input[BatchIdx * InputNum + OutputNum + OutputIdx];
-		const float ActionStd = ActionNoiseScale[BatchIdx] * 
-			exp(LearningSigmoid(LogActionStd) * (LogActionNoiseMax - LogActionNoiseMin) + LogActionNoiseMin);
+		OutputMean[BatchIdx * OutputNum + OutputIdx] = Input[BatchIdx * InputNum + OutputIdx];
+		OutputStd[BatchIdx * OutputNum + OutputIdx] = ActionNoiseScale[BatchIdx] *
+			exp(LearningSigmoid(Input[BatchIdx * InputNum + OutputNum + OutputIdx]) * (LogActionNoiseMax - LogActionNoiseMin) + LogActionNoiseMin);
 
 		Output[BatchIdx * OutputNum + OutputIdx] = LearningRandomGaussian(
 			Seed[BatchIdx] ^ 0xab744615 ^ LearningRandomInt(OutputIdx ^ 0xf8a88a27),
-			ActionMean,
-			ActionStd);
+			OutputMean[BatchIdx * OutputNum + OutputIdx],
+			OutputStd[BatchIdx * OutputNum + OutputIdx]);
 	}
 }
 
@@ -678,8 +675,8 @@ export void LearningKMeansUpdateBounds(
 {
 	foreach(ClusterIdx = 0 ... ClusterNum, DimIdx = 0 ... DimNum)
 	{
-		OutMins[ClusterIdx * DimNum + DimIdx] = MAX_flt;
-		OutMaxs[ClusterIdx * DimNum + DimIdx] = MIN_flt;
+		OutMins[ClusterIdx * DimNum + DimIdx] = +MAX_flt;
+		OutMaxs[ClusterIdx * DimNum + DimIdx] = -MAX_flt;
 	}
 
 	for (uniform int SampleIdx = 0; SampleIdx < SampleNum; SampleIdx++)
diff --git a/Engine/Plugins/Experimental/LearningAgents/Source/Learning/Private/LearningArrayMap.cpp b/Engine/Plugins/Experimental/LearningAgents/Source/Learning/Private/LearningArrayMap.cpp
--- a/Engine/Plugins/Experimental/LearningAgents/Source/Learning/Private/LearningArrayMap.cpp
+++ b/Engine/Plugins/Experimental/LearningAgents/Source/Learning/Private/LearningArrayMap.cpp
@@ -82,7 +82,7 @@ namespace UE::Learning
 		return Handles.Contains(Key);
 	}
 
-	void FArrayMap::Link(const FArrayMapKey Src, const FArrayMapKey Dst)
+	void FArrayMap::LinkKeys(const FArrayMapKey Src, const FArrayMapKey Dst)
 	{
 		const FArrayMapHandle* SrcVariable = Handles.Find(Src);
 		const FArrayMapHandle* DstVariable = Handles.Find(Dst);
@@ -147,7 +147,7 @@ namespace UE::Learning
 		Flags[DstIndex] = FlagLinked;
 	}
 
-	void FArrayMap::Link(const FArrayMapHandle Src, const FArrayMapHandle Dst)
+	void FArrayMap::LinkHandles(const FArrayMapHandle Src, const FArrayMapHandle Dst)
 	{
 		const int16 SrcIndex = Src.Index;
 		const int16 DstIndex = Dst.Index;
@@ -199,6 +199,61 @@ namespace UE::Learning
 		Flags[DstIndex] = FlagLinked;
 	}
 
+	void FArrayMap::LinkHandlesFlat(const FArrayMapHandle Src, const FArrayMapHandle Dst)
+	{
+		const int16 SrcIndex = Src.Index;
+		const int16 DstIndex = Dst.Index;
+
+		UE_LEARNING_CHECKF(!(Flags[DstIndex] & FlagLinked), TEXT("Destination array is already linked!"));
+
+		UE_LEARNING_CHECKF(TypeIds[SrcIndex] == TypeIds[DstIndex],
+			TEXT("Source array is of type %s, while destination array is of type %s"),
+			TypeNames[SrcIndex], TypeNames[DstIndex]);
+
+		int32 SrcTotalSize = Arrays[SrcIndex].Shape[0];
+		for (uint8 SrcDimIdx = 1; SrcDimIdx < DimNums[SrcIndex]; SrcDimIdx++)
+		{
+			SrcTotalSize *= Arrays[SrcIndex].Shape[SrcDimIdx];
+		}
+
+		int32 DstTotalSize = Arrays[DstIndex].Shape[0];
+		for (uint8 DstDimIdx = 1; DstDimIdx < DimNums[DstIndex]; DstDimIdx++)
+		{
+			DstTotalSize *= Arrays[DstIndex].Shape[DstDimIdx];
+		}
+
+		UE_LEARNING_CHECKF(SrcTotalSize == DstTotalSize,
+			TEXT("Array total sizes don't match (%i vs %i)"),
+			SrcTotalSize, DstTotalSize);
+
+		// Free the Destination array data
+		if (Flags[DstIndex] & FlagConstructed)
+		{
+			Destructors[DstIndex](
+				Arrays[DstIndex].Data,
+				DimNums[DstIndex],
+				Arrays[DstIndex].Shape);
+		}
+
+		if (Flags[DstIndex] & FlagAllocated)
+		{
+			FMemory::Free(Arrays[DstIndex].Data);
+		}
+
+		// Re-direct existing Links
+		for (int16 ArrayIdx = 0; ArrayIdx < Arrays.Num(); ArrayIdx++)
+		{
+			if (Flags[ArrayIdx] == FlagLinked && Arrays[ArrayIdx].Data == Arrays[DstIndex].Data)
+			{
+				Arrays[ArrayIdx].Data = Arrays[SrcIndex].Data;
+			}
+		}
+
+		// Insert the new Link
+		Arrays[DstIndex].Data = Arrays[SrcIndex].Data;
+		Flags[DstIndex] = FlagLinked;
+	}
+
 	bool FArrayMap::HasLink(const FArrayMapKey Dst) const
 	{
 		const FArrayMapHandle* DstVariable = Handles.Find(Dst);
diff --git a/Engine/Plugins/Experimental/LearningAgents/Source/Learning/Private/LearningCompletionObject.cpp b/Engine/Plugins/Experimental/LearningAgents/Source/Learning/Private/LearningCompletionObject.cpp
--- a/Engine/Plugins/Experimental/LearningAgents/Source/Learning/Private/LearningCompletionObject.cpp
+++ b/Engine/Plugins/Experimental/LearningAgents/Source/Learning/Private/LearningCompletionObject.cpp
@@ -88,12 +88,13 @@ namespace UE::Learning
 		const FName& InIdentifier,
 		const TSharedRef<FArrayMap>& InInstanceData,
 		const int32 InMaxInstanceNum,
+		const int32 InPositionNum,
 		const float InThreshold,
 		const ECompletionMode InCompletionMode)
 		: FCompletionObject(InIdentifier, InInstanceData, InMaxInstanceNum, InCompletionMode)
 	{
-		Position0Handle = InstanceData->Add<1, float>({ InIdentifier, TEXT("Position0") }, { InMaxInstanceNum }, 0.0f);
-		Position1Handle = InstanceData->Add<1, float>({ InIdentifier, TEXT("Position1") }, { InMaxInstanceNum }, 0.0f);
+		Position0Handle = InstanceData->Add<2, float>({ InIdentifier, TEXT("Position0") }, { InMaxInstanceNum, InPositionNum }, 0.0f);
+		Position1Handle = InstanceData->Add<2, float>({ InIdentifier, TEXT("Position1") }, { InMaxInstanceNum, InPositionNum }, 0.0f);
 		ThresholdHandle = InstanceData->Add<1, float>({ InIdentifier, TEXT("Threshold") }, { InMaxInstanceNum }, InThreshold);
 	}
 
@@ -101,20 +102,24 @@ namespace UE::Learning
 	{
 		UE_LEARNING_TRACE_CPUPROFILER_EVENT_SCOPE(FScalarPositionDifferenceCompletion::Evaluate);
 
-		const TLearningArrayView<1, const float> Position0 = InstanceData->ConstView(Position0Handle);
-		const TLearningArrayView<1, const float> Position1 = InstanceData->ConstView(Position1Handle);
+		const TLearningArrayView<2, const float> Position0 = InstanceData->ConstView(Position0Handle);
+		const TLearningArrayView<2, const float> Position1 = InstanceData->ConstView(Position1Handle);
 		const TLearningArrayView<1, const float> Threshold = InstanceData->ConstView(ThresholdHandle);
 		TLearningArrayView<1, ECompletionMode> Completion = InstanceData->View(CompletionHandle);
 
+		const int32 PositionNum = Position0.Num<1>();
+
 		for (const int32 InstanceIdx : Instances)
 		{
-			if (FMath::Abs(Position0[InstanceIdx] - Position1[InstanceIdx]) > Threshold[InstanceIdx])
-			{
-				Completion[InstanceIdx] = CompletionMode;
-			}
-			else
+			Completion[InstanceIdx] = ECompletionMode::Running;
+
+			for (int32 PositionIdx = 0; PositionIdx < PositionNum; PositionIdx++)
 			{
-				Completion[InstanceIdx] = ECompletionMode::Running;
+				if (FMath::Abs(Position0[InstanceIdx][PositionIdx] - Position1[InstanceIdx][PositionIdx]) > Threshold[InstanceIdx])
+				{
+					Completion[InstanceIdx] = CompletionMode;
+					break;
+				}
 			}
 		}
 	}
@@ -123,6 +128,7 @@ namespace UE::Learning
 		const FName& InIdentifier,
 		const TSharedRef<FArrayMap>& InInstanceData,
 		const int32 InMaxInstanceNum,
+		const int32 InPositionNum,
 		const float InThreshold,
 		const ECompletionMode InCompletionMode,
 		const FVector InAxis0,
@@ -131,8 +137,8 @@ namespace UE::Learning
 		, Axis0(InAxis0)
 		, Axis1(InAxis1)
 	{
-		Position0Handle = InstanceData->Add<1, FVector>({ InIdentifier, TEXT("Position0") }, { InMaxInstanceNum }, FVector::ZeroVector);
-		Position1Handle = InstanceData->Add<1, FVector>({ InIdentifier, TEXT("Position1") }, { InMaxInstanceNum }, FVector::ZeroVector);
+		Position0Handle = InstanceData->Add<2, FVector>({ InIdentifier, TEXT("Position0") }, { InMaxInstanceNum, InPositionNum }, FVector::ZeroVector);
+		Position1Handle = InstanceData->Add<2, FVector>({ InIdentifier, TEXT("Position1") }, { InMaxInstanceNum, InPositionNum }, FVector::ZeroVector);
 		ThresholdHandle = InstanceData->Add<1, float>({ InIdentifier, TEXT("Threshold") }, { InMaxInstanceNum }, InThreshold);
 	}
 
@@ -140,23 +146,27 @@ namespace UE::Learning
 	{
 		UE_LEARNING_TRACE_CPUPROFILER_EVENT_SCOPE(FPlanarPositionDifferenceCompletion::Evaluate);
 
-		const TLearningArrayView<1, const FVector> Position0 = InstanceData->ConstView(Position0Handle);
-		const TLearningArrayView<1, const FVector> Position1 = InstanceData->ConstView(Position1Handle);
+		const TLearningArrayView<2, const FVector> Position0 = InstanceData->ConstView(Position0Handle);
+		const TLearningArrayView<2, const FVector> Position1 = InstanceData->ConstView(Position1Handle);
 		const TLearningArrayView<1, const float> Threshold = InstanceData->ConstView(ThresholdHandle);
 		TLearningArrayView<1, ECompletionMode> Completion = InstanceData->View(CompletionHandle);
 
+		const int32 PositionNum = Position0.Num<1>();
+
 		for (const int32 InstanceIdx : Instances)
 		{
-			const FVector ProjectedPosition0 = FVector(Axis0.Dot(Position0[InstanceIdx]), Axis1.Dot(Position0[InstanceIdx]), 0.0f);
-			const FVector ProjectedPosition1 = FVector(Axis0.Dot(Position1[InstanceIdx]), Axis1.Dot(Position1[InstanceIdx]), 0.0f);
+			Completion[InstanceIdx] = ECompletionMode::Running;
 
-			if (FVector::Distance(ProjectedPosition0, ProjectedPosition1) > Threshold[InstanceIdx])
+			for (int32 PositionIdx = 0; PositionIdx < PositionNum; PositionIdx++)
 			{
-				Completion[InstanceIdx] = CompletionMode;
-			}
-			else
-			{
-				Completion[InstanceIdx] = ECompletionMode::Running;
+				const FVector ProjectedPosition0 = FVector(Axis0.Dot(Position0[InstanceIdx][PositionIdx]), Axis1.Dot(Position0[InstanceIdx][PositionIdx]), 0.0f);
+				const FVector ProjectedPosition1 = FVector(Axis0.Dot(Position1[InstanceIdx][PositionIdx]), Axis1.Dot(Position1[InstanceIdx][PositionIdx]), 0.0f);
+
+				if (FVector::Distance(ProjectedPosition0, ProjectedPosition1) > Threshold[InstanceIdx])
+				{
+					Completion[InstanceIdx] = CompletionMode;
+					break;
+				}
 			}
 		}
 	}
@@ -167,12 +177,13 @@ namespace UE::Learning
 		const FName& InIdentifier,
 		const TSharedRef<FArrayMap>& InInstanceData,
 		const int32 InMaxInstanceNum,
+		const int32 InVelocityNum,
 		const float InThreshold,
 		const ECompletionMode InCompletionMode)
 		: FCompletionObject(InIdentifier, InInstanceData, InMaxInstanceNum, InCompletionMode)
 	{
-		Velocity0Handle = InstanceData->Add<1, float>({ InIdentifier, TEXT("Velocity0") }, { InMaxInstanceNum }, 0.0f);
-		Velocity1Handle = InstanceData->Add<1, float>({ InIdentifier, TEXT("Velocity1") }, { InMaxInstanceNum }, 0.0f);
+		Velocity0Handle = InstanceData->Add<2, float>({ InIdentifier, TEXT("Velocity0") }, { InMaxInstanceNum, InVelocityNum }, 0.0f);
+		Velocity1Handle = InstanceData->Add<2, float>({ InIdentifier, TEXT("Velocity1") }, { InMaxInstanceNum, InVelocityNum }, 0.0f);
 		ThresholdHandle = InstanceData->Add<1, float>({ InIdentifier, TEXT("Threshold") }, { InMaxInstanceNum }, InThreshold);
 	}
 
@@ -180,20 +191,24 @@ namespace UE::Learning
 	{
 		UE_LEARNING_TRACE_CPUPROFILER_EVENT_SCOPE(FScalarVelocityDifferenceCompletion::Evaluate);
 
-		const TLearningArrayView<1, const float> Velocity0 = InstanceData->ConstView(Velocity0Handle);
-		const TLearningArrayView<1, const float> Velocity1 = InstanceData->ConstView(Velocity1Handle);
+		const TLearningArrayView<2, const float> Velocity0 = InstanceData->ConstView(Velocity0Handle);
+		const TLearningArrayView<2, const float> Velocity1 = InstanceData->ConstView(Velocity1Handle);
 		const TLearningArrayView<1, const float> Threshold = InstanceData->ConstView(ThresholdHandle);
 		TLearningArrayView<1, ECompletionMode> Completion = InstanceData->View(CompletionHandle);
 
+		const int32 VelocityNum = Velocity0.Num<1>();
+
 		for (const int32 InstanceIdx : Instances)
 		{
-			if (FMath::Abs(Velocity0[InstanceIdx] - Velocity1[InstanceIdx]) > Threshold[InstanceIdx])
-			{
-				Completion[InstanceIdx] = CompletionMode;
-			}
-			else
+			Completion[InstanceIdx] = ECompletionMode::Running;
+
+			for (int32 VelocityIdx = 0; VelocityIdx < VelocityNum; VelocityIdx++)
 			{
-				Completion[InstanceIdx] = ECompletionMode::Running;
+				if (FMath::Abs(Velocity0[InstanceIdx][VelocityIdx] - Velocity1[InstanceIdx][VelocityIdx]) > Threshold[InstanceIdx])
+				{
+					Completion[InstanceIdx] = CompletionMode;
+					break;
+				}
 			}
 		}
 	}
@@ -204,12 +219,13 @@ namespace UE::Learning
 		const FName& InIdentifier,
 		const TSharedRef<FArrayMap>& InInstanceData,
 		const int32 InMaxInstanceNum,
+		const int32 InRotationNum,
 		const float InThreshold,
 		const ECompletionMode InCompletionMode)
 		: FCompletionObject(InIdentifier, InInstanceData, InMaxInstanceNum, InCompletionMode)
 	{
-		Rotation0Handle = InstanceData->Add<1, float>({ InIdentifier, TEXT("Rotation0") }, { InMaxInstanceNum }, 0.0f);
-		Rotation1Handle = InstanceData->Add<1, float>({ InIdentifier, TEXT("Rotation1") }, { InMaxInstanceNum }, 0.0f);
+		Rotation0Handle = InstanceData->Add<2, float>({ InIdentifier, TEXT("Rotation0") }, { InMaxInstanceNum, InRotationNum }, 0.0f);
+		Rotation1Handle = InstanceData->Add<2, float>({ InIdentifier, TEXT("Rotation1") }, { InMaxInstanceNum, InRotationNum }, 0.0f);
 		ThresholdHandle = InstanceData->Add<1, float>({ InIdentifier, TEXT("Threshold") }, { InMaxInstanceNum }, InThreshold);
 	}
 
@@ -217,20 +233,24 @@ namespace UE::Learning
 	{
 		UE_LEARNING_TRACE_CPUPROFILER_EVENT_SCOPE(FScalarRotationDifferenceCompletion::Evaluate);
 
-		const TLearningArrayView<1, const float> Rotation0 = InstanceData->ConstView(Rotation0Handle);
-		const TLearningArrayView<1, const float> Rotation1 = InstanceData->ConstView(Rotation1Handle);
+		const TLearningArrayView<2, const float> Rotation0 = InstanceData->ConstView(Rotation0Handle);
+		const TLearningArrayView<2, const float> Rotation1 = InstanceData->ConstView(Rotation1Handle);
 		const TLearningArrayView<1, const float> Threshold = InstanceData->ConstView(ThresholdHandle);
 		TLearningArrayView<1, ECompletionMode> Completion = InstanceData->View(CompletionHandle);
 
+		const int32 RotationNum = Rotation0.Num<1>();
+
 		for (const int32 InstanceIdx : Instances)
 		{
-			if (FMath::Abs(FMath::FindDeltaAngleRadians(Rotation1[InstanceIdx], Rotation0[InstanceIdx])) > Threshold[InstanceIdx])
-			{
-				Completion[InstanceIdx] = CompletionMode;
-			}
-			else
+			Completion[InstanceIdx] = ECompletionMode::Running;
+
+			for (int32 RotationIdx = 0; RotationIdx < RotationNum; RotationIdx++)
 			{
-				Completion[InstanceIdx] = ECompletionMode::Running;
+				if (FMath::Abs(FMath::FindDeltaAngleRadians(Rotation1[InstanceIdx][RotationIdx], Rotation0[InstanceIdx][RotationIdx])) > Threshold[InstanceIdx])
+				{
+					Completion[InstanceIdx] = CompletionMode;
+					break;
+				}
 			}
 		}
 	}
@@ -241,12 +261,13 @@ namespace UE::Learning
 		const FName& InIdentifier,
 		const TSharedRef<FArrayMap>& InInstanceData,
 		const int32 InMaxInstanceNum,
+		const int32 InAngularVelocityNum,
 		const float InThreshold,
 		const ECompletionMode InCompletionMode)
 		: FCompletionObject(InIdentifier, InInstanceData, InMaxInstanceNum, InCompletionMode)
 	{
-		AngularVelocity0Handle = InstanceData->Add<1, float>({ InIdentifier, TEXT("AngularVelocity0") }, { InMaxInstanceNum }, 0.0f);
-		AngularVelocity1Handle = InstanceData->Add<1, float>({ InIdentifier, TEXT("AngularVelocity1") }, { InMaxInstanceNum }, 0.0f);
+		AngularVelocity0Handle = InstanceData->Add<2, float>({ InIdentifier, TEXT("AngularVelocity0") }, { InMaxInstanceNum, InAngularVelocityNum }, 0.0f);
+		AngularVelocity1Handle = InstanceData->Add<2, float>({ InIdentifier, TEXT("AngularVelocity1") }, { InMaxInstanceNum, InAngularVelocityNum }, 0.0f);
 		ThresholdHandle = InstanceData->Add<1, float>({ InIdentifier, TEXT("Threshold") }, { InMaxInstanceNum }, InThreshold);
 	}
 
@@ -254,20 +275,24 @@ namespace UE::Learning
 	{
 		UE_LEARNING_TRACE_CPUPROFILER_EVENT_SCOPE(FScalarAngularVelocityDifferenceCompletion::Evaluate);
 
-		const TLearningArrayView<1, const float> AngularVelocity0 = InstanceData->ConstView(AngularVelocity0Handle);
-		const TLearningArrayView<1, const float> AngularVelocity1 = InstanceData->ConstView(AngularVelocity1Handle);
+		const TLearningArrayView<2, const float> AngularVelocity0 = InstanceData->ConstView(AngularVelocity0Handle);
+		const TLearningArrayView<2, const float> AngularVelocity1 = InstanceData->ConstView(AngularVelocity1Handle);
 		const TLearningArrayView<1, const float> Threshold = InstanceData->ConstView(ThresholdHandle);
 		TLearningArrayView<1, ECompletionMode> Completion = InstanceData->View(CompletionHandle);
 
+		const int32 AngularVelocityNum = AngularVelocity0.Num<1>();
+
 		for (const int32 InstanceIdx : Instances)
 		{
-			if (FMath::Abs(AngularVelocity0[InstanceIdx] - AngularVelocity1[InstanceIdx]) > Threshold[InstanceIdx])
-			{
-				Completion[InstanceIdx] = CompletionMode;
-			}
-			else
+			Completion[InstanceIdx] = ECompletionMode::Running;
+
+			for (int32 AngularVelocityIdx = 0; AngularVelocityIdx < AngularVelocityNum; AngularVelocityIdx++)
 			{
-				Completion[InstanceIdx] = ECompletionMode::Running;
+				if (FMath::Abs(AngularVelocity0[InstanceIdx][AngularVelocityIdx] - AngularVelocity1[InstanceIdx][AngularVelocityIdx]) > Threshold[InstanceIdx])
+				{
+					Completion[InstanceIdx] = CompletionMode;
+					break;
+				}
 			}
 		}
 	}
diff --git a/Engine/Plugins/Experimental/LearningAgents/Source/Learning/Private/LearningFunctionObject.cpp b/Engine/Plugins/Experimental/LearningAgents/Source/Learning/Private/LearningFunctionObject.cpp
--- a/Engine/Plugins/Experimental/LearningAgents/Source/Learning/Private/LearningFunctionObject.cpp
+++ b/Engine/Plugins/Experimental/LearningAgents/Source/Learning/Private/LearningFunctionObject.cpp
@@ -96,6 +96,7 @@ namespace UE::Learning
 		const FName& InIdentifier,
 		const TSharedRef<FArrayMap>& InInstanceData,
 		const int32 InMaxInstanceNum,
+		const int32 InValueNum,
 		const uint32 InSeed,
 		const float InMin,
 		const float InMax)
@@ -104,7 +105,7 @@ namespace UE::Learning
 		SeedHandle = InstanceData->Add<1, uint32>({ InIdentifier, TEXT("Seed") }, { InMaxInstanceNum });
 		MinHandle = InstanceData->Add<1, float>({ InIdentifier, TEXT("Min") }, { InMaxInstanceNum }, InMin);
 		MaxHandle = InstanceData->Add<1, float>({ InIdentifier, TEXT("Max") }, { InMaxInstanceNum }, InMax);
-		ValueHandle = InstanceData->Add<1, float>({ InIdentifier, TEXT("Value") }, { InMaxInstanceNum }, 0.0f);
+		ValueHandle = InstanceData->Add<2, float>({ InIdentifier, TEXT("Value") }, { InMaxInstanceNum, InValueNum }, 0.0f);
 
 		Random::IntArray(InstanceData->View(SeedHandle), InSeed);
 	}
@@ -116,11 +117,12 @@ namespace UE::Learning
 		const TLearningArrayView<1, const float> Min = InstanceData->ConstView(MinHandle);
 		const TLearningArrayView<1, const float> Max = InstanceData->ConstView(MaxHandle);
 		TLearningArrayView<1, uint32> Seed = InstanceData->View(SeedHandle);
-		TLearningArrayView<1, float> Values = InstanceData->View(ValueHandle);
+		TLearningArrayView<2, float> Values = InstanceData->View(ValueHandle);
 
 		for (const int32 InstanceIdx : Instances)
 		{
-			Values[InstanceIdx] = Random::SampleUniform(
+			Random::SampleUniformArray(
+				Values[InstanceIdx],
 				Seed[InstanceIdx],
 				Min[InstanceIdx],
 				Max[InstanceIdx]);
@@ -131,17 +133,22 @@ namespace UE::Learning
 		const FName& InIdentifier,
 		const TSharedRef<FArrayMap>& InInstanceData,
 		const int32 InMaxInstanceNum,
+		const int32 InValueNum,
 		const uint32 InSeed,
 		const float InMean,
 		const float InStd,
-		const float InClip)
+		const float InClip,
+		const FVector InAxis0,
+		const FVector InAxis1)
 		: FFunctionObject(InInstanceData)
+		, Axis0(InAxis0)
+		, Axis1(InAxis1)
 	{
 		SeedHandle = InstanceData->Add<1, uint32>({ InIdentifier, TEXT("Seed") }, { InMaxInstanceNum });
 		MeanHandle = InstanceData->Add<1, float>({ InIdentifier, TEXT("Mean") }, { InMaxInstanceNum }, InMean);
 		StdHandle = InstanceData->Add<1, float>({ InIdentifier, TEXT("Std") }, { InMaxInstanceNum }, InStd);
 		ClipHandle = InstanceData->Add<1, float>({ InIdentifier, TEXT("Clip") }, { InMaxInstanceNum }, InClip);
-		ValueHandle = InstanceData->Add<1, FVector>({ InIdentifier, TEXT("Value") }, { InMaxInstanceNum }, FVector::ZeroVector);
+		ValueHandle = InstanceData->Add<2, FVector>({ InIdentifier, TEXT("Value") }, { InMaxInstanceNum, InValueNum }, FVector::ZeroVector);
 
 		Random::IntArray(InstanceData->View(SeedHandle), InSeed);
 	}
@@ -154,27 +161,35 @@ namespace UE::Learning
 		const TLearningArrayView<1, const float> Std = InstanceData->ConstView(StdHandle);
 		const TLearningArrayView<1, const float> Clip = InstanceData->ConstView(ClipHandle);
 		TLearningArrayView<1, uint32> Seed = InstanceData->View(SeedHandle);
-		TLearningArrayView<1, FVector> Values = InstanceData->View(ValueHandle);
+		TLearningArrayView<2, FVector> Values = InstanceData->View(ValueHandle);
 
 		for (const int32 InstanceIdx : Instances)
 		{
-			Values[InstanceIdx] = Random::SamplePlanarClippedGaussian(
+			Random::SamplePlanarClippedGaussianArray(
+				Values[InstanceIdx],
 				Seed[InstanceIdx],
 				Mean[InstanceIdx],
 				Std[InstanceIdx],
-				Clip[InstanceIdx]);
+				Clip[InstanceIdx],
+				Axis0,
+				Axis1);
 		}
 	}
 
 	FRandomPlanarDirectionFunction::FRandomPlanarDirectionFunction(
 		const FName& InIdentifier,
 		const TSharedRef<FArrayMap>& InInstanceData,
+		const int32 InMaxInstanceNum,
+		const int32 InDirectionNum,
 		const uint32 InSeed,
-		const int32 InMaxInstanceNum)
+		const FVector InAxis0,
+		const FVector InAxis1)
 		: FFunctionObject(InInstanceData)
+		, Axis0(InAxis0)
+		, Axis1(InAxis1)
 	{
 		SeedHandle = InstanceData->Add<1, uint32>({ InIdentifier, TEXT("Seed") }, { InMaxInstanceNum });
-		DirectionHandle = InstanceData->Add<1, FVector>({ InIdentifier, TEXT("Direction") }, { InMaxInstanceNum }, FVector::ForwardVector);
+		DirectionHandle = InstanceData->Add<2, FVector>({ InIdentifier, TEXT("Direction") }, { InMaxInstanceNum, InDirectionNum }, FVector::ForwardVector);
 
 		Random::IntArray(InstanceData->View(SeedHandle), InSeed);
 	}
@@ -184,11 +199,15 @@ namespace UE::Learning
 		UE_LEARNING_TRACE_CPUPROFILER_EVENT_SCOPE(Learning::FRandomPlanarDirectionFunction::Evaluate);
 
 		TLearningArrayView<1, uint32> Seed = InstanceData->View(SeedHandle);
-		TLearningArrayView<1, FVector> Direction = InstanceData->View(DirectionHandle);
+		TLearningArrayView<2, FVector> Direction = InstanceData->View(DirectionHandle);
 
 		for (const int32 InstanceIdx : Instances)
 		{
-			Direction[InstanceIdx] = Random::SamplePlanarDirection(Seed[InstanceIdx]);
+			Random::SamplePlanarDirectionArray(
+				Direction[InstanceIdx], 
+				Seed[InstanceIdx],
+				Axis0,
+				Axis1);
 		}
 	}
 
@@ -196,14 +215,19 @@ namespace UE::Learning
 		const FName& InIdentifier,
 		const TSharedRef<FArrayMap>& InInstanceData,
 		const int32 InMaxInstanceNum,
+		const int32 InDirectionVelocityNum,
 		const uint32 InSeed,
-		const float InVelocityScale)
+		const float InVelocityScale,
+		const FVector InAxis0,
+		const FVector InAxis1)
 		: FFunctionObject(InInstanceData)
+		, Axis0(InAxis0)
+		, Axis1(InAxis1)
 	{
 		SeedHandle = InstanceData->Add<1, uint32>({ InIdentifier, TEXT("Seed") }, { InMaxInstanceNum });
 		VelocityScaleHandle = InstanceData->Add<1, float>({ InIdentifier, TEXT("VelocityScale") }, { InMaxInstanceNum }, InVelocityScale);
-		DirectionHandle = InstanceData->Add<1, FVector>({ InIdentifier, TEXT("Direction") }, { InMaxInstanceNum }, FVector::ForwardVector);
-		VelocityHandle = InstanceData->Add<1, FVector>({ InIdentifier, TEXT("Velocity") }, { InMaxInstanceNum }, FVector::ZeroVector);
+		DirectionHandle = InstanceData->Add<2, FVector>({ InIdentifier, TEXT("Direction") }, { InMaxInstanceNum, InDirectionVelocityNum }, FVector::ForwardVector);
+		VelocityHandle = InstanceData->Add<2, FVector>({ InIdentifier, TEXT("Velocity") }, { InMaxInstanceNum, InDirectionVelocityNum }, FVector::ZeroVector);
 
 		Random::IntArray(InstanceData->View(SeedHandle), InSeed);
 	}
@@ -214,14 +238,23 @@ namespace UE::Learning
 
 		const TLearningArrayView<1, const float> VelocityScale = InstanceData->ConstView(VelocityScaleHandle);
 		TLearningArrayView<1, uint32> Seed = InstanceData->View(SeedHandle);
-		TLearningArrayView<1, FVector> Direction = InstanceData->View(DirectionHandle);
-		TLearningArrayView<1, FVector> Velocity = InstanceData->View(VelocityHandle);
+		TLearningArrayView<2, FVector> Direction = InstanceData->View(DirectionHandle);
+		TLearningArrayView<2, FVector> Velocity = InstanceData->View(VelocityHandle);
+
+		const int32 VelocityNum = Velocity.Num<1>();
 
 		for (const int32 InstanceIdx : Instances)
 		{
-			const FVector RandomDirection = Random::SamplePlanarDirection(Seed[InstanceIdx]);
-			Direction[InstanceIdx] = RandomDirection;
-			Velocity[InstanceIdx] = VelocityScale[InstanceIdx] * RandomDirection;
+			Random::SamplePlanarDirectionArray(
+				Direction[InstanceIdx],
+				Seed[InstanceIdx],
+				Axis0,
+				Axis1);
+
+			for (int32 VelocityIdx = 0; VelocityIdx < VelocityNum; VelocityIdx++)
+			{
+				Velocity[InstanceIdx][VelocityIdx] = VelocityScale[InstanceIdx] * Direction[InstanceIdx][VelocityIdx];
+			}
 		}
 	}
 
diff --git a/Engine/Plugins/Experimental/LearningAgents/Source/Learning/Private/LearningKMeans.cpp b/Engine/Plugins/Experimental/LearningAgents/Source/Learning/Private/LearningKMeans.cpp
--- a/Engine/Plugins/Experimental/LearningAgents/Source/Learning/Private/LearningKMeans.cpp
+++ b/Engine/Plugins/Experimental/LearningAgents/Source/Learning/Private/LearningKMeans.cpp
@@ -280,8 +280,8 @@ namespace UE::Learning::KMeans
 			ClusterNum,
 			DimNum);
 #else
-		Array::Set(OutMins, MAX_flt);
-		Array::Set(OutMaxs, MIN_flt);
+		Array::Set(OutMins, +MAX_flt);
+		Array::Set(OutMaxs, -MAX_flt);
 
 		for (int32 SampleIdx = 0; SampleIdx < SampleNum; SampleIdx++)
 		{
diff --git a/Engine/Plugins/Experimental/LearningAgents/Source/Learning/Private/LearningNeuralNetwork.cpp b/Engine/Plugins/Experimental/LearningAgents/Source/Learning/Private/LearningNeuralNetwork.cpp
--- a/Engine/Plugins/Experimental/LearningAgents/Source/Learning/Private/LearningNeuralNetwork.cpp
+++ b/Engine/Plugins/Experimental/LearningAgents/Source/Learning/Private/LearningNeuralNetwork.cpp
@@ -4,7 +4,7 @@
 
 namespace UE::Learning
 {
-	const TCHAR* ActivationFunctionString(const EActivationFunction ActivationFunction)
+	const TCHAR* GetActivationFunctionString(const EActivationFunction ActivationFunction)
 	{
 		switch (ActivationFunction)
 		{
@@ -63,29 +63,16 @@ namespace UE::Learning
 
 	int32 FNeuralNetwork::GetHiddenNum() const
 	{
-		return Weights[0].Num(1);
-	}
-
-	int32 FNeuralNetwork::GetTotalByteNum() const
-	{
-		const int32 LayerNum = Weights.Num();
-
-		int32 Total = 0;
-
-		for (int32 LayerIdx = 0; LayerIdx < LayerNum; LayerIdx++)
-		{
-			Total += Weights[LayerIdx].Num() * sizeof(float);
-			Total += Biases[LayerIdx].Num() * sizeof(float);
-		}
-
-		return Total;
+		return Weights[0].Num<1>();
 	}
 
 	void FNeuralNetwork::DeserializeFromBytes(const TLearningArrayView<1, const uint8> RawBytes)
 	{
 		UE_LEARNING_TRACE_CPUPROFILER_EVENT_SCOPE(Learning::FNeuralNetwork::DeserializeFromBytes);
 
-		UE_LEARNING_CHECK(RawBytes.Num() == GetTotalByteNum());
+		const int32 TotalByteNum = GetSerializationByteNum(GetInputNum(), GetOutputNum(), GetHiddenNum(), GetLayerNum());
+
+		UE_LEARNING_CHECK(RawBytes.Num() == TotalByteNum);
 
 		const int32 LayerNum = Weights.Num();
 
@@ -104,14 +91,16 @@ namespace UE::Learning
 			Offset += BiasByteNum;
 		}
 
-		UE_LEARNING_CHECK(Offset == GetTotalByteNum());
+		UE_LEARNING_CHECK(Offset == TotalByteNum);
 	}
 
 	void FNeuralNetwork::SerializeToBytes(TLearningArrayView<1, uint8> OutRawBytes) const
 	{
 		UE_LEARNING_TRACE_CPUPROFILER_EVENT_SCOPE(Learning::FNeuralNetwork::SerializeToBytes);
 
-		UE_LEARNING_CHECK(OutRawBytes.Num() == GetTotalByteNum());
+		const int32 TotalByteNum = GetSerializationByteNum(GetInputNum(), GetOutputNum(), GetHiddenNum(), GetLayerNum());
+
+		UE_LEARNING_CHECK(OutRawBytes.Num() == TotalByteNum);
 
 		const int32 LayerNum = Weights.Num();
 
@@ -128,6 +117,30 @@ namespace UE::Learning
 			Offset += BiasByteNum;
 		}
 
-		UE_LEARNING_CHECK(Offset == GetTotalByteNum());
+		UE_LEARNING_CHECK(Offset == TotalByteNum);
 	}
+
+	int32 FNeuralNetwork::GetSerializationByteNum(
+		const int32 InputNum,
+		const int32 OutputNum,
+		const int32 HiddenNum,
+		const int32 LayerNum)
+	{
+		int32 Total = 0;
+
+		Total += InputNum * HiddenNum * sizeof(float);
+		Total += HiddenNum * sizeof(float);
+
+		for (int32 LayerIdx = 0; LayerIdx < LayerNum - 2; LayerIdx++)
+		{
+			Total += HiddenNum * HiddenNum * sizeof(float);
+			Total += HiddenNum * sizeof(float);
+		}
+
+		Total += HiddenNum * OutputNum * sizeof(float);
+		Total += OutputNum * sizeof(float);
+
+		return Total;
+	}
+
 }
diff --git a/Engine/Plugins/Experimental/LearningAgents/Source/Learning/Private/LearningNeuralNetworkObject.cpp b/Engine/Plugins/Experimental/LearningAgents/Source/Learning/Private/LearningNeuralNetworkObject.cpp
new file mode 100644
--- /dev/null
+++ b/Engine/Plugins/Experimental/LearningAgents/Source/Learning/Private/LearningNeuralNetworkObject.cpp
@@ -0,0 +1,465 @@
+// Copyright Epic Games, Inc. All Rights Reserved.
+
+#include "LearningNeuralNetworkObject.h"
+
+#include "LearningNeuralNetwork.h"
+#include "LearningEigen.h"
+#include "LearningRandom.h"
+
+#if UE_LEARNING_ISPC
+#include "Learning.ispc.generated.h"
+#endif
+
+namespace UE::Learning
+{
+	namespace NeuralNetworkObject
+	{
+		static inline float ReLU(const float X)
+		{
+			return FMath::Max(X, 0.0f);
+		}
+
+		static inline float ELU(const float X)
+		{
+			return X > 0.0f ? X : FMath::InvExpApprox(-X) - 1.0f;
+		}
+
+		static inline float Sigmoid(const float X)
+		{
+			return 1.0f / (1.0f + FMath::InvExpApprox(X));
+		}
+
+		static inline float TanH(const float X)
+		{
+			return FMath::Tanh(X);
+		}
+
+		static inline void MatMulPlusBias(
+			TLearningArrayView<2, float> Output,
+			const TLearningArrayView<2, const float> Input,
+			const TLearningArrayView<2, const float> Weights,
+			const TLearningArrayView<1, const float> Biases,
+			const FIndexSet Instances)
+		{
+			UE_LEARNING_TRACE_CPUPROFILER_EVENT_SCOPE(Learning::NeuralNetworkObject::MatMulPlusBias);
+
+			const int32 RowNum = Weights.Num<0>();
+			const int32 ColNum = Weights.Num<1>();
+
+			if (UE_LEARNING_ISPC && Instances.IsSlice())
+			{
+#if UE_LEARNING_ISPC
+				ispc::LearningLayerMatMulPlusBias(
+					Output.Slice(Instances.GetSliceStart(), Instances.GetSliceNum()).GetData(),
+					Input.Slice(Instances.GetSliceStart(), Instances.GetSliceNum()).GetData(),
+					Weights.GetData(),
+					Biases.GetData(),
+					Instances.GetSliceNum(),
+					RowNum,
+					ColNum);
+#endif
+			}
+			else
+			{
+				OutEigenMatrix(Output).noalias() =
+					(InEigenMatrix(Weights).transpose() *
+						InEigenMatrix(Input).transpose()).transpose().rowwise() +
+					InEigenRowVector(Biases);
+			}
+		}
+
+		static inline void ActivationReLU(TLearningArrayView<2, float> InputOutput, const FIndexSet Instances)
+		{
+			UE_LEARNING_TRACE_CPUPROFILER_EVENT_SCOPE(Learning::NeuralNetworkObject::ActivationReLU);
+
+			const int32 HiddenNum = InputOutput.Num<1>();
+
+			if (UE_LEARNING_ISPC && Instances.IsSlice())
+			{
+#if UE_LEARNING_ISPC
+				ispc::LearningLayerReLU(
+					InputOutput.Slice(Instances.GetSliceStart(), Instances.GetSliceNum()).GetData(),
+					Instances.GetSliceNum(),
+					HiddenNum);
+#endif
+			}
+			else
+			{
+				for (const int32 InstanceIdx : Instances)
+				{
+					for (int32 HiddenIdx = 0; HiddenIdx < HiddenNum; HiddenIdx++)
+					{
+						InputOutput[InstanceIdx][HiddenIdx] = ReLU(InputOutput[InstanceIdx][HiddenIdx]);
+					}
+				}
+			}
+		}
+
+		static inline void ActivationELU(TLearningArrayView<2, float> InputOutput, const FIndexSet Instances)
+		{
+			UE_LEARNING_TRACE_CPUPROFILER_EVENT_SCOPE(Learning::NeuralNetworkObject::ActivationELU);
+
+			const int32 HiddenNum = InputOutput.Num<1>();
+
+			if (UE_LEARNING_ISPC && Instances.IsSlice())
+			{
+#if UE_LEARNING_ISPC
+				ispc::LearningLayerELU(
+					InputOutput.Slice(Instances.GetSliceStart(), Instances.GetSliceNum()).GetData(),
+					Instances.GetSliceNum(),
+					HiddenNum);
+#endif
+			}
+			else
+			{
+				for (const int32 InstanceIdx : Instances)
+				{
+					for (int32 HiddenIdx = 0; HiddenIdx < HiddenNum; HiddenIdx++)
+					{
+						InputOutput[InstanceIdx][HiddenIdx] = ELU(InputOutput[InstanceIdx][HiddenIdx]);
+					}
+				}
+			}
+		}
+
+		static inline void ActivationTanH(TLearningArrayView<2, float> InputOutput, const FIndexSet Instances)
+		{
+			UE_LEARNING_TRACE_CPUPROFILER_EVENT_SCOPE(Learning::NeuralNetworkObject::ActivationTanH);
+
+			const int32 HiddenNum = InputOutput.Num<1>();
+
+			if (UE_LEARNING_ISPC && Instances.IsSlice())
+			{
+#if UE_LEARNING_ISPC
+				ispc::LearningLayerTanH(
+					InputOutput.Slice(Instances.GetSliceStart(), Instances.GetSliceNum()).GetData(),
+					Instances.GetSliceNum(),
+					HiddenNum);
+#endif
+			}
+			else
+			{
+				for (const int32 InstanceIdx : Instances)
+				{
+					for (int32 HiddenIdx = 0; HiddenIdx < HiddenNum; HiddenIdx++)
+					{
+						InputOutput[InstanceIdx][HiddenIdx] = TanH(InputOutput[InstanceIdx][HiddenIdx]);
+					}
+				}
+			}
+		}
+
+		static inline void ActionNoise(
+			TLearningArrayView<2, float> Output,
+			TLearningArrayView<2, float> OutputMean,
+			TLearningArrayView<2, float> OutputStd,
+			const TLearningArrayView<2, float> Input,
+			TLearningArrayView<1, uint32> Seed,
+			const TLearningArrayView<1, const float> ActionNoiseScale,
+			const float LogActionNoiseMin,
+			const float LogActionNoiseMax,
+			const FIndexSet Instances)
+		{
+			UE_LEARNING_TRACE_CPUPROFILER_EVENT_SCOPE(Learning::NeuralNetworkObject::ActionNoise);
+
+			const int32 InputNum = Input.Num<1>();
+			const int32 OutputNum = Output.Num<1>();
+
+			if (UE_LEARNING_ISPC && Instances.IsSlice())
+			{
+#if UE_LEARNING_ISPC
+				ispc::LearningLayerActionNoise(
+					Output.Slice(Instances.GetSliceStart(), Instances.GetSliceNum()).GetData(),
+					OutputMean.Slice(Instances.GetSliceStart(), Instances.GetSliceNum()).GetData(),
+					OutputStd.Slice(Instances.GetSliceStart(), Instances.GetSliceNum()).GetData(),
+					Input.Slice(Instances.GetSliceStart(), Instances.GetSliceNum()).GetData(),
+					Seed.Slice(Instances.GetSliceStart(), Instances.GetSliceNum()).GetData(),
+					ActionNoiseScale.Slice(Instances.GetSliceStart(), Instances.GetSliceNum()).GetData(),
+					Instances.GetSliceNum(),
+					InputNum,
+					OutputNum,
+					LogActionNoiseMin,
+					LogActionNoiseMax);
+#endif
+			}
+			else
+			{
+				for (const int32 InstanceIdx : Instances)
+				{
+					for (int32 OutputIdx = 0; OutputIdx < OutputNum; OutputIdx++)
+					{
+						OutputMean[InstanceIdx][OutputIdx] = Input[InstanceIdx][OutputIdx];
+						OutputStd[InstanceIdx][OutputIdx] = ActionNoiseScale[InstanceIdx] *
+							FMath::Exp(Sigmoid(Input[InstanceIdx][OutputNum + OutputIdx]) * (LogActionNoiseMax - LogActionNoiseMin) + LogActionNoiseMin);
+
+						Output[InstanceIdx][OutputIdx] = Random::Gaussian(
+							Seed[InstanceIdx] ^ 0xab744615 ^ Random::Int(OutputIdx ^ 0xf8a88a27),
+							OutputMean[InstanceIdx][OutputIdx],
+							OutputStd[InstanceIdx][OutputIdx]);
+					}
+				}
+			}
+
+			Random::ResampleStateArray(Seed, Instances);
+		}
+
+		static inline void EvaluatePolicyLayer(
+			TLearningArrayView<2, float> Outputs,
+			TLearningArrayView<2, float> OutputMeans,
+			TLearningArrayView<2, float> OutputStds,
+			TArrayView<TLearningArray<2, float, TInlineAllocator<128>>> Activations,
+			const TLearningArrayView<2, const float> Inputs,
+			TLearningArrayView<1, uint32> Seed,
+			const TLearningArrayView<1, const float> ActionNoiseScale,
+			const int32 LayerIdx,
+			const int32 LayerNum,
+			const float LogActionNoiseMin,
+			const float LogActionNoiseMax,
+			const FNeuralNetwork& NeuralNetwork,
+			const FIndexSet Instances)
+		{
+			UE_LEARNING_TRACE_CPUPROFILER_EVENT_SCOPE(Learning::NeuralNetworkObject::EvaluatePolicyLayer);
+
+			TLearningArrayView<2, float> LayerOutput = Activations[LayerIdx - 0];
+			const TLearningArrayView<2, const float> LayerInput = LayerIdx == 0 ? Inputs : Activations[LayerIdx - 1];
+
+			// Apply Linear Transformation
+
+			MatMulPlusBias(
+				LayerOutput,
+				LayerInput,
+				NeuralNetwork.Weights[LayerIdx],
+				NeuralNetwork.Biases[LayerIdx],
+				Instances);
+
+			// Apply Noise on final layer otherwise apply activation in-place
+
+			if (LayerIdx == LayerNum - 1)
+			{
+				ActionNoise(
+					Outputs,
+					OutputMeans,
+					OutputStds,
+					LayerOutput,
+					Seed,
+					ActionNoiseScale,
+					LogActionNoiseMin,
+					LogActionNoiseMax,
+					Instances);
+			}
+			else
+			{
+				switch (NeuralNetwork.ActivationFunction)
+				{
+				case EActivationFunction::ReLU: ActivationReLU(LayerOutput, Instances); break;
+				case EActivationFunction::ELU: ActivationELU(LayerOutput, Instances); break;
+				case EActivationFunction::TanH: ActivationTanH(LayerOutput, Instances); break;
+				default: UE_LEARNING_NOT_IMPLEMENTED();
+				}
+			}
+		}
+
+		static inline void EvaluateCriticLayer(
+			TLearningArrayView<1, float> Outputs,
+			TArrayView<TLearningArray<2, float, TInlineAllocator<128>>> Activations,
+			const TLearningArrayView<2, const float> Inputs,
+			const int32 LayerIdx,
+			const int32 LayerNum,
+			const FNeuralNetwork& NeuralNetwork,
+			const FIndexSet Instances)
+		{
+			UE_LEARNING_TRACE_CPUPROFILER_EVENT_SCOPE(Learning::NeuralNetworkObject::EvaluateCriticLayer);
+
+			TLearningArrayView<2, float> LayerOutput = LayerIdx == LayerNum - 1 ? TLearningArrayView<2, float>(Outputs.GetData(), { Outputs.Num<0>(), 1}) : Activations[LayerIdx - 0];
+			const TLearningArrayView<2, const float> LayerInput = LayerIdx == 0 ? Inputs : Activations[LayerIdx - 1];
+
+			// Apply Linear Transformation
+
+			MatMulPlusBias(
+				LayerOutput,
+				LayerInput,
+				NeuralNetwork.Weights[LayerIdx],
+				NeuralNetwork.Biases[LayerIdx],
+				Instances);
+
+			// Apply activation in-place
+
+			if (LayerIdx != LayerNum - 1)
+			{
+				switch (NeuralNetwork.ActivationFunction)
+				{
+				case EActivationFunction::ReLU: ActivationReLU(LayerOutput, Instances); break;
+				case EActivationFunction::ELU: ActivationELU(LayerOutput, Instances); break;
+				case EActivationFunction::TanH: ActivationTanH(LayerOutput, Instances); break;
+				default: UE_LEARNING_NOT_IMPLEMENTED();
+				}
+			}
+		}
+	}
+
+	FNeuralNetworkPolicyFunction::FNeuralNetworkPolicyFunction(
+		const FName& InIdentifier,
+		const TSharedRef<FArrayMap>& InInstanceData,
+		const int32 InMaxInstanceNum,
+		const TSharedRef<FNeuralNetwork>& InNeuralNetwork,
+		const uint32 InSeed,
+		const FNeuralNetworkPolicyFunctionSettings& InSettings)
+		: FFunctionObject(InInstanceData)
+		, NeuralNetwork(InNeuralNetwork)
+		, Settings(InSettings)
+	{
+		SeedHandle = InstanceData->Add<1, uint32>({ InIdentifier, TEXT("Seed") }, { InMaxInstanceNum });
+		InputHandle = InstanceData->Add<2, float>({ InIdentifier, TEXT("Input") }, { InMaxInstanceNum, NeuralNetwork->GetInputNum() }, 0.0f);
+		OutputHandle = InstanceData->Add<2, float>({ InIdentifier, TEXT("Output") }, { InMaxInstanceNum, NeuralNetwork->GetOutputNum() / 2 }, 0.0f);
+		OutputMeanHandle = InstanceData->Add<2, float>({ InIdentifier, TEXT("OutputMean") }, { InMaxInstanceNum, NeuralNetwork->GetOutputNum() / 2 }, 0.0f);
+		OutputStdHandle = InstanceData->Add<2, float>({ InIdentifier, TEXT("OutputStd") }, { InMaxInstanceNum, NeuralNetwork->GetOutputNum() / 2 }, 0.0f);
+		ActionNoiseScaleHandle = InstanceData->Add<1, float>({ InIdentifier, TEXT("ActionNoiseScale") }, { InMaxInstanceNum }, Settings.ActionNoiseScale);
+
+		Random::IntArray(InstanceData->View(SeedHandle), InSeed);
+
+		// Allocate temporary storage for activations
+
+		const int32 LayerNum = InNeuralNetwork->GetLayerNum();
+
+		Activations.SetNum(LayerNum);
+		for (int32 LayerIdx = 0; LayerIdx < LayerNum; LayerIdx++)
+		{
+			Activations[LayerIdx].SetNumUninitialized({ InMaxInstanceNum, NeuralNetwork->Weights[LayerIdx].Num<1>() });
+		}
+	}
+
+	void FNeuralNetworkPolicyFunction::Evaluate(const FIndexSet Instances)
+	{
+		UE_LEARNING_TRACE_CPUPROFILER_EVENT_SCOPE(Learning::FNeuralNetworkPolicyFunction::Evaluate);
+
+		const TLearningArrayView<2, const float> Inputs = InstanceData->ConstView(InputHandle);
+		const TLearningArrayView<1, const float> ActionNoiseScale = InstanceData->ConstView(ActionNoiseScaleHandle);
+		TLearningArrayView<2, float> Outputs = InstanceData->View(OutputHandle);
+		TLearningArrayView<2, float> OutputMeans = InstanceData->View(OutputMeanHandle);
+		TLearningArrayView<2, float> OutputStds = InstanceData->View(OutputStdHandle);
+		TLearningArrayView<1, uint32> Seed = InstanceData->View(SeedHandle);
+
+		UE_LEARNING_CHECK(NeuralNetwork->GetInputNum() == Inputs.Num<1>());
+		UE_LEARNING_CHECK(NeuralNetwork->GetOutputNum() == 2 * Outputs.Num<1>());
+
+		const int32 LayerNum = NeuralNetwork->GetLayerNum();
+
+		if (!ensureMsgf(LayerNum > 0, TEXT("Empty Neural Network used in Policy")))
+		{
+			Array::Zero(Outputs, Instances);
+			return;
+		}
+
+		UE_LEARNING_ARRAY_VALUE_CHECK(Settings.ActionNoiseMin >= 0.0f && Settings.ActionNoiseMax >= 0.0f);
+		const float LogActionNoiseMin = FMath::Loge(Settings.ActionNoiseMin + UE_KINDA_SMALL_NUMBER);
+		const float LogActionNoiseMax = FMath::Loge(Settings.ActionNoiseMax + UE_KINDA_SMALL_NUMBER);
+
+		// Compute Layers
+
+		auto LayerEvaluationFunction = [&](const int32 SliceStart, const int32 SliceNum)
+		{
+			for (int32 LayerIdx = 0; LayerIdx < LayerNum; LayerIdx++)
+			{
+				NeuralNetworkObject::EvaluatePolicyLayer(
+					Outputs,
+					OutputMeans,
+					OutputStds,
+					Activations,
+					Inputs,
+					Seed,
+					ActionNoiseScale,
+					LayerIdx,
+					LayerNum,
+					LogActionNoiseMin,
+					LogActionNoiseMax,
+					*NeuralNetwork,
+					Instances.Slice(SliceStart, SliceNum));
+			}
+		};
+
+		if (Settings.bParallelEvaluation && Instances.Num() > Settings.MinParallelBatchSize)
+		{
+			SlicedParallelFor(Instances.Num(), Settings.MinParallelBatchSize, LayerEvaluationFunction);
+		}
+		else
+		{
+			LayerEvaluationFunction(0, Instances.Num());
+		}
+
+		Array::Check(Outputs, Instances);
+	}
+
+	FNeuralNetworkCriticFunction::FNeuralNetworkCriticFunction(
+		const FName& InIdentifier,
+		const TSharedRef<FArrayMap>& InInstanceData,
+		const int32 InMaxInstanceNum,
+		const TSharedRef<FNeuralNetwork>& InNeuralNetwork,
+		const FNeuralNetworkCriticFunctionSettings& InSettings)
+		: FFunctionObject(InInstanceData)
+		, NeuralNetwork(InNeuralNetwork)
+		, Settings(InSettings)
+	{
+		InputHandle = InstanceData->Add<2, float>({ InIdentifier, TEXT("Input") }, { InMaxInstanceNum, NeuralNetwork->GetInputNum() }, 0.0f);
+		OutputHandle = InstanceData->Add<1, float>({ InIdentifier, TEXT("Output") }, { InMaxInstanceNum }, 0.0f);
+
+		// Allocate temporary storage for activations
+		// 
+		// We will use the output handle as the final layer directly
+		// so we don't need storage for this activation like we do with
+		// the policy network.
+		const int32 LayerNum = InNeuralNetwork->GetLayerNum() - 1;
+
+		Activations.SetNum(LayerNum);
+		for (int32 LayerIdx = 0; LayerIdx < LayerNum; LayerIdx++)
+		{
+			Activations[LayerIdx].SetNumUninitialized({ InMaxInstanceNum, NeuralNetwork->Weights[LayerIdx].Num<1>() });
+		}
+	}
+
+	void FNeuralNetworkCriticFunction::Evaluate(const FIndexSet Instances)
+	{
+		UE_LEARNING_TRACE_CPUPROFILER_EVENT_SCOPE(Learning::FNeuralNetworkCriticFunction::Evaluate);
+
+		const TLearningArrayView<2, const float> Inputs = InstanceData->ConstView(InputHandle);
+		TLearningArrayView<1, float> Outputs = InstanceData->View(OutputHandle);
+
+		UE_LEARNING_CHECK(NeuralNetwork->GetInputNum() == Inputs.Num<1>());
+		UE_LEARNING_CHECK(NeuralNetwork->GetOutputNum() == 1);
+
+		const int32 LayerNum = NeuralNetwork->GetLayerNum();
+
+		if (!ensureMsgf(LayerNum > 0, TEXT("Empty Neural Network used in Critic")))
+		{
+			Array::Zero(Outputs, Instances);
+			return;
+		}
+
+		// Compute Layers
+
+		auto LayerEvaluationFunction = [&](const int32 SliceStart, const int32 SliceNum)
+		{
+			for (int32 LayerIdx = 0; LayerIdx < LayerNum; LayerIdx++)
+			{
+				NeuralNetworkObject::EvaluateCriticLayer(
+					Outputs,
+					Activations,
+					Inputs,
+					LayerIdx,
+					LayerNum,
+					*NeuralNetwork,
+					Instances.Slice(SliceStart, SliceNum));
+			}
+		};
+
+		if (Settings.bParallelEvaluation && Instances.Num() > Settings.MinParallelBatchSize)
+		{
+			SlicedParallelFor(Instances.Num(), Settings.MinParallelBatchSize, LayerEvaluationFunction);
+		}
+		else
+		{
+			LayerEvaluationFunction(0, Instances.Num());
+		}
+
+		Array::Check(Outputs, Instances);
+	}
+}
diff --git a/Engine/Plugins/Experimental/LearningAgents/Source/Learning/Private/LearningRandom.cpp b/Engine/Plugins/Experimental/LearningAgents/Source/Learning/Private/LearningRandom.cpp
--- a/Engine/Plugins/Experimental/LearningAgents/Source/Learning/Private/LearningRandom.cpp
+++ b/Engine/Plugins/Experimental/LearningAgents/Source/Learning/Private/LearningRandom.cpp
@@ -214,6 +214,41 @@ namespace UE::Learning::Random
 #endif
 	}
 
+	void PlanarClippedGaussianArray(
+		TLearningArrayView<1, FVector> Output,
+		const uint32 State,
+		const float Mean,
+		const float Std,
+		const float Clip,
+		const FVector Axis0,
+		const FVector Axis1)
+	{
+		UE_LEARNING_TRACE_CPUPROFILER_EVENT_SCOPE(Learning::Random::PlanarClippedGaussianArray);
+
+		const int32 ElementNum = Output.Num();
+
+		for (int32 ElementIdx = 0; ElementIdx < ElementNum; ElementIdx++)
+		{
+			Output[ElementIdx] = PlanarClippedGaussian(State ^ 0x5eeac916 ^ Int(ElementIdx ^ 0x8527618d), Mean, Std, Clip, Axis0, Axis1);
+		}
+	}
+
+	void PlanarDirectionArray(
+		TLearningArrayView<1, FVector> Output,
+		const uint32 State,
+		const FVector Axis0,
+		const FVector Axis1)
+	{
+		UE_LEARNING_TRACE_CPUPROFILER_EVENT_SCOPE(Learning::Random::PlanarDirectionArray);
+
+		const int32 ElementNum = Output.Num();
+
+		for (int32 ElementIdx = 0; ElementIdx < ElementNum; ElementIdx++)
+		{
+			Output[ElementIdx] = PlanarDirection(State ^ 0xd80bd375 ^ Int(ElementIdx ^ 0x50d8c207), Axis0, Axis1);
+		}
+	}
+
 	////////////
 
 	uint32 SampleInt(uint32& State)
@@ -370,4 +405,27 @@ namespace UE::Learning::Random
 		}
 	}
 
+	void SamplePlanarClippedGaussianArray(
+		TLearningArrayView<1, FVector> Output,
+		uint32& State,
+		const float Mean,
+		const float Std,
+		const float Clip,
+		const FVector Axis0,
+		const FVector Axis1)
+	{
+		State = Int(State ^ 0xb538a8b5);
+		PlanarClippedGaussianArray(Output, State ^ 0x73a55e65, Mean, Std, Clip, Axis0, Axis1);
+	}
+
+	void SamplePlanarDirectionArray(
+		TLearningArrayView<1, FVector> Output,
+		uint32& State,
+		const FVector Axis0,
+		const FVector Axis1)
+	{
+		State = Int(State ^ 0x3219c5db);
+		PlanarDirectionArray(Output, State ^ 0x6bfd3e6a, Axis0, Axis1);
+	}
+
 }
\ No newline at end of file
diff --git a/Engine/Plugins/Experimental/LearningAgents/Source/Learning/Public/LearningArrayMap.h b/Engine/Plugins/Experimental/LearningAgents/Source/Learning/Public/LearningArrayMap.h
--- a/Engine/Plugins/Experimental/LearningAgents/Source/Learning/Public/LearningArrayMap.h
+++ b/Engine/Plugins/Experimental/LearningAgents/Source/Learning/Public/LearningArrayMap.h
@@ -377,24 +377,45 @@ namespace UE::Learning
 		* Concretely, this function will make asking for a View of the 
 		* Dst array return a view of the Src array instead.
 		*/
-		void Link(const FArrayMapKey Src, const FArrayMapKey Dst);
+		void LinkKeys(const FArrayMapKey Src, const FArrayMapKey Dst);
 
 		/**
 		* Link one array to another such that they share the same data.
 		* Concretely, this function will make asking for a View of the
 		* Dst array return a view of the Src array instead.
 		*/
-		void Link(const FArrayMapHandle Src, const FArrayMapHandle Dst);
+		void LinkHandles(const FArrayMapHandle Src, const FArrayMapHandle Dst);
+
+
+		/**
+		* Link one array to another as if they were flattened. This can be 
+		* useful when the array dimensions don't exactly match such as with
+		* one dimension of a single element.
+		*/
+		void LinkHandlesFlat(const FArrayMapHandle Src, const FArrayMapHandle Dst);
 
 		/**
 		* Link one array to another such that they share the same data.
 		* Concretely, this function will make asking for a View of the
-		* Dst array return a view of the Src array instead.
+		* Dst array return a view of the Src array instead. This version
+		* provides a type-safe interface for linking so should be used
+		* by default.
 		*/
 		template<uint8 DimNum, typename ElementType>
 		void Link(const TArrayMapHandle<DimNum, ElementType> Src, const TArrayMapHandle<DimNum, ElementType> Dst)
 		{
-			Link((FArrayMapHandle)Src, (FArrayMapHandle)Dst);
+			LinkHandles(Src, Dst);
+		}
+
+		/**
+		* Link one array to another as if they were flattened. This can be
+		* useful when the array dimensions don't exactly match such as with
+		* one dimension of a single element.
+		*/
+		template<uint8 SrcDimNum, uint8 DstDimNum, typename ElementType>
+		void LinkFlat(const TArrayMapHandle<SrcDimNum, ElementType> Src, const TArrayMapHandle<DstDimNum, ElementType> Dst)
+		{
+			LinkHandlesFlat(Src, Dst);
 		}
 
 		/**
diff --git a/Engine/Plugins/Experimental/LearningAgents/Source/Learning/Public/LearningCompletionObject.h b/Engine/Plugins/Experimental/LearningAgents/Source/Learning/Public/LearningCompletionObject.h
--- a/Engine/Plugins/Experimental/LearningAgents/Source/Learning/Public/LearningCompletionObject.h
+++ b/Engine/Plugins/Experimental/LearningAgents/Source/Learning/Public/LearningCompletionObject.h
@@ -85,7 +85,7 @@ namespace UE::Learning
 	//------------------------------------------------------------------
 
 	/**
-	* Completion due to the difference between two scalar positions exceeding a given threshold
+	* Completion due to the difference between any two scalar positions exceeding a given threshold
 	*/
 	struct LEARNING_API FScalarPositionDifferenceCompletion : public FCompletionObject
 	{
@@ -93,18 +93,19 @@ namespace UE::Learning
 			const FName& InIdentifier,
 			const TSharedRef<FArrayMap>& InInstanceData,
 			const int32 InMaxInstanceNum,
+			const int32 InPositionNum,
 			const float InThreshold = Completion::DefaultThresholdPosition,
 			const ECompletionMode InCompletionMode = ECompletionMode::Terminated);
 
 		virtual void Evaluate(const FIndexSet Instances) override final;
 
-		TArrayMapHandle<1, float> Position0Handle;
-		TArrayMapHandle<1, float> Position1Handle;
+		TArrayMapHandle<2, float> Position0Handle;
+		TArrayMapHandle<2, float> Position1Handle;
 		TArrayMapHandle<1, float> ThresholdHandle;
 	};
 
 	/**
-	* Completion due to the difference between two planar positions exceeding a given threshold
+	* Completion due to the difference between any two planar positions exceeding a given threshold
 	*/
 	struct LEARNING_API FPlanarPositionDifferenceCompletion : public FCompletionObject
 	{
@@ -112,6 +113,7 @@ namespace UE::Learning
 			const FName& InIdentifier,
 			const TSharedRef<FArrayMap>& InInstanceData,
 			const int32 InMaxInstanceNum,
+			const int32 InPositionNum,
 			const float InThreshold = Completion::DefaultThresholdPosition,
 			const ECompletionMode InCompletionMode = ECompletionMode::Terminated,
 			const FVector InAxis0 = FVector::ForwardVector,
@@ -122,15 +124,15 @@ namespace UE::Learning
 		FVector Axis0 = FVector::ForwardVector;
 		FVector Axis1 = FVector::RightVector;
 
-		TArrayMapHandle<1, FVector> Position0Handle;
-		TArrayMapHandle<1, FVector> Position1Handle;
+		TArrayMapHandle<2, FVector> Position0Handle;
+		TArrayMapHandle<2, FVector> Position1Handle;
 		TArrayMapHandle<1, float> ThresholdHandle;
 	};
 
 	//------------------------------------------------------------------
 
 	/**
-	* Completion due to the difference between two scalar velocities exceeding a given threshold
+	* Completion due to the difference between any two scalar velocities exceeding a given threshold
 	*/
 	struct LEARNING_API FScalarVelocityDifferenceCompletion : public FCompletionObject
 	{
@@ -138,20 +140,21 @@ namespace UE::Learning
 			const FName& InIdentifier,
 			const TSharedRef<FArrayMap>& InInstanceData,
 			const int32 InMaxInstanceNum,
+			const int32 InVelocityNum,
 			const float InThreshold = Completion::DefaultThresholdVelocity,
 			const ECompletionMode InCompletionMode = ECompletionMode::Terminated);
 
 		virtual void Evaluate(const FIndexSet Instances) override final;
 
-		TArrayMapHandle<1, float> Velocity0Handle;
-		TArrayMapHandle<1, float> Velocity1Handle;
+		TArrayMapHandle<2, float> Velocity0Handle;
+		TArrayMapHandle<2, float> Velocity1Handle;
 		TArrayMapHandle<1, float> ThresholdHandle;
 	};
 
 	//------------------------------------------------------------------
 
 	/**
-	* Completion due to the difference between two scalar rotations exceeding a given threshold
+	* Completion due to the difference between any two scalar rotations exceeding a given threshold
 	*/
 	struct LEARNING_API FScalarRotationDifferenceCompletion : public FCompletionObject
 	{
@@ -159,13 +162,14 @@ namespace UE::Learning
 			const FName& InIdentifier,
 			const TSharedRef<FArrayMap>& InInstanceData,
 			const int32 InMaxInstanceNum,
+			const int32 InRotationNum,
 			const float InThreshold = Completion::DefaultThresholdAngle,
 			const ECompletionMode InCompletionMode = ECompletionMode::Terminated);
 
 		virtual void Evaluate(const FIndexSet Instances) override final;
 
-		TArrayMapHandle<1, float> Rotation0Handle;
-		TArrayMapHandle<1, float> Rotation1Handle;
+		TArrayMapHandle<2, float> Rotation0Handle;
+		TArrayMapHandle<2, float> Rotation1Handle;
 		TArrayMapHandle<1, float> ThresholdHandle;
 	};
 
@@ -180,13 +184,14 @@ namespace UE::Learning
 			const FName& InIdentifier,
 			const TSharedRef<FArrayMap>& InInstanceData,
 			const int32 InMaxInstanceNum,
+			const int32 InAngularVelocityNum,
 			const float InThreshold = Completion::DefaultThresholdAngularVelocity,
 			const ECompletionMode InCompletionMode = ECompletionMode::Terminated);
 
 		virtual void Evaluate(const FIndexSet Instances) override final;
 
-		TArrayMapHandle<1, float> AngularVelocity0Handle;
-		TArrayMapHandle<1, float> AngularVelocity1Handle;
+		TArrayMapHandle<2, float> AngularVelocity0Handle;
+		TArrayMapHandle<2, float> AngularVelocity1Handle;
 		TArrayMapHandle<1, float> ThresholdHandle;
 	};
 
diff --git a/Engine/Plugins/Experimental/LearningAgents/Source/Learning/Public/LearningFunctionObject.h b/Engine/Plugins/Experimental/LearningAgents/Source/Learning/Public/LearningFunctionObject.h
--- a/Engine/Plugins/Experimental/LearningAgents/Source/Learning/Public/LearningFunctionObject.h
+++ b/Engine/Plugins/Experimental/LearningAgents/Source/Learning/Public/LearningFunctionObject.h
@@ -92,7 +92,7 @@ namespace UE::Learning
 	};
 
 	/**
-	* Samples a scalar value from a uniform distribution
+	* Samples scalar values from a uniform distribution
 	*/
 	struct LEARNING_API FRandomUniformFunction : public FFunctionObject
 	{
@@ -100,6 +100,7 @@ namespace UE::Learning
 			const FName& InIdentifier,
 			const TSharedRef<FArrayMap>& InInstanceData,
 			const int32 InMaxInstanceNum,
+			const int32 InValueNum,
 			const uint32 InSeed,
 			const float InMin = 0.0f,
 			const float InMax = 1.0f);
@@ -109,7 +110,7 @@ namespace UE::Learning
 		TArrayMapHandle<1, uint32> SeedHandle;
 		TArrayMapHandle<1, float> MinHandle;
 		TArrayMapHandle<1, float> MaxHandle;
-		TArrayMapHandle<1, float> ValueHandle;
+		TArrayMapHandle<2, float> ValueHandle;
 	};
 
 	/**
@@ -121,18 +122,24 @@ namespace UE::Learning
 			const FName& InIdentifier,
 			const TSharedRef<FArrayMap>& InInstanceData,
 			const int32 InMaxInstanceNum,
+			const int32 InValueNum,
 			const uint32 InSeed,
 			const float InMean = 0.0f,
 			const float InStd = 1.0f,
-			const float InClip = 10.0f);
+			const float InClip = 10.0f,
+			const FVector InAxis0 = FVector::ForwardVector,
+			const FVector InAxis1 = FVector::RightVector);
 
 		virtual void Evaluate(const FIndexSet Instances) override final;
 
+		FVector Axis0 = FVector::ForwardVector;
+		FVector Axis1 = FVector::RightVector;
+
 		TArrayMapHandle<1, uint32> SeedHandle;
 		TArrayMapHandle<1, float> MeanHandle;
 		TArrayMapHandle<1, float> StdHandle;
 		TArrayMapHandle<1, float> ClipHandle;
-		TArrayMapHandle<1, FVector> ValueHandle;
+		TArrayMapHandle<2, FVector> ValueHandle;
 	};
 
 	/**
@@ -143,13 +150,19 @@ namespace UE::Learning
 		FRandomPlanarDirectionFunction(
 			const FName& InIdentifier,
 			const TSharedRef<FArrayMap>& InInstanceData,
+			const int32 InMaxInstanceNum,
+			const int32 InDirectionNum,
 			const uint32 InSeed,
-			const int32 InMaxInstanceNum);
+			const FVector InAxis0 = FVector::ForwardVector,
+			const FVector InAxis1 = FVector::RightVector);
 
 		virtual void Evaluate(const FIndexSet Instances) override final;
 
+		FVector Axis0 = FVector::ForwardVector;
+		FVector Axis1 = FVector::RightVector;
+
 		TArrayMapHandle<1, uint32> SeedHandle;
-		TArrayMapHandle<1, FVector> DirectionHandle;
+		TArrayMapHandle<2, FVector> DirectionHandle;
 	};
 
 	/**
@@ -161,15 +174,21 @@ namespace UE::Learning
 			const FName& InIdentifier,
 			const TSharedRef<FArrayMap>& InInstanceData,
 			const int32 InMaxInstanceNum,
+			const int32 InDirectionVelocityNum,
 			const uint32 InSeed,
-			const float InVelocityScale = 1.0f);
+			const float InVelocityScale = 1.0f,
+			const FVector InAxis0 = FVector::ForwardVector,
+			const FVector InAxis1 = FVector::RightVector);
 
 		virtual void Evaluate(const FIndexSet Instances) override final;
 
+		FVector Axis0 = FVector::ForwardVector;
+		FVector Axis1 = FVector::RightVector;
+
 		TArrayMapHandle<1, uint32> SeedHandle;
 		TArrayMapHandle<1, float> VelocityScaleHandle;
-		TArrayMapHandle<1, FVector> DirectionHandle;
-		TArrayMapHandle<1, FVector> VelocityHandle;
+		TArrayMapHandle<2, FVector> DirectionHandle;
+		TArrayMapHandle<2, FVector> VelocityHandle;
 	};
 
 }
diff --git a/Engine/Plugins/Experimental/LearningAgents/Source/Learning/Public/LearningNeuralNetwork.h b/Engine/Plugins/Experimental/LearningAgents/Source/Learning/Public/LearningNeuralNetwork.h
--- a/Engine/Plugins/Experimental/LearningAgents/Source/Learning/Public/LearningNeuralNetwork.h
+++ b/Engine/Plugins/Experimental/LearningAgents/Source/Learning/Public/LearningNeuralNetwork.h
@@ -21,7 +21,7 @@ namespace UE::Learning
 		TanH = 2,
 	};
 
-	LEARNING_API const TCHAR* ActivationFunctionString(const EActivationFunction ActivationFunction);
+	LEARNING_API const TCHAR* GetActivationFunctionString(const EActivationFunction ActivationFunction);
 
 	/**
 	* Simple feed-forward network data.
@@ -31,6 +31,12 @@ namespace UE::Learning
 		void DeserializeFromBytes(const TLearningArrayView<1, const uint8> RawBytes);
 		void SerializeToBytes(TLearningArrayView<1, uint8> OutRawBytes) const;
 
+		static int32 GetSerializationByteNum(
+			const int32 InputNum,
+			const int32 OutputNum,
+			const int32 HiddenNum,
+			const int32 LayerNum);
+
 		void Resize(
 			const int32 InputNum,
 			const int32 OutputNum,
@@ -41,7 +47,6 @@ namespace UE::Learning
 		int32 GetOutputNum() const;
 		int32 GetLayerNum() const;
 		int32 GetHiddenNum() const;
-		int32 GetTotalByteNum() const;
 
 		EActivationFunction ActivationFunction = EActivationFunction::ELU;
 		TArray<TLearningArray<2, float>, TInlineAllocator<16>> Weights;
diff --git a/Engine/Plugins/Experimental/LearningAgents/Source/Learning/Public/LearningNeuralNetworkObject.h b/Engine/Plugins/Experimental/LearningAgents/Source/Learning/Public/LearningNeuralNetworkObject.h
new file mode 100644
--- /dev/null
+++ b/Engine/Plugins/Experimental/LearningAgents/Source/Learning/Public/LearningNeuralNetworkObject.h
@@ -0,0 +1,99 @@
+// Copyright Epic Games, Inc. All Rights Reserved.
+
+#pragma once
+
+#include "LearningArrayMap.h"
+#include "LearningFunctionObject.h"
+
+#include "Templates/SharedPointer.h"
+
+namespace UE::Learning
+{
+	struct FNeuralNetwork;
+
+	/**
+	* Settings object for a neural network based policy
+	*/
+	struct LEARNING_API FNeuralNetworkPolicyFunctionSettings
+	{
+		// Minimum amount of action noise to allow
+		float ActionNoiseMin = 0.0f;
+
+		// Maximum amount of action noise to allow
+		float ActionNoiseMax = 0.0f;
+
+		// Overall scale of the action noise
+		float ActionNoiseScale = 1.0f;
+
+		// If to allow for multi-threaded evaluation
+		bool bParallelEvaluation = true;
+
+		// Minimum batch size to use for multi-threaded evaluation
+		uint16 MinParallelBatchSize = 16;
+	};
+
+	/**
+	* Neural-network based policy object. Stores various settings and intermediate
+	* storage required to evaluate the given network for the provided number of instances.
+	*/
+	struct LEARNING_API FNeuralNetworkPolicyFunction : public FFunctionObject
+	{
+		FNeuralNetworkPolicyFunction(
+			const FName& InIdentifier,
+			const TSharedRef<FArrayMap>& InInstanceData,
+			const int32 InMaxInstanceNum,
+			const TSharedRef<FNeuralNetwork>& InNeuralNetwork,
+			const uint32 InSeed,
+			const FNeuralNetworkPolicyFunctionSettings& InSettings = FNeuralNetworkPolicyFunctionSettings());
+
+		virtual void Evaluate(const FIndexSet Instances) override final;
+
+		TSharedRef<FNeuralNetwork> NeuralNetwork;
+		FNeuralNetworkPolicyFunctionSettings Settings;
+
+		TArray<TLearningArray<2, float, TInlineAllocator<128>>, TInlineAllocator<16>> Activations;
+
+		TArrayMapHandle<1, uint32> SeedHandle;
+		TArrayMapHandle<2, float> InputHandle;
+		TArrayMapHandle<2, float> OutputHandle;
+		TArrayMapHandle<2, float> OutputMeanHandle;
+		TArrayMapHandle<2, float> OutputStdHandle;
+		TArrayMapHandle<1, float> ActionNoiseScaleHandle;
+	};
+
+	/**
+	* Settings object for a neural network based critic
+	*/
+	struct LEARNING_API FNeuralNetworkCriticFunctionSettings
+	{
+		// If to allow for multi-threaded evaluation
+		bool bParallelEvaluation = true;
+
+		// Minimum batch size to use for multi-threaded evaluation
+		uint16 MinParallelBatchSize = 16;
+	};
+
+	/**
+	* Neural-network based critic object. Stores various settings and intermediate
+	* storage required to evaluate the given network for the provided number of instances.
+	*/
+	struct LEARNING_API FNeuralNetworkCriticFunction : public FFunctionObject
+	{
+		FNeuralNetworkCriticFunction(
+			const FName& InIdentifier,
+			const TSharedRef<FArrayMap>& InInstanceData,
+			const int32 InMaxInstanceNum,
+			const TSharedRef<FNeuralNetwork>& InNeuralNetwork,
+			const FNeuralNetworkCriticFunctionSettings& InSettings = FNeuralNetworkCriticFunctionSettings());
+
+		virtual void Evaluate(const FIndexSet Instances) override final;
+
+		TSharedRef<FNeuralNetwork> NeuralNetwork;
+		FNeuralNetworkCriticFunctionSettings Settings;
+
+		TArray<TLearningArray<2, float, TInlineAllocator<128>>, TInlineAllocator<16>> Activations;
+
+		TArrayMapHandle<2, float> InputHandle;
+		TArrayMapHandle<1, float> OutputHandle;
+	};
+}
diff --git a/Engine/Plugins/Experimental/LearningAgents/Source/Learning/Public/LearningRandom.h b/Engine/Plugins/Experimental/LearningAgents/Source/Learning/Public/LearningRandom.h
--- a/Engine/Plugins/Experimental/LearningAgents/Source/Learning/Public/LearningRandom.h
+++ b/Engine/Plugins/Experimental/LearningAgents/Source/Learning/Public/LearningRandom.h
@@ -104,6 +104,14 @@ namespace UE::Learning::Random
 		const FVector Axis0 = FVector::ForwardVector,
 		const FVector Axis1 = FVector::RightVector);
 
+	LEARNING_API void PlanarDirectionVelocity(
+		FVector& OutDirection,
+		FVector& OutVelocity,
+		const uint32 State,
+		const float VelocityScale = 100.0f,
+		const FVector Axis0 = FVector::ForwardVector,
+		const FVector Axis1 = FVector::RightVector);
+
 	LEARNING_API FQuat Rotation(const uint32 Input);
 
 	LEARNING_API void IntArray(
@@ -126,6 +134,21 @@ namespace UE::Learning::Random
 		const float Mean = 0.0f,
 		const float Std = 1.0f);
 
+	LEARNING_API void PlanarClippedGaussianArray(
+		TLearningArrayView<1, FVector> Output,
+		const uint32 State,
+		const float Mean = 0.0f,
+		const float Std = 1.0f,
+		const float Clip = 10.0f,
+		const FVector Axis0 = FVector::ForwardVector,
+		const FVector Axis1 = FVector::RightVector);
+
+	LEARNING_API void PlanarDirectionArray(
+		TLearningArrayView<1, FVector> Output,
+		const uint32 State,
+		const FVector Axis0 = FVector::ForwardVector,
+		const FVector Axis1 = FVector::RightVector);
+
 	/////////////////////////////////////////////////////////
 	//
 	// Here we provide a more normal interface which updates
@@ -207,4 +230,19 @@ namespace UE::Learning::Random
 	LEARNING_API void ResampleStateArray(
 		TLearningArrayView<1, uint32> InOutStates,
 		const FIndexSet Indices);
+
+	LEARNING_API void SamplePlanarClippedGaussianArray(
+		TLearningArrayView<1, FVector> Output,
+		uint32& State,
+		const float Mean = 0.0f,
+		const float Std = 1.0f,
+		const float Clip = 10.0f,
+		const FVector Axis0 = FVector::ForwardVector,
+		const FVector Axis1 = FVector::RightVector);
+
+	LEARNING_API void SamplePlanarDirectionArray(
+		TLearningArrayView<1, FVector> Output,
+		uint32& State,
+		const FVector Axis0 = FVector::ForwardVector,
+		const FVector Axis1 = FVector::RightVector);
 }
\ No newline at end of file
diff --git a/Engine/Plugins/Experimental/LearningAgents/Source/LearningAgents/Private/LearningAgentsActions.cpp b/Engine/Plugins/Experimental/LearningAgents/Source/LearningAgents/Private/LearningAgentsActions.cpp
--- a/Engine/Plugins/Experimental/LearningAgents/Source/LearningAgents/Private/LearningAgentsActions.cpp
+++ b/Engine/Plugins/Experimental/LearningAgents/Source/LearningAgents/Private/LearningAgentsActions.cpp
@@ -22,7 +22,7 @@
 	UE_LEARNING_AGENTS_VLOG_STRING(Owner, Category, Verbosity, Location + FVector(0.0f, 0.0f, 20.0f), Color, Format, ##__VA_ARGS__)
 
 #define UE_LEARNING_AGENTS_VLOG_PLANE(Owner, Category, Verbosity, Location, Rotation, Axis0, Axis1, Color, Format, ...) \
-	UE_VLOG_OBOX(Owner, Category, Verbosity, FBox(25.0f * FVector(-1, -1, 0), 25.0f * FVector(1, 1, 0)), FTransform(Rotation, Location, FVector::OneVector).ToMatrixNoScale(), Color, Format, ##__VA_ARGS__)
+	UE_VLOG_OBOX(Owner, Category, Verbosity, FBox(-25.0f * (Axis0 + Axis1), 25.0f * (Axis0 + Axis1)), FTransform(Rotation, Location, FVector::OneVector).ToMatrixNoScale(), Color, Format, ##__VA_ARGS__)
 
 namespace UE::Learning::Agents::Private
 {
@@ -356,4 +356,4 @@ void URotationVectorArrayAction::VisualLog(const UE::Learning::FIndexSet Instanc
 
 #undef UE_LEARNING_AGENTS_VLOG_STRING
 #undef UE_LEARNING_AGENTS_VLOG_TRANSFORM
-#undef UE_LEARNING_AGENTS_VLOG_PLANE
+#undef UE_LEARNING_AGENTS_VLOG_PLANE
\ No newline at end of file
diff --git a/Engine/Plugins/Experimental/LearningAgents/Source/LearningAgents/Private/LearningAgentsController.cpp b/Engine/Plugins/Experimental/LearningAgents/Source/LearningAgents/Private/LearningAgentsController.cpp
new file mode 100644
--- /dev/null
+++ b/Engine/Plugins/Experimental/LearningAgents/Source/LearningAgents/Private/LearningAgentsController.cpp
@@ -0,0 +1,131 @@
+// Copyright Epic Games, Inc. All Rights Reserved.
+
+#include "LearningAgentsController.h"
+
+#include "LearningAgentsType.h"
+#include "LearningAgentsActions.h"
+#include "LearningFeatureObject.h"
+#include "LearningLog.h"
+
+#include "EngineDefines.h"
+
+ULearningAgentsController::ULearningAgentsController() : UActorComponent() {}
+ULearningAgentsController::ULearningAgentsController(FVTableHelper& Helper) : ULearningAgentsController() {}
+ULearningAgentsController::~ULearningAgentsController() {}
+
+void ULearningAgentsController::SetActions_Implementation(const TArray<int32>& AgentIds)
+{
+	// Can be overridden to get actions without blueprints
+}
+
+void ULearningAgentsController::SetupController(ULearningAgentsType* InAgentType)
+{
+	if (IsControllerSetupPerformed())
+	{
+		UE_LOG(LogLearning, Error, TEXT("Setup already performed!"));
+		return;
+	}
+
+	if (!InAgentType)
+	{
+		UE_LOG(LogLearning, Error, TEXT("SetupController called but AgentType is nullptr."));
+		return;
+	}
+
+	if (!InAgentType->IsSetupPerformed())
+	{
+		UE_LOG(LogLearning, Error, TEXT("AgentType Setup must be run before controller can be setup."));
+		return;
+	}
+
+	AgentType = InAgentType;
+}
+
+bool ULearningAgentsController::IsControllerSetupPerformed() const
+{
+	return AgentType ? true : false;
+}
+
+void ULearningAgentsController::AddAgent(int32 AgentId)
+{
+	if (!IsControllerSetupPerformed())
+	{
+		UE_LOG(LogLearning, Error, TEXT("Controller setup must be run before agents can be added!"));
+		return;
+	}
+
+	if (!AgentType->GetOccupiedAgentSet().Contains(AgentId))
+	{
+		UE_LOG(LogLearning, Error, TEXT("Unable to add: AgentId %d not found on AgentType. Make sure to add agents to the agent type before adding."), AgentId);
+		return;
+	}
+
+	if (SelectedAgentIds.Contains(AgentId))
+	{
+		UE_LOG(LogLearning, Error, TEXT("AgentId %i is already included in agents set"), AgentId);
+		return;
+	}
+
+	SelectedAgentIds.Add(AgentId);
+	SelectedAgentsSet = SelectedAgentIds;
+	SelectedAgentsSet.TryMakeSlice();
+}
+
+void ULearningAgentsController::RemoveAgent(int32 AgentId)
+{
+	if (!IsControllerSetupPerformed())
+	{
+		UE_LOG(LogLearning, Error, TEXT("Controller setup must be run before agents can be removed!"));
+		return;
+	}
+
+	if (SelectedAgentIds.RemoveSingleSwap(AgentId, false) == 0)
+	{
+		UE_LOG(LogLearning, Error, TEXT("Unable to remove: AgentId %d not found in the added agents set."), AgentId);
+		return;
+	}
+
+	SelectedAgentsSet = SelectedAgentIds;
+	SelectedAgentsSet.TryMakeSlice();
+}
+
+bool ULearningAgentsController::HasAgent(int32 AgentId) const
+{
+	return SelectedAgentsSet.Contains(AgentId);
+}
+
+ULearningAgentsType* ULearningAgentsController::GetAgentType(TSubclassOf<ULearningAgentsType> AgentClass)
+{
+	if (!IsControllerSetupPerformed())
+	{
+		UE_LOG(LogLearning, Error, TEXT("Controller setup must be run before getting the agent type!"));
+		return nullptr;
+	}
+
+	return AgentType;
+}
+
+void ULearningAgentsController::EncodeActions()
+{
+	UE_LEARNING_TRACE_CPUPROFILER_EVENT_SCOPE(ULearningAgentsController::EncodeActions);
+
+	if (!IsControllerSetupPerformed())
+	{
+		UE_LOG(LogLearning, Error, TEXT("Setup must be run before actions can be encoded."));
+		return;
+	}
+
+	SetActions(SelectedAgentIds);
+
+	AgentType->GetActionFeature().Encode(SelectedAgentsSet);
+
+#if ENABLE_VISUAL_LOG
+	for (const ULearningAgentsAction* ActionObject : AgentType->GetActionObjects())
+	{
+		if (ActionObject)
+		{
+			ActionObject->VisualLog(SelectedAgentsSet);
+		}
+	}
+#endif
+}
diff --git a/Engine/Plugins/Experimental/LearningAgents/Source/LearningAgents/Private/LearningAgentsCritic.cpp b/Engine/Plugins/Experimental/LearningAgents/Source/LearningAgents/Private/LearningAgentsCritic.cpp
new file mode 100644
--- /dev/null
+++ b/Engine/Plugins/Experimental/LearningAgents/Source/LearningAgents/Private/LearningAgentsCritic.cpp
@@ -0,0 +1,364 @@
+// Copyright Epic Games, Inc. All Rights Reserved.
+
+#include "LearningAgentsCritic.h"
+
+#include "LearningAgentsType.h"
+#include "LearningFeatureObject.h"
+#include "LearningNeuralNetwork.h"
+#include "LearningNeuralNetworkObject.h"
+#include "LearningLog.h"
+
+#include "UObject/Package.h"
+#include "Misc/FileHelper.h"
+#include "Misc/Paths.h"
+#include "GameFramework/Actor.h"
+#include "VisualLogger/VisualLogger.h"
+
+#define UE_LEARNING_AGENTS_VLOG_STRING(Owner, Category, Verbosity, Location, Color, Format, ...) \
+	UE_VLOG_LOCATION(Owner, Category, Verbosity, Location, 0.0f, Color, Format, ##__VA_ARGS__)
+
+namespace UE::Learning::Agents::Critic::Private
+{
+	static inline FString ArrayToString(const TLearningArrayView<1, const float> Array)
+	{
+		const int32 ItemNum = Array.Num();
+		const int32 MaxItemNum = 32;
+		const int32 OutputItemNum = FMath::Min(ItemNum, MaxItemNum);
+
+		FString Output = TEXT("[");
+
+		for (int32 Idx = 0; Idx < OutputItemNum; Idx++)
+		{
+			Output.Appendf(TEXT("% 6.3f"), Array[Idx]);
+
+			if (Idx < OutputItemNum - 1)
+			{
+				Output += TEXT(" ");
+			}
+			else if (Idx == ItemNum - 1)
+			{
+				Output += TEXT("]");
+			}
+			else if (Idx == MaxItemNum - 1)
+			{
+				Output += TEXT("...]");
+			}
+		}
+
+
+		return Output;
+	}
+
+	static inline FString ArrayToStatsString(const TLearningArrayView<1, const float> Array)
+	{
+		const int32 ItemNum = Array.Num();
+
+		float Min = +FLT_MAX, Max = -FLT_MAX, Mean = 0.0f;
+		for (int32 Idx = 0; Idx < ItemNum; Idx++)
+		{
+			Min = FMath::Min(Min, Array[Idx]);
+			Max = FMath::Max(Max, Array[Idx]);
+			Mean += Array[Idx] / ItemNum;
+		}
+
+		float Var = 0.0f;
+		for (int32 Idx = 0; Idx < ItemNum; Idx++)
+		{
+			Var += FMath::Square(Array[Idx] - Mean) / ItemNum;
+		}
+
+		return FString::Printf(TEXT("[% 6.3f/% 6.3f/% 6.3f/% 6.3f]"),
+			Min, Max, Mean, FMath::Sqrt(Var));
+	}
+}
+
+ULearningAgentsCritic::ULearningAgentsCritic() : UActorComponent() {}
+ULearningAgentsCritic::ULearningAgentsCritic(FVTableHelper& Helper) : ULearningAgentsCritic() {}
+ULearningAgentsCritic::~ULearningAgentsCritic() {}
+
+void ULearningAgentsCritic::SetupCritic(ULearningAgentsType* InAgentType, const FLearningAgentsCriticSettings& CriticSettings)
+{
+	if (IsCriticSetupPerformed())
+	{
+		UE_LOG(LogLearning, Error, TEXT("Setup already performed!"));
+		return;
+	}
+
+	// Setup Agent Type
+
+	if (!InAgentType)
+	{
+		UE_LOG(LogLearning, Error, TEXT("SetupCritic called but AgentType is nullptr."));
+		return;
+	}
+
+	if (!InAgentType->IsSetupPerformed())
+	{
+		UE_LOG(LogLearning, Error, TEXT("AgentType Setup must be run before critic can be setup."));
+		return;
+	}
+
+	AgentType = InAgentType;
+
+	// Setup Neural Network
+
+	Network = NewObject<ULearningAgentsNeuralNetwork>(this, TEXT("CriticNetwork"));
+	Network->NeuralNetwork = MakeShared<UE::Learning::FNeuralNetwork>();
+	Network->NeuralNetwork->Resize(
+		AgentType->GetObservationFeature().DimNum(),
+		1,
+		CriticSettings.HiddenLayerSize,
+		CriticSettings.LayerNum);
+	Network->NeuralNetwork->ActivationFunction = UE::Learning::Agents::GetActivationFunction(CriticSettings.ActivationFunction);
+
+	// Create Critic Object
+
+	CriticObject = MakeShared<UE::Learning::FNeuralNetworkCriticFunction>(
+		TEXT("CriticObject"),
+		AgentType->GetInstanceData().ToSharedRef(),
+		AgentType->GetMaxInstanceNum(),
+		Network->NeuralNetwork.ToSharedRef());
+
+	AgentType->GetInstanceData()->Link(AgentType->GetObservationFeature().FeatureHandle, CriticObject->InputHandle);
+
+	// Done!
+	bCriticSetupPerformed = true;
+}
+
+bool ULearningAgentsCritic::IsCriticSetupPerformed() const
+{
+	return bCriticSetupPerformed;
+}
+
+void ULearningAgentsCritic::AddAgent(int32 AgentId)
+{
+	if (!IsCriticSetupPerformed())
+	{
+		UE_LOG(LogLearning, Error, TEXT("Critic setup must be run before agents can be added!"));
+		return;
+	}
+
+	if (!AgentType->GetOccupiedAgentSet().Contains(AgentId))
+	{
+		UE_LOG(LogLearning, Error, TEXT("Unable to add: AgentId %d not found on AgentType. Make sure to add agents to the agent type before adding."), AgentId);
+		return;
+	}
+
+	if (SelectedAgentIds.Contains(AgentId))
+	{
+		UE_LOG(LogLearning, Error, TEXT("AgentId %i is already included in agents set"), AgentId);
+		return;
+	}
+
+	SelectedAgentIds.Add(AgentId);
+	SelectedAgentsSet = SelectedAgentIds;
+	SelectedAgentsSet.TryMakeSlice();
+}
+
+void ULearningAgentsCritic::RemoveAgent(int32 AgentId)
+{
+	if (!IsCriticSetupPerformed())
+	{
+		UE_LOG(LogLearning, Error, TEXT("Critic setup must be run before agents can be removed!"));
+		return;
+	}
+
+	if (SelectedAgentIds.RemoveSingleSwap(AgentId, false) == 0)
+	{
+		UE_LOG(LogLearning, Error, TEXT("Unable to remove: AgentId %d not found in the added agents set."), AgentId);
+		return;
+	}
+
+	SelectedAgentsSet = SelectedAgentIds;
+	SelectedAgentsSet.TryMakeSlice();
+}
+
+bool ULearningAgentsCritic::HasAgent(int32 AgentId) const
+{
+	return SelectedAgentsSet.Contains(AgentId);
+}
+
+ULearningAgentsType* ULearningAgentsCritic::GetAgentType(TSubclassOf<ULearningAgentsType> AgentClass)
+{
+	if (!IsCriticSetupPerformed())
+	{
+		UE_LOG(LogLearning, Error, TEXT("Critic setup must be run before getting the agent type!"));
+		return nullptr;
+	}
+
+	return AgentType;
+}
+
+UE::Learning::FNeuralNetwork& ULearningAgentsCritic::GetCriticNetwork()
+{
+	return *Network->NeuralNetwork;
+}
+
+UE::Learning::FNeuralNetworkCriticFunction& ULearningAgentsCritic::GetCriticObject()
+{
+	return *CriticObject;
+}
+
+void ULearningAgentsCritic::LoadCriticFromSnapshot(const FDirectoryPath& Directory, const FString Filename)
+{
+	if (!IsCriticSetupPerformed())
+	{
+		UE_LOG(LogLearning, Error, TEXT("Critic setup must be run before network can be loaded."));
+		return;
+	}
+
+	TArray64<uint8> NetworkData;
+	FString FilePath = Directory.Path + FGenericPlatformMisc::GetDefaultPathSeparator() + Filename;
+	if (FFileHelper::LoadFileToArray(NetworkData, *FilePath))
+	{
+		const int32 TotalByteNum = UE::Learning::FNeuralNetwork::GetSerializationByteNum(
+			Network->NeuralNetwork->GetInputNum(),
+			Network->NeuralNetwork->GetOutputNum(),
+			Network->NeuralNetwork->GetHiddenNum(),
+			Network->NeuralNetwork->GetLayerNum());
+
+		if (NetworkData.Num() != TotalByteNum)
+		{
+			UE_LOG(LogLearning, Error, TEXT("Failed to load network from file %s. File size incorrect."), *FilePath);
+			return;
+		}
+
+		Network->NeuralNetwork->DeserializeFromBytes(NetworkData);
+	}
+	else
+	{
+		UE_LOG(LogLearning, Error, TEXT("Failed to load network. File not found: %s"), *FilePath);
+	}
+}
+
+void ULearningAgentsCritic::SaveCriticToSnapshot(const FDirectoryPath& Directory, const FString Filename) const
+{
+	if (!IsCriticSetupPerformed())
+	{
+		UE_LOG(LogLearning, Error, TEXT("Critic setup must be run before network can be saved."));
+		return;
+	}
+
+	TArray64<uint8> NetworkData;
+	NetworkData.SetNumUninitialized(UE::Learning::FNeuralNetwork::GetSerializationByteNum(
+		Network->NeuralNetwork->GetInputNum(),
+		Network->NeuralNetwork->GetOutputNum(),
+		Network->NeuralNetwork->GetHiddenNum(),
+		Network->NeuralNetwork->GetLayerNum()));
+
+	Network->NeuralNetwork->SerializeToBytes(NetworkData);
+
+	FString FilePath = Directory.Path + FGenericPlatformMisc::GetDefaultPathSeparator() + Filename;
+	if (!FFileHelper::SaveArrayToFile(NetworkData, *FilePath))
+	{
+		UE_LOG(LogLearning, Error, TEXT("Failed to save network to file: %s"), *FilePath);
+	}
+}
+
+void ULearningAgentsCritic::LoadCriticFromAsset(const ULearningAgentsNeuralNetwork* NeuralNetworkAsset)
+{
+	if (!IsCriticSetupPerformed())
+	{
+		UE_LOG(LogLearning, Error, TEXT("Critic setup must be run before network can be loaded."));
+		return;
+	}
+
+	if (!NeuralNetworkAsset || !NeuralNetworkAsset->NeuralNetwork)
+	{
+		UE_LOG(LogLearning, Error, TEXT("Cannot load critic from invalid asset."));
+		return;
+	}
+
+	if (NeuralNetworkAsset->NeuralNetwork->GetInputNum() != Network->NeuralNetwork->GetInputNum() ||
+		NeuralNetworkAsset->NeuralNetwork->GetOutputNum() != Network->NeuralNetwork->GetOutputNum())
+	{
+		UE_LOG(LogLearning, Error, TEXT("Failed to load critic from asset. Network Asset inputs and outputs don't match."));
+		return;
+	}
+
+	*Network->NeuralNetwork = *NeuralNetworkAsset->NeuralNetwork;
+}
+
+void ULearningAgentsCritic::SaveCriticToAsset(ULearningAgentsNeuralNetwork* NeuralNetworkAsset) const
+{
+	if (!IsCriticSetupPerformed())
+	{
+		UE_LOG(LogLearning, Error, TEXT("Critic setup must be run before network can be saved."));
+		return;
+	}
+
+	if (!NeuralNetworkAsset)
+	{
+		UE_LOG(LogLearning, Error, TEXT("Cannot save critic to invalid asset."));
+		return;
+	}
+
+	if (!NeuralNetworkAsset->NeuralNetwork)
+	{
+		NeuralNetworkAsset->NeuralNetwork = MakeShared<UE::Learning::FNeuralNetwork>();
+	}
+
+	*NeuralNetworkAsset->NeuralNetwork = *Network->NeuralNetwork;
+
+	// Manually mark the package as dirty since just using `Modify` prevents 
+	// marking packages as dirty during PIE which is most likely when this
+	// is being used.
+	if (UPackage* Package = NeuralNetworkAsset->GetPackage())
+	{
+		const bool bIsDirty = Package->IsDirty();
+
+		if (!bIsDirty)
+		{
+			Package->SetDirtyFlag(true);
+		}
+
+		Package->PackageMarkedDirtyEvent.Broadcast(Package, bIsDirty);
+	}
+}
+
+void ULearningAgentsCritic::EvaluateCritic()
+{
+	UE_LEARNING_TRACE_CPUPROFILER_EVENT_SCOPE(ULearningAgentsCritic::EvaluateCritic);
+
+	if (!IsCriticSetupPerformed())
+	{
+		UE_LOG(LogLearning, Error, TEXT("Setup must be run before the critic can be evaluated."));
+		return;
+	}
+
+	CriticObject->Evaluate(SelectedAgentIds);
+
+#if ENABLE_VISUAL_LOG
+	VisualLog(SelectedAgentIds);
+#endif
+}
+
+#if ENABLE_VISUAL_LOG
+void ULearningAgentsCritic::VisualLog(const UE::Learning::FIndexSet Instances) const
+{
+	UE_LEARNING_TRACE_CPUPROFILER_EVENT_SCOPE(ULearningAgentsCritic::VisualLog);
+
+	const TLearningArrayView<2, const float> InputView = CriticObject->InstanceData->ConstView(CriticObject->InputHandle);
+	const TLearningArrayView<1, const float> OutputView = CriticObject->InstanceData->ConstView(CriticObject->OutputHandle);
+
+	for (const int32 Instance : Instances)
+	{
+		if (const AActor* Actor = Cast<AActor>(AgentType->GetAgent(Instance)))
+		{
+			const FString InputArrayString = UE::Learning::Agents::Critic::Private::ArrayToString(InputView[Instance]);
+			const FString InputStatsString = UE::Learning::Agents::Critic::Private::ArrayToStatsString(InputView[Instance]);
+
+			UE_LEARNING_AGENTS_VLOG_STRING(this, LogLearning, Display,
+				Actor->GetActorLocation(),
+				VisualLogColor.ToFColor(true),
+				TEXT("Agent %i\nInput: %s\nInput Stats (Min/Max/Mean/Std): %s\nOutput: [% 6.3f]"),
+				Instance,
+				*InputArrayString,
+				*InputStatsString,
+				OutputView[Instance]);
+		}
+	}
+}
+#endif
+
+#undef UE_LEARNING_AGENTS_VLOG_STRING
\ No newline at end of file
diff --git a/Engine/Plugins/Experimental/LearningAgents/Source/LearningAgents/Private/LearningAgentsNeuralNetwork.cpp b/Engine/Plugins/Experimental/LearningAgents/Source/LearningAgents/Private/LearningAgentsNeuralNetwork.cpp
new file mode 100644
--- /dev/null
+++ b/Engine/Plugins/Experimental/LearningAgents/Source/LearningAgents/Private/LearningAgentsNeuralNetwork.cpp
@@ -0,0 +1,91 @@
+// Copyright Epic Games, Inc. All Rights Reserved.
+
+#include "LearningAgentsNeuralNetwork.h"
+
+#include "LearningNeuralNetwork.h"
+#include "LearningLog.h"
+
+ULearningAgentsNeuralNetwork::ULearningAgentsNeuralNetwork() {}
+ULearningAgentsNeuralNetwork::ULearningAgentsNeuralNetwork(FVTableHelper& Helper) : ULearningAgentsNeuralNetwork() {}
+ULearningAgentsNeuralNetwork::~ULearningAgentsNeuralNetwork() {}
+
+namespace UE::Learning::Agents
+{
+	ELearningAgentsActivationFunction GetLearningAgentsActivationFunction(const EActivationFunction ActivationFunction)
+	{
+		switch (ActivationFunction)
+		{
+		case EActivationFunction::ReLU: return ELearningAgentsActivationFunction::ReLU;
+		case EActivationFunction::TanH: return ELearningAgentsActivationFunction::TanH;
+		case EActivationFunction::ELU: return ELearningAgentsActivationFunction::ELU;
+		default:UE_LOG(LogLearning, Error, TEXT("Unknown Activation Function.")); return ELearningAgentsActivationFunction::ELU;
+		}
+	}
+
+	EActivationFunction GetActivationFunction(const ELearningAgentsActivationFunction ActivationFunction)
+	{
+		switch (ActivationFunction)
+		{
+		case ELearningAgentsActivationFunction::ReLU: return EActivationFunction::ReLU;
+		case ELearningAgentsActivationFunction::TanH: return EActivationFunction::TanH;
+		case ELearningAgentsActivationFunction::ELU: return EActivationFunction::ELU;
+		default:UE_LOG(LogLearning, Error, TEXT("Unknown Activation Function.")); return EActivationFunction::ELU;
+		}
+	}
+}
+
+void ULearningAgentsNeuralNetwork::Serialize(FArchive& Ar)
+{
+	Super::Serialize(Ar);
+
+	if (Ar.IsLoading())
+	{
+		bool bValid;
+		Ar << bValid;
+		if (bValid)
+		{
+			int32 InputNum;
+			int32 OutputNum;
+			int32 HiddenNum;
+			int32 LayerNum;
+			UE::Learning::EActivationFunction Activation;
+			Ar << InputNum;
+			Ar << OutputNum;
+			Ar << HiddenNum;
+			Ar << LayerNum;
+			Ar << Activation;
+			NeuralNetwork = MakeShared<UE::Learning::FNeuralNetwork>();
+			NeuralNetwork->Resize(InputNum, OutputNum, HiddenNum, LayerNum);
+			NeuralNetwork->ActivationFunction = Activation;
+			TArray<uint8> Bytes;
+			Ar << Bytes;
+			NeuralNetwork->DeserializeFromBytes(Bytes);
+		}
+		else
+		{
+			NeuralNetwork.Reset();
+		}
+	}
+	else if (Ar.IsSaving())
+	{
+		bool bValid = NeuralNetwork.IsValid();
+		Ar << bValid;
+		if (bValid)
+		{
+			int32 InputNum = NeuralNetwork->GetInputNum();
+			int32 OutputNum = NeuralNetwork->GetOutputNum();
+			int32 HiddenNum = NeuralNetwork->GetHiddenNum();
+			int32 LayerNum = NeuralNetwork->GetLayerNum();
+			UE::Learning::EActivationFunction Activation = NeuralNetwork->ActivationFunction;
+			Ar << InputNum;
+			Ar << OutputNum;
+			Ar << HiddenNum;
+			Ar << LayerNum;
+			Ar << Activation;
+			TArray<uint8> Bytes;
+			Bytes.SetNumUninitialized(UE::Learning::FNeuralNetwork::GetSerializationByteNum(InputNum, OutputNum, HiddenNum, LayerNum));
+			NeuralNetwork->SerializeToBytes(Bytes);
+			Ar << Bytes;
+		}
+	}
+}
\ No newline at end of file
diff --git a/Engine/Plugins/Experimental/LearningAgents/Source/LearningAgents/Private/LearningAgentsObservations.cpp b/Engine/Plugins/Experimental/LearningAgents/Source/LearningAgents/Private/LearningAgentsObservations.cpp
--- a/Engine/Plugins/Experimental/LearningAgents/Source/LearningAgents/Private/LearningAgentsObservations.cpp
+++ b/Engine/Plugins/Experimental/LearningAgents/Source/LearningAgents/Private/LearningAgentsObservations.cpp
@@ -22,7 +22,7 @@
 	UE_LEARNING_AGENTS_VLOG_STRING(Owner, Category, Verbosity, Location + FVector(0.0f, 0.0f, 20.0f), Color, Format, ##__VA_ARGS__)
 
 #define UE_LEARNING_AGENTS_VLOG_PLANE(Owner, Category, Verbosity, Location, Rotation, Axis0, Axis1, Color, Format, ...) \
-	UE_VLOG_OBOX(Owner, Category, Verbosity, FBox(25.0f * FVector(-1, -1, 0), 25.0f * FVector(1, 1, 0)), FTransform(Rotation, Location, FVector::OneVector).ToMatrixNoScale(), Color, Format, ##__VA_ARGS__)
+	UE_VLOG_OBOX(Owner, Category, Verbosity, FBox(-25.0f * (Axis0 + Axis1), 25.0f * (Axis0 + Axis1)), FTransform(Rotation, Location, FVector::OneVector).ToMatrixNoScale(), Color, Format, ##__VA_ARGS__)
 
 namespace UE::Learning::Agents::Private
 {
@@ -960,3 +960,4 @@ void UVelocityObservation::VisualLog(const UE::Learning::FIndexSet Instances) co
 #undef UE_LEARNING_AGENTS_VLOG_STRING
 #undef UE_LEARNING_AGENTS_VLOG_TRANSFORM
 #undef UE_LEARNING_AGENTS_VLOG_PLANE
+
