LearningAgents: Large refactoring to allow critic to be loaded and used at runtime. Fixed some other misc bugs and fallout from switching features to all be arrays by default. Some refactoring of diff

diff --git a/Engine/Plugins/Experimental/LearningAgents/Source/LearningAgents/Private/LearningAgentsPolicy.cpp b/Engine/Plugins/Experimental/LearningAgents/Source/LearningAgents/Private/LearningAgentsPolicy.cpp
new file mode 100644
--- /dev/null
+++ b/Engine/Plugins/Experimental/LearningAgents/Source/LearningAgents/Private/LearningAgentsPolicy.cpp
@@ -0,0 +1,389 @@
+// Copyright Epic Games, Inc. All Rights Reserved.
+
+#include "LearningAgentsPolicy.h"
+
+#include "LearningAgentsType.h"
+#include "LearningFeatureObject.h"
+#include "LearningNeuralNetwork.h"
+#include "LearningNeuralNetworkObject.h"
+#include "LearningLog.h"
+
+#include "UObject/Package.h"
+#include "Misc/FileHelper.h"
+#include "Misc/Paths.h"
+#include "GameFramework/Actor.h"
+#include "VisualLogger/VisualLogger.h"
+
+#define UE_LEARNING_AGENTS_VLOG_STRING(Owner, Category, Verbosity, Location, Color, Format, ...) \
+	UE_VLOG_LOCATION(Owner, Category, Verbosity, Location, 0.0f, Color, Format, ##__VA_ARGS__)
+
+namespace UE::Learning::Agents::Policy::Private
+{
+	static inline FString ArrayToString(const TLearningArrayView<1, const float> Array)
+	{
+		const int32 ItemNum = Array.Num();
+		const int32 MaxItemNum = 32;
+		const int32 OutputItemNum = FMath::Min(ItemNum, MaxItemNum);
+
+		FString Output = TEXT("[");
+
+		for (int32 Idx = 0; Idx < OutputItemNum; Idx++)
+		{
+			Output.Appendf(TEXT("% 6.3f"), Array[Idx]);
+
+			if (Idx < OutputItemNum - 1)
+			{
+				Output += TEXT(" ");
+			}
+			else if (Idx == ItemNum - 1)
+			{
+				Output += TEXT("]");
+			}
+			else if (Idx == MaxItemNum - 1)
+			{
+				Output += TEXT("...]");
+			}
+		}
+
+
+		return Output;
+	}
+
+	static inline FString ArrayToStatsString(const TLearningArrayView<1, const float> Array)
+	{
+		const int32 ItemNum = Array.Num();
+
+		float Min = +FLT_MAX, Max = -FLT_MAX, Mean = 0.0f;
+		for (int32 Idx = 0; Idx < ItemNum; Idx++)
+		{
+			Min = FMath::Min(Min, Array[Idx]);
+			Max = FMath::Max(Max, Array[Idx]);
+			Mean += Array[Idx] / ItemNum;
+		}
+
+		float Var = 0.0f;
+		for (int32 Idx = 0; Idx < ItemNum; Idx++)
+		{
+			Var += FMath::Square(Array[Idx] - Mean) / ItemNum;
+		}
+
+		return FString::Printf(TEXT("[% 6.3f/% 6.3f/% 6.3f/% 6.3f]"),
+			Min, Max, Mean, FMath::Sqrt(Var));
+	}
+}
+
+ULearningAgentsPolicy::ULearningAgentsPolicy() : UActorComponent() {}
+ULearningAgentsPolicy::ULearningAgentsPolicy(FVTableHelper& Helper) : ULearningAgentsPolicy() {}
+ULearningAgentsPolicy::~ULearningAgentsPolicy() {}
+
+void ULearningAgentsPolicy::SetupPolicy(ULearningAgentsType* InAgentType, const FLearningAgentsPolicySettings& PolicySettings)
+{
+	if (IsPolicySetupPerformed())
+	{
+		UE_LOG(LogLearning, Error, TEXT("Setup already performed!"));
+		return;
+	}
+
+	// Setup Agent Type
+
+	if (!InAgentType)
+	{
+		UE_LOG(LogLearning, Error, TEXT("SetupPolicy called but AgentType is nullptr."));
+		return;
+	}
+
+	if (!InAgentType->IsSetupPerformed())
+	{
+		UE_LOG(LogLearning, Error, TEXT("AgentType Setup must be run before policy can be setup."));
+		return;
+	}
+
+	AgentType = InAgentType;
+
+	// Setup Neural Network
+
+	Network = NewObject<ULearningAgentsNeuralNetwork>(this, TEXT("PolicyNetwork"));
+	Network->NeuralNetwork = MakeShared<UE::Learning::FNeuralNetwork>();
+	Network->NeuralNetwork->Resize(
+		AgentType->GetObservationFeature().DimNum(),
+		2 * AgentType->GetActionFeature().DimNum(),
+		PolicySettings.HiddenLayerSize,
+		PolicySettings.LayerNum);
+	Network->NeuralNetwork->ActivationFunction = UE::Learning::Agents::GetActivationFunction(PolicySettings.ActivationFunction);
+
+	// Create Policy Object
+
+	UE::Learning::FNeuralNetworkPolicyFunctionSettings PolicyFunctionSettings;
+	PolicyFunctionSettings.ActionNoiseMin = PolicySettings.ActionNoiseMin;
+	PolicyFunctionSettings.ActionNoiseMax = PolicySettings.ActionNoiseMax;
+
+	PolicyObject = MakeShared<UE::Learning::FNeuralNetworkPolicyFunction>(
+		TEXT("PolicyObject"),
+		AgentType->GetInstanceData().ToSharedRef(),
+		AgentType->GetMaxInstanceNum(),
+		Network->NeuralNetwork.ToSharedRef(),
+		PolicySettings.ActionNoiseSeed,
+		PolicyFunctionSettings);
+
+	if (AgentType->GetInstanceData()->HasLink(AgentType->GetActionFeature().FeatureHandle))
+	{
+		UE_LOG(LogLearning, Warning, TEXT("AgentType actions are already linked to a difference policy."));
+	}
+	else
+	{
+		AgentType->GetInstanceData()->Link(AgentType->GetObservationFeature().FeatureHandle, PolicyObject->InputHandle);
+		AgentType->GetInstanceData()->Link(PolicyObject->OutputHandle, AgentType->GetActionFeature().FeatureHandle);
+	}
+
+	// Done!
+	bPolicySetupPerformed = true;
+}
+
+bool ULearningAgentsPolicy::IsPolicySetupPerformed() const
+{
+	return bPolicySetupPerformed;
+}
+
+void ULearningAgentsPolicy::AddAgent(int32 AgentId)
+{
+	if (!IsPolicySetupPerformed())
+	{
+		UE_LOG(LogLearning, Error, TEXT("Policy setup must be run before agents can be added!"));
+		return;
+	}
+
+	if (!AgentType->GetOccupiedAgentSet().Contains(AgentId))
+	{
+		UE_LOG(LogLearning, Error, TEXT("Unable to add: AgentId %d not found on AgentType. Make sure to add agents to the agent type before adding."), AgentId);
+		return;
+	}
+
+	if (SelectedAgentIds.Contains(AgentId))
+	{
+		UE_LOG(LogLearning, Error, TEXT("AgentId %i is already included in agents set"), AgentId);
+		return;
+	}
+
+	SelectedAgentIds.Add(AgentId);
+	SelectedAgentsSet = SelectedAgentIds;
+	SelectedAgentsSet.TryMakeSlice();
+}
+
+void ULearningAgentsPolicy::RemoveAgent(int32 AgentId)
+{
+	if (!IsPolicySetupPerformed())
+	{
+		UE_LOG(LogLearning, Error, TEXT("Policy setup must be run before agents can be removed!"));
+		return;
+	}
+
+	if (SelectedAgentIds.RemoveSingleSwap(AgentId, false) == 0)
+	{
+		UE_LOG(LogLearning, Error, TEXT("Unable to remove: AgentId %d not found in the added agents set."), AgentId);
+		return;
+	}
+
+	SelectedAgentsSet = SelectedAgentIds;
+	SelectedAgentsSet.TryMakeSlice();
+}
+
+bool ULearningAgentsPolicy::HasAgent(int32 AgentId) const
+{
+	return SelectedAgentsSet.Contains(AgentId);
+}
+
+ULearningAgentsType* ULearningAgentsPolicy::GetAgentType(TSubclassOf<ULearningAgentsType> AgentClass)
+{
+	if (!IsPolicySetupPerformed())
+	{
+		UE_LOG(LogLearning, Error, TEXT("Policy setup must be run before getting the agent type!"));
+		return nullptr;
+	}
+
+	return AgentType;
+}
+
+UE::Learning::FNeuralNetwork& ULearningAgentsPolicy::GetPolicyNetwork()
+{
+	return *Network->NeuralNetwork;
+}
+
+UE::Learning::FNeuralNetworkPolicyFunction& ULearningAgentsPolicy::GetPolicyObject()
+{
+	return *PolicyObject;
+}
+
+void ULearningAgentsPolicy::LoadPolicyFromSnapshot(const FDirectoryPath& Directory, const FString Filename)
+{
+	if (!IsPolicySetupPerformed())
+	{
+		UE_LOG(LogLearning, Error, TEXT("Policy setup must be run before network can be loaded."));
+		return;
+	}
+
+	TArray64<uint8> NetworkData;
+	FString FilePath = Directory.Path + FGenericPlatformMisc::GetDefaultPathSeparator() + Filename;
+	if (FFileHelper::LoadFileToArray(NetworkData, *FilePath))
+	{
+		const int32 TotalByteNum = UE::Learning::FNeuralNetwork::GetSerializationByteNum(
+			Network->NeuralNetwork->GetInputNum(),
+			Network->NeuralNetwork->GetOutputNum(),
+			Network->NeuralNetwork->GetHiddenNum(),
+			Network->NeuralNetwork->GetLayerNum());
+
+		if (NetworkData.Num() != TotalByteNum)
+		{
+			UE_LOG(LogLearning, Error, TEXT("Failed to load network from file %s. File size incorrect."), *FilePath);
+			return;
+		}
+
+		Network->NeuralNetwork->DeserializeFromBytes(NetworkData);
+	}
+	else
+	{
+		UE_LOG(LogLearning, Warning, TEXT("Failed to load network. File not found: %s"), *FilePath);
+	}
+}
+
+void ULearningAgentsPolicy::SavePolicyToSnapshot(const FDirectoryPath& Directory, const FString Filename) const
+{
+	if (!IsPolicySetupPerformed())
+	{
+		UE_LOG(LogLearning, Error, TEXT("Policy setup must be run before network can be saved."));
+		return;
+	}
+
+	TArray64<uint8> NetworkData;
+	NetworkData.SetNumUninitialized(UE::Learning::FNeuralNetwork::GetSerializationByteNum(
+		Network->NeuralNetwork->GetInputNum(),
+		Network->NeuralNetwork->GetOutputNum(),
+		Network->NeuralNetwork->GetHiddenNum(),
+		Network->NeuralNetwork->GetLayerNum()));
+
+	Network->NeuralNetwork->SerializeToBytes(NetworkData);
+
+	FString FilePath = Directory.Path + FGenericPlatformMisc::GetDefaultPathSeparator() + Filename;
+	if (!FFileHelper::SaveArrayToFile(NetworkData, *FilePath))
+	{
+		UE_LOG(LogLearning, Error, TEXT("Failed to save network to file: %s"), *FilePath);
+	}
+}
+
+void ULearningAgentsPolicy::LoadPolicyFromAsset(const ULearningAgentsNeuralNetwork* NeuralNetworkAsset)
+{
+	if (!IsPolicySetupPerformed())
+	{
+		UE_LOG(LogLearning, Error, TEXT("Policy setup must be run before network can be loaded."));
+		return;
+	}
+
+	if (!NeuralNetworkAsset || !NeuralNetworkAsset->NeuralNetwork)
+	{
+		UE_LOG(LogLearning, Error, TEXT("Cannot load policy from invalid asset."));
+		return;
+	}
+
+	if (NeuralNetworkAsset->NeuralNetwork->GetInputNum() != Network->NeuralNetwork->GetInputNum() ||
+		NeuralNetworkAsset->NeuralNetwork->GetOutputNum() != Network->NeuralNetwork->GetOutputNum())
+	{
+		UE_LOG(LogLearning, Error, TEXT("Failed to load policy from asset. Network Asset inputs and outputs don't match."));
+		return;
+	}
+
+	*Network->NeuralNetwork = *NeuralNetworkAsset->NeuralNetwork;
+}
+
+void ULearningAgentsPolicy::SavePolicyToAsset(ULearningAgentsNeuralNetwork* NeuralNetworkAsset) const
+{
+	if (!IsPolicySetupPerformed())
+	{
+		UE_LOG(LogLearning, Error, TEXT("Policy setup must be run before network can be saved."));
+		return;
+	}
+
+	if (!NeuralNetworkAsset)
+	{
+		UE_LOG(LogLearning, Error, TEXT("Cannot save policy to invalid asset."));
+		return;
+	}
+
+	if (!NeuralNetworkAsset->NeuralNetwork)
+	{
+		NeuralNetworkAsset->NeuralNetwork = MakeShared<UE::Learning::FNeuralNetwork>();
+	}
+
+	*NeuralNetworkAsset->NeuralNetwork = *Network->NeuralNetwork;
+
+	// Manually mark the package as dirty since just using `Modify` prevents 
+	// marking packages as dirty during PIE which is most likely when this
+	// is being used.
+	if (UPackage* Package = NeuralNetworkAsset->GetPackage())
+	{
+		const bool bIsDirty = Package->IsDirty();
+
+		if (!bIsDirty)
+		{
+			Package->SetDirtyFlag(true);
+		}
+
+		Package->PackageMarkedDirtyEvent.Broadcast(Package, bIsDirty);
+	}
+}
+
+void ULearningAgentsPolicy::EvaluatePolicy()
+{
+	UE_LEARNING_TRACE_CPUPROFILER_EVENT_SCOPE(ULearningAgentsPolicy::EvaluatePolicy);
+
+	if (!IsPolicySetupPerformed())
+	{
+		UE_LOG(LogLearning, Error, TEXT("Setup must be run before the policy can be evaluated."));
+		return;
+	}
+
+	PolicyObject->Evaluate(SelectedAgentIds);
+
+#if ENABLE_VISUAL_LOG
+	VisualLog(SelectedAgentIds);
+#endif
+}
+
+#if ENABLE_VISUAL_LOG
+void ULearningAgentsPolicy::VisualLog(const UE::Learning::FIndexSet Instances) const
+{
+	UE_LEARNING_TRACE_CPUPROFILER_EVENT_SCOPE(ULearningAgentsPolicy::VisualLog);
+
+	const TLearningArrayView<2, const float> InputView = PolicyObject->InstanceData->ConstView(PolicyObject->InputHandle);
+	const TLearningArrayView<2, const float> OutputView = PolicyObject->InstanceData->ConstView(PolicyObject->OutputHandle);
+	const TLearningArrayView<2, const float> OutputMeanView = PolicyObject->InstanceData->ConstView(PolicyObject->OutputMeanHandle);
+	const TLearningArrayView<2, const float> OutputStdView = PolicyObject->InstanceData->ConstView(PolicyObject->OutputStdHandle);
+	const TLearningArrayView<1, const float> ActionNoiseScaleView = PolicyObject->InstanceData->ConstView(PolicyObject->ActionNoiseScaleHandle);
+
+	for (const int32 Instance : Instances)
+	{
+		if (const AActor* Actor = Cast<AActor>(AgentType->GetAgent(Instance)))
+		{
+			const FString InputArrayString = UE::Learning::Agents::Policy::Private::ArrayToString(InputView[Instance]);
+			const FString OutputArrayString = UE::Learning::Agents::Policy::Private::ArrayToString(OutputView[Instance]);
+			const FString OutputMeanArrayString = UE::Learning::Agents::Policy::Private::ArrayToString(OutputMeanView[Instance]);
+			const FString OutputStdArrayString = UE::Learning::Agents::Policy::Private::ArrayToString(OutputStdView[Instance]);
+			const FString InputStatsString = UE::Learning::Agents::Policy::Private::ArrayToStatsString(InputView[Instance]);
+			const FString OutputStatsString = UE::Learning::Agents::Policy::Private::ArrayToStatsString(OutputView[Instance]);
+
+			UE_LEARNING_AGENTS_VLOG_STRING(this, LogLearning, Display,
+				Actor->GetActorLocation(),
+				VisualLogColor.ToFColor(true),
+				TEXT("Agent %i\nAction Noise Scale: [% 6.3f]\nInput: %s\nInput Stats (Min/Max/Mean/Std): %s\nOutput Mean: %s\nOutput Std: %s\nOutput Sample: %s\nOutput Stats (Min/Max/Mean/Std): %s"),
+				Instance,
+				ActionNoiseScaleView[Instance],
+				*InputArrayString,
+				*InputStatsString,
+				*OutputMeanArrayString,
+				*OutputStdArrayString,
+				*OutputArrayString,
+				*OutputStatsString);
+		}
+	}
+}
+#endif
+
+#undef UE_LEARNING_AGENTS_VLOG_STRING
\ No newline at end of file
diff --git a/Engine/Plugins/Experimental/LearningAgents/Source/LearningAgents/Private/LearningAgentsType.cpp b/Engine/Plugins/Experimental/LearningAgents/Source/LearningAgents/Private/LearningAgentsType.cpp
--- a/Engine/Plugins/Experimental/LearningAgents/Source/LearningAgents/Private/LearningAgentsType.cpp
+++ b/Engine/Plugins/Experimental/LearningAgents/Source/LearningAgents/Private/LearningAgentsType.cpp
@@ -8,18 +8,9 @@
 #include "LearningArrayMap.h"
 #include "LearningFeatureObject.h"
 #include "LearningLog.h"
-#include "LearningNeuralNetwork.h"
-#include "LearningPolicyObject.h"
-#include "Misc/FileHelper.h"
-#include "Misc/Paths.h"
 #include "EngineDefines.h"
 
-ULearningAgentsType::ULearningAgentsType()
-{
-	PrimaryComponentTick.bCanEverTick = true;
-	PrimaryComponentTick.bStartWithTickEnabled = false;
-}
-
+ULearningAgentsType::ULearningAgentsType() : UActorComponent() {}
 ULearningAgentsType::ULearningAgentsType(FVTableHelper& Helper) : ULearningAgentsType() {}
 ULearningAgentsType::~ULearningAgentsType() {}
 
@@ -44,12 +35,13 @@ void ULearningAgentsType::PostInitProperties()
 	UpdateAgentSets();
 }
 
-void ULearningAgentsType::SetupAgentType(
-	const FLearningAgentsTypeSettings& Settings,
-	const FLearningAgentsNetworkSettings& NetworkSettings)
+void ULearningAgentsType::SetupAgentType()
 {
-	// Reset Setup Flag
-	bSetupPerformed = false;
+	if (IsSetupPerformed())
+	{
+		UE_LOG(LogLearning, Error, TEXT("Setup already performed!"));
+		return;
+	}
 
 	// Allocate Instance Data
 	InstanceData = MakeShared<UE::Learning::FArrayMap>();
@@ -72,47 +64,6 @@ void ULearningAgentsType::SetupAgentType(
 		InstanceData.ToSharedRef(),
 		MaxInstanceNum);
 
-	// Create Neural Network
-	NeuralNetwork = MakeShared<UE::Learning::FNeuralNetwork>();
-	NeuralNetwork->Resize(
-		Observations->DimNum(), 
-		2 * Actions->DimNum(), 
-		NetworkSettings.HiddenLayerSize, 
-		NetworkSettings.LayerNum);
-
-	switch (NetworkSettings.ActivationFunction)
-	{
-		case ELearningAgentsActivationFunction::ReLU:
-			NeuralNetwork->ActivationFunction = UE::Learning::EActivationFunction::ReLU;
-			break;
-		case ELearningAgentsActivationFunction::ELU:
-			NeuralNetwork->ActivationFunction = UE::Learning::EActivationFunction::ELU;
-			break;
-		case ELearningAgentsActivationFunction::TanH:
-			NeuralNetwork->ActivationFunction = UE::Learning::EActivationFunction::TanH;
-			break;
-		default:
-			UE_LOG(LogLearning, Error, TEXT("NetworkSetting's activation function was not found. Defaulting to ELU."));
-			NeuralNetwork->ActivationFunction = UE::Learning::EActivationFunction::ELU;
-			break;
-	}
-
-	// Create Policy
-	UE::Learning::FNeuralNetworkPolicyFunctionSettings PolicySettings;
-	PolicySettings.ActionNoiseMin = Settings.ActionNoiseMin;
-	PolicySettings.ActionNoiseMax = Settings.ActionNoiseMax;
-
-	Policy = MakeUnique<UE::Learning::FNeuralNetworkPolicyFunction>(
-		TEXT("Policy"), 
-		InstanceData.ToSharedRef(), 
-		MaxInstanceNum, 
-		NeuralNetwork.ToSharedRef(), 
-		Settings.ActionNoiseSeed,
-		PolicySettings);
-
-	InstanceData->Link(Observations->FeatureHandle, Policy->InputHandle);
-	InstanceData->Link(Policy->OutputHandle, Actions->FeatureHandle);
-
 	// Done!
 	bSetupPerformed = true;
 }
@@ -142,14 +93,14 @@ UE::Learning::FFeatureObject& ULearningAgentsType::GetActionFeature() const
 	return *Actions;
 }
 
-UE::Learning::FNeuralNetwork& ULearningAgentsType::GetNeuralNetwork() const
+TConstArrayView<ULearningAgentsObservation*> ULearningAgentsType::GetObservationObjects() const
 {
-	return *NeuralNetwork;
+	return ObservationObjects;
 }
 
-const UE::Learning::FNeuralNetworkPolicyFunction& ULearningAgentsType::GetPolicy() const
+TConstArrayView<ULearningAgentsAction*> ULearningAgentsType::GetActionObjects() const
 {
-	return *Policy;
+	return ActionObjects;
 }
 
 const TConstArrayView<TObjectPtr<UObject>> ULearningAgentsType::GetAgents() const
@@ -189,8 +140,6 @@ int32 ULearningAgentsType::AddAgent(UObject* Agent)
 
 	UpdateAgentSets();
 
-	OnAgentAdded.Broadcast(NewAgentId, Agent);
-
 	return NewAgentId;
 }
 
@@ -217,8 +166,6 @@ void ULearningAgentsType::RemoveAgentById(int32 AgentId)
 	Agents[AgentId] = nullptr;
 
 	UpdateAgentSets();
-
-	OnAgentRemoved.Broadcast(AgentId, RemovedAgent);
 }
 
 void ULearningAgentsType::RemoveAgent(UObject* Agent)
@@ -241,6 +188,16 @@ void ULearningAgentsType::RemoveAgent(UObject* Agent)
 	RemoveAgentById(AgentId);
 }
 
+bool ULearningAgentsType::HasAgent(UObject* Agent) const
+{
+	return Agents.Find(Agent) ? true : false;
+}
+
+bool ULearningAgentsType::HasAgentById(int32 AgentId) const
+{
+	return OccupiedAgentSet.Contains(AgentId);
+}
+
 UObject* ULearningAgentsType::GetAgent(int32 AgentId, TSubclassOf<UObject> AgentClass)
 {
 	if (AgentId < 0 || AgentId >= Agents.Num())
@@ -254,12 +211,11 @@ UObject* ULearningAgentsType::GetAgent(int32 AgentId, TSubclassOf<UObject> Agent
 
 const UObject* ULearningAgentsType::GetAgent(int32 AgentId) const
 {
-	if (AgentId < 0 || AgentId >= Agents.Num())
-	{
-		UE_LOG(LogLearning, Warning, TEXT("AgentId %d outside valid range [0, %d]"), AgentId, Agents.Num() - 1);
-		return nullptr;
-	}
+	return Agents[AgentId];
+}
 
+UObject* ULearningAgentsType::GetAgent(int32 AgentId)
+{
 	return Agents[AgentId];
 }
 
@@ -275,7 +231,7 @@ void ULearningAgentsType::SetObservations_Implementation(const TArray<int32>& Ag
 
 void ULearningAgentsType::AddObservation(TObjectPtr<ULearningAgentsObservation> Object, const TSharedRef<UE::Learning::FFeatureObject>& Feature)
 {
-	check(!bSetupPerformed);
+	UE_LEARNING_CHECK(!IsSetupPerformed());
 	ObservationObjects.Add(Object);
 	ObservationFeatures.Add(Feature);
 }
@@ -292,37 +248,16 @@ void ULearningAgentsType::GetActions_Implementation(const TArray<int32>& AgentId
 
 void ULearningAgentsType::AddAction(TObjectPtr<ULearningAgentsAction> Object, const TSharedRef<UE::Learning::FFeatureObject>& Feature)
 {
-	check(!bSetupPerformed);
+	UE_LEARNING_CHECK(!IsSetupPerformed());
 	ActionObjects.Add(Object);
 	ActionFeatures.Add(Feature);
 }
 
-void ULearningAgentsType::LoadNetwork(const FDirectoryPath& Directory, const FString Filename)
-{
-	if (!bSetupPerformed)
-	{
-		UE_LOG(LogLearning, Error, TEXT("Setup must be run before network can be loaded."));
-		return;
-	}
-
-	TArray64<uint8> NetworkData;
-	FString FilePath = Directory.Path + FGenericPlatformMisc::GetDefaultPathSeparator() + Filename;
-	if (FFileHelper::LoadFileToArray(NetworkData, *FilePath))
-	{
-		NeuralNetwork->DeserializeFromBytes(NetworkData);
-		bNetworkLoaded = true;
-	}
-	else
-	{
-		UE_LOG(LogLearning, Warning, TEXT("Failed to load network. File not found: %s..."), *FilePath);
-	}
-}
-
 void ULearningAgentsType::EncodeObservations()
 {
 	UE_LEARNING_TRACE_CPUPROFILER_EVENT_SCOPE(ULearningAgentsType::EncodeObservations);
 
-	if (!bSetupPerformed)
+	if (!IsSetupPerformed())
 	{
 		UE_LOG(LogLearning, Error, TEXT("Setup must be run before observations can be encoded."));
 		return;
@@ -343,25 +278,11 @@ void ULearningAgentsType::EncodeObservations()
 #endif
 }
 
-
-void ULearningAgentsType::EvaluatePolicy()
-{
-	UE_LEARNING_TRACE_CPUPROFILER_EVENT_SCOPE(ULearningAgentsType::EvaluatePolicy);
-
-	if (!bSetupPerformed)
-	{
-		UE_LOG(LogLearning, Error, TEXT("Setup must be run before the policy can be evaluated."));
-		return;
-	}
-
-	Policy->Evaluate(OccupiedAgentSet);
-}
-
 void ULearningAgentsType::DecodeActions()
 {
 	UE_LEARNING_TRACE_CPUPROFILER_EVENT_SCOPE(ULearningAgentsType::DecodeActions);
 
-	if (!bSetupPerformed)
+	if (!IsSetupPerformed())
 	{
 		UE_LOG(LogLearning, Error, TEXT("Setup must be run before actions can be decoded."));
 		return;
@@ -382,11 +303,6 @@ void ULearningAgentsType::DecodeActions()
 #endif
 }
 
-void ULearningAgentsType::BroadcastSetupComplete()
-{
-	OnSetupComplete.Broadcast();
-}
-
 void ULearningAgentsType::UpdateAgentSets()
 {
 	OccupiedAgentSet = OccupiedAgentIds;
diff --git a/Engine/Plugins/Experimental/LearningAgents/Source/LearningAgents/Public/LearningAgentsController.h b/Engine/Plugins/Experimental/LearningAgents/Source/LearningAgents/Public/LearningAgentsController.h
new file mode 100644
--- /dev/null
+++ b/Engine/Plugins/Experimental/LearningAgents/Source/LearningAgents/Public/LearningAgentsController.h
@@ -0,0 +1,66 @@
+// Copyright Epic Games, Inc. All Rights Reserved.
+
+#pragma once
+
+#include "LearningArray.h"
+
+#include "Components/ActorComponent.h"
+
+#include "LearningAgentsController.generated.h"
+
+UCLASS(Abstract, BlueprintType, Blueprintable)
+class LEARNINGAGENTS_API ULearningAgentsController : public UActorComponent
+{
+	GENERATED_BODY()
+
+// ----- Setup -----
+public:
+
+	// These constructors/destructors are needed to make forward declarations happy
+	ULearningAgentsController();
+	ULearningAgentsController(FVTableHelper& Helper);
+	virtual ~ULearningAgentsController();
+
+	UFUNCTION(BlueprintCallable, Category = "LearningAgents")
+	void SetupController(ULearningAgentsType* InAgentType);
+
+	UFUNCTION(BlueprintPure, Category = "LearningAgents")
+	bool IsControllerSetupPerformed() const;
+
+// ----- Agent Management -----
+public:
+
+	UFUNCTION(BlueprintCallable, Category = "LearningAgents")
+	void AddAgent(int32 AgentId);
+
+	UFUNCTION(BlueprintCallable, Category = "LearningAgents")
+	void RemoveAgent(int32 AgentId);
+
+	UFUNCTION(BlueprintPure, Category = "LearningAgents")
+	bool HasAgent(int32 AgentId) const;
+
+	UFUNCTION(BlueprintPure, Category = "LearningAgents", meta = (DeterminesOutputType = "AgentClass"))
+	ULearningAgentsType* GetAgentType(TSubclassOf<ULearningAgentsType> AgentClass);
+
+// ----- Actions -----
+public:
+
+	UFUNCTION(BlueprintNativeEvent, Category = "LearningAgents")
+	void SetActions(const TArray<int32>& AgentIds);
+
+	UFUNCTION(BlueprintCallable, Category = "LearningAgents")
+	void EncodeActions();
+
+// ----- Private Data -----
+private:
+
+	UPROPERTY(VisibleAnywhere, Category = "LearningAgents")
+	TObjectPtr<ULearningAgentsType> AgentType;
+
+	UPROPERTY(VisibleAnywhere, Category = "LearningAgents")
+	TArray<int32> SelectedAgentIds;
+
+private:
+
+	UE::Learning::FIndexSet SelectedAgentsSet;
+};
diff --git a/Engine/Plugins/Experimental/LearningAgents/Source/LearningAgents/Public/LearningAgentsCritic.h b/Engine/Plugins/Experimental/LearningAgents/Source/LearningAgents/Public/LearningAgentsCritic.h
new file mode 100644
--- /dev/null
+++ b/Engine/Plugins/Experimental/LearningAgents/Source/LearningAgents/Public/LearningAgentsCritic.h
@@ -0,0 +1,137 @@
+// Copyright Epic Games, Inc. All Rights Reserved.
+
+#pragma once
+
+#include "LearningAgentsNeuralNetwork.h" // Included for ELearningAgentsActivationFunction
+#include "LearningArray.h"
+
+#include "Components/ActorComponent.h"
+#include "UObject/ObjectPtr.h"
+#include "EngineDefines.h"
+
+#include "LearningAgentsCritic.generated.h"
+
+namespace UE::Learning
+{
+	struct FNeuralNetwork;
+	struct FNeuralNetworkCriticFunction;
+}
+
+struct FDirectoryPath;
+class ULearningAgentsNeuralNetwork;
+
+//------------------------------------------------------------------
+
+USTRUCT(BlueprintType, Category = "LearningAgents")
+struct LEARNINGAGENTS_API FLearningAgentsCriticSettings
+{
+	GENERATED_BODY()
+
+public:
+
+	/** Total layers for critic network including input, hidden, and output layers */
+	UPROPERTY(EditAnywhere, Category = "LearningAgents", meta = (ClampMin = "2", UIMin = "2"))
+	int32 LayerNum = 3;
+
+	/** Number of neurons in each hidden layer of the critic network */
+	UPROPERTY(EditAnywhere, Category = "LearningAgents", meta = (ClampMin = "1", UIMin = "1"))
+	int32 HiddenLayerSize = 128;
+
+	/** Activation function to use on hidden layers of the critic network */
+	UPROPERTY(EditAnywhere, Category = "LearningAgents")
+	ELearningAgentsActivationFunction ActivationFunction = ELearningAgentsActivationFunction::ELU;
+};
+
+//------------------------------------------------------------------
+
+UCLASS(BlueprintType, Blueprintable)
+class LEARNINGAGENTS_API ULearningAgentsCritic : public UActorComponent
+{
+	GENERATED_BODY()
+
+// ----- Setup -----
+public:
+
+	// These constructors/destructors are needed to make forward declarations happy
+	ULearningAgentsCritic();
+	ULearningAgentsCritic(FVTableHelper& Helper);
+	virtual ~ULearningAgentsCritic();
+
+	UFUNCTION(BlueprintCallable, Category = "LearningAgents")
+	void SetupCritic(ULearningAgentsType* InAgentType, const FLearningAgentsCriticSettings& CriticSettings = FLearningAgentsCriticSettings());
+
+	UFUNCTION(BlueprintPure, Category = "LearningAgents")
+	bool IsCriticSetupPerformed() const;
+
+// ----- Agent Management -----
+public:
+
+	UFUNCTION(BlueprintCallable, Category = "LearningAgents")
+	void AddAgent(int32 AgentId);
+
+	UFUNCTION(BlueprintCallable, Category = "LearningAgents")
+	void RemoveAgent(int32 AgentId);
+
+	UFUNCTION(BlueprintPure, Category = "LearningAgents")
+	bool HasAgent(int32 AgentId) const;
+
+	UFUNCTION(BlueprintPure, Category = "LearningAgents", meta = (DeterminesOutputType = "AgentClass"))
+	ULearningAgentsType* GetAgentType(TSubclassOf<ULearningAgentsType> AgentClass);
+
+// ----- Load / Save -----
+public:
+
+	UFUNCTION(BlueprintCallable, Category = "LearningAgents")
+	void LoadCriticFromSnapshot(const FDirectoryPath& Directory, const FString Filename);
+
+	UFUNCTION(BlueprintCallable, Category = "LearningAgents")
+	void SaveCriticToSnapshot(const FDirectoryPath& Directory, const FString Filename) const;
+
+	UFUNCTION(BlueprintCallable, Category = "LearningAgents")
+	void LoadCriticFromAsset(const ULearningAgentsNeuralNetwork* NeuralNetworkAsset);
+
+	UFUNCTION(BlueprintCallable, Category = "LearningAgents", Meta=(DevelopmentOnly))
+	void SaveCriticToAsset(ULearningAgentsNeuralNetwork* NeuralNetworkAsset) const;
+
+// ----- Evaluation -----
+public:
+
+	UFUNCTION(BlueprintCallable, Category = "LearningAgents")
+	void EvaluateCritic();
+
+// ----- Non-blueprint public interface -----
+public:
+
+	UE::Learning::FNeuralNetwork& GetCriticNetwork();
+	UE::Learning::FNeuralNetworkCriticFunction& GetCriticObject();
+
+// ----- Private Data -----
+private:
+
+	UPROPERTY(VisibleAnywhere, Category = "LearningAgents")
+	TObjectPtr<ULearningAgentsType> AgentType;
+
+	UPROPERTY(VisibleAnywhere, Category = "LearningAgents")
+	TArray<int32> SelectedAgentIds;
+
+	UPROPERTY(VisibleAnywhere, Category = "LearningAgents")
+	bool bCriticSetupPerformed = false;
+
+	UPROPERTY(VisibleAnywhere, Category = "LearningAgents")
+	TObjectPtr<ULearningAgentsNeuralNetwork> Network;
+
+private:
+
+	TSharedPtr<UE::Learning::FNeuralNetworkCriticFunction> CriticObject;
+
+	UE::Learning::FIndexSet SelectedAgentsSet;
+
+private:
+
+#if ENABLE_VISUAL_LOG
+	/** Color used to draw this action in the visual log */
+	FLinearColor VisualLogColor = FColor::Orange;
+
+	void VisualLog(const UE::Learning::FIndexSet Instances) const;
+#endif
+};
diff --git a/Engine/Plugins/Experimental/LearningAgents/Source/LearningAgents/Public/LearningAgentsNeuralNetwork.h b/Engine/Plugins/Experimental/LearningAgents/Source/LearningAgents/Public/LearningAgentsNeuralNetwork.h
new file mode 100644
--- /dev/null
+++ b/Engine/Plugins/Experimental/LearningAgents/Source/LearningAgents/Public/LearningAgentsNeuralNetwork.h
@@ -0,0 +1,54 @@
+// Copyright Epic Games, Inc. All Rights Reserved.
+
+#pragma once
+
+#include "Engine/DataAsset.h"
+
+#include "LearningAgentsNeuralNetwork.generated.h"
+
+namespace UE::Learning
+{
+	struct FNeuralNetwork;
+	enum class EActivationFunction : uint8;
+}
+
+//------------------------------------------------------------------
+
+UENUM(BlueprintType, Category = "LearningAgents")
+enum class ELearningAgentsActivationFunction : uint8
+{
+	/** ReLU Activation - Fast to train and evaluate but occasionally causes gradient collapse and untrainable networks. */
+	ReLU	UMETA(DisplayName = "ReLU"),
+
+	/** ELU Activation - Generally performs better than ReLU and is not prone to gradient collapse but slower to evaluate. */
+	ELU		UMETA(DisplayName = "ELU"),
+
+	/** TanH Activation - Smooth activation function that is slower to train and evaluate but sometimes more stable for certain tasks. */
+	TanH	UMETA(DisplayName = "TanH"),
+};
+
+//------------------------------------------------------------------
+
+namespace UE::Learning::Agents
+{
+	LEARNINGAGENTS_API ELearningAgentsActivationFunction GetLearningAgentsActivationFunction(const EActivationFunction ActivationFunction);
+	LEARNINGAGENTS_API EActivationFunction GetActivationFunction(const ELearningAgentsActivationFunction ActivationFunction);
+}
+
+//------------------------------------------------------------------
+
+UCLASS(BlueprintType)
+class LEARNINGAGENTS_API ULearningAgentsNeuralNetwork : public UDataAsset
+{
+	GENERATED_BODY()
+
+public:
+	// These constructors/destructors are needed to make forward declarations happy
+	ULearningAgentsNeuralNetwork();
+	ULearningAgentsNeuralNetwork(FVTableHelper& Helper);
+	virtual ~ULearningAgentsNeuralNetwork();
+
+	virtual void Serialize(FArchive& Ar) override;
+
+	TSharedPtr<UE::Learning::FNeuralNetwork> NeuralNetwork;
+};
diff --git a/Engine/Plugins/Experimental/LearningAgents/Source/LearningAgents/Public/LearningAgentsPolicy.h b/Engine/Plugins/Experimental/LearningAgents/Source/LearningAgents/Public/LearningAgentsPolicy.h
new file mode 100644
--- /dev/null
+++ b/Engine/Plugins/Experimental/LearningAgents/Source/LearningAgents/Public/LearningAgentsPolicy.h
@@ -0,0 +1,149 @@
+// Copyright Epic Games, Inc. All Rights Reserved.
+
+#pragma once
+
+#include "LearningAgentsNeuralNetwork.h" // Included for ELearningAgentsActivationFunction
+#include "LearningArray.h"
+
+#include "Components/ActorComponent.h"
+#include "UObject/ObjectPtr.h"
+#include "EngineDefines.h"
+
+#include "LearningAgentsPolicy.generated.h"
+
+namespace UE::Learning
+{
+	struct FNeuralNetwork;
+	struct FNeuralNetworkPolicyFunction;
+}
+
+struct FDirectoryPath;
+class ULearningAgentsNeuralNetwork;
+
+//------------------------------------------------------------------
+
+USTRUCT(BlueprintType, Category = "LearningAgents")
+struct LEARNINGAGENTS_API FLearningAgentsPolicySettings
+{
+	GENERATED_BODY()
+
+public:
+
+	/** Seed for the action noise used by the policy  */
+	UPROPERTY(EditAnywhere, Category = "LearningAgents", meta = (ClampMin = "0", UIMin = "0"))
+	int32 ActionNoiseSeed = 1234;
+
+	/** Minimum action noise used by the policy */
+	UPROPERTY(EditAnywhere, Category = "LearningAgents", meta = (ClampMin = "0.0", UIMin = "0.0"))
+	float ActionNoiseMin = 0.25f;
+
+	/** Maximum action noise used by the policy */
+	UPROPERTY(EditAnywhere, Category = "LearningAgents", meta = (ClampMin = "0.0", UIMin = "0.0"))
+	float ActionNoiseMax = 0.25f;
+
+	/** Total layers for policy network including input, hidden, and output layers */
+	UPROPERTY(EditAnywhere, Category = "LearningAgents", meta = (ClampMin = "2", UIMin = "2"))
+	int32 LayerNum = 3;
+
+	/** Number of neurons in each hidden layer of the policy network */
+	UPROPERTY(EditAnywhere, Category = "LearningAgents", meta = (ClampMin = "1", UIMin = "1"))
+	int32 HiddenLayerSize = 128;
+
+	/** Activation function to use on hidden layers of the policy network */
+	UPROPERTY(EditAnywhere, Category = "LearningAgents")
+	ELearningAgentsActivationFunction ActivationFunction = ELearningAgentsActivationFunction::ELU;
+};
+
+//------------------------------------------------------------------
+
+UCLASS(BlueprintType, Blueprintable)
+class LEARNINGAGENTS_API ULearningAgentsPolicy : public UActorComponent
+{
+	GENERATED_BODY()
+
+// ----- Setup -----
+public:
+
+	// These constructors/destructors are needed to make forward declarations happy
+	ULearningAgentsPolicy();
+	ULearningAgentsPolicy(FVTableHelper& Helper);
+	virtual ~ULearningAgentsPolicy();
+
+	UFUNCTION(BlueprintCallable, Category = "LearningAgents")
+	void SetupPolicy(ULearningAgentsType* InAgentType, const FLearningAgentsPolicySettings& PolicySettings = FLearningAgentsPolicySettings());
+
+	UFUNCTION(BlueprintPure, Category = "LearningAgents")
+	bool IsPolicySetupPerformed() const;
+
+// ----- Agent Management -----
+public:
+
+	UFUNCTION(BlueprintCallable, Category = "LearningAgents")
+	void AddAgent(int32 AgentId);
+
+	UFUNCTION(BlueprintCallable, Category = "LearningAgents")
+	void RemoveAgent(int32 AgentId);
+
+	UFUNCTION(BlueprintPure, Category = "LearningAgents")
+	bool HasAgent(int32 AgentId) const;
+
+	UFUNCTION(BlueprintPure, Category = "LearningAgents", meta = (DeterminesOutputType = "AgentClass"))
+	ULearningAgentsType* GetAgentType(TSubclassOf<ULearningAgentsType> AgentClass);
+
+// ----- Load / Save -----
+public:
+
+	UFUNCTION(BlueprintCallable, Category = "LearningAgents")
+	void LoadPolicyFromSnapshot(const FDirectoryPath& Directory, const FString Filename);
+
+	UFUNCTION(BlueprintCallable, Category = "LearningAgents")
+	void SavePolicyToSnapshot(const FDirectoryPath& Directory, const FString Filename) const;
+
+	UFUNCTION(BlueprintCallable, Category = "LearningAgents")
+	void LoadPolicyFromAsset(const ULearningAgentsNeuralNetwork* NeuralNetworkAsset);
+
+	UFUNCTION(BlueprintCallable, Category = "LearningAgents", Meta = (DevelopmentOnly))
+	void SavePolicyToAsset(ULearningAgentsNeuralNetwork* NeuralNetworkAsset) const;
+
+// ----- Evaluation -----
+public:
+
+	UFUNCTION(BlueprintCallable, Category = "LearningAgents")
+	void EvaluatePolicy();
+
+// ----- Non-blueprint public interface -----
+public:
+
+	UE::Learning::FNeuralNetwork& GetPolicyNetwork();
+	UE::Learning::FNeuralNetworkPolicyFunction& GetPolicyObject();
+
+// ----- Private Data -----
+private:
+
+	UPROPERTY(VisibleAnywhere, Category = "LearningAgents")
+	TObjectPtr<ULearningAgentsType> AgentType;
+
+	UPROPERTY(VisibleAnywhere, Category = "LearningAgents")
+	TArray<int32> SelectedAgentIds;
+
+	UPROPERTY(VisibleAnywhere, Category = "LearningAgents")
+	bool bPolicySetupPerformed = false;
+
+	UPROPERTY(VisibleAnywhere, Category = "LearningAgents")
+	TObjectPtr<ULearningAgentsNeuralNetwork> Network;
+
+private:
+
+	TSharedPtr<UE::Learning::FNeuralNetworkPolicyFunction> PolicyObject;
+
+	UE::Learning::FIndexSet SelectedAgentsSet;
+
+private:
+
+#if ENABLE_VISUAL_LOG
+	/** Color used to draw this action in the visual log */
+	FLinearColor VisualLogColor = FColor::Purple;
+
+	void VisualLog(const UE::Learning::FIndexSet Instances) const;
+#endif
+};
diff --git a/Engine/Plugins/Experimental/LearningAgents/Source/LearningAgents/Public/LearningAgentsType.h b/Engine/Plugins/Experimental/LearningAgents/Source/LearningAgents/Public/LearningAgentsType.h
--- a/Engine/Plugins/Experimental/LearningAgents/Source/LearningAgents/Public/LearningAgentsType.h
+++ b/Engine/Plugins/Experimental/LearningAgents/Source/LearningAgents/Public/LearningAgentsType.h
@@ -2,12 +2,12 @@
 
 #pragma once
 
-#include "Components/SceneComponent.h"
+#include "LearningArray.h"
+
+#include "Components/ActorComponent.h"
 #include "Containers/Array.h"
 #include "Containers/ArrayView.h"
 #include "Containers/ContainerAllocationPolicies.h"
-#include "Containers/ContainersFwd.h"
-#include "LearningArray.h"
 #include "Templates/SharedPointer.h"
 #include "UObject/Object.h"
 #include "UObject/ObjectPtr.h"
@@ -19,83 +19,23 @@ namespace UE::Learning
 	struct FArrayMap;
 	struct FConcatenateFeature;
 	struct FFeatureObject;
-	struct FNeuralNetwork;
-	struct FNeuralNetworkPolicyFunction;
 }
 
-struct FDirectoryPath;
 class ULearningAgentsAction;
 class ULearningAgentsObservation;
 
-USTRUCT(BlueprintType, Category = "LearningAgents")
-struct FLearningAgentsTypeSettings
-{
-	GENERATED_BODY()
-
-public:
-
-	/** Seed for the action noise used by the policy  */
-	UPROPERTY(EditAnywhere, Category = "LearningAgents", meta = (ClampMin = "0", UIMin = "0"))
-	int32 ActionNoiseSeed = 1234;
-
-	/** Minimum action noise used by the policy */
-	UPROPERTY(EditAnywhere, Category = "LearningAgents", meta = (ClampMin = "0.0", UIMin = "0.0"))
-	float ActionNoiseMin = 0.25f;
-
-	/** Maximum action noise used by the policy */
-	UPROPERTY(EditAnywhere, Category = "LearningAgents", meta = (ClampMin = "0.0", UIMin = "0.0"))
-	float ActionNoiseMax = 0.25f;
-};
-
-UENUM(BlueprintType, Category = "LearningAgents")
-enum class ELearningAgentsActivationFunction : uint8
-{
-	/** ReLU Activation - Fast to train and evaluate but occasionally causes gradient collapse and untrainable networks. */
-	ReLU	UMETA(DisplayName = "ReLU"),
-
-	/** ELU Activation - Generally performs better than ReLU and is not prone to gradient collapse but slower to evaluate. */
-	ELU		UMETA(DisplayName = "ELU"),
-
-	/** TanH Activation - Smooth activation function that is slower to train and evaluate but sometimes more stable for certain tasks. */
-	TanH	UMETA(DisplayName = "TanH"),
-};
-
-USTRUCT(BlueprintType, Category = "LearningAgents")
-struct FLearningAgentsNetworkSettings
-{
-	GENERATED_BODY()
-
-public:
-
-	/** Total layers for neural network including input, hidden, and output layers */
-	UPROPERTY(EditAnywhere, Category = "LearningAgents", meta = (ClampMin = "2", UIMin = "2"))
-	int32 LayerNum = 3;
-
-	/** Number of neurons in each hidden layer */
-	UPROPERTY(EditAnywhere, Category = "LearningAgents", meta = (ClampMin = "1", UIMin = "1"))
-	int32 HiddenLayerSize = 128;
-
-	/** Activation function to use on hidden layers */
-	UPROPERTY(EditAnywhere, Category = "LearningAgents")
-	ELearningAgentsActivationFunction ActivationFunction = ELearningAgentsActivationFunction::ELU;
-};
-
 UCLASS(Abstract, BlueprintType, Blueprintable)
-class LEARNINGAGENTS_API ULearningAgentsType : public USceneComponent
+class LEARNINGAGENTS_API ULearningAgentsType : public UActorComponent
 {
 	GENERATED_BODY()
 
 public:
 
-	DECLARE_MULTICAST_DELEGATE(FOnSetupComplete);
-	DECLARE_MULTICAST_DELEGATE_TwoParams(FOnAgentEvent, int32 /*AgentId*/, UObject* /*Agent*/);
-
 	// These constructors/destructors are needed to make forward declarations happy
 	ULearningAgentsType();
 	ULearningAgentsType(FVTableHelper& Helper);
 	virtual ~ULearningAgentsType();
 
-	/** Sets up the agent ids so that we can add/remove agents at any time. */
 	virtual void PostInitProperties() override;
 
 // ----- Setup -----
@@ -105,9 +45,7 @@ public:
 	 * Initializes this object and runs the setup functions for observations and actions
 	 */
 	UFUNCTION(BlueprintCallable, Category = "LearningAgents")
-	void SetupAgentType(
-		const FLearningAgentsTypeSettings& Settings = FLearningAgentsTypeSettings(),
-		const FLearningAgentsNetworkSettings& NetworkSettings = FLearningAgentsNetworkSettings());
+	void SetupAgentType();
 
 	UFUNCTION(BlueprintPure, Category = "LearningAgents")
 	const bool IsSetupPerformed() const;
@@ -130,7 +68,16 @@ public:
 	UFUNCTION(BlueprintPure, Category = "LearningAgents", meta = (DeterminesOutputType = "AgentClass"))
 	UObject* GetAgent(int32 AgentId, TSubclassOf<UObject> AgentClass);
 
+	UFUNCTION(BlueprintPure, Category = "LearningAgents")
+	bool HasAgent(UObject* Agent) const;
+
+	UFUNCTION(BlueprintPure, Category = "LearningAgents")
+	bool HasAgentById(int32 AgentId) const;
+
+public:
+
 	const UObject* GetAgent(int32 AgentId) const;
+	UObject* GetAgent(int32 AgentId);
 
 // ----- Observations -----
 public:
@@ -162,39 +109,28 @@ public:
 
 	void AddAction(TObjectPtr<ULearningAgentsAction> Object, const TSharedRef<UE::Learning::FFeatureObject>& Feature);
 
-// ----- Inference Process -----
+// ----- Encoding / Decoding -----
 public:
 
-	UFUNCTION(BlueprintCallable, Category = "LearningAgents")
-	void LoadNetwork(const FDirectoryPath& Directory, const FString Filename);
-
 	UFUNCTION(BlueprintCallable, Category = "LearningAgents")
 	void EncodeObservations();
 
-	UFUNCTION(BlueprintCallable, Category = "LearningAgents")
-	void EvaluatePolicy();
-
 	UFUNCTION(BlueprintCallable, Category = "LearningAgents")
 	void DecodeActions();
 
-	UFUNCTION(BlueprintCallable, Category = "LearningAgents")
-	void BroadcastSetupComplete();
-
 // ----- Non-blueprint public interface -----
 public:
+
 	const TSharedPtr<UE::Learning::FArrayMap>& GetInstanceData() const;
+
 	UE::Learning::FFeatureObject& GetObservationFeature() const;
 	UE::Learning::FFeatureObject& GetActionFeature() const;
-	UE::Learning::FNeuralNetwork& GetNeuralNetwork() const;
-	const UE::Learning::FNeuralNetworkPolicyFunction& GetPolicy() const;
+	TConstArrayView<ULearningAgentsObservation*> GetObservationObjects() const;
+	TConstArrayView<ULearningAgentsAction*> GetActionObjects() const;
+
 	const TConstArrayView<TObjectPtr<UObject>> GetAgents() const;
 	const UE::Learning::FIndexSet GetOccupiedAgentSet() const;
 	const UE::Learning::FIndexSet GetVacantAgentSet() const;
-	FOnSetupComplete& GetOnSetupComplete() { return OnSetupComplete; }
-	FOnAgentEvent& GetOnAgentAdded() { return OnAgentAdded; }
-	FOnAgentEvent& GetOnAgentRemoved() { return OnAgentRemoved; }
-	TConstArrayView<ULearningAgentsObservation*> GetObservationObjects() const { return ObservationObjects; }
-	TConstArrayView<ULearningAgentsAction*> GetActionObjects() const { return ActionObjects; }
 
 // ----- Private Data -----
 private:
@@ -209,17 +145,12 @@ private:
 	TSharedPtr<UE::Learning::FArrayMap> InstanceData;
 
 	TArray<TSharedRef<UE::Learning::FFeatureObject>, TInlineAllocator<16>> ObservationFeatures;
-	TSharedPtr<UE::Learning::FConcatenateFeature> Observations;
-	
 	TArray<TSharedRef<UE::Learning::FFeatureObject>, TInlineAllocator<16>> ActionFeatures;
-	TSharedPtr<UE::Learning::FConcatenateFeature> Actions;
 
-	TSharedPtr<UE::Learning::FNeuralNetwork> NeuralNetwork;
-	TUniquePtr<UE::Learning::FNeuralNetworkPolicyFunction> Policy;
-
-	FOnSetupComplete OnSetupComplete;
-	FOnAgentEvent OnAgentAdded;
-	FOnAgentEvent OnAgentRemoved;
+	TSharedPtr<UE::Learning::FConcatenateFeature> Observations;	
+	TSharedPtr<UE::Learning::FConcatenateFeature> Actions;
+	
+private:
 
 	/** Maximum number of agent instances. Used to preallocate internal buffers. */
 	UPROPERTY(EditDefaultsOnly, Category = "LearningAgents", meta = (ClampMin = "1", UIMin = "1"))
@@ -228,9 +159,6 @@ private:
 	UPROPERTY(VisibleAnywhere, Category = "LearningAgents")
 	bool bSetupPerformed = false;
 
-	UPROPERTY(VisibleAnywhere, Category = "LearningAgents")
-	bool bNetworkLoaded = false;
-
 	UPROPERTY(VisibleAnywhere, Transient, Category = "LearningAgents")
 	TArray<TObjectPtr<UObject>> Agents;
 
diff --git a/Engine/Plugins/Experimental/LearningAgents/Source/LearningAgentsTraining/Private/LearningAgentsCompletions.cpp b/Engine/Plugins/Experimental/LearningAgents/Source/LearningAgentsTraining/Private/LearningAgentsCompletions.cpp
--- a/Engine/Plugins/Experimental/LearningAgents/Source/LearningAgentsTraining/Private/LearningAgentsCompletions.cpp
+++ b/Engine/Plugins/Experimental/LearningAgents/Source/LearningAgentsTraining/Private/LearningAgentsCompletions.cpp
@@ -24,7 +24,8 @@
 	UE_LEARNING_AGENTS_VLOG_STRING(Owner, Category, Verbosity, Location + FVector(0.0f, 0.0f, 20.0f), Color, Format, ##__VA_ARGS__)
 
 #define UE_LEARNING_AGENTS_VLOG_PLANE(Owner, Category, Verbosity, Location, Rotation, Axis0, Axis1, Color, Format, ...) \
-	UE_VLOG_OBOX(Owner, Category, Verbosity, FBox(25.0f * FVector(-1, -1, 0), 25.0f * FVector(1, 1, 0)), FTransform(Rotation, Location, FVector::OneVector).ToMatrixNoScale(), Color, Format, ##__VA_ARGS__)
+	UE_VLOG_OBOX(Owner, Category, Verbosity, FBox(-25.0f * (Axis0 + Axis1), 25.0f * (Axis0 + Axis1)), FTransform(Rotation, Location, FVector::OneVector).ToMatrixNoScale(), Color, Format, ##__VA_ARGS__)
+
 
 namespace UE::Learning::Agents::Private
 {
@@ -40,8 +41,8 @@ namespace UE::Learning::Agents::Private
 
 		Completion->CompletionObject = MakeShared<CompletionFObject>(
 			Completion->GetFName(),
-			AgentTrainer->GetAgentType().GetInstanceData().ToSharedRef(),
-			AgentTrainer->GetAgentType().GetMaxInstanceNum(),
+			AgentTrainer->GetAgentType()->GetInstanceData().ToSharedRef(),
+			AgentTrainer->GetAgentType()->GetMaxInstanceNum(),
 			Forward<InArgTypes>(Args)...);
 
 		AgentTrainer->AddCompletion(Completion, Completion->CompletionObject.ToSharedRef());
@@ -115,6 +116,7 @@ UPlanarPositionDifferenceCompletion* UPlanarPositionDifferenceCompletion::AddPla
 	return UE::Learning::Agents::Private::AddCompletion<UPlanarPositionDifferenceCompletion, UE::Learning::FPlanarPositionDifferenceCompletion>(
 		AgentTrainer,
 		Name,
+		1,
 		Threshold,
 		InCompletionMode == ELearningAgentsCompletion::Termination ?
 		UE::Learning::ECompletionMode::Terminated :
@@ -125,8 +127,8 @@ UPlanarPositionDifferenceCompletion* UPlanarPositionDifferenceCompletion::AddPla
 
 void UPlanarPositionDifferenceCompletion::SetPlanarPositionDifferenceCompletion(int32 AgentId, FVector Position0, FVector Position1)
 {
-	TLearningArrayView<1, FVector> Position0View = CompletionObject->InstanceData->View(CompletionObject->Position0Handle);
-	TLearningArrayView<1, FVector> Position1View = CompletionObject->InstanceData->View(CompletionObject->Position1Handle);
+	TLearningArrayView<2, FVector> Position0View = CompletionObject->InstanceData->View(CompletionObject->Position0Handle);
+	TLearningArrayView<2, FVector> Position1View = CompletionObject->InstanceData->View(CompletionObject->Position1Handle);
 
 	if (AgentId == INDEX_NONE)
 	{
@@ -140,8 +142,8 @@ void UPlanarPositionDifferenceCompletion::SetPlanarPositionDifferenceCompletion(
 		return;
 	}
 
-	Position0View[AgentId] = Position0;
-	Position1View[AgentId] = Position1;
+	Position0View[AgentId][0] = Position0;
+	Position1View[AgentId][0] = Position1;
 }
 
 #if ENABLE_VISUAL_LOG
@@ -149,8 +151,8 @@ void UPlanarPositionDifferenceCompletion::VisualLog(const UE::Learning::FIndexSe
 {
 	UE_LEARNING_TRACE_CPUPROFILER_EVENT_SCOPE(UPlanarPositionDifferenceCompletion::VisualLog);
 
-	const TLearningArrayView<1, const FVector> Position0View = CompletionObject->InstanceData->ConstView(CompletionObject->Position0Handle);
-	const TLearningArrayView<1, const FVector> Position1View = CompletionObject->InstanceData->ConstView(CompletionObject->Position1Handle);
+	const TLearningArrayView<2, const FVector> Position0View = CompletionObject->InstanceData->ConstView(CompletionObject->Position0Handle);
+	const TLearningArrayView<2, const FVector> Position1View = CompletionObject->InstanceData->ConstView(CompletionObject->Position1Handle);
 	const TLearningArrayView<1, const float> ThresholdView = CompletionObject->InstanceData->ConstView(CompletionObject->ThresholdHandle);
 	const TLearningArrayView<1, const UE::Learning::ECompletionMode> CompletionView = CompletionObject->InstanceData->ConstView(CompletionObject->CompletionHandle);
 
@@ -160,8 +162,8 @@ void UPlanarPositionDifferenceCompletion::VisualLog(const UE::Learning::FIndexSe
 		{
 			if (const AActor* Actor = Cast<AActor>(AgentTrainer->GetAgent(Instance)))
 			{
-				const FVector Position0 = Position0View[Instance];
-				const FVector Position1 = Position1View[Instance];
+				const FVector Position0 = Position0View[Instance][0];
+				const FVector Position1 = Position1View[Instance][0];
 
 				const FVector PlanarPosition0 = FVector(CompletionObject->Axis0.Dot(Position0), CompletionObject->Axis1.Dot(Position0), 0.0f);
 				const FVector PlanarPosition1 = FVector(CompletionObject->Axis0.Dot(Position1), CompletionObject->Axis1.Dot(Position1), 0.0f);
@@ -177,8 +179,8 @@ void UPlanarPositionDifferenceCompletion::VisualLog(const UE::Learning::FIndexSe
 				UE_LEARNING_AGENTS_VLOG_PLANE(this, LogLearning, Display,
 					Position0,
 					FQuat::Identity,
-					FeatureObject->Axis0,
-					FeatureObject->Axis1,
+					CompletionObject->Axis0,
+					CompletionObject->Axis1,
 					VisualLogColor.ToFColor(true),
 					TEXT(""));
 
@@ -193,8 +195,8 @@ void UPlanarPositionDifferenceCompletion::VisualLog(const UE::Learning::FIndexSe
 				UE_LEARNING_AGENTS_VLOG_PLANE(this, LogLearning, Display,
 					Position1,
 					FQuat::Identity,
-					FeatureObject->Axis0,
-					FeatureObject->Axis1,
+					CompletionObject->Axis0,
+					CompletionObject->Axis1,
 					VisualLogColor.ToFColor(true),
 					TEXT(""));
 
@@ -218,4 +220,4 @@ void UPlanarPositionDifferenceCompletion::VisualLog(const UE::Learning::FIndexSe
 
 #undef UE_LEARNING_AGENTS_VLOG_STRING
 #undef UE_LEARNING_AGENTS_VLOG_TRANSFORM
-#undef UE_LEARNING_AGENTS_VLOG_PLANE
+#undef UE_LEARNING_AGENTS_VLOG_PLANE
\ No newline at end of file
diff --git a/Engine/Plugins/Experimental/LearningAgents/Source/LearningAgentsTraining/Private/LearningAgentsDataStorage.cpp b/Engine/Plugins/Experimental/LearningAgents/Source/LearningAgentsTraining/Private/LearningAgentsDataStorage.cpp
--- a/Engine/Plugins/Experimental/LearningAgents/Source/LearningAgentsTraining/Private/LearningAgentsDataStorage.cpp
+++ b/Engine/Plugins/Experimental/LearningAgents/Source/LearningAgentsTraining/Private/LearningAgentsDataStorage.cpp
@@ -78,6 +78,12 @@ ULearningAgentsRecord* ULearningAgentsDataStorage::CreateRecord(FName RecordName
 		return nullptr;
 	}
 
+	if (!AgentType->IsSetupPerformed())
+	{
+		UE_LOG(LogLearning, Warning, TEXT("CreateRecord: AgentType Setup must be performed before record can be created."));
+		return nullptr;
+	}
+
 	ULearningAgentsRecord* NewRecord = NewObject<ULearningAgentsRecord>(this, RecordName);
 	NewRecord->MetaData.AgentType = AgentType;
 	NewRecord->MetaData.CreatedOn = FDateTime::Now();
@@ -128,6 +134,12 @@ ULearningAgentsRecord* ULearningAgentsDataStorage::LoadRecord(ULearningAgentsTyp
 		return nullptr;
 	}
 
+	if (!AgentType->IsSetupPerformed())
+	{
+		UE_LOG(LogLearning, Warning, TEXT("LoadRecord: AgentType Setup must be performed before record can be created."));
+		return nullptr;
+	}
+
 	TArray<uint8> Bytes;
 
 	FString InputFilePath = Directory.Path + FGenericPlatformMisc::GetDefaultPathSeparator() + Filename;
diff --git a/Engine/Plugins/Experimental/LearningAgents/Source/LearningAgentsTraining/Private/LearningAgentsImitationTrainer.cpp b/Engine/Plugins/Experimental/LearningAgents/Source/LearningAgentsTraining/Private/LearningAgentsImitationTrainer.cpp
--- a/Engine/Plugins/Experimental/LearningAgents/Source/LearningAgentsTraining/Private/LearningAgentsImitationTrainer.cpp
+++ b/Engine/Plugins/Experimental/LearningAgents/Source/LearningAgentsTraining/Private/LearningAgentsImitationTrainer.cpp
@@ -12,150 +12,189 @@
 #include "LearningNeuralNetwork.h"
 #include "LearningImitationTrainer.h"
 #include "LearningAgentsDataStorage.h"
-#include "LearningPolicyObject.h"
+#include "LearningAgentsPolicy.h"
+#include "LearningNeuralNetworkObject.h"
 #include "Misc/Paths.h"
 
-ULearningAgentsImitationTrainer::ULearningAgentsImitationTrainer() : ULearningAgentsTypeComponent() {}
+ULearningAgentsImitationTrainer::ULearningAgentsImitationTrainer() : UActorComponent() {}
 ULearningAgentsImitationTrainer::ULearningAgentsImitationTrainer(FVTableHelper& Helper) : ULearningAgentsImitationTrainer() {}
 ULearningAgentsImitationTrainer::~ULearningAgentsImitationTrainer() {}
 
 void ULearningAgentsImitationTrainer::EndPlay(const EEndPlayReason::Type EndPlayReason)
 {
-	if (bIsTraining)
+	if (IsTraining())
 	{
 		EndTraining();
 	}
-}
 
-const bool ULearningAgentsImitationTrainer::IsTraining() const
-{
-	return bIsTraining;
+	Super::EndPlay(EndPlayReason);
 }
 
-void ULearningAgentsImitationTrainer::BeginTraining(TArray<UObject*> Records)
+void ULearningAgentsImitationTrainer::BeginTraining(
+	ULearningAgentsPolicy* InPolicy, 
+	const TArray<ULearningAgentsRecord*>& Records, 
+	const FLearningAgentsImitationTrainerTrainingSettings& TrainingSettings,
+	const bool bReinitializePolicyNetwork)
 {
-	if (bIsTraining)
+	if (IsTraining())
 	{
-		UE_LOG(LogLearning, Warning, TEXT("Already Training!"));
+		UE_LOG(LogLearning, Error, TEXT("Already Training!"));
 		return;
 	}
 
-	if (Records.IsEmpty())
+	if (!InPolicy)
 	{
-		UE_LOG(LogLearning, Error, TEXT("Records list is empty! Imitation Trainer needs at least one record to work correctly."));
+		UE_LOG(LogLearning, Error, TEXT("BeginTraining called with nullptr for Policy."));
 		return;
 	}
 
-	UE_LOG(LogLearning, Display, TEXT("Imitation Training Started"));
-
-	bIsTraining = true;
-	bIsTrainingComplete = false;
-	bRequestImitationTrainingStop = false;
+	if (!InPolicy->IsPolicySetupPerformed())
+	{
+		UE_LOG(LogLearning, Error, TEXT("Policy Setup not Performed"));
+		return;
+	}
 
-	// Start Python Training Process (this must be done on main thread)
+	if (Records.IsEmpty())
+	{
+		UE_LOG(LogLearning, Error, TEXT("Records list is empty! Imitation Trainer needs at least one record to work correctly."));
+		return;
+	}
 
-#if WITH_EDITOR
-	const FString PythonExecutablePath = UE::Learning::Trainer::DefaultEditorPythonExecutablePath();
-	const FString SitePackagesPath = UE::Learning::Trainer::DefaultEditorSitePackagesPath();
-	const FString PythonContentPath = UE::Learning::Trainer::DefaultEditorPythonContentPath();
-	const FString IntermediatePath = UE::Learning::Trainer::DefaultEditorIntermediatePath();
-#else
-	UE_LEARNING_CHECKF(PLATFORM_WINDOWS || PLATFORM_MAC || PLATFORM_LINUX, TEXT("Python only supported on Windows, Mac, and Linux."));
-	const FString PythonExe = PLATFORM_WINDOWS ? TEXT("python.exe") : TEXT("bin/python");
-	const FString EnginePath = IFileManager::Get().ConvertToAbsolutePathForExternalAppForRead(*(FPaths::RootDir() / TEXT("../../../../../../Engine")));
-	const FString PythonExecutablePath = EnginePath / TEXT("Binaries/ThirdParty/Python3") / FPlatformMisc::GetUBTPlatform() / PythonExe;
-	const FString SitePackagesPath = EnginePath / TEXT("Plugins/Experimental/PythonFoundationPackages/Content/Python/Lib/") / FPlatformMisc::GetUBTPlatform() / TEXT("site-packages");
-	const FString PythonContentPath = EnginePath / TEXT("Plugins/Experimental/LearningAgents/Content/Python/");
-	const FString IntermediatePath = EnginePath / TEXT("Plugins/Experimental/LearningAgents/Intermediate");
-#endif
+	Policy = InPolicy;
 
-	UE::Learning::FImitationTrainerSettings Settings;
-	Settings.bReinitializeNetwork = true;
-	Settings.Device = UE::Learning::ETrainerDevice::CPU;
-	Settings.bUseTensorboard = false;
+	// Get Number of Steps
 
-	// Get length
 	int32 StepNum = 0;
-	for (UObject* Object : Records)
+	for (const ULearningAgentsRecord* Record : Records)
 	{
-		ULearningAgentsRecord* Record = Cast<ULearningAgentsRecord>(Object);
 		if (!Record)
 		{
-			UE_LOG(LogLearning, Warning, TEXT("BeginTraining: input contains object (%s) that can't be cast to ULearningAgentsRecord."), *(Object ? Object->GetName() : FString(TEXT("nullptr"))));
+			UE_LOG(LogLearning, Warning, TEXT("BeginTraining: Null record object."));
 			continue;
 		}
 
-		TConstArrayView<TLearningArray<2, float>> RecordObsChunks = Record->GetObservations();
-
 		// Assuming everyone has been trimmed already
-		for (TLearningArray<2, float> Chunk : RecordObsChunks)
+		for (const TLearningArray<2, float>& Chunk : Record->GetObservations())
 		{
 			StepNum += Chunk.Num<0>();
 		}
 	}
 
-	if (StepNum == 0)
-	{
-		UE_LOG(LogLearning, Warning, TEXT("BeginTraining: input contains no valid training data. Returning early."));
-		return;
-	}
+	// Parse into Flat Arrays
+
+	const int32 PolicyInputNum = Policy->GetPolicyNetwork().GetInputNum();
+	const int32 PolicyOutputNum = Policy->GetPolicyNetwork().GetOutputNum();
 
-	RecordedObservations.SetNumUninitialized({ StepNum, AgentType->GetObservationFeature().DimNum() });
-	RecordedActions.SetNumUninitialized({ StepNum, AgentType->GetActionFeature().DimNum() });
+	RecordedObservations.SetNumUninitialized({ StepNum, PolicyInputNum });
+	RecordedActions.SetNumUninitialized({ StepNum, PolicyOutputNum / 2 });
 
 	int32 DataIndex = 0;
-	for (UObject* Object : Records)
+	for (const ULearningAgentsRecord* Record : Records)
 	{
-		ULearningAgentsRecord* Record = Cast<ULearningAgentsRecord>(Object);
-		if (!Record)
-		{
-			UE_LOG(LogLearning, Warning, TEXT("BeginTraining: input contains object (%s) that can't be cast to ULearningAgentsRecord."), *(Object ? Object->GetName() : FString(TEXT("nullptr"))));
-			continue;
-		}
+		if (!Record) { continue; }
 
 		TConstArrayView<TLearningArray<2, float>> RecordObsChunks = Record->GetObservations();
 		TConstArrayView<TLearningArray<2, float>> RecordActionChunks = Record->GetActions();
 
-		// Assuming everyone has been trimmed already
-		for (int32 i = 0; i < RecordObsChunks.Num(); i++)
+		for (int32 Idx = 0; Idx < RecordObsChunks.Num(); Idx++)
 		{
-			int32 ChunkLength = RecordObsChunks[i].Num<0>();
-			UE::Learning::Array::Copy(RecordedObservations.Slice(DataIndex, ChunkLength), RecordObsChunks[i]);
-			UE::Learning::Array::Copy(RecordedActions.Slice(DataIndex, ChunkLength), RecordActionChunks[i]);
-			DataIndex += ChunkLength;
+			if (RecordObsChunks[Idx].Num<1>() == PolicyInputNum &&
+				RecordActionChunks[Idx].Num<1>() == PolicyOutputNum / 2)
+			{
+				int32 ChunkLength = RecordObsChunks[Idx].Num<0>();
+				UE::Learning::Array::Copy(RecordedObservations.Slice(DataIndex, ChunkLength), RecordObsChunks[Idx]);
+				UE::Learning::Array::Copy(RecordedActions.Slice(DataIndex, ChunkLength), RecordActionChunks[Idx]);
+				DataIndex += ChunkLength;
+			}
+			else
+			{
+				UE_LOG(LogLearning, Warning, TEXT("Record input or output size does not match policy."));
+			}
 		}
 	}
 
+	// Return if no valid records found
+
+	if (DataIndex == 0)
+	{
+		UE_LOG(LogLearning, Warning, TEXT("BeginTraining: input contains no valid training data."));
+		RecordedObservations.Empty();
+		RecordedActions.Empty();
+		return;
+	}
+
+	// Resize to final size
+
+	RecordedObservations.SetNumUninitialized({ DataIndex, PolicyInputNum });
+	RecordedActions.SetNumUninitialized({ DataIndex, PolicyOutputNum / 2 });
+
+	// Begin Training Properly
+
+	UE_LOG(LogLearning, Display, TEXT("Imitation Training Started"));
+
+	bIsTraining = true;
+	bIsTrainingComplete = false;
+	bRequestImitationTrainingStop = false;
+
+#if WITH_EDITOR
+	const FString PythonExecutablePath = UE::Learning::Trainer::DefaultEditorPythonExecutablePath();
+	const FString SitePackagesPath = UE::Learning::Trainer::DefaultEditorSitePackagesPath();
+	const FString PythonContentPath = UE::Learning::Trainer::DefaultEditorPythonContentPath();
+	const FString IntermediatePath = UE::Learning::Trainer::DefaultEditorIntermediatePath();
+#else
+	UE_LEARNING_CHECKF(PLATFORM_WINDOWS || PLATFORM_MAC || PLATFORM_LINUX, TEXT("Python only supported on Windows, Mac, and Linux."));
+	const FString PythonExe = PLATFORM_WINDOWS ? TEXT("python.exe") : TEXT("bin/python");
+	const FString EnginePath = IFileManager::Get().ConvertToAbsolutePathForExternalAppForRead(*(FPaths::RootDir() / TEXT("../../../../../../Engine")));
+	const FString PythonExecutablePath = EnginePath / TEXT("Binaries/ThirdParty/Python3") / FPlatformMisc::GetUBTPlatform() / PythonExe;
+	const FString SitePackagesPath = EnginePath / TEXT("Plugins/Experimental/PythonFoundationPackages/Content/Python/Lib/") / FPlatformMisc::GetUBTPlatform() / TEXT("site-packages");
+	const FString PythonContentPath = EnginePath / TEXT("Plugins/Experimental/LearningAgents/Content/Python/");
+	const FString IntermediatePath = EnginePath / TEXT("Plugins/Experimental/LearningAgents/Intermediate");
+#endif
+
+	UE::Learning::FImitationTrainerTrainingSettings ImitationTrainingSettings;
+	ImitationTrainingSettings.IterationNum = TrainingSettings.NumberOfIterations;
+	ImitationTrainingSettings.LearningRateActor = TrainingSettings.LearningRate;
+	ImitationTrainingSettings.LearningRateDecay = TrainingSettings.LearningRateDecay;
+	ImitationTrainingSettings.WeightDecay = TrainingSettings.WeightDecay;
+	ImitationTrainingSettings.BatchSize = TrainingSettings.BatchSize;
+	ImitationTrainingSettings.Seed = TrainingSettings.RandomSeed;
+	ImitationTrainingSettings.Device = TrainingSettings.Device == ELearningAgentsTrainerDevice::CPU ? UE::Learning::ETrainerDevice::CPU : UE::Learning::ETrainerDevice::GPU;
+	ImitationTrainingSettings.bUseTensorboard = TrainingSettings.bUseTensorboard;
+
+	const UE::Learning::EImitationTrainerFlags TrainerFlags = 
+		bReinitializePolicyNetwork ? 
+		UE::Learning::EImitationTrainerFlags::None : 
+		UE::Learning::EImitationTrainerFlags::UseInitialPolicyNetwork;
+
 	ImitationTrainer = MakeUnique<UE::Learning::FSharedMemoryImitationTrainer>(
-		TEXT("ImitationLearner"),
+		GetName(),
 		PythonExecutablePath,
 		SitePackagesPath,
 		PythonContentPath,
 		IntermediatePath,
-		AgentType->GetNeuralNetwork(),
-		AgentType->GetPolicy().Settings.ActionNoiseMin,
-		AgentType->GetPolicy().Settings.ActionNoiseMax,
 		RecordedObservations.Num<0>(),
-		Settings);
+		RecordedObservations.Num<1>(),
+		RecordedActions.Num<1>(),
+		ImitationTrainingSettings);
 
-	ImitationTrainingTask = UE::Tasks::Launch(UE_SOURCE_LOCATION, [&]()
-		{
-			UE::Learning::ImitationTrainer::Train(
-				*ImitationTrainer,
-				AgentType->GetNeuralNetwork(),
-				RecordedObservations,
-				RecordedActions,
-				true,
-				&bRequestImitationTrainingStop,
-				&NetworkLock);
+	ImitationTrainingTask = UE::Tasks::Launch(UE_SOURCE_LOCATION, [this, TrainerFlags]()
+	{
+		UE::Learning::ImitationTrainer::Train(
+			*ImitationTrainer,
+			Policy->GetPolicyNetwork(),
+			RecordedObservations,
+			RecordedActions,
+			TrainerFlags,
+			&bRequestImitationTrainingStop,
+			&NetworkLock);
 
 			bIsTrainingComplete = true;
-		});
+	});
 }
 
 void ULearningAgentsImitationTrainer::EndTraining()
 {
-	if (bIsTraining)
+	if (IsTraining())
 	{
 		UE_LOG(LogLearning, Display, TEXT("Imitation Training Ended"));
 
@@ -170,7 +209,12 @@ void ULearningAgentsImitationTrainer::EndTraining()
 	}
 }
 
-const bool ULearningAgentsImitationTrainer::IsTrainingComplete() const
+bool ULearningAgentsImitationTrainer::IsTraining() const
+{
+	return bIsTraining;
+}
+
+bool ULearningAgentsImitationTrainer::IsTrainingComplete() const
 {
 	return bIsTrainingComplete;
 }
diff --git a/Engine/Plugins/Experimental/LearningAgents/Source/LearningAgentsTraining/Private/LearningAgentsRecorder.cpp b/Engine/Plugins/Experimental/LearningAgents/Source/LearningAgentsTraining/Private/LearningAgentsRecorder.cpp
--- a/Engine/Plugins/Experimental/LearningAgents/Source/LearningAgentsTraining/Private/LearningAgentsRecorder.cpp
+++ b/Engine/Plugins/Experimental/LearningAgents/Source/LearningAgentsTraining/Private/LearningAgentsRecorder.cpp
@@ -9,78 +9,203 @@
 #include "LearningLog.h"
 #include "LearningTrainer.h"
 
-ULearningAgentsRecorder::ULearningAgentsRecorder() : ULearningAgentsTypeComponent() {}
+ULearningAgentsRecorder::ULearningAgentsRecorder() : UActorComponent() {}
 ULearningAgentsRecorder::ULearningAgentsRecorder(FVTableHelper& Helper) : ULearningAgentsRecorder() {}
 ULearningAgentsRecorder::~ULearningAgentsRecorder() {}
 
-void ULearningAgentsRecorder::OnRegister()
+void ULearningAgentsRecorder::EndPlay(const EEndPlayReason::Type EndPlayReason)
 {
-	Super::OnRegister();
+	if (IsRecording())
+	{
+		EndRecording();
+	}
+
+	Super::EndPlay(EndPlayReason);
+}
 
-	if (!GetWorld()->IsGameWorld())
+void ULearningAgentsRecorder::SetupRecorder(ULearningAgentsType* InAgentType)
+{
+	if (IsRecorderSetupPerformed())
 	{
-		// We're not in a game yet so we don't need the storage yet
+		UE_LOG(LogLearning, Error, TEXT("Setup already performed!"));
 		return;
 	}
 
-	DataStorage = NewObject<ULearningAgentsDataStorage>(this);
+	if (!InAgentType)
+	{
+		UE_LOG(LogLearning, Error, TEXT("SetupRecorder called but AgentType is nullptr."));
+		return;
+	}
 
-	if (DataDirectory.Path.IsEmpty())
+	if (!InAgentType->IsSetupPerformed())
 	{
+		UE_LOG(LogLearning, Error, TEXT("AgentType Setup must be run before recorder can be setup."));
+		return;
+	}
+
+	AgentType = InAgentType;
+
+	// Create Data Storage
+
+	DataStorage = NewObject<ULearningAgentsDataStorage>(this, TEXT("DataStorage"));
 #if WITH_EDITOR
-		UE_LOG(LogLearning, Display, TEXT("ULearningAgentsRecorder writing data to default editor intermediate directory since DataDirectory was not set."));
-		DataDirectory.Path = UE::Learning::Trainer::DefaultEditorIntermediatePath();
+	DataDirectory.Path = UE::Learning::Trainer::DefaultEditorIntermediatePath() / TEXT("Recordings");
 #else
-		UE_LOG(LogLearning, Error, TEXT("ULearningAgentsRecorder: DataDirectory was not set. This is non-editor build so need a directory setting."));
+	UE_LOG(LogLearning, Error, TEXT("ULearningAgentsRecorder: DataDirectory was not set. This is non-editor build so need a directory setting."));
 #endif
+}
+
+bool ULearningAgentsRecorder::IsRecorderSetupPerformed() const
+{
+	return AgentType ? true : false;
+}
+
+void ULearningAgentsRecorder::AddAgent(int32 AgentId)
+{
+	if (!IsRecorderSetupPerformed())
+	{
+		UE_LOG(LogLearning, Error, TEXT("Recorder setup must be run before agents can be added!"));
+		return;
+	}
+
+	if (!AgentType->GetOccupiedAgentSet().Contains(AgentId))
+	{
+		UE_LOG(LogLearning, Error, TEXT("Unable to add: AgentId %d not found on AgentType. Make sure to add agents to the agent type before adding."), AgentId);
+		return;
+	}
+
+	if (SelectedAgentIds.Contains(AgentId))
+	{
+		UE_LOG(LogLearning, Error, TEXT("AgentId %i is already included in agents set"), AgentId);
+		return;
+	}
+
+	SelectedAgentIds.Add(AgentId);
+	SelectedAgentsSet = SelectedAgentIds;
+	SelectedAgentsSet.TryMakeSlice();
+
+	if (IsRecording() && !CurrentRecords.Contains(AgentId))
+	{
+		CurrentRecords.Add(AgentId, DataStorage->CreateRecord(FName(AgentType->GetName() + "_id" + FString::FromInt(AgentId)), AgentType));
 	}
 }
 
-void ULearningAgentsRecorder::EndPlay(const EEndPlayReason::Type EndPlayReason)
+void ULearningAgentsRecorder::RemoveAgent(int32 AgentId)
 {
-	StopRecording();
+	if (!IsRecorderSetupPerformed())
+	{
+		UE_LOG(LogLearning, Error, TEXT("Recorder setup must be run before agents can be removed!"));
+		return;
+	}
 
-	if (bSaveDataOnEndPlay)
+	if (SelectedAgentIds.RemoveSingleSwap(AgentId, false) == 0)
 	{
-		DataStorage->SaveAllRecords(DataDirectory);
+		UE_LOG(LogLearning, Error, TEXT("Unable to remove: AgentId %d not found in the added agents set."), AgentId);
+		return;
 	}
 
-	Super::EndPlay(EndPlayReason);
+	SelectedAgentsSet = SelectedAgentIds;
+	SelectedAgentsSet.TryMakeSlice();
+
+	if (IsRecording() && CurrentRecords.Contains(AgentId))
+	{
+		CurrentRecords[AgentId]->Trim();
+		CurrentRecords.Remove(AgentId);
+	}
+}
+
+bool ULearningAgentsRecorder::HasAgent(int32 AgentId) const
+{
+	return SelectedAgentsSet.Contains(AgentId);
+}
+
+ULearningAgentsType* ULearningAgentsRecorder::GetAgentType(TSubclassOf<ULearningAgentsType> AgentClass)
+{
+	if (!IsRecorderSetupPerformed())
+	{
+		UE_LOG(LogLearning, Error, TEXT("Recorder setup must be run before getting the agent type!"));
+		return nullptr;
+	}
+
+	return AgentType;
 }
 
 void ULearningAgentsRecorder::AddExperience()
 {
 	if (!IsRecording())
 	{
-		UE_LOG(LogLearning, Warning, TEXT("Trying to add experience but we aren't currently recording. Call StartRecording() before AddExperience()."));
+		UE_LOG(LogLearning, Warning, TEXT("Trying to add experience but we aren't currently recording. Call BeginRecording() before AddExperience()."));
 		return;
 	}
 
+	// Add Experience to Records
+
 	for (const int32 AgentId : SelectedAgentsSet)
 	{
-		CurrentRecords[AgentId]->AddExperience(AgentType->GetObservationFeature().FeatureBuffer()[AgentId], AgentType->GetActionFeature().FeatureBuffer()[AgentId]);
+		CurrentRecords[AgentId]->AddExperience(
+			AgentType->GetObservationFeature().FeatureBuffer()[AgentId], 
+			AgentType->GetActionFeature().FeatureBuffer()[AgentId]);
 	}
 }
 
-const bool ULearningAgentsRecorder::IsRecording() const
+bool ULearningAgentsRecorder::IsRecording() const
 {
-	return !CurrentRecords.IsEmpty();
+	return bIsRecording;
 }
 
-void ULearningAgentsRecorder::StartRecording()
+void ULearningAgentsRecorder::EndRecording()
 {
+	if (!IsRecording())
+	{
+		UE_LOG(LogLearning, Error, TEXT("Not Recording!"));
+		return;
+	}
+
+	// Save Records
+
+	if (bSaveDataOnEndPlay)
+	{
+		DataStorage->SaveAllRecords(DataDirectory);
+	}
+
+	// Trim Records
+
 	for (const int32 AgentId : SelectedAgentsSet)
 	{
-		CurrentRecords.Add(AgentId, DataStorage->CreateRecord(FName(AgentType->GetName() + "_id" + FString::FromInt(AgentId)), AgentType));
+		CurrentRecords[AgentId]->Trim();
 	}
+
+	// Reset Current Records
+
+	CurrentRecords.Empty();
+
+	// Done
+
+	bIsRecording = false;
 }
 
-void ULearningAgentsRecorder::StopRecording()
+void ULearningAgentsRecorder::BeginRecording()
 {
-	for (TPair<int32, ULearningAgentsRecord*> Record : CurrentRecords)
+	if (IsRecorderSetupPerformed())
 	{
-		Record.Value->Trim();
+		UE_LOG(LogLearning, Error, TEXT("Setup must be run before recording can begin."));
+		return;
 	}
 
-	CurrentRecords.Empty();
+	if (IsRecording())
+	{
+		UE_LOG(LogLearning, Error, TEXT("Already Recording!"));
+		return;
+	}
+
+	// Create Records
+
+	for (const int32 AgentId : SelectedAgentsSet)
+	{
+		CurrentRecords.Add(AgentId, DataStorage->CreateRecord(FName(AgentType->GetName() + "_id" + FString::FromInt(AgentId)), AgentType));
+	}
+
+	// Done
+
+	bIsRecording = true;
 }
diff --git a/Engine/Plugins/Experimental/LearningAgents/Source/LearningAgentsTraining/Private/LearningAgentsRewards.cpp b/Engine/Plugins/Experimental/LearningAgents/Source/LearningAgentsTraining/Private/LearningAgentsRewards.cpp
--- a/Engine/Plugins/Experimental/LearningAgents/Source/LearningAgentsTraining/Private/LearningAgentsRewards.cpp
+++ b/Engine/Plugins/Experimental/LearningAgents/Source/LearningAgentsTraining/Private/LearningAgentsRewards.cpp
@@ -23,7 +23,7 @@
 	UE_LEARNING_AGENTS_VLOG_STRING(Owner, Category, Verbosity, Location + FVector(0.0f, 0.0f, 20.0f), Color, Format, ##__VA_ARGS__)
 
 #define UE_LEARNING_AGENTS_VLOG_PLANE(Owner, Category, Verbosity, Location, Rotation, Axis0, Axis1, Color, Format, ...) \
-	UE_VLOG_OBOX(Owner, Category, Verbosity, FBox(25.0f * FVector(-1, -1, 0), 25.0f * FVector(1, 1, 0)), FTransform(Rotation, Location, FVector::OneVector).ToMatrixNoScale(), Color, Format, ##__VA_ARGS__)
+	UE_VLOG_OBOX(Owner, Category, Verbosity, FBox(-25.0f * (Axis0 + Axis1), 25.0f * (Axis0 + Axis1)), FTransform(Rotation, Location, FVector::OneVector).ToMatrixNoScale(), Color, Format, ##__VA_ARGS__)
 
 namespace UE::Learning::Agents::Private
 {
@@ -40,8 +40,8 @@ namespace UE::Learning::Agents::Private
 
 		Reward->RewardObject = MakeShared<RewardFObject>(
 			Reward->GetFName(),
-			AgentTrainer->GetAgentType().GetInstanceData().ToSharedRef(),
-			AgentTrainer->GetAgentType().GetMaxInstanceNum(),
+			AgentTrainer->GetAgentType()->GetInstanceData().ToSharedRef(),
+			AgentTrainer->GetAgentType()->GetMaxInstanceNum(),
 			Forward<InArgTypes>(Args)...);
 
 		AgentTrainer->AddReward(Reward, Reward->RewardObject.ToSharedRef());
@@ -315,8 +315,8 @@ void UPlanarPositionDifferencePenalty::VisualLog(const UE::Learning::FIndexSet I
 				UE_LEARNING_AGENTS_VLOG_PLANE(this, LogLearning, Display,
 					Position0,
 					FQuat::Identity,
-					FeatureObject->Axis0,
-					FeatureObject->Axis1,
+					RewardObject->Axis0,
+					RewardObject->Axis1,
 					VisualLogColor.ToFColor(true),
 					TEXT(""));
 
@@ -331,8 +331,8 @@ void UPlanarPositionDifferencePenalty::VisualLog(const UE::Learning::FIndexSet I
 				UE_LEARNING_AGENTS_VLOG_PLANE(this, LogLearning, Display,
 					Position1,
 					FQuat::Identity,
-					FeatureObject->Axis0,
-					FeatureObject->Axis1,
+					RewardObject->Axis0,
+					RewardObject->Axis1,
 					VisualLogColor.ToFColor(true),
 					TEXT(""));
 
@@ -498,4 +498,4 @@ void UPositionArraySimilarityReward::VisualLog(const UE::Learning::FIndexSet Ins
 
 #undef UE_LEARNING_AGENTS_VLOG_STRING
 #undef UE_LEARNING_AGENTS_VLOG_TRANSFORM
-#undef UE_LEARNING_AGENTS_VLOG_PLANE
+#undef UE_LEARNING_AGENTS_VLOG_PLANE
\ No newline at end of file
diff --git a/Engine/Plugins/Experimental/LearningAgents/Source/LearningAgentsTraining/Private/LearningAgentsTrainer.cpp b/Engine/Plugins/Experimental/LearningAgents/Source/LearningAgentsTraining/Private/LearningAgentsTrainer.cpp
--- a/Engine/Plugins/Experimental/LearningAgents/Source/LearningAgentsTraining/Private/LearningAgentsTrainer.cpp
+++ b/Engine/Plugins/Experimental/LearningAgents/Source/LearningAgentsTraining/Private/LearningAgentsTrainer.cpp
@@ -8,13 +8,15 @@
 #include "LearningAgentsRewards.h"
 #include "LearningAgentsCompletions.h"
 #include "LearningAgentsObservations.h"
+#include "LearningAgentsPolicy.h"
+#include "LearningAgentsCritic.h"
 #include "LearningArray.h"
 #include "LearningArrayMap.h"
 #include "LearningExperience.h"
 #include "LearningFeatureObject.h"
 #include "LearningLog.h"
 #include "LearningNeuralNetwork.h"
-#include "LearningPolicyObject.h"
+#include "LearningNeuralNetworkObject.h"
 #include "LearningPPOTrainer.h"
 #include "LearningRewardObject.h"
 #include "LearningCompletion.h"
@@ -26,45 +28,65 @@
 #include "Engine/GameViewportClient.h"
 #include "EngineDefines.h"
 
-ULearningAgentsTrainer::ULearningAgentsTrainer() : ULearningAgentsTypeComponent() {}
+ULearningAgentsTrainer::ULearningAgentsTrainer() : UActorComponent() {}
 ULearningAgentsTrainer::ULearningAgentsTrainer(FVTableHelper& Helper) : ULearningAgentsTrainer() {}
 ULearningAgentsTrainer::~ULearningAgentsTrainer() {}
 
-void ULearningAgentsTrainer::OnAgentAdded_Implementation(int32 AgentId, UObject* Agent)
+void ULearningAgentsTrainer::SetupTrainer(
+	ULearningAgentsType* InAgentType,
+	ULearningAgentsPolicy* InPolicy,
+	ULearningAgentsCritic* InCritic,
+	const FLearningAgentsTrainerSettings& Settings)
 {
-	Super::OnAgentAdded_Implementation(AgentId, Agent);
+	if (IsTrainerSetupPerformed())
+	{
+		UE_LOG(LogLearning, Error, TEXT("Setup already performed!"));
+		return;
+	}
+
+	// Setup Agent Type, Policy, and Critic
 
-	if (bSetupPerformed)
+	if (!InAgentType)
 	{
-		// Reset the instance and the buffer in case we have stale data from another agent that was using this id previously
-		ResetInstance({AgentId});
-		EpisodeBuffer->Reset(AgentId);
+		UE_LOG(LogLearning, Error, TEXT("SetupTrainer called with nullptr for AgentType."));
+		return;
 	}
-}
 
-void ULearningAgentsTrainer::SetupTrainer(ULearningAgentsType* AgentTypeOverride, const FLearningAgentsTrainerSettings& Settings)
-{
-	if (AgentTypeOverride)
+	if (!InAgentType->IsSetupPerformed())
 	{
-		AgentType = AgentTypeOverride;
+		UE_LOG(LogLearning, Error, TEXT("AgentType Setup not performed."));
+		return;
 	}
-	else if (!AgentType)
+
+	AgentType = InAgentType;
+
+	if (!InPolicy)
 	{
-		UE_LOG(LogLearning, Error, TEXT("LearningAgentsTrainer setup called but AgentType is nullptr."));
+		UE_LOG(LogLearning, Error, TEXT("SetupTrainer called with nullptr for Policy."));
 		return;
 	}
 
-	if (!AgentType->IsSetupPerformed())
+	if (!InPolicy->IsPolicySetupPerformed())
 	{
-		UE_LOG(LogLearning, Error, TEXT("ULearningAgentsTrainer: AgentType not setup. "
-			"Consider using ULearningAgentsTrainer::OnAgentTypeSetupComplete when calling SetupTrainer."));
+		UE_LOG(LogLearning, Error, TEXT("Policy Setup not Performed"));
 		return;
 	}
 
-	// Reset Setup Flag
-	bSetupPerformed = false;
+	Policy = InPolicy;
+
+	if (InCritic)
+	{
+		if (!InCritic->IsCriticSetupPerformed())
+		{
+			UE_LOG(LogLearning, Error, TEXT("Critic Setup not Performed"));
+			return;
+		}
+
+		Critic = InCritic;
+	}
 
 	// Setup Rewards
+
 	RewardObjects.Empty();
 	RewardFeatures.Empty();
 	SetupRewards(this);
@@ -74,6 +96,7 @@ void ULearningAgentsTrainer::SetupTrainer(ULearningAgentsType* AgentTypeOverride
 		AgentType->GetMaxInstanceNum());
 
 	// Setup Completions
+
 	CompletionObjects.Empty();
 	CompletionFeatures.Empty();
 	SetupCompletions(this);
@@ -83,6 +106,7 @@ void ULearningAgentsTrainer::SetupTrainer(ULearningAgentsType* AgentTypeOverride
 		AgentType->GetMaxInstanceNum());
 
 	// Create Episode Buffer
+
 	EpisodeBuffer = MakeUnique<UE::Learning::FEpisodeBuffer>();
 	EpisodeBuffer->Resize(
 		AgentType->GetMaxInstanceNum(),
@@ -93,6 +117,7 @@ void ULearningAgentsTrainer::SetupTrainer(ULearningAgentsType* AgentTypeOverride
 	MaxStepsCompletion = Settings.MaxStepsCompletion;
 
 	// Create Replay Buffer
+
 	ReplayBuffer = MakeUnique<UE::Learning::FReplayBuffer>();
 	ReplayBuffer->Resize(
 		AgentType->GetObservationFeature().DimNum(),
@@ -101,23 +126,120 @@ void ULearningAgentsTrainer::SetupTrainer(ULearningAgentsType* AgentTypeOverride
 		Settings.MaximumRecordedStepsPerIteration);
 
 	// Create Reset Buffer
+
 	ResetBuffer = MakeUnique<UE::Learning::FResetInstanceBuffer>();
 	ResetBuffer->Resize(AgentType->GetMaxInstanceNum());
 
+	// Record Timeout Setting
+
+	TrainerTimeout = Settings.TrainerCommunicationTimeout;
+
 	// Done!
-	bSetupPerformed = true;
+
+	bTrainerSetupPerformed = true;
+}
+
+bool ULearningAgentsTrainer::IsTrainerSetupPerformed() const
+{
+	return bTrainerSetupPerformed;
 }
 
-const bool ULearningAgentsTrainer::IsSetupPerformed() const
+void ULearningAgentsTrainer::AddAgent(int32 AgentId)
 {
-	return bSetupPerformed;
+	if (!IsTrainerSetupPerformed())
+	{
+		UE_LOG(LogLearning, Error, TEXT("Trainer setup must be run before agents can be added!"));
+		return;
+	}
+
+	if (!AgentType->GetOccupiedAgentSet().Contains(AgentId))
+	{
+		UE_LOG(LogLearning, Error, TEXT("Unable to add: AgentId %d not found on AgentType. Make sure to add agents to the agent type before adding."), AgentId);
+		return;
+	}
+
+	if (SelectedAgentIds.Contains(AgentId))
+	{
+		UE_LOG(LogLearning, Error, TEXT("AgentId %i is already included in agents set"), AgentId);
+		return;
+	}
+
+	SelectedAgentIds.Add(AgentId);
+	SelectedAgentsSet = SelectedAgentIds;
+	SelectedAgentsSet.TryMakeSlice();
+
+	if (IsTraining())
+	{
+		// Reset the instance and the buffer in case we have stale 
+		// data from another agent that was using this id previously
+		ResetInstance({ AgentId });
+		EpisodeBuffer->Reset(AgentId);
+	}
+}
+
+void ULearningAgentsTrainer::RemoveAgent(int32 AgentId)
+{
+	if (!IsTrainerSetupPerformed())
+	{
+		UE_LOG(LogLearning, Error, TEXT("Trainer setup must be run before agents can be removed!"));
+		return;
+	}
+
+	if (SelectedAgentIds.RemoveSingleSwap(AgentId, false) == 0)
+	{
+		UE_LOG(LogLearning, Error, TEXT("Unable to remove: AgentId %d not found in the added agents set."), AgentId);
+		return;
+	}
+
+	SelectedAgentsSet = SelectedAgentIds;
+	SelectedAgentsSet.TryMakeSlice();
+}
+
+bool ULearningAgentsTrainer::HasAgent(int32 AgentId) const
+{
+	return SelectedAgentsSet.Contains(AgentId);
+}
+
+ULearningAgentsType* ULearningAgentsTrainer::GetAgentType(TSubclassOf<ULearningAgentsType> AgentClass)
+{
+	if (!IsTrainerSetupPerformed())
+	{
+		UE_LOG(LogLearning, Error, TEXT("Trainer setup must be run before getting the agent type!"));
+		return nullptr;
+	}
+
+	return AgentType;
+}
+
+const UObject* ULearningAgentsTrainer::GetAgent(int32 AgentId) const
+{
+	UE_LEARNING_CHECK(AgentType);
+	return AgentType->GetAgent(AgentId);
+}
+
+UObject* ULearningAgentsTrainer::GetAgent(int32 AgentId)
+{
+	UE_LEARNING_CHECK(AgentType);
+	return AgentType->GetAgent(AgentId);
+}
+
+const ULearningAgentsType* ULearningAgentsTrainer::GetAgentType() const
+{
+	UE_LEARNING_CHECK(AgentType);
+	return AgentType;
+}
+
+ULearningAgentsType* ULearningAgentsTrainer::GetAgentType()
+{
+	UE_LEARNING_CHECK(AgentType);
+	return AgentType;
 }
 
 void ULearningAgentsTrainer::EndPlay(const EEndPlayReason::Type EndPlayReason)
 {
 	if (bIsTraining)
 	{
-		StopTraining();
+		EndTraining();
 	}
 
 	Super::EndPlay(EndPlayReason);
@@ -135,7 +257,7 @@ void ULearningAgentsTrainer::SetRewards_Implementation(const TArray<int32>& Agen
 
 void ULearningAgentsTrainer::AddReward(TObjectPtr<ULearningAgentsReward> Object, const TSharedRef<UE::Learning::FRewardObject>& Reward)
 {
-	check(!bSetupPerformed);
+	UE_LEARNING_CHECK(!IsTrainerSetupPerformed());
 	RewardObjects.Add(Object);
 	RewardFeatures.Add(Reward);
 }
@@ -152,7 +274,7 @@ void ULearningAgentsTrainer::SetCompletions_Implementation(const TArray<int32>&
 
 void ULearningAgentsTrainer::AddCompletion(TObjectPtr<ULearningAgentsCompletion> Object, const TSharedRef<UE::Learning::FCompletionObject>& Completion)
 {
-	check(!bSetupPerformed);
+	UE_LEARNING_CHECK(!IsTrainerSetupPerformed());
 	CompletionObjects.Add(Object);
 	CompletionFeatures.Add(Completion);
 }
@@ -162,11 +284,14 @@ const bool ULearningAgentsTrainer::IsTraining() const
 	return bIsTraining;
 }
 
-void ULearningAgentsTrainer::StartTraining(
+void ULearningAgentsTrainer::BeginTraining(
 	const FLearningAgentsTrainerTrainingSettings& TrainingSettings, 
-	const FLearningAgentsTrainerGameSettings& Settings)
+	const FLearningAgentsTrainerGameSettings& TrainerGameSettings,
+	const FLearningAgentsCriticSettings& CriticSettings,
+	const bool bReinitializePolicyNetwork,
+	const bool bReinitializeCriticNetwork)
 {
-	if (!bSetupPerformed)
+	if (!IsTrainerSetupPerformed())
 	{
 		UE_LOG(LogLearning, Error, TEXT("Setup must be run before training can start."));
 		return;
@@ -196,33 +321,32 @@ void ULearningAgentsTrainer::StartTraining(
 	}
 	// Apply Training GameState Settings
 
-	FApp::SetUseFixedTimeStep(Settings.bUseFixedTimeStep);
+	FApp::SetUseFixedTimeStep(TrainerGameSettings.bUseFixedTimeStep);
 
-	if (Settings.FixedTimeStepFrequency > UE_SMALL_NUMBER)
+	if (TrainerGameSettings.FixedTimeStepFrequency > UE_SMALL_NUMBER)
 	{
-		FApp::SetFixedDeltaTime(1.0f / Settings.FixedTimeStepFrequency);
-		if (Settings.bSetMaxPhysicsStepToFixedTimeStep && PhysicsSettings)
+		FApp::SetFixedDeltaTime(1.0f / TrainerGameSettings.FixedTimeStepFrequency);
+		if (TrainerGameSettings.bSetMaxPhysicsStepToFixedTimeStep && PhysicsSettings)
 		{
-			PhysicsSettings->MaxPhysicsDeltaTime = 1.0f / Settings.FixedTimeStepFrequency;
+			PhysicsSettings->MaxPhysicsDeltaTime = 1.0f / TrainerGameSettings.FixedTimeStepFrequency;
 		}
 	}
 	else
 	{
-		UE_LOG(LogLearning, Warning, TEXT("Provided invalid FixedTimeStepFrequency: %0.5f"), Settings.FixedTimeStepFrequency);
+		UE_LOG(LogLearning, Warning, TEXT("Provided invalid FixedTimeStepFrequency: %0.5f"), TrainerGameSettings.FixedTimeStepFrequency);
 	}
 
-	if (Settings.bDisableVSync && GameSettings)
+	if (TrainerGameSettings.bDisableVSync && GameSettings)
 	{
 		GameSettings->SetVSyncEnabled(false);
 		GameSettings->ApplySettings(false);
 	}
 
-	if (Settings.bUseUnlitViewportRendering && ViewportClient)
+	if (TrainerGameSettings.bUseUnlitViewportRendering && ViewportClient)
 	{
 		ViewportClient->ViewModeIndex = EViewModeIndex::VMI_Unlit;
 	}
 
-
 	// Start Trainer
 
 #if WITH_EDITOR
@@ -243,7 +367,6 @@ void ULearningAgentsTrainer::StartTraining(
 	// run training when these things actually exist somewhere on your machine, which will usually be the case on 
 	// a normal development machine.
 
-	// todo bmulcahy all these paths should be configurable by the dev
 	UE_LEARNING_CHECKF(PLATFORM_WINDOWS || PLATFORM_MAC || PLATFORM_LINUX, TEXT("Python only supported on Windows, Mac, and Linux."));
 	const FString PythonExe = PLATFORM_WINDOWS ? TEXT("python.exe") : TEXT("bin/python");
 	const FString EnginePath = IFileManager::Get().ConvertToAbsolutePathForExternalAppForRead(*(FPaths::RootDir() / TEXT("../../../../../../Engine")));
@@ -253,56 +376,108 @@ void ULearningAgentsTrainer::StartTraining(
 	const FString IntermediatePath = EnginePath / TEXT("Plugins/Experimental/LearningAgents/Intermediate");
 #endif
 
-	UE::Learning::FPPOTrainerSettings PPOSettings;
-	PPOSettings.IterationNum = TrainingSettings.NumberOfIterations;
-	PPOSettings.bReinitializeNetwork = TrainingSettings.bReinitializeNetwork;
-	PPOSettings.bUseTensorboard = TrainingSettings.bUseTensorboard;
-	PPOSettings.InitialActionScale = TrainingSettings.InitialActionScale;
-	PPOSettings.DiscountFactor = TrainingSettings.DiscountFactor;
-	PPOSettings.Seed = TrainingSettings.RandomSeed;
-	PPOSettings.TrimEpisodeStartStepNum = TrainingSettings.NumberOfStepsToTrimAtStartOfEpisode;
-	PPOSettings.TrimEpisodeEndStepNum = TrainingSettings.NumberOfStepsToTrimAtEndOfEpisode;
-	PPOSettings.Device = TrainingSettings.Device == ELearningAgentsTrainerDevice::CPU ? UE::Learning::ETrainerDevice::CPU : UE::Learning::ETrainerDevice::GPU;
-
-	UE::Learning::FNeuralNetwork& NeuralNetwork = AgentType->GetNeuralNetwork();
-	const UE::Learning::FNeuralNetworkPolicyFunction& Policy = AgentType->GetPolicy();
-
-	// Start Python Training Process (this must be done on main thread)
+	UE::Learning::FPPOTrainerTrainingSettings PPOTrainingSettings;
+	PPOTrainingSettings.IterationNum = TrainingSettings.NumberOfIterations;
+	PPOTrainingSettings.bUseTensorboard = TrainingSettings.bUseTensorboard;
+	PPOTrainingSettings.InitialActionScale = TrainingSettings.InitialActionScale;
+	PPOTrainingSettings.DiscountFactor = TrainingSettings.DiscountFactor;
+	PPOTrainingSettings.Seed = TrainingSettings.RandomSeed;
+	PPOTrainingSettings.TrimEpisodeStartStepNum = TrainingSettings.NumberOfStepsToTrimAtStartOfEpisode;
+	PPOTrainingSettings.TrimEpisodeEndStepNum = TrainingSettings.NumberOfStepsToTrimAtEndOfEpisode;
+	PPOTrainingSettings.Device = TrainingSettings.Device == ELearningAgentsTrainerDevice::CPU ? UE::Learning::ETrainerDevice::CPU : UE::Learning::ETrainerDevice::GPU;
+
+	UE::Learning::FPPOTrainerNetworkSettings PPONetworkSettings;
+	PPONetworkSettings.PolicyActionNoiseMin = Policy->GetPolicyObject().Settings.ActionNoiseMin;
+	PPONetworkSettings.PolicyActionNoiseMax = Policy->GetPolicyObject().Settings.ActionNoiseMax;
+	PPONetworkSettings.PolicyActivationFunction = Policy->GetPolicyNetwork().ActivationFunction;
+	PPONetworkSettings.PolicyHiddenLayerSize = Policy->GetPolicyNetwork().GetHiddenNum();
+	PPONetworkSettings.PolicyLayerNum = Policy->GetPolicyNetwork().GetLayerNum();
+
+	if (Critic)
+	{
+		if (CriticSettings.HiddenLayerSize != Critic->GetCriticNetwork().GetHiddenNum() ||
+			CriticSettings.LayerNum != Critic->GetCriticNetwork().GetLayerNum() ||
+			UE::Learning::Agents::GetActivationFunction(CriticSettings.ActivationFunction) != Critic->GetCriticNetwork().ActivationFunction)
+		{
+			UE_LOG(LogLearning, Warning, TEXT("StartTraining got different Critic Network Settings to those provided to SetupCritic."));
+		}
+
+		PPONetworkSettings.CriticHiddenLayerSize = Critic->GetCriticNetwork().GetHiddenNum();
+		PPONetworkSettings.CriticLayerNum = Critic->GetCriticNetwork().GetLayerNum();
+		PPONetworkSettings.CriticActivationFunction = Critic->GetCriticNetwork().ActivationFunction;
+	}
+	else
+	{
+		PPONetworkSettings.CriticHiddenLayerSize = CriticSettings.HiddenLayerSize;
+		PPONetworkSettings.CriticLayerNum = CriticSettings.LayerNum;
+		PPONetworkSettings.CriticActivationFunction = UE::Learning::Agents::GetActivationFunction(CriticSettings.ActivationFunction);
+	}
+
+	// We assume that if the critic has been setup on the agent type, then
+	// the user wants the critic network to be synced during training.
+	UE::Learning::EPPOTrainerFlags TrainerFlags = 
+		Critic ?
+		UE::Learning::EPPOTrainerFlags::SynchronizeCriticNetwork :
+		UE::Learning::EPPOTrainerFlags::None;
+
+	if (!bReinitializePolicyNetwork) { TrainerFlags |= UE::Learning::EPPOTrainerFlags::UseInitialPolicyNetwork; }
+	if (!bReinitializeCriticNetwork && Critic) { TrainerFlags |= UE::Learning::EPPOTrainerFlags::UseInitialCriticNetwork; }
+
+	// Start Python Training Process (this must be done on game thread)
 	Trainer = MakeUnique<UE::Learning::FSharedMemoryPPOTrainer>(
-		TEXT("Training"),
+		GetName(),
 		PythonExecutablePath,
 		SitePackagesPath,
 		PythonContentPath,
 		IntermediatePath,
-		NeuralNetwork,
-		Policy.Settings.ActionNoiseMin,
-		Policy.Settings.ActionNoiseMax,
 		*ReplayBuffer,
-		PPOSettings);
+		PPOTrainingSettings,
+		PPONetworkSettings,
+		TrainerFlags);
 
-	TrainerTimeout = TrainingSettings.TrainerCommunicationTimeout;
+	UE_LOG(LogLearning, Display, TEXT("Receiving initial policy..."));
 
-	UE::Learning::ETrainerResponse Response;
+	UE::Learning::ETrainerResponse Response = UE::Learning::ETrainerResponse::Success;
 
-	if (TrainingSettings.bReinitializeNetwork)
+	if ((bool)(TrainerFlags & UE::Learning::EPPOTrainerFlags::UseInitialPolicyNetwork))
 	{
-		UE_LOG(LogLearning, Display, TEXT("Receiving initial policy..."));
-		Response = Trainer->RecvPolicy(NeuralNetwork, TrainerTimeout);
+		Response = Trainer->SendPolicy(Policy->GetPolicyNetwork(), TrainerTimeout);
 	}
 	else
 	{
-		UE_LOG(LogLearning, Display, TEXT("Sending initial policy..."));
-		Response = Trainer->SendPolicy(NeuralNetwork, TrainerTimeout);
+		Response = Trainer->RecvPolicy(Policy->GetPolicyNetwork(), TrainerTimeout);
 	}
 
 	if (Response != UE::Learning::ETrainerResponse::Success)
 	{
-		UE_LOG(LogLearning, Error, TEXT("Error sending or receiving policy from trainer. Check log for errors."));
+		UE_LOG(LogLearning, Error, TEXT("Error sending or receiving policy from trainer: %s. Check log for errors."), UE::Learning::Trainer::GetResponseString(Response));
 		Trainer->Terminate();
-		UE_DEBUG_BREAK();
 		return;
 	}
 
+	if (Critic)
+	{
+		if ((bool)(TrainerFlags & UE::Learning::EPPOTrainerFlags::UseInitialCriticNetwork))
+		{
+			Response = Trainer->SendCritic(Critic->GetCriticNetwork(), TrainerTimeout);
+		}
+		else if ((bool)(TrainerFlags & UE::Learning::EPPOTrainerFlags::SynchronizeCriticNetwork))
+		{
+			Response = Trainer->RecvCritic(Critic->GetCriticNetwork(), TrainerTimeout);
+		}
+		else
+		{
+			Response = UE::Learning::ETrainerResponse::Success;
+		}
+
+		if (Response != UE::Learning::ETrainerResponse::Success)
+		{
+			UE_LOG(LogLearning, Error, TEXT("Error sending or receiving critic from trainer: %s. Check log for errors."), UE::Learning::Trainer::GetResponseString(Response));
+			Trainer->Terminate();
+			return;
+		}
+	}
+
 	// Reset Agents, Episode Buffer, and Replay Buffer
 	ResetInstance(SelectedAgentIds);
 
@@ -315,7 +490,7 @@ void ULearningAgentsTrainer::StartTraining(
 
 void ULearningAgentsTrainer::DoneTraining()
 {
-	if (bIsTraining)
+	if (IsTraining())
 	{
 		// Wait for Trainer to finish
 		Trainer->Wait(1.0f);
@@ -350,9 +525,9 @@ void ULearningAgentsTrainer::DoneTraining()
 	}
 }
 
-void ULearningAgentsTrainer::StopTraining()
+void ULearningAgentsTrainer::EndTraining()
 {
-	if (bIsTraining)
+	if (IsTraining())
 	{
 		UE_LOG(LogLearning, Display, TEXT("Stopping training..."));
 		Trainer->SendStop();
@@ -364,7 +539,7 @@ void ULearningAgentsTrainer::EvaluateRewards()
 {
 	UE_LEARNING_TRACE_CPUPROFILER_EVENT_SCOPE(ULearningAgentsTrainer::EvaluateRewards);
 
-	if (!bSetupPerformed)
+	if (!IsTrainerSetupPerformed())
 	{
 		UE_LOG(LogLearning, Error, TEXT("Setup must be run before rewards can be evaluated."));
 		return;
@@ -389,7 +564,7 @@ void ULearningAgentsTrainer::EvaluateCompletions()
 {
 	UE_LEARNING_TRACE_CPUPROFILER_EVENT_SCOPE(ULearningAgentsTrainer::EvaluateCompletions);
 
-	if (!bSetupPerformed)
+	if (!IsTrainerSetupPerformed())
 	{
 		UE_LOG(LogLearning, Error, TEXT("Setup must be run before completions can be evaluated."));
 		return;
@@ -412,21 +587,15 @@ void ULearningAgentsTrainer::EvaluateCompletions()
 
 void ULearningAgentsTrainer::IterateTraining()
 {
-	if (!bSetupPerformed)
+	if (!IsTrainerSetupPerformed())
 	{
 		UE_LOG(LogLearning, Error, TEXT("Setup must be run before training can be performed."));
 		return;
 	}
 
-	if (!bIsTraining)
+	if (!IsTraining())
 	{
-		if (!bPreviouslyLogged)
-		{
-			UE_LOG(LogLearning, Error, TEXT("Attempted to iterate training but the trainer is not running. Ensure StartTraining() is called before IterateTraining() and no errors occurred."));
-			UE_DEBUG_BREAK();
-			bPreviouslyLogged = true;
-		}
-
+		UE_LOG(LogLearning, Error, TEXT("Attempted to iterate training but the trainer is not running. Ensure StartTraining() is called before IterateTraining() and no errors occurred."));
 		return;
 	}
 
@@ -493,8 +662,7 @@ void ULearningAgentsTrainer::IterateTraining()
 			if (Response != UE::Learning::ETrainerResponse::Success)
 			{
 				UE_LOG(LogLearning, Error, TEXT("Error waiting to push experience to trainer. Check log for errors."));
-				StopTraining();
-				UE_DEBUG_BREAK();
+				EndTraining();
 				return;
 			}
 
@@ -504,7 +672,7 @@ void ULearningAgentsTrainer::IterateTraining()
 
 			// Get Updated Policy
 
-			Response = Trainer->RecvPolicy(AgentType->GetNeuralNetwork(), TrainerTimeout);
+			Response = Trainer->RecvPolicy(Policy->GetPolicyNetwork(), TrainerTimeout);
 
 			if (Response == UE::Learning::ETrainerResponse::Completed)
 			{
@@ -515,11 +683,24 @@ void ULearningAgentsTrainer::IterateTraining()
 			else if (Response != UE::Learning::ETrainerResponse::Success)
 			{
 				UE_LOG(LogLearning, Error, TEXT("Error waiting for policy from trainer. Check log for errors."));
-				StopTraining();
-				UE_DEBUG_BREAK();
+				EndTraining();
 				return;
 			}
 
+			// Get Updated Critic
+
+			if (Critic)
+			{
+				Response = Trainer->RecvCritic(Critic->GetCriticNetwork(), TrainerTimeout);
+
+				if (Response != UE::Learning::ETrainerResponse::Success)
+				{
+					UE_LOG(LogLearning, Error, TEXT("Error waiting for critic from trainer. Check log for errors."));
+					EndTraining();
+					return;
+				}
+			}
+
 			// Mark all instances for reset since we have a new policy
 			ResetBuffer->SetResetInstances(SelectedAgentsSet);
 		}
@@ -535,7 +716,7 @@ void ULearningAgentsTrainer::ResetAllInstances()
 {
 	UE_LEARNING_TRACE_CPUPROFILER_EVENT_SCOPE(ULearningAgentsTrainer::ResetAllInstances);
 
-	if (!bSetupPerformed)
+	if (!IsTrainerSetupPerformed())
 	{
 		UE_LOG(LogLearning, Error, TEXT("Setup must be run before agents can be reset."));
 		return;
diff --git a/Engine/Plugins/Experimental/LearningAgents/Source/LearningAgentsTraining/Public/LearningAgentsImitationTrainer.h b/Engine/Plugins/Experimental/LearningAgents/Source/LearningAgentsTraining/Public/LearningAgentsImitationTrainer.h
--- a/Engine/Plugins/Experimental/LearningAgents/Source/LearningAgentsTraining/Public/LearningAgentsImitationTrainer.h
+++ b/Engine/Plugins/Experimental/LearningAgents/Source/LearningAgentsTraining/Public/LearningAgentsImitationTrainer.h
@@ -2,32 +2,65 @@
 
 #pragma once
 
-#include "Engine/EngineTypes.h"
-#include "LearningAgentsTypeComponent.h"
+#include "LearningAgentsTrainer.h" // Included for ELearningAgentsTrainerDevice
+#include "LearningArray.h"
+
 #include "Templates/SharedPointer.h"
-#include "UObject/Object.h"
 #include "UObject/ObjectPtr.h"
-#include "UObject/NameTypes.h"
 #include "Tasks/Task.h"
-#include "LearningArray.h"
 
 #include "LearningAgentsImitationTrainer.generated.h"
 
 namespace UE::Learning
 {
-	struct FArrayMap;
-	struct FNeuralNetwork;
-	struct FNeuralNetworkPolicyFunction;
 	struct FSharedMemoryImitationTrainer;
 }
 
 class ULearningAgentsType;
+class ULearningAgentsPolicy;
+
+//------------------------------------------------------------------
+
+USTRUCT(BlueprintType, Category = "LearningAgents")
+struct FLearningAgentsImitationTrainerTrainingSettings
+{
+	GENERATED_BODY()
+
+public:
+
+	UPROPERTY(EditAnywhere, Category = "LearningAgents", meta = (ClampMin = "0", UIMin = "0"))
+	int32 NumberOfIterations = 1000000;
+
+	UPROPERTY(EditAnywhere, Category = "LearningAgents", meta = (ClampMin = "0.0", UIMin = "0.0"))
+	float LearningRate = 0.0001f;
+
+	UPROPERTY(EditAnywhere, Category = "LearningAgents", meta = (ClampMin = "0.0", UIMin = "0.0"))
+	float LearningRateDecay = 0.99f;
+
+	UPROPERTY(EditAnywhere, Category = "LearningAgents", meta = (ClampMin = "0.0", UIMin = "0.0"))
+	float WeightDecay = 0.001f;
+
+	UPROPERTY(EditAnywhere, Category = "LearningAgents", meta = (ClampMin = "1", UIMin = "1"))
+	uint32 BatchSize = 128;
 
-UCLASS(Abstract, BlueprintType, Blueprintable)
-class LEARNINGAGENTSTRAINING_API ULearningAgentsImitationTrainer : public ULearningAgentsTypeComponent
+	UPROPERTY(EditAnywhere, Category = "LearningAgents", meta = (ClampMin = "0", UIMin = "0"))
+	int32 RandomSeed = 1234;
+
+	UPROPERTY(EditAnywhere, Category = "LearningAgents")
+	ELearningAgentsTrainerDevice Device = ELearningAgentsTrainerDevice::CPU;
+
+	UPROPERTY(EditAnywhere, Category = "LearningAgents")
+	bool bUseTensorboard = false;
+};
+
+//------------------------------------------------------------------
+
+UCLASS(BlueprintType, Blueprintable)
+class LEARNINGAGENTSTRAINING_API ULearningAgentsImitationTrainer : public UActorComponent
 {
 	GENERATED_BODY()
 
+// ----- Setup -----
 public:
 
 	// These constructors/destructors are needed to make forward declarations happy
@@ -37,23 +70,28 @@ public:
 
 	virtual void EndPlay(const EEndPlayReason::Type EndPlayReason) override;
 
-public:
-
-	UFUNCTION(BlueprintPure, Category = "LearningAgents")
-	const bool IsTraining() const;
-
 	UFUNCTION(BlueprintCallable, Category = "LearningAgents")
-	void BeginTraining(TArray<UObject*> Records);
+	void BeginTraining(
+		ULearningAgentsPolicy* InPolicy,
+		const TArray<ULearningAgentsRecord*>& Records,
+		const FLearningAgentsImitationTrainerTrainingSettings& TrainingSettings = FLearningAgentsImitationTrainerTrainingSettings(),
+		const bool bReinitializePolicyNetwork = true);
 
 	UFUNCTION(BlueprintCallable, Category = "LearningAgents")
 	void EndTraining();
 
 	UFUNCTION(BlueprintPure, Category = "LearningAgents")
-	const bool IsTrainingComplete() const;
+	bool IsTraining() const;
+
+	UFUNCTION(BlueprintPure, Category = "LearningAgents")
+	bool IsTrainingComplete() const;
 
-//** ----- Private ----- */
+// ----- Private Data -----
 private:
 
+	UPROPERTY(VisibleAnywhere, Category = "LearningAgents")
+	TObjectPtr<ULearningAgentsPolicy> Policy;
+
 	/** True if training is currently in-progress. Otherwise, false. */
 	UPROPERTY(VisibleAnywhere, Category = "LearningAgents")
 	bool bIsTraining = false;
@@ -69,6 +107,4 @@ private:
 	FRWLock NetworkLock;
 
 	TAtomic<bool> bRequestImitationTrainingStop = false;
-
-	float TrainerTimeout = 10.0f;
 };
