LearningAgents: Large refactoring to allow critic to be loaded and used at runtime. Fixed some other misc bugs and fallout from switching features to all be arrays by default. Some refactoring of diff

diff --git a/Engine/Plugins/Experimental/LearningAgents/Source/LearningAgentsTraining/Public/LearningAgentsRecorder.h b/Engine/Plugins/Experimental/LearningAgents/Source/LearningAgentsTraining/Public/LearningAgentsRecorder.h
--- a/Engine/Plugins/Experimental/LearningAgents/Source/LearningAgentsTraining/Public/LearningAgentsRecorder.h
+++ b/Engine/Plugins/Experimental/LearningAgents/Source/LearningAgentsTraining/Public/LearningAgentsRecorder.h
@@ -2,35 +2,54 @@
 
 #pragma once
 
-#include "LearningAgentsTypeComponent.h"
+#include "LearningArray.h"
 
+#include "Components/ActorComponent.h"
 #include "Containers/Map.h"
-#include "Engine/EngineTypes.h"
-#include "UObject/Object.h"
 #include "UObject/ObjectPtr.h"
 
 #include "LearningAgentsRecorder.generated.h"
 
-class ULearningAgentsType;
 class ULearningAgentsDataStorage;
 class ULearningAgentsRecord;
 
-UCLASS(Abstract, BlueprintType, Blueprintable)
-class LEARNINGAGENTSTRAINING_API ULearningAgentsRecorder : public ULearningAgentsTypeComponent
+UCLASS(BlueprintType, Blueprintable)
+class LEARNINGAGENTSTRAINING_API ULearningAgentsRecorder : public UActorComponent
 {
 	GENERATED_BODY()
 
+// ----- Setup -----
 public:
+
 	// These constructors/destructors are needed to make forward declarations happy
 	ULearningAgentsRecorder();
 	ULearningAgentsRecorder(FVTableHelper& Helper);
 	virtual ~ULearningAgentsRecorder();
 
-	virtual void OnRegister() override;
-
 	virtual void EndPlay(const EEndPlayReason::Type EndPlayReason) override;
 
-//** ----- Recording Process ----- */
+	UFUNCTION(BlueprintCallable, Category = "LearningAgents")
+	void SetupRecorder(ULearningAgentsType* InAgentType);
+
+	UFUNCTION(BlueprintPure, Category = "LearningAgents")
+	bool IsRecorderSetupPerformed() const;
+
+// ----- Agent Management -----
+public:
+
+	UFUNCTION(BlueprintCallable, Category = "LearningAgents")
+	void AddAgent(int32 AgentId);
+
+	UFUNCTION(BlueprintCallable, Category = "LearningAgents")
+	void RemoveAgent(int32 AgentId);
+
+	UFUNCTION(BlueprintPure, Category = "LearningAgents")
+	bool HasAgent(int32 AgentId) const;
+
+	UFUNCTION(BlueprintPure, Category = "LearningAgents", meta = (DeterminesOutputType = "AgentClass"))
+	ULearningAgentsType* GetAgentType(TSubclassOf<ULearningAgentsType> AgentClass);
+
+// ----- Recording Process -----
 public:
 
 	/** Adds experience to tracked agents' records. Each agent id whose record is full will be returned, e.g. so you can reset them. */
@@ -38,15 +57,24 @@ public:
 	void AddExperience();
 
 	UFUNCTION(BlueprintCallable, Category = "LearningAgents")
-	const bool IsRecording() const;
+	void BeginRecording();
 
 	UFUNCTION(BlueprintCallable, Category = "LearningAgents")
-	void StartRecording();
+	void EndRecording();
 
 	UFUNCTION(BlueprintCallable, Category = "LearningAgents")
-	void StopRecording();
+	bool IsRecording() const;
+
+// ----- Private Data ----- 
+private:
+
+	UPROPERTY(VisibleAnywhere, Category = "LearningAgents")
+	TObjectPtr<ULearningAgentsType> AgentType;
+
+	UPROPERTY(VisibleAnywhere, Category = "LearningAgents")
+	TArray<int32> SelectedAgentIds;
 
-//** ----- Recorder Configuration ----- */
+// ----- Recorder Configuration -----
 private:
 
 	/** Directory where records will be saved. If not set, BeginPlay will automatically set this to the editor's default intermediate folder. */
@@ -57,9 +85,12 @@ private:
 	UPROPERTY(EditAnywhere, Category = "LearningAgents")
 	bool bSaveDataOnEndPlay = true;
 
-//** ----- Recorder State ----- */
+// ----- Recorder State -----
 private:
 
+	UPROPERTY(VisibleAnywhere, Category = "LearningAgents")
+	bool bIsRecording = false;
+
 	/** The data storage manager. It can be used to save/load agent records. */
 	UPROPERTY(VisibleAnywhere, Category = "LearningAgents")
 	TObjectPtr<ULearningAgentsDataStorage> DataStorage;
@@ -67,4 +98,8 @@ private:
 	/** All records which are currently being written */
 	UPROPERTY(VisibleAnywhere, Category = "LearningAgents")
 	TMap<int32, TObjectPtr<ULearningAgentsRecord>> CurrentRecords;
+
+private:
+
+	UE::Learning::FIndexSet SelectedAgentsSet;
 };
diff --git a/Engine/Plugins/Experimental/LearningAgents/Source/LearningAgentsTraining/Public/LearningAgentsTrainer.h b/Engine/Plugins/Experimental/LearningAgents/Source/LearningAgentsTraining/Public/LearningAgentsTrainer.h
--- a/Engine/Plugins/Experimental/LearningAgents/Source/LearningAgentsTraining/Public/LearningAgentsTrainer.h
+++ b/Engine/Plugins/Experimental/LearningAgents/Source/LearningAgentsTraining/Public/LearningAgentsTrainer.h
@@ -2,12 +2,12 @@
 
 #pragma once
 
-#include "Engine/EngineTypes.h"
-#include "LearningAgentsTypeComponent.h"
+#include "Components/ActorComponent.h"
+#include "LearningAgentsCritic.h" // Included for FLearningAgentsCriticSettings()
+
 #include "Templates/SharedPointer.h"
-#include "UObject/Object.h"
 #include "UObject/ObjectPtr.h"
-#include "UObject/NameTypes.h"
+#include "Engine/EngineTypes.h"
 
 #include "LearningAgentsTrainer.generated.h"
 
@@ -16,8 +16,6 @@ namespace UE::Learning
 	struct FAnyCompletion;
 	struct FArrayMap;
 	struct FEpisodeBuffer;
-	struct FNeuralNetwork;
-	struct FNeuralNetworkPolicyFunction;
 	struct FReplayBuffer;
 	struct FResetInstanceBuffer;
 	struct FRewardObject;
@@ -26,20 +24,26 @@ namespace UE::Learning
 	struct FCompletionObject;
 }
 
+class ULearningAgentsType;
 class ULearningAgentsCompletion;
 class ULearningAgentsReward;
-class ULearningAgentsType;
+class ULearningAgentsPolicy;
+class ULearningAgentsCritic;
+
+//------------------------------------------------------------------
 
 UENUM(BlueprintType, Category = "LearningAgents")
 enum class ELearningAgentsCompletion : uint8
 {
-	/** Episode ended early but was still in progress. Value function will be used to estimate final return. */
+	/** Episode ended early but was still in progress. Critic will be used to estimate final return. */
 	Truncation	UMETA(DisplayName = "Truncation"),
 
 	/** Episode ended early and zero reward was expected for all future steps. */
 	Termination	UMETA(DisplayName = "Termination"),
 };
 
+//------------------------------------------------------------------
+
 USTRUCT(BlueprintType, Category = "LearningAgents")
 struct FLearningAgentsTrainerSettings
 {
@@ -62,8 +66,14 @@ public:
 	/** Maximum number of steps to record before running a training iteration */
 	UPROPERTY(EditAnywhere, Category = "LearningAgents", meta = (ClampMin = "1", UIMin = "1"))
 	int32 MaximumRecordedStepsPerIteration = 10000;
+
+	/** How long to wait for the training subprocess before timing out */
+	UPROPERTY(EditAnywhere, Category = "LearningAgents", meta = (ClampMin = "0.0", UIMin = "0.0"))
+	float TrainerCommunicationTimeout = 20.0f;
 };
 
+//------------------------------------------------------------------
+
 USTRUCT(BlueprintType, Category="LearningAgents")
 struct FLearningAgentsTrainerGameSettings
 {
@@ -87,6 +97,8 @@ public:
 	bool bUseUnlitViewportRendering = false;
 };
 
+//------------------------------------------------------------------
+
 UENUM(BlueprintType, Category = "LearningAgents")
 enum class ELearningAgentsTrainerDevice : uint8
 {
@@ -94,6 +106,8 @@ enum class ELearningAgentsTrainerDevice : uint8
 	GPU,
 };
 
+//------------------------------------------------------------------
+
 USTRUCT(BlueprintType, Category = "LearningAgents")
 struct FLearningAgentsTrainerTrainingSettings
 {
@@ -104,9 +118,6 @@ public:
 	UPROPERTY(EditAnywhere, Category = "LearningAgents", meta = (ClampMin = "0", UIMin = "0"))
 	int32 NumberOfIterations = 1000000;
 
-	UPROPERTY(EditAnywhere, Category = "LearningAgents")
-	bool bReinitializeNetwork = true;
-
 	UPROPERTY(EditAnywhere, Category = "LearningAgents")
 	bool bUseTensorboard = false;
 
@@ -127,17 +138,16 @@ public:
 
 	UPROPERTY(EditAnywhere, Category = "LearningAgents")
 	ELearningAgentsTrainerDevice Device = ELearningAgentsTrainerDevice::GPU;
-
-	UPROPERTY(EditAnywhere, Category = "LearningAgents", meta = (ClampMin = "0.0", UIMin = "0.0"))
-	float TrainerCommunicationTimeout = 20.0f;
 };
 
+//------------------------------------------------------------------
 
 UCLASS(Abstract, BlueprintType, Blueprintable)
-class LEARNINGAGENTSTRAINING_API ULearningAgentsTrainer : public ULearningAgentsTypeComponent
+class LEARNINGAGENTSTRAINING_API ULearningAgentsTrainer : public UActorComponent
 {
 	GENERATED_BODY()
 
+// ----- Setup -----
 public:
 
 	// These constructors/destructors are needed to make forward declarations happy
@@ -145,31 +155,42 @@ public:
 	ULearningAgentsTrainer(FVTableHelper& Helper);
 	virtual ~ULearningAgentsTrainer();
 
-	virtual void OnAgentAdded_Implementation(int32 AgentId, UObject* Agent) override;
-
 	virtual void EndPlay(const EEndPlayReason::Type EndPlayReason) override;
 
-// ----- Setup -----
-public:
-
-	/**
-	* Initializes this object and runs the setup functions for observations and actions.
-	* If the AgentType has not yet been setup, this call will fail, so consider using the
-	* OnAgentTypeSetupComplete event to ensure the timing is correct.
-	* 
-	* @param AgentTypeOverride This agent type will override any parent agent type we found. Must be provided if this trainer component is not attached to a parent agent type component.
-	* @param Settings The settings used to setup this trainer
-	* @see OnAgentTypeSetupComplete
-	*/
 	UFUNCTION(BlueprintCallable, Category = "LearningAgents")
 	void SetupTrainer(
-		ULearningAgentsType* AgentTypeOverride,
+		ULearningAgentsType* InAgentType,
+		ULearningAgentsPolicy* InPolicy,
+		ULearningAgentsCritic* InCritic = nullptr,
 		const FLearningAgentsTrainerSettings& Settings = FLearningAgentsTrainerSettings());
 
 	UFUNCTION(BlueprintPure, Category = "LearningAgents")
-	const bool IsSetupPerformed() const;
+	bool IsTrainerSetupPerformed() const;
+
+// ----- Agent Management -----
+public:
+
+	UFUNCTION(BlueprintCallable, Category = "LearningAgents")
+	void AddAgent(int32 AgentId);
+
+	UFUNCTION(BlueprintCallable, Category = "LearningAgents")
+	void RemoveAgent(int32 AgentId);
+
+	UFUNCTION(BlueprintPure, Category = "LearningAgents")
+	bool HasAgent(int32 AgentId) const;
 
-//** ----- Rewards ----- */
+	UFUNCTION(BlueprintPure, Category = "LearningAgents", meta = (DeterminesOutputType = "AgentClass"))
+	ULearningAgentsType* GetAgentType(TSubclassOf<ULearningAgentsType> AgentClass);
+
+public:
+
+	const UObject* GetAgent(int32 AgentId) const;
+	UObject* GetAgent(int32 AgentId);
+
+	const ULearningAgentsType* GetAgentType() const;
+	ULearningAgentsType* GetAgentType();
+
+// ----- Rewards -----
 public:
 	
 	/**
@@ -185,7 +206,7 @@ public:
 
 	void AddReward(TObjectPtr<ULearningAgentsReward> Object, const TSharedRef<UE::Learning::FRewardObject>& Reward);
 
-//** ----- Completions ----- */
+// ----- Completions ----- 
 public:
 
 	/**
@@ -201,7 +222,7 @@ public:
 
 	void AddCompletion(TObjectPtr<ULearningAgentsCompletion> Object, const TSharedRef<UE::Learning::FCompletionObject>& Completion);
 
-//** ----- Resets ----- */
+// ----- Resets ----- 
 public:
 
 	/**
@@ -214,19 +235,22 @@ public:
 	UFUNCTION(BlueprintNativeEvent, Category = "LearningAgents")
 	void ResetInstance(const TArray<int32>& AgentIds);
 
-//** ----- Training Process ----- */
+// ----- Training Process -----
 public:
 
 	UFUNCTION(BlueprintPure, Category = "LearningAgents")
 	const bool IsTraining() const;
 
 	UFUNCTION(BlueprintCallable, Category = "LearningAgents")
-	void StartTraining(
+	void BeginTraining(
 		const FLearningAgentsTrainerTrainingSettings& TrainingSettings = FLearningAgentsTrainerTrainingSettings(),
-		const FLearningAgentsTrainerGameSettings& GameSettings = FLearningAgentsTrainerGameSettings());
+		const FLearningAgentsTrainerGameSettings& GameSettings = FLearningAgentsTrainerGameSettings(),
+		const FLearningAgentsCriticSettings& CriticSettings = FLearningAgentsCriticSettings(),
+		const bool bReinitializePolicyNetwork = true,
+		const bool bReinitializeCriticNetwork = true);
 
 	UFUNCTION(BlueprintCallable, Category = "LearningAgents")
-	void StopTraining();
+	void EndTraining();
 
 	UFUNCTION(BlueprintCallable, Category = "LearningAgents")
 	void EvaluateRewards();
@@ -243,11 +267,23 @@ public:
 	UFUNCTION(BlueprintCallable, Category = "LearningAgents")
 	void ResetAllInstances();
 
-//** ----- Private ----- */
+// ----- Private Data ----- 
 private:
 
 	UPROPERTY(VisibleAnywhere, Category = "LearningAgents")
-	bool bSetupPerformed = false;
+	TObjectPtr<ULearningAgentsType> AgentType;
+
+	UPROPERTY(VisibleAnywhere, Category = "LearningAgents")
+	TArray<int32> SelectedAgentIds;
+
+	UPROPERTY(VisibleAnywhere, Category = "LearningAgents")
+	TObjectPtr<ULearningAgentsPolicy> Policy;
+
+	UPROPERTY(VisibleAnywhere, Category = "LearningAgents")
+	TObjectPtr<ULearningAgentsCritic> Critic;
+
+	UPROPERTY(VisibleAnywhere, Category = "LearningAgents")
+	bool bTrainerSetupPerformed = false;
 
 	/** True if training is currently in-progress. Otherwise, false. */
 	UPROPERTY(VisibleAnywhere, Category = "LearningAgents")
@@ -260,9 +296,9 @@ private:
 	TArray<TObjectPtr<ULearningAgentsCompletion>> CompletionObjects;
 
 	TArray<TSharedRef<UE::Learning::FRewardObject>, TInlineAllocator<16>> RewardFeatures;
-	TSharedPtr<UE::Learning::FSumReward> Rewards;
-
 	TArray<TSharedRef<UE::Learning::FCompletionObject>, TInlineAllocator<16>> CompletionFeatures;
+
+	TSharedPtr<UE::Learning::FSumReward> Rewards;
 	TSharedPtr<UE::Learning::FAnyCompletion> Completions;
 
 	TUniquePtr<UE::Learning::FEpisodeBuffer> EpisodeBuffer;
@@ -270,16 +306,15 @@ private:
 	TUniquePtr<UE::Learning::FResetInstanceBuffer> ResetBuffer;
 	TUniquePtr<UE::Learning::FSharedMemoryPPOTrainer> Trainer;
 
-	/** One-time Error Logging State */
-	bool bPreviouslyLogged = false;
-
 	ELearningAgentsCompletion MaxStepsCompletion = ELearningAgentsCompletion::Truncation;
 
 	float TrainerTimeout = 10.0f;
 
 	void DoneTraining();
 
-//** ----- Private Recording of GameSettings ----- */
+	UE::Learning::FIndexSet SelectedAgentsSet;
+
+// ----- Private Recording of GameSettings ----- 
 private:
 
 	bool bFixedTimestepUsed = false;
diff --git a/Engine/Plugins/Experimental/LearningAgents/Source/LearningTraining/Private/LearningImitationTrainer.cpp b/Engine/Plugins/Experimental/LearningAgents/Source/LearningTraining/Private/LearningImitationTrainer.cpp
--- a/Engine/Plugins/Experimental/LearningAgents/Source/LearningTraining/Private/LearningImitationTrainer.cpp
+++ b/Engine/Plugins/Experimental/LearningAgents/Source/LearningTraining/Private/LearningImitationTrainer.cpp
@@ -94,6 +94,7 @@ int32 ULearningSocketImitationTrainerServerCommandlet::Main(const FString& Comma
 		IntermediatePath,
 		IpAddress,
 		Port,
+		UE::Learning::ESubprocessFlags::None,
 		LogSettings);
 
 	while (ServerProcess.IsRunning())
@@ -112,24 +113,27 @@ namespace UE::Learning
 		const FString& SitePackagesPath,
 		const FString& PythonContentPath,
 		const FString& IntermediatePath,
-		const FNeuralNetwork& Network,
-		const float ActionNoiseMin,
-		const float ActionNoiseMax,
 		const int32 MaxSampleNum,
-		const FImitationTrainerSettings& Settings,
+		const int32 ObservationDimNum,
+		const int32 ActionDimNum,
+		const FImitationTrainerTrainingSettings& TrainingSettings,
+		const FImitationTrainerNetworkSettings& NetworkSettings,
+		const EImitationTrainerFlags TrainerFlags,
+		const ESubprocessFlags TrainingProcessFlags,
 		const ELogSetting LogSettings)
 	{
-		UE_LEARNING_CHECK(Network.GetOutputNum() % 2 == 0);
-
-		const int32 ObservationVectorDimensionNum = Network.GetInputNum();
-		const int32 ActionVectorDimensionNum = Network.GetOutputNum() / 2;
+		const int32 TotalPolicyByteNum = FNeuralNetwork::GetSerializationByteNum(
+			ObservationDimNum,
+			2 * ActionDimNum,
+			NetworkSettings.PolicyHiddenLayerSize,
+			NetworkSettings.PolicyLayerNum);
 
 		// Allocate Shared Memory
 
-		Policy = SharedMemory::Allocate<1, uint8>({ Network.GetTotalByteNum() });
+		Policy = SharedMemory::Allocate<1, uint8>({ TotalPolicyByteNum });
 		Controls = SharedMemory::Allocate<1, volatile int32>({ SharedMemoryTraining::GetControlNum() });
-		Observations = SharedMemory::Allocate<2, float>({ MaxSampleNum, ObservationVectorDimensionNum });
-		Actions = SharedMemory::Allocate<2, float>({ MaxSampleNum, ActionVectorDimensionNum });
+		Observations = SharedMemory::Allocate<2, float>({ MaxSampleNum, ObservationDimNum });
+		Actions = SharedMemory::Allocate<2, float>({ MaxSampleNum, ActionDimNum });
 
 		// We need to zero the control memory before we start
 		// the training sub-process since it may contain uninitialized 
@@ -158,26 +162,28 @@ namespace UE::Learning
 		ConfigObject->SetStringField(TEXT("ObservationsGuid"), *Observations.Guid.ToString(EGuidFormats::DigitsWithHyphensInBraces));
 		ConfigObject->SetStringField(TEXT("ActionsGuid"), *Actions.Guid.ToString(EGuidFormats::DigitsWithHyphensInBraces));
 
-		ConfigObject->SetNumberField(TEXT("ObservationVectorDimensionNum"), ObservationVectorDimensionNum);
-		ConfigObject->SetNumberField(TEXT("ActionVectorDimensionNum"), ActionVectorDimensionNum);
-		ConfigObject->SetNumberField(TEXT("NetworkTotalByteNum"), Network.GetTotalByteNum());
+		ConfigObject->SetNumberField(TEXT("ObservationVectorDimensionNum"), ObservationDimNum);
+		ConfigObject->SetNumberField(TEXT("ActionVectorDimensionNum"), ActionDimNum);
 		ConfigObject->SetNumberField(TEXT("MaxSampleNum"), MaxSampleNum);
-		ConfigObject->SetNumberField(TEXT("HiddenUnitNum"), Network.GetHiddenNum());
-		ConfigObject->SetNumberField(TEXT("LayerNum"), Network.GetLayerNum());
-		ConfigObject->SetStringField(TEXT("ActivationFunction"), ActivationFunctionString(Network.ActivationFunction));
-
-		ConfigObject->SetNumberField(TEXT("ActionNoiseMin"), ActionNoiseMin);
-		ConfigObject->SetNumberField(TEXT("ActionNoiseMax"), ActionNoiseMax);
-
-		ConfigObject->SetNumberField(TEXT("IterationNum"), Settings.IterationNum);
-		ConfigObject->SetNumberField(TEXT("LearningRateActor"), Settings.LearningRateActor);
-		ConfigObject->SetNumberField(TEXT("LearningRateDecay"), Settings.LearningRateDecay);
-		ConfigObject->SetNumberField(TEXT("WeightDecay"), Settings.WeightDecay);
-		ConfigObject->SetNumberField(TEXT("BatchSize"), Settings.BatchSize);
-		ConfigObject->SetNumberField(TEXT("Seed"), Settings.Seed);
-		ConfigObject->SetBoolField(TEXT("ReinitializeNetwork"), Settings.bReinitializeNetwork);
-		ConfigObject->SetStringField(TEXT("Device"), Trainer::GetDeviceString(Settings.Device));
-		ConfigObject->SetBoolField(TEXT("UseTensorBoard"), Settings.bUseTensorboard);
+
+		ConfigObject->SetNumberField(TEXT("PolicyNetworkByteNum"), TotalPolicyByteNum);
+		ConfigObject->SetNumberField(TEXT("PolicyHiddenUnitNum"), NetworkSettings.PolicyHiddenLayerSize);
+		ConfigObject->SetNumberField(TEXT("PolicyLayerNum"), NetworkSettings.PolicyLayerNum);
+		ConfigObject->SetStringField(TEXT("PolicyActivationFunction"), GetActivationFunctionString(NetworkSettings.PolicyActivationFunction));
+		ConfigObject->SetNumberField(TEXT("PolicyActionNoiseMin"), NetworkSettings.PolicyActionNoiseMin);
+		ConfigObject->SetNumberField(TEXT("PolicyActionNoiseMax"), NetworkSettings.PolicyActionNoiseMax);
+
+		ConfigObject->SetNumberField(TEXT("IterationNum"), TrainingSettings.IterationNum);
+		ConfigObject->SetNumberField(TEXT("LearningRateActor"), TrainingSettings.LearningRateActor);
+		ConfigObject->SetNumberField(TEXT("LearningRateDecay"), TrainingSettings.LearningRateDecay);
+		ConfigObject->SetNumberField(TEXT("WeightDecay"), TrainingSettings.WeightDecay);
+		ConfigObject->SetNumberField(TEXT("BatchSize"), TrainingSettings.BatchSize);
+		ConfigObject->SetNumberField(TEXT("Seed"), TrainingSettings.Seed);
+		ConfigObject->SetStringField(TEXT("Device"), Trainer::GetDeviceString(TrainingSettings.Device));
+		ConfigObject->SetBoolField(TEXT("UseTensorBoard"), TrainingSettings.bUseTensorboard);
+
+		ConfigObject->SetBoolField(TEXT("UseInitialPolicyNetwork"), (bool)(TrainerFlags & EImitationTrainerFlags::UseInitialPolicyNetwork));
+
 		ConfigObject->SetBoolField(TEXT("LoggingEnabled"), LogSettings == ELogSetting::Silent ? false : true);
 
 		FString JsonString;
@@ -190,10 +196,19 @@ namespace UE::Learning
 
 		const FString CommandLineArguments = FString::Printf(TEXT("\"%s\" SharedMemory \"%s\""), *(PythonContentPath / TEXT("train_imitation.py")), *ConfigPath);
 
-		TrainingProcess = MakeShared<FMonitoredProcess>(PythonExecutiblePath, CommandLineArguments, true, true);
-		TrainingProcess->OnCanceled().BindRaw(this, &FSharedMemoryImitationTrainer::HandleTrainingProcessCanceled);
-		TrainingProcess->OnCompleted().BindRaw(this, &FSharedMemoryImitationTrainer::HandleTrainingProcessCompleted);
-		TrainingProcess->OnOutput().BindStatic(&FSharedMemoryImitationTrainer::HandleTrainingProcessOutput);
+		TrainingProcess = MakeShared<FMonitoredProcess>(
+			PythonExecutiblePath, 
+			CommandLineArguments,
+			!(TrainingProcessFlags & ESubprocessFlags::ShowWindow),
+			!(TrainingProcessFlags & ESubprocessFlags::NoRedirectOutput));
+
+		if (!(TrainingProcessFlags & ESubprocessFlags::NoRedirectOutput))
+		{
+			TrainingProcess->OnCanceled().BindRaw(this, &FSharedMemoryImitationTrainer::HandleTrainingProcessCanceled);
+			TrainingProcess->OnCompleted().BindRaw(this, &FSharedMemoryImitationTrainer::HandleTrainingProcessCompleted);
+			TrainingProcess->OnOutput().BindStatic(&FSharedMemoryImitationTrainer::HandleTrainingProcessOutput);
+		}
+
 		TrainingProcess->Launch();
 	}
 
@@ -333,15 +348,25 @@ namespace UE::Learning
 		const FString& IntermediatePath,
 		const TCHAR* IpAddress,
 		const uint32 Port,
+		const ESubprocessFlags TrainingProcessFlags,
 		const ELogSetting LogSettings)
 	{
 		const FString CommandLineArguments = FString::Printf(TEXT("\"%s\" Socket \"%s:%i\" \"%s\" \"%s\" %i"),
 			*(PythonContentPath / TEXT("train_imitation.py")), IpAddress, Port, *SitePackagesPath, *IntermediatePath, LogSettings == ELogSetting::Normal ? 1 : 0);
 
-		TrainingProcess = MakeShared<FMonitoredProcess>(PythonExecutiblePath, CommandLineArguments, true, true);
-		TrainingProcess->OnCanceled().BindRaw(this, &FSocketImitationTrainerServerProcess::HandleTrainingProcessCanceled);
-		TrainingProcess->OnCompleted().BindRaw(this, &FSocketImitationTrainerServerProcess::HandleTrainingProcessCompleted);
-		TrainingProcess->OnOutput().BindStatic(&FSocketImitationTrainerServerProcess::HandleTrainingProcessOutput);
+		TrainingProcess = MakeShared<FMonitoredProcess>(
+			PythonExecutiblePath, 
+			CommandLineArguments,
+			!(TrainingProcessFlags & ESubprocessFlags::ShowWindow),
+			!(TrainingProcessFlags & ESubprocessFlags::NoRedirectOutput));
+
+		if (!(TrainingProcessFlags & ESubprocessFlags::NoRedirectOutput))
+		{
+			TrainingProcess->OnCanceled().BindRaw(this, &FSocketImitationTrainerServerProcess::HandleTrainingProcessCanceled);
+			TrainingProcess->OnCompleted().BindRaw(this, &FSocketImitationTrainerServerProcess::HandleTrainingProcessCompleted);
+			TrainingProcess->OnOutput().BindStatic(&FSocketImitationTrainerServerProcess::HandleTrainingProcessOutput);
+		}
+
 		TrainingProcess->Launch();
 	}
 
@@ -412,19 +437,21 @@ namespace UE::Learning
 	FSocketImitationTrainer::FSocketImitationTrainer(
 		ETrainerResponse& OutResponse,
 		const FString& TaskName,
-		const FNeuralNetwork& Network,
-		const float ActionNoiseMin,
-		const float ActionNoiseMax,
 		const int32 MaxSampleNum,
+		const int32 ObservationDimNum,
+		const int32 ActionDimNum,
 		const TCHAR* IpAddress,
 		const uint32 Port,
 		const float Timeout,
-		const FImitationTrainerSettings& Settings)
+		const FImitationTrainerTrainingSettings& TrainingSettings,
+		const FImitationTrainerNetworkSettings& NetworkSettings,
+		const EImitationTrainerFlags TrainerFlags)
 	{
-		UE_LEARNING_CHECK(Network.GetOutputNum() % 2 == 0);
-
-		const int32 ObservationVectorDimensionNum = Network.GetInputNum();
-		const int32 ActionVectorDimensionNum = Network.GetOutputNum() / 2;
+		const int32 TotalPolicyByteNum = FNeuralNetwork::GetSerializationByteNum(
+			ObservationDimNum,
+			2 * ActionDimNum,
+			NetworkSettings.PolicyHiddenLayerSize,
+			NetworkSettings.PolicyLayerNum);
 
 		// Write Config
 
@@ -438,26 +465,27 @@ namespace UE::Learning
 		ConfigObject->SetStringField(TEXT("TrainerType"), TrainerType);
 		ConfigObject->SetStringField(TEXT("TimeStamp"), *TimeStamp);
 
-		ConfigObject->SetNumberField(TEXT("ObservationVectorDimensionNum"), ObservationVectorDimensionNum);
-		ConfigObject->SetNumberField(TEXT("ActionVectorDimensionNum"), ActionVectorDimensionNum);
-		ConfigObject->SetNumberField(TEXT("NetworkTotalByteNum"), Network.GetTotalByteNum());
+		ConfigObject->SetNumberField(TEXT("ObservationVectorDimensionNum"), ObservationDimNum);
+		ConfigObject->SetNumberField(TEXT("ActionVectorDimensionNum"), ActionDimNum);
 		ConfigObject->SetNumberField(TEXT("MaxSampleNum"), MaxSampleNum);
-		ConfigObject->SetNumberField(TEXT("HiddenUnitNum"), Network.GetHiddenNum());
-		ConfigObject->SetNumberField(TEXT("LayerNum"), Network.GetLayerNum());
-		ConfigObject->SetStringField(TEXT("ActivationFunction"), ActivationFunctionString(Network.ActivationFunction));
-
-		ConfigObject->SetNumberField(TEXT("ActionNoiseMin"), ActionNoiseMin);
-		ConfigObject->SetNumberField(TEXT("ActionNoiseMax"), ActionNoiseMax);
-
-		ConfigObject->SetNumberField(TEXT("IterationNum"), Settings.IterationNum);
-		ConfigObject->SetNumberField(TEXT("LearningRateActor"), Settings.LearningRateActor);
-		ConfigObject->SetNumberField(TEXT("LearningRateDecay"), Settings.LearningRateDecay);
-		ConfigObject->SetNumberField(TEXT("WeightDecay"), Settings.WeightDecay);
-		ConfigObject->SetNumberField(TEXT("BatchSize"), Settings.BatchSize);
-		ConfigObject->SetNumberField(TEXT("Seed"), Settings.Seed);
-		ConfigObject->SetBoolField(TEXT("ReinitializeNetwork"), Settings.bReinitializeNetwork);
-		ConfigObject->SetStringField(TEXT("Device"), Trainer::GetDeviceString(Settings.Device));
-		ConfigObject->SetBoolField(TEXT("UseTensorBoard"), Settings.bUseTensorboard);
+
+		ConfigObject->SetNumberField(TEXT("PolicyNetworkByteNum"), TotalPolicyByteNum);
+		ConfigObject->SetNumberField(TEXT("PolicyHiddenUnitNum"), NetworkSettings.PolicyHiddenLayerSize);
+		ConfigObject->SetNumberField(TEXT("PolicyLayerNum"), NetworkSettings.PolicyLayerNum);
+		ConfigObject->SetStringField(TEXT("PolicyActivationFunction"), GetActivationFunctionString(NetworkSettings.PolicyActivationFunction));
+		ConfigObject->SetNumberField(TEXT("PolicyActionNoiseMin"), NetworkSettings.PolicyActionNoiseMin);
+		ConfigObject->SetNumberField(TEXT("PolicyActionNoiseMax"), NetworkSettings.PolicyActionNoiseMax);
+
+		ConfigObject->SetNumberField(TEXT("IterationNum"), TrainingSettings.IterationNum);
+		ConfigObject->SetNumberField(TEXT("LearningRateActor"), TrainingSettings.LearningRateActor);
+		ConfigObject->SetNumberField(TEXT("LearningRateDecay"), TrainingSettings.LearningRateDecay);
+		ConfigObject->SetNumberField(TEXT("WeightDecay"), TrainingSettings.WeightDecay);
+		ConfigObject->SetNumberField(TEXT("BatchSize"), TrainingSettings.BatchSize);
+		ConfigObject->SetNumberField(TEXT("Seed"), TrainingSettings.Seed);
+		ConfigObject->SetStringField(TEXT("Device"), Trainer::GetDeviceString(TrainingSettings.Device));
+		ConfigObject->SetBoolField(TEXT("UseTensorBoard"), TrainingSettings.bUseTensorboard);
+
+		ConfigObject->SetBoolField(TEXT("UseInitialPolicyNetwork"), (bool)(TrainerFlags & EImitationTrainerFlags::UseInitialPolicyNetwork));
 
 		FString JsonString;
 		TSharedRef<TJsonWriter<>> JsonWriter = TJsonWriterFactory<>::Create(&JsonString, 0);
@@ -465,7 +493,7 @@ namespace UE::Learning
 
 		// Allocate buffer to receive network data in
 
-		NetworkBuffer.SetNumUninitialized({ Network.GetTotalByteNum() });
+		NetworkBuffer.SetNumUninitialized({ TotalPolicyByteNum });
 
 		// Create Socket
 
@@ -567,7 +595,7 @@ namespace UE::Learning
 			FNeuralNetwork& Network,
 			const TLearningArrayView<2, const float> ObservationVectors,
 			const TLearningArrayView<2, const float> ActionVectors,
-			const bool bReinitializeNetwork,
+			const EImitationTrainerFlags TrainerFlags, 
 			TAtomic<bool>* bRequestTrainingStopSignal,
 			FRWLock* NetworkLock,
 			TAtomic<bool>* bNetworkUpdatedSignal,
@@ -575,22 +603,22 @@ namespace UE::Learning
 		{
 			ETrainerResponse Response = ETrainerResponse::Success;
 
-			if (bReinitializeNetwork)
+			if ((bool)(TrainerFlags & EImitationTrainerFlags::UseInitialPolicyNetwork))
 			{
-				// Receive initial Policy
+				// Send initial Policy
 
 				if (LogSettings != ELogSetting::Silent)
 				{
-					UE_LOG(LogLearning, Display, TEXT("Receiving initial Policy..."));
+					UE_LOG(LogLearning, Display, TEXT("Sending initial Policy..."));
 				}
 
-				Response = Trainer.RecvPolicy(Network, 20.0f, NetworkLock);
+				Response = Trainer.SendPolicy(Network, 20.0f, NetworkLock);
 
 				if (Response != ETrainerResponse::Success)
 				{
 					if (LogSettings != ELogSetting::Silent)
 					{
-						UE_LOG(LogLearning, Error, TEXT("Error receiving initial policy from trainer: %s. Check log for errors."), Trainer::ResponseString(Response));
+						UE_LOG(LogLearning, Error, TEXT("Error sending initial policy from trainer: %s. Check log for errors."), Trainer::GetResponseString(Response));
 					}
 
 					Trainer.Terminate();
@@ -599,20 +627,20 @@ namespace UE::Learning
 			}
 			else
 			{
-				// Send initial Policy
+				// Receive initial Policy
 
 				if (LogSettings != ELogSetting::Silent)
 				{
-					UE_LOG(LogLearning, Display, TEXT("Sending initial Policy..."));
+					UE_LOG(LogLearning, Display, TEXT("Receiving initial Policy..."));
 				}
 
-				Response = Trainer.SendPolicy(Network, 20.0f, NetworkLock);
+				Response = Trainer.RecvPolicy(Network, 20.0f, NetworkLock);
 
 				if (Response != ETrainerResponse::Success)
 				{
 					if (LogSettings != ELogSetting::Silent)
 					{
-						UE_LOG(LogLearning, Error, TEXT("Error sending initial policy from trainer: %s. Check log for errors."), Trainer::ResponseString(Response));
+						UE_LOG(LogLearning, Error, TEXT("Error receiving initial policy from trainer: %s. Check log for errors."), Trainer::GetResponseString(Response));
 					}
 
 					Trainer.Terminate();
@@ -628,7 +656,7 @@ namespace UE::Learning
 			{
 				if (LogSettings != ELogSetting::Silent)
 				{
-					UE_LOG(LogLearning, Error, TEXT("Error sending experience to trainer: %s. Check log for errors."), Trainer::ResponseString(Response));
+					UE_LOG(LogLearning, Error, TEXT("Error sending experience to trainer: %s. Check log for errors."), Trainer::GetResponseString(Response));
 				}
 
 				Trainer.Terminate();
@@ -654,7 +682,7 @@ namespace UE::Learning
 					{
 						if (LogSettings != ELogSetting::Silent)
 						{
-							UE_LOG(LogLearning, Error, TEXT("Error sending stop signal to trainer: %s. Check log for errors."), Trainer::ResponseString(Response));
+							UE_LOG(LogLearning, Error, TEXT("Error sending stop signal to trainer: %s. Check log for errors."), Trainer::GetResponseString(Response));
 						}
 
 						Trainer.Terminate();
@@ -671,7 +699,7 @@ namespace UE::Learning
 					{
 						if (LogSettings != ELogSetting::Silent)
 						{
-							UE_LOG(LogLearning, Error, TEXT("Error sending continue signal to trainer: %s. Check log for errors."), Trainer::ResponseString(Response));
+							UE_LOG(LogLearning, Error, TEXT("Error sending continue signal to trainer: %s. Check log for errors."), Trainer::GetResponseString(Response));
 						}
 
 						Trainer.Terminate();
@@ -693,7 +721,7 @@ namespace UE::Learning
 				{
 					if (LogSettings != ELogSetting::Silent)
 					{
-						UE_LOG(LogLearning, Error, TEXT("Error receiving policy from trainer: %s. Check log for errors."), Trainer::ResponseString(Response));
+						UE_LOG(LogLearning, Error, TEXT("Error receiving policy from trainer: %s. Check log for errors."), Trainer::GetResponseString(Response));
 					}
 					break;
 				}
@@ -712,7 +740,7 @@ namespace UE::Learning
 			{
 				if (LogSettings != ELogSetting::Silent)
 				{
-					UE_LOG(LogLearning, Error, TEXT("Error waiting for trainer to exit: %s. Check log for errors."), Trainer::ResponseString(Response));
+					UE_LOG(LogLearning, Error, TEXT("Error waiting for trainer to exit: %s. Check log for errors."), Trainer::GetResponseString(Response));
 				}
 			}
 
diff --git a/Engine/Plugins/Experimental/LearningAgents/Source/LearningTraining/Private/LearningPPOTrainer.cpp b/Engine/Plugins/Experimental/LearningAgents/Source/LearningTraining/Private/LearningPPOTrainer.cpp
--- a/Engine/Plugins/Experimental/LearningAgents/Source/LearningTraining/Private/LearningPPOTrainer.cpp
+++ b/Engine/Plugins/Experimental/LearningAgents/Source/LearningTraining/Private/LearningPPOTrainer.cpp
@@ -97,6 +97,7 @@ int32 ULearningSocketPPOTrainerServerCommandlet::Main(const FString& Commandline
 		IntermediatePath,
 		IpAddress,
 		Port,
+		UE::Learning::ESubprocessFlags::None,
 		LogSettings);
 
 	while (ServerProcess.IsRunning())
@@ -111,23 +112,35 @@ namespace UE::Learning
 {
 	FSharedMemoryPPOTrainer::FSharedMemoryPPOTrainer(
 		const FString& TaskName,
-		const FString& PythonExecutiblePath,
+		const FString& PythonExecutablePath,
 		const FString& SitePackagesPath,
 		const FString& PythonContentPath,
 		const FString& IntermediatePath,
-		const FNeuralNetwork& Network,
-		const float ActionNoiseMin,
-		const float ActionNoiseMax,
 		const FReplayBuffer& ReplayBuffer,
-		const FPPOTrainerSettings& Settings,
+		const FPPOTrainerTrainingSettings& TrainingSettings,
+		const FPPOTrainerNetworkSettings& NetworkSettings,
+		const EPPOTrainerFlags TrainerFlags,
 		const ELogSetting LogSettings,
-		const bool bHideTrainingWindow,
-		const bool bRedirectTrainingOutput)
+		const ESubprocessFlags TrainingProcessFlags,
+		const uint16 ProcessNum,
+		const ESubprocessFlags MultiProcessFlags)
 	{
-		UE_LEARNING_CHECK(Network.GetOutputNum() % 2 == 0);
+		const int32 ObservationVectorDimensionNum = ReplayBuffer.GetObservations().Num<1>();
+		const int32 ActionVectorDimensionNum = ReplayBuffer.GetActions().Num<1>();
+
+		const int32 TotalPolicyByteNum = FNeuralNetwork::GetSerializationByteNum(
+			ObservationVectorDimensionNum,
+			2 * ActionVectorDimensionNum,
+			NetworkSettings.PolicyHiddenLayerSize,
+			NetworkSettings.PolicyLayerNum);
+
+		const int32 TotalCriticByteNum = FNeuralNetwork::GetSerializationByteNum(
+				ObservationVectorDimensionNum,
+				1,
+				NetworkSettings.CriticHiddenLayerSize,
+				NetworkSettings.CriticLayerNum);
 
-		const int32 ObservationVectorDimensionNum = Network.GetInputNum();
-		const int32 ActionVectorDimensionNum = Network.GetOutputNum() / 2;
+		const bool bRequiresCritic = (bool)(TrainerFlags & EPPOTrainerFlags::SynchronizeCriticNetwork) || (bool)(TrainerFlags & EPPOTrainerFlags::UseInitialCriticNetwork);
 
 		if (!ensure(Policy.Region == nullptr))
 		{
@@ -144,15 +157,16 @@ namespace UE::Learning
 		{
 			// Allocate Shared Memory
 
-			Policy = SharedMemory::Allocate<1, uint8>({ Network.GetTotalByteNum() });
-			Controls = SharedMemory::Allocate<2, volatile int32>({ Settings.ProcessNum, SharedMemoryTraining::GetControlNum() });
-			EpisodeStarts = SharedMemory::Allocate<2, int32>({ Settings.ProcessNum, ReplayBuffer.GetMaxEpisodeNum() });
-			EpisodeLengths = SharedMemory::Allocate<2, int32>({ Settings.ProcessNum, ReplayBuffer.GetMaxEpisodeNum() });
-			EpisodeCompletionModes = SharedMemory::Allocate<2, ECompletionMode>({ Settings.ProcessNum, ReplayBuffer.GetMaxEpisodeNum() });
-			EpisodeFinalObservations = SharedMemory::Allocate<3, float>({ Settings.ProcessNum, ReplayBuffer.GetMaxEpisodeNum(), ObservationVectorDimensionNum });
-			Observations = SharedMemory::Allocate<3, float>({ Settings.ProcessNum, ReplayBuffer.GetMaxStepNum(), ObservationVectorDimensionNum });
-			Actions = SharedMemory::Allocate<3, float>({ Settings.ProcessNum, ReplayBuffer.GetMaxStepNum(), ActionVectorDimensionNum });
-			Rewards = SharedMemory::Allocate<2, float>({ Settings.ProcessNum, ReplayBuffer.GetMaxStepNum() });
+			Policy = SharedMemory::Allocate<1, uint8>({ TotalPolicyByteNum });
+			if (bRequiresCritic) { Critic = SharedMemory::Allocate<1, uint8>({ TotalCriticByteNum }); }
+			Controls = SharedMemory::Allocate<2, volatile int32>({ ProcessNum, SharedMemoryTraining::GetControlNum() });
+			EpisodeStarts = SharedMemory::Allocate<2, int32>({ ProcessNum, ReplayBuffer.GetMaxEpisodeNum() });
+			EpisodeLengths = SharedMemory::Allocate<2, int32>({ ProcessNum, ReplayBuffer.GetMaxEpisodeNum() });
+			EpisodeCompletionModes = SharedMemory::Allocate<2, ECompletionMode>({ ProcessNum, ReplayBuffer.GetMaxEpisodeNum() });
+			EpisodeFinalObservations = SharedMemory::Allocate<3, float>({ ProcessNum, ReplayBuffer.GetMaxEpisodeNum(), ObservationVectorDimensionNum });
+			Observations = SharedMemory::Allocate<3, float>({ ProcessNum, ReplayBuffer.GetMaxStepNum(), ObservationVectorDimensionNum });
+			Actions = SharedMemory::Allocate<3, float>({ ProcessNum, ReplayBuffer.GetMaxStepNum(), ActionVectorDimensionNum });
+			Rewards = SharedMemory::Allocate<2, float>({ ProcessNum, ReplayBuffer.GetMaxStepNum() });
 
 			// We need to zero the control memory before we start
 			// the training sub-process since it may contain uninitialized 
@@ -162,7 +176,7 @@ namespace UE::Learning
 
 			// Create Experience Gathering Sub-processes
 
-			for (uint16 SubprocessIdx = 1; SubprocessIdx < Settings.ProcessNum; SubprocessIdx++)
+			for (uint16 SubprocessIdx = 1; SubprocessIdx < ProcessNum; SubprocessIdx++)
 			{
 				ensureMsgf(!WITH_EDITOR || IsRunningCommandlet(), TEXT("Multi-processing generally does not work in-editor as it requires a standalone executable."));
 
@@ -170,6 +184,7 @@ namespace UE::Learning
 
 				SubprocessCommandLine += FString::Printf(TEXT(" -LearningProcessIdx %i"), SubprocessIdx);
 				SubprocessCommandLine += FString(TEXT(" -LearningPolicyGuid ")) + Policy.Guid.ToString();
+				if (bRequiresCritic) { SubprocessCommandLine += FString(TEXT(" -LearningCriticGuid ")) + Critic.Guid.ToString(); }
 				SubprocessCommandLine += FString(TEXT(" -LearningControlsGuid ")) + Controls.Guid.ToString();
 				SubprocessCommandLine += FString(TEXT(" -LearningEpisodeStartsGuid ")) + EpisodeStarts.Guid.ToString();
 				SubprocessCommandLine += FString(TEXT(" -LearningEpisodeLengthsGuid ")) + EpisodeLengths.Guid.ToString();
@@ -182,10 +197,10 @@ namespace UE::Learning
 				const TSharedPtr<FMonitoredProcess> Subprocess = MakeShared<FMonitoredProcess>(
 					FPlatformProcess::ExecutablePath(), 
 					SubprocessCommandLine, 
-					Settings.bMultiProcessHideTrainingWindow, 
-					Settings.bMultiProcessRedirectTrainingOutput);
+					!(MultiProcessFlags & ESubprocessFlags::ShowWindow),
+					!(MultiProcessFlags & ESubprocessFlags::NoRedirectOutput));
 
-				if (Settings.bMultiProcessRedirectTrainingOutput)
+				if (!(MultiProcessFlags & ESubprocessFlags::NoRedirectOutput))
 				{
 					Subprocess->OnCanceled().BindRaw(this, &FSharedMemoryPPOTrainer::HandleSubprocessCanceled);
 					Subprocess->OnCompleted().BindRaw(this, &FSharedMemoryPPOTrainer::HandleSubprocessCompleted);
@@ -216,6 +231,7 @@ namespace UE::Learning
 			ConfigObject->SetStringField(TEXT("IntermediatePath"), *IntermediatePath);
 
 			ConfigObject->SetStringField(TEXT("PolicyGuid"), *Policy.Guid.ToString(EGuidFormats::DigitsWithHyphensInBraces));
+			if (bRequiresCritic) { ConfigObject->SetStringField(TEXT("CriticGuid"), *Critic.Guid.ToString(EGuidFormats::DigitsWithHyphensInBraces)); }
 			ConfigObject->SetStringField(TEXT("ControlsGuid"), *Controls.Guid.ToString(EGuidFormats::DigitsWithHyphensInBraces));
 			ConfigObject->SetStringField(TEXT("EpisodeStartsGuid"), *EpisodeStarts.Guid.ToString(EGuidFormats::DigitsWithHyphensInBraces));
 			ConfigObject->SetStringField(TEXT("EpisodeLengthsGuid"), *EpisodeLengths.Guid.ToString(EGuidFormats::DigitsWithHyphensInBraces));
@@ -227,37 +243,47 @@ namespace UE::Learning
 
 			ConfigObject->SetNumberField(TEXT("ObservationVectorDimensionNum"), ObservationVectorDimensionNum);
 			ConfigObject->SetNumberField(TEXT("ActionVectorDimensionNum"), ActionVectorDimensionNum);
-			ConfigObject->SetNumberField(TEXT("NetworkTotalByteNum"), Network.GetTotalByteNum());
 			ConfigObject->SetNumberField(TEXT("MaxEpisodeNum"), ReplayBuffer.GetMaxEpisodeNum());
 			ConfigObject->SetNumberField(TEXT("MaxStepNum"), ReplayBuffer.GetMaxStepNum());
-			ConfigObject->SetNumberField(TEXT("HiddenUnitNum"), Network.GetHiddenNum());
-			ConfigObject->SetNumberField(TEXT("LayerNum"), Network.GetLayerNum());
-			ConfigObject->SetStringField(TEXT("ActivationFunction"), ActivationFunctionString(Network.ActivationFunction));
-
-			ConfigObject->SetNumberField(TEXT("ActionNoiseMin"), ActionNoiseMin);
-			ConfigObject->SetNumberField(TEXT("ActionNoiseMax"), ActionNoiseMax);
-
-			ConfigObject->SetNumberField(TEXT("ProcessNum"), Settings.ProcessNum);
-			ConfigObject->SetNumberField(TEXT("IterationNum"), Settings.IterationNum);
-			ConfigObject->SetNumberField(TEXT("LearningRateActor"), Settings.LearningRateActor);
-			ConfigObject->SetNumberField(TEXT("LearningRateCritic"), Settings.LearningRateCritic);
-			ConfigObject->SetNumberField(TEXT("LearningRateDecay"), Settings.LearningRateDecay);
-			ConfigObject->SetNumberField(TEXT("WeightDecay"), Settings.WeightDecay);
-			ConfigObject->SetNumberField(TEXT("InitialActionScale"), Settings.InitialActionScale);
-			ConfigObject->SetNumberField(TEXT("BatchSize"), Settings.BatchSize);
-			ConfigObject->SetNumberField(TEXT("EpsilonClip"), Settings.EpsilonClip);
-			ConfigObject->SetNumberField(TEXT("ActionRegularizationWeight"), Settings.ActionRegularizationWeight);
-			ConfigObject->SetNumberField(TEXT("EntropyWeight"), Settings.EntropyWeight);
-			ConfigObject->SetNumberField(TEXT("GaeLambda"), Settings.GaeLambda);
-			ConfigObject->SetBoolField(TEXT("ClipAdvantages"), Settings.bClipAdvantages);
-			ConfigObject->SetBoolField(TEXT("AdvantageNormalization"), Settings.bAdvantageNormalization);
-			ConfigObject->SetNumberField(TEXT("TrimEpisodeStartStepNum"), Settings.TrimEpisodeStartStepNum);
-			ConfigObject->SetNumberField(TEXT("TrimEpisodeEndStepNum"), Settings.TrimEpisodeEndStepNum);
-			ConfigObject->SetNumberField(TEXT("Seed"), Settings.Seed);
-			ConfigObject->SetNumberField(TEXT("DiscountFactor"), Settings.DiscountFactor);
-			ConfigObject->SetBoolField(TEXT("ReinitializeNetwork"), Settings.bReinitializeNetwork);
-			ConfigObject->SetStringField(TEXT("Device"), Trainer::GetDeviceString(Settings.Device));
-			ConfigObject->SetBoolField(TEXT("UseTensorBoard"), Settings.bUseTensorboard);
+
+			ConfigObject->SetNumberField(TEXT("PolicyNetworkByteNum"), TotalPolicyByteNum);
+			ConfigObject->SetNumberField(TEXT("PolicyHiddenUnitNum"), NetworkSettings.PolicyHiddenLayerSize);
+			ConfigObject->SetNumberField(TEXT("PolicyLayerNum"), NetworkSettings.PolicyLayerNum);
+			ConfigObject->SetStringField(TEXT("PolicyActivationFunction"), GetActivationFunctionString(NetworkSettings.PolicyActivationFunction));
+			ConfigObject->SetNumberField(TEXT("PolicyActionNoiseMin"), NetworkSettings.PolicyActionNoiseMin);
+			ConfigObject->SetNumberField(TEXT("PolicyActionNoiseMax"), NetworkSettings.PolicyActionNoiseMax);
+
+			ConfigObject->SetNumberField(TEXT("CriticNetworkByteNum"), TotalCriticByteNum);
+			ConfigObject->SetNumberField(TEXT("CriticHiddenUnitNum"), NetworkSettings.CriticHiddenLayerSize);
+			ConfigObject->SetNumberField(TEXT("CriticLayerNum"), NetworkSettings.CriticLayerNum);
+			ConfigObject->SetStringField(TEXT("CriticActivationFunction"), GetActivationFunctionString(NetworkSettings.CriticActivationFunction));
+
+			ConfigObject->SetNumberField(TEXT("ProcessNum"), ProcessNum);
+
+			ConfigObject->SetNumberField(TEXT("IterationNum"), TrainingSettings.IterationNum);
+			ConfigObject->SetNumberField(TEXT("LearningRatePolicy"), TrainingSettings.LearningRatePolicy);
+			ConfigObject->SetNumberField(TEXT("LearningRateCritic"), TrainingSettings.LearningRateCritic);
+			ConfigObject->SetNumberField(TEXT("LearningRateDecay"), TrainingSettings.LearningRateDecay);
+			ConfigObject->SetNumberField(TEXT("WeightDecay"), TrainingSettings.WeightDecay);
+			ConfigObject->SetNumberField(TEXT("InitialActionScale"), TrainingSettings.InitialActionScale);
+			ConfigObject->SetNumberField(TEXT("BatchSize"), TrainingSettings.BatchSize);
+			ConfigObject->SetNumberField(TEXT("EpsilonClip"), TrainingSettings.EpsilonClip);
+			ConfigObject->SetNumberField(TEXT("ActionRegularizationWeight"), TrainingSettings.ActionRegularizationWeight);
+			ConfigObject->SetNumberField(TEXT("EntropyWeight"), TrainingSettings.EntropyWeight);
+			ConfigObject->SetNumberField(TEXT("GaeLambda"), TrainingSettings.GaeLambda);
+			ConfigObject->SetBoolField(TEXT("ClipAdvantages"), TrainingSettings.bClipAdvantages);
+			ConfigObject->SetBoolField(TEXT("AdvantageNormalization"), TrainingSettings.bAdvantageNormalization);
+			ConfigObject->SetNumberField(TEXT("TrimEpisodeStartStepNum"), TrainingSettings.TrimEpisodeStartStepNum);
+			ConfigObject->SetNumberField(TEXT("TrimEpisodeEndStepNum"), TrainingSettings.TrimEpisodeEndStepNum);
+			ConfigObject->SetNumberField(TEXT("Seed"), TrainingSettings.Seed);
+			ConfigObject->SetNumberField(TEXT("DiscountFactor"), TrainingSettings.DiscountFactor);
+			ConfigObject->SetStringField(TEXT("Device"), Trainer::GetDeviceString(TrainingSettings.Device));
+			ConfigObject->SetBoolField(TEXT("UseTensorBoard"), TrainingSettings.bUseTensorboard);
+
+			ConfigObject->SetBoolField(TEXT("UseInitialPolicyNetwork"), (bool)(TrainerFlags& EPPOTrainerFlags::UseInitialPolicyNetwork));
+			ConfigObject->SetBoolField(TEXT("UseInitialCriticNetwork"), (bool)(TrainerFlags& EPPOTrainerFlags::UseInitialCriticNetwork));
+			ConfigObject->SetBoolField(TEXT("SynchronizeCriticNetwork"), (bool)(TrainerFlags& EPPOTrainerFlags::SynchronizeCriticNetwork));
+
 			ConfigObject->SetBoolField(TEXT("LoggingEnabled"), LogSettings == ELogSetting::Silent ? false : true);
 
 			FString JsonString;
@@ -270,9 +296,13 @@ namespace UE::Learning
 
 			const FString CommandLineArguments = FString::Printf(TEXT("\"%s\" SharedMemory \"%s\""), *(PythonContentPath / TEXT("train_ppo.py")), *ConfigPath);
 
-			TrainingProcess = MakeShared<FMonitoredProcess>(PythonExecutiblePath, CommandLineArguments, bHideTrainingWindow, bRedirectTrainingOutput);
+			TrainingProcess = MakeShared<FMonitoredProcess>(
+				PythonExecutablePath, 
+				CommandLineArguments, 
+				!(TrainingProcessFlags & ESubprocessFlags::ShowWindow),
+				!(TrainingProcessFlags & ESubprocessFlags::NoRedirectOutput));
 
-			if (bRedirectTrainingOutput)
+			if (!(TrainingProcessFlags & ESubprocessFlags::NoRedirectOutput))
 			{
 				TrainingProcess->OnCanceled().BindRaw(this, &FSharedMemoryPPOTrainer::HandleTrainingProcessCanceled);
 				TrainingProcess->OnCompleted().BindRaw(this, &FSharedMemoryPPOTrainer::HandleTrainingProcessCompleted);
@@ -286,6 +316,7 @@ namespace UE::Learning
 			// Parse Guids from command line args
 
 			FGuid PolicyGuid; ensure(FParse::Value(FCommandLine::Get(), TEXT("LearningPolicyGuid"), PolicyGuid));
+			FGuid CriticGuid; if (bRequiresCritic) { ensure(FParse::Value(FCommandLine::Get(), TEXT("LearningCriticGuid"), CriticGuid)); }
 			FGuid ControlsGuid; ensure(FParse::Value(FCommandLine::Get(), TEXT("LearningControlsGuid"), ControlsGuid));
 			FGuid EpisodeStartsGuid; ensure(FParse::Value(FCommandLine::Get(), TEXT("LearningEpisodeStartsGuid"), EpisodeStartsGuid));
 			FGuid EpisodeLengthsGuid; ensure(FParse::Value(FCommandLine::Get(), TEXT("LearningEpisodeLengthsGuid"), EpisodeLengthsGuid));
@@ -297,15 +328,16 @@ namespace UE::Learning
 
 			// Map shared memory
 
-			Policy = SharedMemory::Map<1, uint8>(PolicyGuid, { Network.GetTotalByteNum() });
-			Controls = SharedMemory::Map<2, volatile int32>(ControlsGuid, { Settings.ProcessNum, SharedMemoryTraining::GetControlNum() });
-			EpisodeStarts = SharedMemory::Map<2, int32>(EpisodeStartsGuid, { Settings.ProcessNum, ReplayBuffer.GetMaxEpisodeNum() });
-			EpisodeLengths = SharedMemory::Map<2, int32>(EpisodeLengthsGuid, { Settings.ProcessNum, ReplayBuffer.GetMaxEpisodeNum() });
-			EpisodeCompletionModes = SharedMemory::Map<2, ECompletionMode>(EpisodeCompletionModesGuid, { Settings.ProcessNum, ReplayBuffer.GetMaxEpisodeNum() });
-			EpisodeFinalObservations = SharedMemory::Map<3, float>(EpisodeFinalObservationsGuid, { Settings.ProcessNum, ReplayBuffer.GetMaxEpisodeNum(), ObservationVectorDimensionNum });
-			Observations = SharedMemory::Map<3, float>(ObservationsGuid, { Settings.ProcessNum, ReplayBuffer.GetMaxStepNum(), ObservationVectorDimensionNum });
-			Actions = SharedMemory::Map<3, float>(ActionsGuid, { Settings.ProcessNum, ReplayBuffer.GetMaxStepNum(), ActionVectorDimensionNum });
-			Rewards = SharedMemory::Map<2, float>(RewardsGuid, { Settings.ProcessNum, ReplayBuffer.GetMaxStepNum() });
+			Policy = SharedMemory::Map<1, uint8>(PolicyGuid, { TotalPolicyByteNum });
+			if (bRequiresCritic) { Critic = SharedMemory::Map<1, uint8>(CriticGuid, { TotalCriticByteNum }); }
+			Controls = SharedMemory::Map<2, volatile int32>(ControlsGuid, { ProcessNum, SharedMemoryTraining::GetControlNum() });
+			EpisodeStarts = SharedMemory::Map<2, int32>(EpisodeStartsGuid, { ProcessNum, ReplayBuffer.GetMaxEpisodeNum() });
+			EpisodeLengths = SharedMemory::Map<2, int32>(EpisodeLengthsGuid, { ProcessNum, ReplayBuffer.GetMaxEpisodeNum() });
+			EpisodeCompletionModes = SharedMemory::Map<2, ECompletionMode>(EpisodeCompletionModesGuid, { ProcessNum, ReplayBuffer.GetMaxEpisodeNum() });
+			EpisodeFinalObservations = SharedMemory::Map<3, float>(EpisodeFinalObservationsGuid, { ProcessNum, ReplayBuffer.GetMaxEpisodeNum(), ObservationVectorDimensionNum });
+			Observations = SharedMemory::Map<3, float>(ObservationsGuid, { ProcessNum, ReplayBuffer.GetMaxStepNum(), ObservationVectorDimensionNum });
+			Actions = SharedMemory::Map<3, float>(ActionsGuid, { ProcessNum, ReplayBuffer.GetMaxStepNum(), ActionVectorDimensionNum });
+			Rewards = SharedMemory::Map<2, float>(RewardsGuid, { ProcessNum, ReplayBuffer.GetMaxStepNum() });
 		}
 	}
 
@@ -368,6 +400,24 @@ namespace UE::Learning
 			LogSettings);
 	}
 
+	ETrainerResponse FSharedMemoryPPOTrainer::RecvCritic(
+		FNeuralNetwork& OutNetwork,
+		const float Timeout,
+		FRWLock* NetworkLock,
+		const ELogSetting LogSettings)
+	{
+		UE_LEARNING_CHECKF(Critic.View.GetData(), 
+			TEXT("Buffer for critic was not allocated. Were the correct EPPOTrainerFlags provided to construct the FSharedMemoryPPOTrainer?"));
+
+		return SharedMemoryTraining::RecvCritic(
+			Controls.View[ProcessIdx],
+			OutNetwork,
+			Critic.View,
+			Timeout,
+			NetworkLock,
+			LogSettings);
+	}
+
 	ETrainerResponse FSharedMemoryPPOTrainer::SendPolicy(
 		const FNeuralNetwork& Network,
 		const float Timeout,
@@ -383,6 +433,24 @@ namespace UE::Learning
 			LogSettings);
 	}
 
+	ETrainerResponse FSharedMemoryPPOTrainer::SendCritic(
+		const FNeuralNetwork& Network,
+		const float Timeout,
+		FRWLock* NetworkLock,
+		const ELogSetting LogSettings)
+	{
+		UE_LEARNING_CHECKF(Critic.View.GetData(),
+			TEXT("Buffer for critic was not allocated. Were the correct EPPOTrainerFlags provided to construct the FSharedMemoryPPOTrainer?"));
+
+		return SharedMemoryTraining::SendCritic(
+			Controls.View[ProcessIdx],
+			Critic.View,
+			Network,
+			Timeout,
+			NetworkLock,
+			LogSettings);
+	}
+
 	ETrainerResponse FSharedMemoryPPOTrainer::SendExperience(
 		const FReplayBuffer& ReplayBuffer,
 		const float Timeout,
@@ -472,15 +540,25 @@ namespace UE::Learning
 		const FString& IntermediatePath,
 		const TCHAR* IpAddress,
 		const uint32 Port,
+		const ESubprocessFlags TrainingProcessFlags,
 		const ELogSetting LogSettings)
 	{
 		const FString CommandLineArguments = FString::Printf(TEXT("\"%s\" Socket \"%s:%i\" \"%s\" \"%s\" %i"), 
 			*(PythonContentPath / TEXT("train_ppo.py")), IpAddress, Port, *SitePackagesPath, *IntermediatePath, LogSettings == ELogSetting::Normal ? 1 : 0);
 
-		TrainingProcess = MakeShared<FMonitoredProcess>(PythonExecutiblePath, CommandLineArguments, true, true);
-		TrainingProcess->OnCanceled().BindRaw(this, &FSocketPPOTrainerServerProcess::HandleTrainingProcessCanceled);
-		TrainingProcess->OnCompleted().BindRaw(this, &FSocketPPOTrainerServerProcess::HandleTrainingProcessCompleted);
-		TrainingProcess->OnOutput().BindStatic(&FSocketPPOTrainerServerProcess::HandleTrainingProcessOutput);
+		TrainingProcess = MakeShared<FMonitoredProcess>(
+			PythonExecutiblePath, 
+			CommandLineArguments,
+			!(TrainingProcessFlags & ESubprocessFlags::ShowWindow),
+			!(TrainingProcessFlags & ESubprocessFlags::NoRedirectOutput));
+
+		if (!(TrainingProcessFlags & ESubprocessFlags::NoRedirectOutput))
+		{
+			TrainingProcess->OnCanceled().BindRaw(this, &FSocketPPOTrainerServerProcess::HandleTrainingProcessCanceled);
+			TrainingProcess->OnCompleted().BindRaw(this, &FSocketPPOTrainerServerProcess::HandleTrainingProcessCompleted);
+			TrainingProcess->OnOutput().BindStatic(&FSocketPPOTrainerServerProcess::HandleTrainingProcessOutput);
+		}
+
 		TrainingProcess->Launch();
 	}
 
@@ -551,64 +629,77 @@ namespace UE::Learning
 	FSocketPPOTrainer::FSocketPPOTrainer(
 		ETrainerResponse& OutResponse,
 		const FString& TaskName,
-		const FNeuralNetwork& Network,
-		const float ActionNoiseMin,
-		const float ActionNoiseMax,
 		const FReplayBuffer& ReplayBuffer,
 		const TCHAR* IpAddress,
 		const uint32 Port,
 		const float Timeout,
-		const FPPOTrainerSettings& Settings)
+		const FPPOTrainerTrainingSettings& TrainingSettings,
+		const FPPOTrainerNetworkSettings& NetworkSettings,
+		const EPPOTrainerFlags TrainerFlags)
 	{
-		UE_LEARNING_CHECK(Network.GetOutputNum() % 2 == 0);
+		const int32 ObservationVectorDimensionNum = ReplayBuffer.GetObservations().Num<1>();
+		const int32 ActionVectorDimensionNum = ReplayBuffer.GetActions().Num<1>();
 
-		const int32 ObservationVectorDimensionNum = Network.GetInputNum();
-		const int32 ActionVectorDimensionNum = Network.GetOutputNum() / 2;
+		const int32 TotalPolicyByteNum = FNeuralNetwork::GetSerializationByteNum(
+			ObservationVectorDimensionNum,
+			2 * ActionVectorDimensionNum,
+			NetworkSettings.PolicyHiddenLayerSize,
+			NetworkSettings.PolicyLayerNum);
 
-		// Write Config
+		const int32 TotalCriticByteNum = FNeuralNetwork::GetSerializationByteNum(
+			ObservationVectorDimensionNum,
+			1,
+			NetworkSettings.CriticHiddenLayerSize,
+			NetworkSettings.CriticLayerNum);
 
-		const FString TimeStamp = FDateTime::Now().ToFormattedString(TEXT("%Y-%m-%d_%H-%M-%S"));
-		const FString TrainerMethod = TEXT("PPO");
-		const FString TrainerType = TEXT("Network");
+		// Write Config
 
 		TSharedRef<FJsonObject> ConfigObject = MakeShared<FJsonObject>();
 		ConfigObject->SetStringField(TEXT("TaskName"), *TaskName);
-		ConfigObject->SetStringField(TEXT("TrainerMethod"), TrainerMethod);
-		ConfigObject->SetStringField(TEXT("TrainerType"), TrainerType);
-		ConfigObject->SetStringField(TEXT("TimeStamp"), *TimeStamp);
+		ConfigObject->SetStringField(TEXT("TrainerMethod"), TEXT("PPO"));
+		ConfigObject->SetStringField(TEXT("TrainerType"), TEXT("Network"));
+		ConfigObject->SetStringField(TEXT("TimeStamp"), *FDateTime::Now().ToFormattedString(TEXT("%Y-%m-%d_%H-%M-%S")));
 
 		ConfigObject->SetNumberField(TEXT("ObservationVectorDimensionNum"), ObservationVectorDimensionNum);
 		ConfigObject->SetNumberField(TEXT("ActionVectorDimensionNum"), ActionVectorDimensionNum);
-		ConfigObject->SetNumberField(TEXT("NetworkTotalByteNum"), Network.GetTotalByteNum());
 		ConfigObject->SetNumberField(TEXT("MaxEpisodeNum"), ReplayBuffer.GetMaxEpisodeNum());
 		ConfigObject->SetNumberField(TEXT("MaxStepNum"), ReplayBuffer.GetMaxStepNum());
-		ConfigObject->SetNumberField(TEXT("HiddenUnitNum"), Network.GetHiddenNum());
-		ConfigObject->SetNumberField(TEXT("LayerNum"), Network.GetLayerNum());
-		ConfigObject->SetStringField(TEXT("ActivationFunction"), ActivationFunctionString(Network.ActivationFunction));
-
-		ConfigObject->SetNumberField(TEXT("ActionNoiseMin"), ActionNoiseMin);
-		ConfigObject->SetNumberField(TEXT("ActionNoiseMax"), ActionNoiseMax);
-
-		ConfigObject->SetNumberField(TEXT("IterationNum"), Settings.IterationNum);
-		ConfigObject->SetNumberField(TEXT("LearningRateActor"), Settings.LearningRateActor);
-		ConfigObject->SetNumberField(TEXT("LearningRateCritic"), Settings.LearningRateCritic);
-		ConfigObject->SetNumberField(TEXT("LearningRateDecay"), Settings.LearningRateDecay);
-		ConfigObject->SetNumberField(TEXT("WeightDecay"), Settings.WeightDecay);
-		ConfigObject->SetNumberField(TEXT("InitialActionScale"), Settings.InitialActionScale);
-		ConfigObject->SetNumberField(TEXT("BatchSize"), Settings.BatchSize);
-		ConfigObject->SetNumberField(TEXT("EpsilonClip"), Settings.EpsilonClip);
-		ConfigObject->SetNumberField(TEXT("ActionRegularizationWeight"), Settings.ActionRegularizationWeight);
-		ConfigObject->SetNumberField(TEXT("EntropyWeight"), Settings.EntropyWeight);
-		ConfigObject->SetNumberField(TEXT("GaeLambda"), Settings.GaeLambda);
-		ConfigObject->SetBoolField(TEXT("ClipAdvantages"), Settings.bClipAdvantages);
-		ConfigObject->SetBoolField(TEXT("AdvantageNormalization"), Settings.bAdvantageNormalization);
-		ConfigObject->SetNumberField(TEXT("TrimEpisodeStartStepNum"), Settings.TrimEpisodeStartStepNum);
-		ConfigObject->SetNumberField(TEXT("TrimEpisodeEndStepNum"), Settings.TrimEpisodeEndStepNum);
-		ConfigObject->SetNumberField(TEXT("Seed"), Settings.Seed);
-		ConfigObject->SetNumberField(TEXT("DiscountFactor"), Settings.DiscountFactor);
-		ConfigObject->SetBoolField(TEXT("ReinitializeNetwork"), Settings.bReinitializeNetwork);
-		ConfigObject->SetStringField(TEXT("Device"), Trainer::GetDeviceString(Settings.Device));
-		ConfigObject->SetBoolField(TEXT("UseTensorBoard"), Settings.bUseTensorboard);
+
+		ConfigObject->SetNumberField(TEXT("PolicyNetworkByteNum"), TotalPolicyByteNum);
+		ConfigObject->SetNumberField(TEXT("PolicyHiddenUnitNum"), NetworkSettings.PolicyHiddenLayerSize);
+		ConfigObject->SetNumberField(TEXT("PolicyLayerNum"), NetworkSettings.PolicyLayerNum);
+		ConfigObject->SetStringField(TEXT("PolicyActivationFunction"), GetActivationFunctionString(NetworkSettings.PolicyActivationFunction));
+		ConfigObject->SetNumberField(TEXT("PolicyActionNoiseMin"), NetworkSettings.PolicyActionNoiseMin);
+		ConfigObject->SetNumberField(TEXT("PolicyActionNoiseMax"), NetworkSettings.PolicyActionNoiseMax);
+
+		ConfigObject->SetNumberField(TEXT("CriticNetworkByteNum"), TotalCriticByteNum);
+		ConfigObject->SetNumberField(TEXT("CriticHiddenUnitNum"), NetworkSettings.CriticHiddenLayerSize);
+		ConfigObject->SetNumberField(TEXT("CriticLayerNum"), NetworkSettings.CriticLayerNum);
+		ConfigObject->SetStringField(TEXT("CriticActivationFunction"), GetActivationFunctionString(NetworkSettings.CriticActivationFunction));
+
+		ConfigObject->SetNumberField(TEXT("IterationNum"), TrainingSettings.IterationNum);
+		ConfigObject->SetNumberField(TEXT("LearningRatePolicy"), TrainingSettings.LearningRatePolicy);
+		ConfigObject->SetNumberField(TEXT("LearningRateCritic"), TrainingSettings.LearningRateCritic);
+		ConfigObject->SetNumberField(TEXT("LearningRateDecay"), TrainingSettings.LearningRateDecay);
+		ConfigObject->SetNumberField(TEXT("WeightDecay"), TrainingSettings.WeightDecay);
+		ConfigObject->SetNumberField(TEXT("InitialActionScale"), TrainingSettings.InitialActionScale);
+		ConfigObject->SetNumberField(TEXT("BatchSize"), TrainingSettings.BatchSize);
+		ConfigObject->SetNumberField(TEXT("EpsilonClip"), TrainingSettings.EpsilonClip);
+		ConfigObject->SetNumberField(TEXT("ActionRegularizationWeight"), TrainingSettings.ActionRegularizationWeight);
+		ConfigObject->SetNumberField(TEXT("EntropyWeight"), TrainingSettings.EntropyWeight);
+		ConfigObject->SetNumberField(TEXT("GaeLambda"), TrainingSettings.GaeLambda);
+		ConfigObject->SetBoolField(TEXT("ClipAdvantages"), TrainingSettings.bClipAdvantages);
+		ConfigObject->SetBoolField(TEXT("AdvantageNormalization"), TrainingSettings.bAdvantageNormalization);
+		ConfigObject->SetNumberField(TEXT("TrimEpisodeStartStepNum"), TrainingSettings.TrimEpisodeStartStepNum);
+		ConfigObject->SetNumberField(TEXT("TrimEpisodeEndStepNum"), TrainingSettings.TrimEpisodeEndStepNum);
+		ConfigObject->SetNumberField(TEXT("Seed"), TrainingSettings.Seed);
+		ConfigObject->SetNumberField(TEXT("DiscountFactor"), TrainingSettings.DiscountFactor);
+		ConfigObject->SetStringField(TEXT("Device"), Trainer::GetDeviceString(TrainingSettings.Device));
+		ConfigObject->SetBoolField(TEXT("UseTensorBoard"), TrainingSettings.bUseTensorboard);
+
+		ConfigObject->SetBoolField(TEXT("UseInitialPolicyNetwork"), (bool)(TrainerFlags & EPPOTrainerFlags::UseInitialPolicyNetwork));
+		ConfigObject->SetBoolField(TEXT("UseInitialCriticNetwork"), (bool)(TrainerFlags & EPPOTrainerFlags::UseInitialCriticNetwork));
+		ConfigObject->SetBoolField(TEXT("SynchronizeCriticNetwork"), (bool)(TrainerFlags & EPPOTrainerFlags::SynchronizeCriticNetwork));
 
 		FString JsonString;
 		TSharedRef<TJsonWriter<>> JsonWriter = TJsonWriterFactory<>::Create(&JsonString, 0);
@@ -616,7 +707,12 @@ namespace UE::Learning
 
 		// Allocate buffer to receive network data in
 
-		NetworkBuffer.SetNumUninitialized({ Network.GetTotalByteNum() });
+		PolicyNetworkBuffer.SetNumUninitialized({ TotalPolicyByteNum });
+
+		if ((bool)(TrainerFlags & EPPOTrainerFlags::SynchronizeCriticNetwork) || (bool)(TrainerFlags & EPPOTrainerFlags::UseInitialCriticNetwork))
+		{
+			CriticNetworkBuffer.SetNumUninitialized({ TotalCriticByteNum });
+		}
 
 		// Create Socket
 
@@ -684,7 +780,19 @@ namespace UE::Learning
 		FRWLock* NetworkLock,
 		const ELogSetting LogSettings)
 	{
-		return SocketTraining::RecvPolicy(*Socket, OutNetwork, NetworkBuffer, Timeout, NetworkLock, LogSettings);
+		return SocketTraining::RecvPolicy(*Socket, OutNetwork, PolicyNetworkBuffer, Timeout, NetworkLock, LogSettings);
+	}
+
+	ETrainerResponse FSocketPPOTrainer::RecvCritic(
+		FNeuralNetwork& OutNetwork,
+		const float Timeout,
+		FRWLock* NetworkLock,
+		const ELogSetting LogSettings)
+	{
+		UE_LEARNING_CHECKF(CriticNetworkBuffer.GetData(),
+			TEXT("Buffer for critic was not allocated. Were the correct EPPOTrainerFlags provided to construct the FSocketPPOTrainer?"));
+
+		return SocketTraining::RecvCritic(*Socket, OutNetwork, CriticNetworkBuffer, Timeout, NetworkLock, LogSettings);
 	}
 
 	ETrainerResponse FSocketPPOTrainer::SendPolicy(
@@ -693,7 +801,19 @@ namespace UE::Learning
 		FRWLock* NetworkLock,
 		const ELogSetting LogSettings)
 	{
-		return SocketTraining::SendPolicy(*Socket, NetworkBuffer, Network, Timeout, NetworkLock, LogSettings);
+		return SocketTraining::SendPolicy(*Socket, PolicyNetworkBuffer, Network, Timeout, NetworkLock, LogSettings);
+	}
+
+	ETrainerResponse FSocketPPOTrainer::SendCritic(
+		const FNeuralNetwork& Network,
+		const float Timeout,
+		FRWLock* NetworkLock,
+		const ELogSetting LogSettings)
+	{
+		UE_LEARNING_CHECKF(CriticNetworkBuffer.GetData(),
+			TEXT("Buffer for critic was not allocated. Were the correct EPPOTrainerFlags provided to construct the FSocketPPOTrainer?"));
+
+		return SocketTraining::SendCritic(*Socket, CriticNetworkBuffer, Network, Timeout, NetworkLock, LogSettings);
 	}
 
 	ETrainerResponse FSocketPPOTrainer::SendExperience(
@@ -711,7 +831,8 @@ namespace UE::Learning
 			FReplayBuffer& ReplayBuffer,
 			FEpisodeBuffer& EpisodeBuffer,
 			FResetInstanceBuffer& ResetBuffer,
-			FNeuralNetwork& Network,
+			FNeuralNetwork& PolicyNetwork,
+			FNeuralNetwork* CriticNetwork,
 			TLearningArrayView<2, float> ObservationVectorBuffer,
 			TLearningArrayView<2, float> ActionVectorBuffer,
 			TLearningArrayView<1, float> RewardBuffer,
@@ -725,17 +846,41 @@ namespace UE::Learning
 			const TFunctionRef<void(const FIndexSet Instances)> RewardFunction,
 			const TFunctionRef<void(const FIndexSet Instances)> CompletionFunction,
 			const FIndexSet Instances,
-			const bool bReinitializeNetwork,
+			const EPPOTrainerFlags TrainerFlags,
 			TAtomic<bool>* bRequestTrainingStopSignal,
-			FRWLock* NetworkLock,
-			TAtomic<bool>* bNetworkUpdatedSignal,
+			FRWLock* PolicyNetworkLock,
+			FRWLock* CriticNetworkLock,
+			TAtomic<bool>* bPolicyNetworkUpdatedSignal,
+			TAtomic<bool>* bCriticNetworkUpdatedSignal,
 			const ELogSetting LogSettings)
 		{
 			TRACE_CPUPROFILER_EVENT_SCOPE(Learning::PPOTrainer::Train);
 
 			ETrainerResponse Response = ETrainerResponse::Success;
 
-			if (bReinitializeNetwork)
+			if ((bool)(TrainerFlags & EPPOTrainerFlags::UseInitialPolicyNetwork))
+			{
+				// Send initial Policy
+
+				if (LogSettings != ELogSetting::Silent)
+				{
+					UE_LOG(LogLearning, Display, TEXT("Sending initial Policy..."));
+				}
+
+				Response = Trainer.SendPolicy(PolicyNetwork, 20.0f, PolicyNetworkLock);
+
+				if (Response != ETrainerResponse::Success)
+				{
+					if (LogSettings != ELogSetting::Silent)
+					{
+						UE_LOG(LogLearning, Error, TEXT("Error sending initial policy to trainer: %s. Check log for errors."), Trainer::GetResponseString(Response));
+					}
+
+					Trainer.Terminate();
+					return Response;
+				}
+			}
+			else
 			{
 				// Receive initial Policy
 
@@ -744,46 +889,78 @@ namespace UE::Learning
 					UE_LOG(LogLearning, Display, TEXT("Receiving initial Policy..."));
 				}
 
-				Response = Trainer.RecvPolicy(Network, 20.0f, NetworkLock);
+				Response = Trainer.RecvPolicy(PolicyNetwork, 20.0f, PolicyNetworkLock);
 
 				if (Response != ETrainerResponse::Success)
 				{
 					if (LogSettings != ELogSetting::Silent)
 					{
-						UE_LOG(LogLearning, Error, TEXT("Error receiving initial policy from trainer: %s. Check log for errors."), Trainer::ResponseString(Response));
+						UE_LOG(LogLearning, Error, TEXT("Error receiving initial policy from trainer: %s. Check log for errors."), Trainer::GetResponseString(Response));
 					}
 
 					Trainer.Terminate();
 					return Response;
 				}
 
-				if (bNetworkUpdatedSignal)
+				if (bPolicyNetworkUpdatedSignal)
 				{
-					*bNetworkUpdatedSignal = true;
+					*bPolicyNetworkUpdatedSignal = true;
 				}
 			}
-			else
+
+			if ((bool)(TrainerFlags & EPPOTrainerFlags::UseInitialCriticNetwork))
 			{
-				// Send initial Policy
+				UE_LEARNING_CHECKF(CriticNetwork, TEXT("Valid Critic Network must be provided if you want to send it"));
+
+				// Send initial Critic
 
 				if (LogSettings != ELogSetting::Silent)
 				{
-					UE_LOG(LogLearning, Display, TEXT("Sending initial Policy..."));
+					UE_LOG(LogLearning, Display, TEXT("Sending initial Critic..."));
 				}
 
-				Response = Trainer.SendPolicy(Network, 20.0f, NetworkLock);
+				Response = Trainer.SendCritic(*CriticNetwork, 20.0f, CriticNetworkLock);
 
 				if (Response != ETrainerResponse::Success)
 				{
 					if (LogSettings != ELogSetting::Silent)
 					{
-						UE_LOG(LogLearning, Error, TEXT("Error sending initial policy to trainer: %s. Check log for errors."), Trainer::ResponseString(Response));
+						UE_LOG(LogLearning, Error, TEXT("Error sending initial critic to trainer: %s. Check log for errors."), Trainer::GetResponseString(Response));
 					}
 
 					Trainer.Terminate();
 					return Response;
 				}
 			}
+			else if ((bool)(TrainerFlags & EPPOTrainerFlags::SynchronizeCriticNetwork))
+			{
+				UE_LEARNING_CHECKF(CriticNetwork, TEXT("Valid Critic Network must be provided if you want to sync it"));
+
+				// Receive initial Critic
+
+				if (LogSettings != ELogSetting::Silent)
+				{
+					UE_LOG(LogLearning, Display, TEXT("Receiving initial Critic..."));
+				}
+
+				Response = Trainer.RecvCritic(*CriticNetwork, 20.0f, CriticNetworkLock);
+
+				if (Response != ETrainerResponse::Success)
+				{
+					if (LogSettings != ELogSetting::Silent)
+					{
+						UE_LOG(LogLearning, Error, TEXT("Error receiving initial critic from trainer: %s. Check log for errors."), Trainer::GetResponseString(Response));
+					}
+
+					Trainer.Terminate();
+					return Response;
+				}
+
+				if (bCriticNetworkUpdatedSignal)
+				{
+					*bCriticNetworkUpdatedSignal = true;
+				}
+			}
 
 			// Start Training Loop
 
@@ -804,7 +981,7 @@ namespace UE::Learning
 					{
 						if (LogSettings != ELogSetting::Silent)
 						{
-							UE_LOG(LogLearning, Error, TEXT("Error sending stop signal to trainer: %s. Check log for errors."), Trainer::ResponseString(Response));
+							UE_LOG(LogLearning, Error, TEXT("Error sending stop signal to trainer: %s. Check log for errors."), Trainer::GetResponseString(Response));
 						}
 
 						Trainer.Terminate();
@@ -839,7 +1016,7 @@ namespace UE::Learning
 					{
 						if (LogSettings != ELogSetting::Silent)
 						{
-							UE_LOG(LogLearning, Error, TEXT("Error sending experience to trainer: %s. Check log for errors."), Trainer::ResponseString(Response));
+							UE_LOG(LogLearning, Error, TEXT("Error sending experience to trainer: %s. Check log for errors."), Trainer::GetResponseString(Response));
 						}
 
 						Trainer.Terminate();
@@ -847,7 +1024,7 @@ namespace UE::Learning
 					}
 				}
 
-				Response = Trainer.RecvPolicy(Network, 10.0f, NetworkLock);
+				Response = Trainer.RecvPolicy(PolicyNetwork, 10.0f, PolicyNetworkLock);
 
 				if (Response == ETrainerResponse::Completed)
 				{
@@ -861,14 +1038,35 @@ namespace UE::Learning
 				{
 					if (LogSettings != ELogSetting::Silent)
 					{
-						UE_LOG(LogLearning, Error, TEXT("Error receiving policy from trainer: %s. Check log for errors."), Trainer::ResponseString(Response));
+						UE_LOG(LogLearning, Error, TEXT("Error receiving policy from trainer: %s. Check log for errors."), Trainer::GetResponseString(Response));
 					}
 					break;
 				}
 
-				if (bNetworkUpdatedSignal)
+				if (bPolicyNetworkUpdatedSignal)
 				{
-					*bNetworkUpdatedSignal = true;
+					*bPolicyNetworkUpdatedSignal = true;
+				}
+
+				if ((bool)(TrainerFlags & EPPOTrainerFlags::SynchronizeCriticNetwork))
+				{
+					UE_LEARNING_CHECKF(CriticNetwork, TEXT("Valid Critic Network must be provided if you want to sync it"));
+
+					Response = Trainer.RecvCritic(*CriticNetwork, 10.0f, CriticNetworkLock);
+
+					if (Response != ETrainerResponse::Success)
+					{
+						if (LogSettings != ELogSetting::Silent)
+						{
+							UE_LOG(LogLearning, Error, TEXT("Error receiving critic from trainer: %s. Check log for errors."), Trainer::GetResponseString(Response));
+						}
+						break;
+					}
+
+					if (bCriticNetworkUpdatedSignal)
+					{
+						*bCriticNetworkUpdatedSignal = true;
+					}
 				}
 			}
 
@@ -880,7 +1078,7 @@ namespace UE::Learning
 			{
 				if (LogSettings != ELogSetting::Silent)
 				{
-					UE_LOG(LogLearning, Error, TEXT("Error waiting for trainer to exit: %s. Check log for errors."), Trainer::ResponseString(Response));
+					UE_LOG(LogLearning, Error, TEXT("Error waiting for trainer to exit: %s. Check log for errors."), Trainer::GetResponseString(Response));
 				}
 			}
 
diff --git a/Engine/Plugins/Experimental/LearningAgents/Source/LearningTraining/Private/LearningSharedMemoryTraining.cpp b/Engine/Plugins/Experimental/LearningAgents/Source/LearningTraining/Private/LearningSharedMemoryTraining.cpp
--- a/Engine/Plugins/Experimental/LearningAgents/Source/LearningTraining/Private/LearningSharedMemoryTraining.cpp
+++ b/Engine/Plugins/Experimental/LearningAgents/Source/LearningTraining/Private/LearningSharedMemoryTraining.cpp
@@ -17,11 +17,12 @@ namespace UE::Learning::SharedMemoryTraining
 		ExperienceStepNum = 1,
 		ExperienceSignal = 2,
 		PolicySignal = 3,
-		CompleteSignal = 4,
-		StopSignal = 5,
-		ContinueSignal = 6,
+		CriticSignal = 4,
+		CompleteSignal = 5,
+		StopSignal = 6,
+		ContinueSignal = 7,
 
-		ControlNum = 7,
+		ControlNum = 8,
 	};
 
 	uint8 GetControlNum()
@@ -89,6 +90,46 @@ namespace UE::Learning::SharedMemoryTraining
 		return ETrainerResponse::Success;
 	}
 
+	ETrainerResponse RecvCritic(
+		TLearningArrayView<1, volatile int32> Controls,
+		FNeuralNetwork& OutNetwork,
+		const TLearningArrayView<1, const uint8> Critic,
+		const float Timeout,
+		FRWLock* NetworkLock,
+		const ELogSetting LogSettings)
+	{
+		const float SleepTime = 0.001f;
+		float WaitTime = 0.0f;
+
+		// Wait until the Critic is done being written by the sub-process
+		while (!Controls[(uint8)EControls::CriticSignal])
+		{
+			FPlatformProcess::Sleep(SleepTime);
+			WaitTime += SleepTime;
+
+			if (WaitTime > Timeout)
+			{
+				return ETrainerResponse::Timeout;
+			}
+		}
+
+		if (LogSettings != ELogSetting::Silent)
+		{
+			UE_LOG(LogLearning, Display, TEXT("Pulling Critic..."));
+		}
+
+		// Read the policy
+		{
+			FScopeNullableWriteLock ScopeLock(NetworkLock);
+			OutNetwork.DeserializeFromBytes(Critic);
+		}
+
+		// Confirm we have read the critic
+		Controls[(uint8)EControls::CriticSignal] = false;
+
+		return ETrainerResponse::Success;
+	}
+
 	ETrainerResponse SendPolicy(
 		TLearningArrayView<1, volatile int32> Controls,
 		TLearningArrayView<1, uint8> Policy,
@@ -129,6 +170,46 @@ namespace UE::Learning::SharedMemoryTraining
 		return ETrainerResponse::Success;
 	}
 
+	ETrainerResponse SendCritic(
+		TLearningArrayView<1, volatile int32> Controls,
+		TLearningArrayView<1, uint8> Critic,
+		const FNeuralNetwork& Network,
+		const float Timeout,
+		FRWLock* NetworkLock,
+		const ELogSetting LogSettings)
+	{
+		const float SleepTime = 0.001f;
+		float WaitTime = 0.0f;
+
+		// Wait until the Critic is requested by the sub-process
+		while (!Controls[(uint8)EControls::CriticSignal])
+		{
+			FPlatformProcess::Sleep(SleepTime);
+			WaitTime += SleepTime;
+
+			if (WaitTime > Timeout)
+			{
+				return ETrainerResponse::Timeout;
+			}
+		}
+
+		if (LogSettings != ELogSetting::Silent)
+		{
+			UE_LOG(LogLearning, Display, TEXT("Pushing Critic..."));
+		}
+
+		// Write the critic
+		{
+			FScopeNullableReadLock ScopeLock(NetworkLock);
+			Network.SerializeToBytes(Critic);
+		}
+
+		// Confirm we have written the critic
+		Controls[(uint8)EControls::CriticSignal] = false;
+
+		return ETrainerResponse::Success;
+	}
+
 	ETrainerResponse SendExperience(
 		TLearningArrayView<1, int32> EpisodeStarts,
 		TLearningArrayView<1, int32> EpisodeLengths,
diff --git a/Engine/Plugins/Experimental/LearningAgents/Source/LearningTraining/Private/LearningSharedMemoryTraining.h b/Engine/Plugins/Experimental/LearningAgents/Source/LearningTraining/Private/LearningSharedMemoryTraining.h
--- a/Engine/Plugins/Experimental/LearningAgents/Source/LearningTraining/Private/LearningSharedMemoryTraining.h
+++ b/Engine/Plugins/Experimental/LearningAgents/Source/LearningTraining/Private/LearningSharedMemoryTraining.h
@@ -22,6 +22,14 @@ namespace UE::Learning
 			FRWLock* NetworkLock = nullptr,
 			const ELogSetting LogSettings = Trainer::DefaultLogSettings);
 
+		LEARNINGTRAINING_API ETrainerResponse RecvCritic(
+			TLearningArrayView<1, volatile int32> Controls,
+			FNeuralNetwork& OutNetwork,
+			const TLearningArrayView<1, const uint8> Critic,
+			const float Timeout = Trainer::DefaultTimeout,
+			FRWLock* NetworkLock = nullptr,
+			const ELogSetting LogSettings = Trainer::DefaultLogSettings);
+
 		LEARNINGTRAINING_API ETrainerResponse SendStop(
 			TLearningArrayView<1, volatile int32> Controls);
 
@@ -36,6 +44,14 @@ namespace UE::Learning
 			FRWLock* NetworkLock = nullptr,
 			const ELogSetting LogSettings = Trainer::DefaultLogSettings);
 
+		LEARNINGTRAINING_API ETrainerResponse SendCritic(
+			TLearningArrayView<1, volatile int32> Controls,
+			TLearningArrayView<1, uint8> Critic,
+			const FNeuralNetwork& Network,
+			const float Timeout = Trainer::DefaultTimeout,
+			FRWLock* NetworkLock = nullptr,
+			const ELogSetting LogSettings = Trainer::DefaultLogSettings);
+
 		LEARNINGTRAINING_API ETrainerResponse SendExperience(
 			TLearningArrayView<1, int32> EpisodeStarts,
 			TLearningArrayView<1, int32> EpisodeLengths,
diff --git a/Engine/Plugins/Experimental/LearningAgents/Source/LearningTraining/Private/LearningSocketTraining.cpp b/Engine/Plugins/Experimental/LearningAgents/Source/LearningTraining/Private/LearningSocketTraining.cpp
--- a/Engine/Plugins/Experimental/LearningAgents/Source/LearningTraining/Private/LearningSocketTraining.cpp
+++ b/Engine/Plugins/Experimental/LearningAgents/Source/LearningTraining/Private/LearningSocketTraining.cpp
@@ -18,9 +18,11 @@ namespace UE::Learning::SocketTraining
 		SendExperience = 2,
 		RecvPolicy = 3,
 		SendPolicy = 4,
-		RecvComplete = 5,
-		SendStop = 6,
-		SendContinue = 7,
+		RecvCritic = 5,
+		SendCritic = 6,
+		RecvComplete = 7,
+		SendStop = 8,
+		SendContinue = 9,
 	};
 
 	ETrainerResponse WaitForConnection(FSocket& Socket, const float Timeout)
@@ -98,9 +100,8 @@ namespace UE::Learning::SocketTraining
 			UE_LOG(LogLearning, Display, TEXT("Pulling Policy..."));
 		}
 
-		ETrainerResponse Response = ETrainerResponse::Success;
 		uint8 Signal = (uint8)ESignal::Invalid;
-		Response = RecvWithTimeout(Socket, &Signal, 1, Timeout);
+		ETrainerResponse Response = RecvWithTimeout(Socket, &Signal, 1, Timeout);
 		if (Response != ETrainerResponse::Success) { return Response; }
 
 		if (Signal == (uint8)ESignal::RecvComplete)
@@ -124,6 +125,39 @@ namespace UE::Learning::SocketTraining
 		return ETrainerResponse::Success;
 	}
 
+	ETrainerResponse RecvCritic(
+		FSocket& Socket,
+		FNeuralNetwork& OutNetwork,
+		TLearningArrayView<1, uint8> OutNetworkBuffer,
+		const float Timeout,
+		FRWLock* NetworkLock,
+		const ELogSetting LogSettings)
+	{
+		if (LogSettings != ELogSetting::Silent)
+		{
+			UE_LOG(LogLearning, Display, TEXT("Pulling Critic..."));
+		}
+
+		uint8 Signal = (uint8)ESignal::Invalid;
+		ETrainerResponse Response = RecvWithTimeout(Socket, &Signal, 1, Timeout);
+		if (Response != ETrainerResponse::Success) { return Response; }
+
+		if (Signal != (uint8)ESignal::RecvCritic)
+		{
+			return ETrainerResponse::Unexpected;
+		}
+
+		Response = RecvWithTimeout(Socket, OutNetworkBuffer.GetData(), OutNetworkBuffer.Num(), Timeout);
+		if (Response != ETrainerResponse::Success) { return Response; }
+
+		{
+			FScopeNullableWriteLock ScopeLock(NetworkLock);
+			OutNetwork.DeserializeFromBytes(OutNetworkBuffer);
+		}
+
+		return ETrainerResponse::Success;
+	}
+
 	ETrainerResponse SendWithTimeout(FSocket& Socket, const uint8* Bytes, const int32 ByteNum, const float Timeout)
 	{
 		float WaitTime = 0.0f;
@@ -214,6 +248,34 @@ namespace UE::Learning::SocketTraining
 		return ETrainerResponse::Success;
 	}
 
+	ETrainerResponse SendCritic(
+		FSocket& Socket,
+		TLearningArrayView<1, uint8> NetworkBuffer,
+		const FNeuralNetwork& Network,
+		const float Timeout,
+		FRWLock* NetworkLock,
+		const ELogSetting LogSettings)
+	{
+		if (LogSettings != ELogSetting::Silent)
+		{
+			UE_LOG(LogLearning, Display, TEXT("Pushing Critic..."));
+		}
+
+		{
+			FScopeNullableReadLock ScopeLock(NetworkLock);
+			Network.SerializeToBytes(NetworkBuffer);
+		}
+
+		const uint8 Signal = (uint8)ESignal::SendCritic;
+		ETrainerResponse Response = SendWithTimeout(Socket, &Signal, 1, Timeout);
+		if (Response != ETrainerResponse::Success) { return Response; }
+
+		Response = SendWithTimeout(Socket, NetworkBuffer.GetData(), NetworkBuffer.Num(), Timeout);
+		if (Response != ETrainerResponse::Success) { return Response; }
+
+		return ETrainerResponse::Success;
+	}
+
 	ETrainerResponse SendExperience(
 		FSocket& Socket,
 		const FReplayBuffer& ReplayBuffer,
diff --git a/Engine/Plugins/Experimental/LearningAgents/Source/LearningTraining/Private/LearningSocketTraining.h b/Engine/Plugins/Experimental/LearningAgents/Source/LearningTraining/Private/LearningSocketTraining.h
--- a/Engine/Plugins/Experimental/LearningAgents/Source/LearningTraining/Private/LearningSocketTraining.h
+++ b/Engine/Plugins/Experimental/LearningAgents/Source/LearningTraining/Private/LearningSocketTraining.h
@@ -31,6 +31,14 @@ namespace UE::Learning
 			FRWLock* NetworkLock = nullptr,
 			const ELogSetting LogSettings = Trainer::DefaultLogSettings);
 
+		LEARNINGTRAINING_API ETrainerResponse RecvCritic(
+			FSocket& Socket,
+			FNeuralNetwork& OutNetwork,
+			TLearningArrayView<1, uint8> OutNetworkBuffer,
+			const float Timeout = Trainer::DefaultTimeout,
+			FRWLock* NetworkLock = nullptr,
+			const ELogSetting LogSettings = Trainer::DefaultLogSettings);
+
 		LEARNINGTRAINING_API ETrainerResponse SendWithTimeout(
 			FSocket& Socket, 
 			const uint8* Bytes, 
@@ -58,6 +66,14 @@ namespace UE::Learning
 			FRWLock* NetworkLock = nullptr,
 			const ELogSetting LogSettings = Trainer::DefaultLogSettings);
 
+		LEARNINGTRAINING_API ETrainerResponse SendCritic(
+			FSocket& Socket,
+			TLearningArrayView<1, uint8> NetworkBuffer,
+			const FNeuralNetwork& Network,
+			const float Timeout = Trainer::DefaultTimeout,
+			FRWLock* NetworkLock = nullptr,
+			const ELogSetting LogSettings = Trainer::DefaultLogSettings);
+
 		LEARNINGTRAINING_API ETrainerResponse SendExperience(
 			FSocket& Socket,
 			const FReplayBuffer& ReplayBuffer,
diff --git a/Engine/Plugins/Experimental/LearningAgents/Source/LearningTraining/Private/LearningTrainer.cpp b/Engine/Plugins/Experimental/LearningAgents/Source/LearningTraining/Private/LearningTrainer.cpp
--- a/Engine/Plugins/Experimental/LearningAgents/Source/LearningTraining/Private/LearningTrainer.cpp
+++ b/Engine/Plugins/Experimental/LearningAgents/Source/LearningTraining/Private/LearningTrainer.cpp
@@ -18,7 +18,7 @@ namespace UE::Learning::Trainer
 		}
 	}
 
-	const TCHAR* ResponseString(const ETrainerResponse Response)
+	const TCHAR* GetResponseString(const ETrainerResponse Response)
 	{
 		switch (Response)
 		{
diff --git a/Engine/Plugins/Experimental/LearningAgents/Source/LearningTraining/Public/LearningImitationTrainer.h b/Engine/Plugins/Experimental/LearningAgents/Source/LearningTraining/Public/LearningImitationTrainer.h
--- a/Engine/Plugins/Experimental/LearningAgents/Source/LearningTraining/Public/LearningImitationTrainer.h
+++ b/Engine/Plugins/Experimental/LearningAgents/Source/LearningTraining/Public/LearningImitationTrainer.h
@@ -6,6 +6,7 @@
 #include "LearningLog.h"
 #include "LearningTrainer.h"
 #include "LearningSharedMemory.h"
+#include "LearningNeuralNetwork.h" // Included for EActivationFunction::ELU
 
 #include "Commandlets/Commandlet.h"
 #include "Templates/SharedPointer.h"
@@ -18,7 +19,9 @@ class FMonitoredProcess;
 UCLASS()
 class LEARNINGTRAINING_API ULearningSocketImitationTrainerServerCommandlet : public UCommandlet
 {
-	GENERATED_UCLASS_BODY()
+	GENERATED_BODY()
+
+	ULearningSocketImitationTrainerServerCommandlet(const FObjectInitializer& ObjectInitializer);
 
 	/** Runs the commandlet */
 	virtual int32 Main(const FString& Params) override;
@@ -26,9 +29,29 @@ class LEARNINGTRAINING_API ULearningSocketImitationTrainerServerCommandlet : pub
 
 namespace UE::Learning
 {
-	struct FNeuralNetwork;
+	/**
+	* Settings for the network used for training. These settings must match the Neural Network
+	* objects passed to ImitationTrainer::Train.
+	*/
+	struct FImitationTrainerNetworkSettings
+	{
+		/** Minimum action noise used by the policy */
+		float PolicyActionNoiseMin = 0.25f;
+
+		/** Maximum action noise used by the policy */
+		float PolicyActionNoiseMax = 0.25f;
 
-	struct FImitationTrainerSettings
+		/** Total layers for policy network including input, hidden, and output layers */
+		int32 PolicyLayerNum = 3;
+
+		/** Number of neurons in each hidden layer of the policy network */
+		int32 PolicyHiddenLayerSize = 128;
+
+		/** Activation function to use on hidden layers of the policy network */
+		EActivationFunction PolicyActivationFunction = EActivationFunction::ELU;
+	};
+
+	struct FImitationTrainerTrainingSettings
 	{
 		// Number of iterations to train the network for. Controls the overall training time.
 		// Training for about 100000 iterations should give you well trained network, but
@@ -52,11 +75,6 @@ namespace UE::Learning
 		// Random seed to use for training
 		uint32 Seed = 1234;
 
-		// If true, will reinitialize the network to random weights at the start of training. Set this
-		// to false if you are starting from a network which has already been trained. Make sure 
-		// to use the same settings here for ImitationTrainer::Train.
-		bool bReinitializeNetwork = true;
-
 		// Which device to use for training
 		ETrainerDevice Device = ETrainerDevice::GPU;
 
@@ -65,10 +83,23 @@ namespace UE::Learning
 		// Even when enabled, TensorBoard will only work if it is installed in your Unreal Editor
 		// bundled version of Python, which is not the case by default. TensorBoard can be installed 
 		// for this version of Python by going to your Unreal Editor Python Binaries directory 
-		// (e.g. "\Engine\Binaries\ThirdParty\Python3\Win64") and running `.\python -m pip install tensorboard`. 
+		// (e.g. "\Engine\Binaries\ThirdParty\Python3\Win64") and running `python -m pip install tensorboard`. 
 		bool bUseTensorboard = false;
 	};
 
+	/**
+	* ImitationTrainer flags controlling some aspects of the process of communication with the trainer
+	*/
+	enum class EImitationTrainerFlags : uint8
+	{
+		None = 0,
+
+		// If to send over the initial provided policy network rather than reinitialize it from random weights at 
+		// the start of training. Use this if you want to start from a network which has already been trained.
+		UseInitialPolicyNetwork = 1 << 0,
+	};
+	ENUM_CLASS_FLAGS(EImitationTrainerFlags)
+
 	/**
 	* Interface for an object which can train a policy from experience using imitation learning.
 	*/
@@ -164,11 +195,12 @@ namespace UE::Learning
 		* @param SitePackagesPath		Path to the site-packages shipped with the PythonFoundationPackages plugin
 		* @param PythonContentPath		Path to the Python Content folder provided by the Learning plugin
 		* @param IntermediatePath		Path to the intermediate folder to write temporary files, logs, and snapshots to
-		* @param Network				Neural Network object
-		* @param ActionNoiseMin			Minimum amount of action noise being used during episode generation
-		* @param ActionNoiseMax			Maximum amount of action noise being used during episode generation
 		* @param MaxSampleNum			Maximum number of samples in the training data
-		* @param Settings				Trainer settings
+		* @param ObservationDimNum		Number of dimensions in the observation vector
+		* @param ActionDimNum			Number of dimensions in the action vector
+		* @param TrainingSettings		Trainer Training settings
+		* @param NetworkSettings		Trainer Network settings
+		* @param TrainingProcessFlags	Training subprocess flags
 		* @param LogSettings			Logging settings to use
 		*/
 		FSharedMemoryImitationTrainer(
@@ -177,11 +209,13 @@ namespace UE::Learning
 			const FString& SitePackagesPath,
 			const FString& PythonContentPath,
 			const FString& IntermediatePath,
-			const FNeuralNetwork& Network,
-			const float ActionNoiseMin,
-			const float ActionNoiseMax,
 			const int32 MaxSampleNum,
-			const FImitationTrainerSettings& Settings = FImitationTrainerSettings(),
+			const int32 ObservationDimNum,
+			const int32 ActionDimNum,
+			const FImitationTrainerTrainingSettings& TrainingSettings = FImitationTrainerTrainingSettings(),
+			const FImitationTrainerNetworkSettings& NetworkSettings = FImitationTrainerNetworkSettings(),
+			const EImitationTrainerFlags TrainerFlags = EImitationTrainerFlags::None,
+			const ESubprocessFlags TrainingProcessFlags = ESubprocessFlags::None,
 			const ELogSetting LogSettings = ELogSetting::Normal);
 
 		~FSharedMemoryImitationTrainer();
@@ -250,6 +284,7 @@ namespace UE::Learning
 		* @param IntermediatePath			Path to the intermediate folder to write temporary files, logs, and snapshots to
 		* @param IpAddress					Ip address to bind the listening socket to. For a local server you will want to use 127.0.0.1
 		* @param Port						Port to use for the listening socket.
+		* @param TrainingProcessFlags		Training server subprocess flags
 		* @param LogSettings				Logging settings to use
 		*/
 		FSocketImitationTrainerServerProcess(
@@ -259,6 +294,7 @@ namespace UE::Learning
 			const FString& IntermediatePath,
 			const TCHAR* IpAddress = Trainer::DefaultIp,
 			const uint32 Port = Trainer::DefaultPort,
+			const ESubprocessFlags TrainingProcessFlags = ESubprocessFlags::None,
 			const ELogSetting LogSettings = ELogSetting::Normal);
 
 		~FSocketImitationTrainerServerProcess();
@@ -297,27 +333,28 @@ namespace UE::Learning
 		*
 		* @param OutResponse				Response to the initial connection
 		* @param TaskName					Name of the training task - used to help identify the logs, snapshots, and other files generated by training
-		* @param Network					Neural Network object
-		* @param ActionNoiseMin				Minimum amount of action noise being used during episode generation
-		* @param ActionNoiseMax				Maximum amount of action noise being used during episode generation
-		* @param MaxSampleNum				Maximum number of samples in the training data
+		* @param MaxSampleNum			Maximum number of samples in the training data
+		* @param ObservationDimNum		Number of dimensions in the observation vector
+		* @param ActionDimNum			Number of dimensions in the action vector
 		* @param IpAddress					Server Ip address
 		* @param Port						Server Port
 		* @param Timeout					Timeout to wait in seconds for connection and initial data transfer
-		* @param Settings					Training settings
-		* @param LogSettings				Logging settings to use
+		* @param TrainingSettings			Trainer Training settings
+		* @param NetworkSettings			Trainer Network settings
+		* @param TrainerFlags				Flags for the trainer
 		*/
 		FSocketImitationTrainer(
 			ETrainerResponse& OutResponse,
 			const FString& TaskName,
-			const FNeuralNetwork& Network,
-			const float ActionNoiseMin,
-			const float ActionNoiseMax,
 			const int32 MaxSampleNum,
+			const int32 ObservationDimNum,
+			const int32 ActionDimNum,
 			const TCHAR* IpAddress = Trainer::DefaultIp,
 			const uint32 Port = Trainer::DefaultPort,
 			const float Timeout = Trainer::DefaultTimeout,
-			const FImitationTrainerSettings& Settings = FImitationTrainerSettings());
+			const FImitationTrainerTrainingSettings& TrainingSettings = FImitationTrainerTrainingSettings(),
+			const FImitationTrainerNetworkSettings& NetworkSettings = FImitationTrainerNetworkSettings(),
+			const EImitationTrainerFlags TrainerFlags = EImitationTrainerFlags::None);
 
 		~FSocketImitationTrainer();
 
@@ -362,7 +399,7 @@ namespace UE::Learning
 		* @param Network							Policy network
 		* @param ObservationVectors					Observation Data
 		* @param ActionVectors						Action Data
-		* @param bReinitializeNetwork				Was the network reinitialized to random weights or used as-is
+		* @param TrainerFlags						Flags for the trainer, should match what was used to initialize the Trainer object.
 		* @param bRequestTrainingStopSignal			Optional signal that can be raised to indicate training should be stopped
 		* @param NetworkLock						Optional Lock to use when updating the policy network
 		* @param bNetworkUpdatedSignal				Optional signal that will be raised when the policy network is updated
@@ -374,7 +411,7 @@ namespace UE::Learning
 			FNeuralNetwork& Network,
 			const TLearningArrayView<2, const float> ObservationVectors,
 			const TLearningArrayView<2, const float> ActionVectors,
-			const bool bReinitializeNetwork = true,
+			const EImitationTrainerFlags TrainerFlags = EImitationTrainerFlags::None,
 			TAtomic<bool>* bRequestTrainingStopSignal = nullptr,
 			FRWLock* NetworkLock = nullptr,
 			TAtomic<bool>* bNetworkUpdatedSignal = nullptr,
