LearningAgents: Large refactoring to allow critic to be loaded and used at runtime. Fixed some other misc bugs and fallout from switching features to all be arrays by default. Some refactoring of diff

diff --git a/Engine/Plugins/Experimental/LearningAgents/Source/LearningTraining/Public/LearningPPOTrainer.h b/Engine/Plugins/Experimental/LearningAgents/Source/LearningTraining/Public/LearningPPOTrainer.h
--- a/Engine/Plugins/Experimental/LearningAgents/Source/LearningTraining/Public/LearningPPOTrainer.h
+++ b/Engine/Plugins/Experimental/LearningAgents/Source/LearningTraining/Public/LearningPPOTrainer.h
@@ -6,6 +6,7 @@
 #include "LearningLog.h"
 #include "LearningTrainer.h"
 #include "LearningSharedMemory.h"
+#include "LearningNeuralNetwork.h" // Included for EActivationFunction::ELU
 
 #include "Commandlets/Commandlet.h"
 #include "Templates/SharedPointer.h"
@@ -18,7 +19,9 @@ class FMonitoredProcess;
 UCLASS()
 class LEARNINGTRAINING_API ULearningSocketPPOTrainerServerCommandlet : public UCommandlet
 {
-	GENERATED_UCLASS_BODY()
+	GENERATED_BODY()
+
+	ULearningSocketPPOTrainerServerCommandlet(const FObjectInitializer& ObjectInitializer);
 
 	/** Runs the commandlet */
 	virtual int32 Main(const FString& Params) override;
@@ -26,24 +29,55 @@ class LEARNINGTRAINING_API ULearningSocketPPOTrainerServerCommandlet : public UC
 
 namespace UE::Learning
 {
-	struct FNeuralNetwork;
 	struct FReplayBuffer;
 	struct FResetInstanceBuffer;
 	struct FEpisodeBuffer;
 	enum class ECompletionMode : uint8;
 
 	/**
-	* Settings to use when training using PPO
+	* Settings for the networks used for training PPO. These settings must match the Neural Network 
+	* objects passed to PPOTrainer::Train and the action noise must match the action noise used by 
+	* the policy while gathering experience. 
+	*/
+	struct FPPOTrainerNetworkSettings
+	{
+		/** Minimum action noise used by the policy */
+		float PolicyActionNoiseMin = 0.25f;
+
+		/** Maximum action noise used by the policy */
+		float PolicyActionNoiseMax = 0.25f;
+
+		/** Total layers for policy network including input, hidden, and output layers */
+		int32 PolicyLayerNum = 3;
+
+		/** Number of neurons in each hidden layer of the policy network */
+		int32 PolicyHiddenLayerSize = 128;
+
+		/** Activation function to use on hidden layers of the policy network */
+		EActivationFunction PolicyActivationFunction = EActivationFunction::ELU;
+
+		/** Total layers for critic network including input, hidden, and output layers */
+		int32 CriticLayerNum = 4;
+
+		/** Number of neurons in each hidden layer of the critic network */
+		int32 CriticHiddenLayerSize = 256;
+
+		/** Activation function to use on hidden layers of the critic network */
+		EActivationFunction CriticActivationFunction = EActivationFunction::ELU;
+	};
+
+	/**
+	* Settings used for training with PPO
 	*/
-	struct FPPOTrainerSettings
+	struct FPPOTrainerTrainingSettings
 	{
 		// Number of iterations to train the network for. Controls the overall training time.
 		// Training for about 100000 iterations should give you well trained network, but
 		// closer to 1000000 iterations or more is required for an exhaustively trained network.
 		uint32 IterationNum = 1000000;
 
-		// Learning rate of the actor network. Typical values are between 0.001f and 0.0001f
-		float LearningRateActor = 0.0001f;
+		// Learning rate of the policy network. Typical values are between 0.001f and 0.0001f
+		float LearningRatePolicy = 0.0001f;
 
 		// Learning rate of the critic network. To avoid instability generally the critic 
 		// should have a larger learning rate then the actor.
@@ -108,47 +142,41 @@ namespace UE::Learning
 		// encourage the system to "look-ahead" but make training more difficult.
 		float DiscountFactor = 0.99f;
 		
-		// If to reinitialize the network to random weights at the start of training. Set this
-		// to false if you are start from a network which has already been trained such as via 
-		// Imitation Training. Make sure to use the same settings here for PPOTrainer::Train.
-		bool bReinitializeNetwork = true;
-
 		// Which device to use for training
 		ETrainerDevice Device = ETrainerDevice::GPU;
 
-		// The number of processes to use for multi-processed training.
-		// 
-		// It is important to know how this multi-process training works so that it can be used
-		// correctly when you set this >1:
-		// 
-		// When called with this argument set >1, the process will spawn additional processes running 
-		// the same command as is currently being run but with the additional command line argument 
-		// `LearningProcessIdx`. Once the training starts on these subprocesses process, this command 
-		// line argument will be used to indicate that this subprocess should only be used for gathering 
-		// experience and should not start a new training process. 
-		// 
-		// This means that this will generally not work in editor (or PIE)
-		// 
-		// IMPORTANT: If you are seeding your experience generation process to make it deterministic
-		// you should use this `LearningProcessIdx` command line argument to change the seed of your
-		// experience gathering - otherwise each process will gather identical experience.
-		uint16 ProcessNum = 1;
-
-		// If to hide the window for each multi-processing subprocess
-		bool bMultiProcessHideTrainingWindow = false;
-
-		// If to redirect the output of each multi-processing subprocess to the main process
-		bool bMultiProcessRedirectTrainingOutput = false;
-
 		// If to use TensorBoard for logging and tracking the training progress.
 		// 
 		// Even when enabled, TensorBoard will only work if it is installed in your Unreal Editor
 		// bundled version of Python, which is not the case by default. TensorBoard can be installed 
 		// for this version of Python by going to your Unreal Editor Python Binaries directory 
-		// (e.g. "\Engine\Binaries\ThirdParty\Python3\Win64") and running `.\python -m pip install tensorboard`. 
+		// (e.g. "\Engine\Binaries\ThirdParty\Python3\Win64") and running `python -m pip install tensorboard`. 
 		bool bUseTensorboard = false;
 	};
 
+	/**
+	* PPOTrainer flags controlling some aspects of the process of communication with the trainer
+	*/
+	enum class EPPOTrainerFlags : uint8
+	{
+		None = 0,
+
+		// If to send over the initial provided policy network rather than reinitialize it from random weights at 
+		// the start of training. Use this if you want to start from a network which has already been trained 
+		// such as via Imitation Training.
+		UseInitialPolicyNetwork = 1 << 0,
+
+		// If to send over the initial provided critic network rather than reinitialize it from random weights at 
+		// the start of training. Use this if you want to start from a network which has already been trained 
+		// such as via Imitation Training.
+		UseInitialCriticNetwork = 1 << 1,
+
+		// If the critic network should be synchronized between the training process 
+		// and experience gathering process during training
+		SynchronizeCriticNetwork = 1 << 2,
+	};
+	ENUM_CLASS_FLAGS(EPPOTrainerFlags)
+
 	/**
 	* Interface for an object which can train a policy using PPO.
 	*/
@@ -184,6 +212,21 @@ namespace UE::Learning
 			FRWLock* NetworkLock = nullptr,
 			const ELogSetting LogSettings = Trainer::DefaultLogSettings) = 0;
 
+		/**
+		* Wait for the trainer to push an updated critic network.
+		*
+		* @param OutNetwork		Network to update
+		* @param Timeout		Timeout to wait in seconds
+		* @param NetworkLock	Lock to use when updating network
+		* @param LogSettings	Log settings
+		* @returns				Trainer response
+		*/
+		virtual ETrainerResponse RecvCritic(
+			FNeuralNetwork& OutNetwork,
+			const float Timeout = Trainer::DefaultTimeout,
+			FRWLock* NetworkLock = nullptr,
+			const ELogSetting LogSettings = Trainer::DefaultLogSettings) = 0;
+
 		/*
 		* Signal for the trainer to stop.
 		*/
@@ -204,6 +247,21 @@ namespace UE::Learning
 			FRWLock* NetworkLock = nullptr,
 			const ELogSetting LogSettings = Trainer::DefaultLogSettings) = 0;
 
+		/**
+		* Wait for the trainer to be ready and push the current critic network.
+		*
+		* @param Network		Network to push
+		* @param Timeout		Timeout to wait in seconds
+		* @param NetworkLock	Lock to use when pushing network
+		* @param LogSettings	Log settings
+		* @returns				Trainer response
+		*/
+		virtual ETrainerResponse SendCritic(
+			const FNeuralNetwork& Network,
+			const float Timeout = Trainer::DefaultTimeout,
+			FRWLock* NetworkLock = nullptr,
+			const ELogSetting LogSettings = Trainer::DefaultLogSettings) = 0;
+
 		/**
 		* Wait for the trainer to be ready and push new experience.
 		*
@@ -234,14 +292,33 @@ namespace UE::Learning
 		* @param SitePackagesPath			Path to the site-packages shipped with the PythonFoundationPackages plugin
 		* @param PythonContentPath			Path to the Python Content folder provided by the Learning plugin
 		* @param IntermediatePath			Path to the intermediate folder to write temporary files, logs, and snapshots to
-		* @param Network					Neural Network object
-		* @param ActionNoiseMin				Minimum Action Noise used by Policy during experience gathering
-		* @param ActionNoiseMax				Maximum Action Noise used by Policy during experience gathering
 		* @param ReplayBuffer				Replay buffer used to collect experience
-		* @param Settings					Trainer settings
+		* @param TrainingSettings			Trainer Training settings
+		* @param NetworkSettings			Trainer Network settings
+		* @param TrainerFlags				Flags for the trainer
 		* @param LogSettings				Logging settings to use
-		* @param bHideTrainingWindow		If to hide the training sub-process console window
-		* @param bRedirectTrainingOutput	If to redirect the training console window output to the log
+		* @param TrainingProcessFlags		Training subprocess flags
+
+		* @param ProcessNum					Number of processes to use for multi-processed experience gathering
+		*                                   
+		*                                   It is important to know how this multi-process training works so that it can be used
+		*                                   correctly when you set this >1:
+		*                                   
+		*                                   When called with this argument set >1, the process will spawn additional processes running
+		*                                   the same command as is currently being run but with the additional command line argument
+		*                                   `LearningProcessIdx`. Once the training starts on these sub-processes, this command
+		*                                   line argument will be used to indicate that this subprocess should only be used for gathering
+		*                                   experience and should not start a new training process itself.
+		*                                   
+		*                                   This means that this will generally not work in editor (or PIE) since it will end up spawning
+		*									Multiple additional editor windows. Instead, this should generally only be used in cooked, 
+		*									builds - and in most cases only when running headless.
+		*                                   
+		*                                   IMPORTANT: If you are seeding your experience generation process to make it deterministic
+		*                                   you should use this `LearningProcessIdx` command line argument to change the seed of your
+		*                                   experience gathering - otherwise each process will end up gathering identical experience!
+		* 
+		* @param MultiProcessFlags			Flags for the additional experience gathering subprocesses
 		*/
 		FSharedMemoryPPOTrainer(
 			const FString& TaskName,
@@ -249,14 +326,14 @@ namespace UE::Learning
 			const FString& SitePackagesPath,
 			const FString& PythonContentPath,
 			const FString& IntermediatePath,
-			const FNeuralNetwork& Network,
-			const float ActionNoiseMin,
-			const float ActionNoiseMax,
 			const FReplayBuffer& ReplayBuffer,
-			const FPPOTrainerSettings& Settings = FPPOTrainerSettings(),
+			const FPPOTrainerTrainingSettings& TrainingSettings = FPPOTrainerTrainingSettings(),
+			const FPPOTrainerNetworkSettings& NetworkSettings = FPPOTrainerNetworkSettings(),
+			const EPPOTrainerFlags TrainerFlags = EPPOTrainerFlags::None,
 			const ELogSetting LogSettings = ELogSetting::Normal,
-			const bool bHideTrainingWindow = true,
-			const bool bRedirectTrainingOutput = true);
+			const ESubprocessFlags TrainingProcessFlags = ESubprocessFlags::None,
+			const uint16 ProcessNum = 1,
+			const ESubprocessFlags MultiProcessFlags = ESubprocessFlags::None);
 
 		~FSharedMemoryPPOTrainer();
 
@@ -270,6 +347,12 @@ namespace UE::Learning
 			FRWLock* NetworkLock = nullptr,
 			const ELogSetting LogSettings = Trainer::DefaultLogSettings) override final;
 
+		virtual ETrainerResponse RecvCritic(
+			FNeuralNetwork& OutNetwork,
+			const float Timeout = Trainer::DefaultTimeout,
+			FRWLock* NetworkLock = nullptr,
+			const ELogSetting LogSettings = Trainer::DefaultLogSettings) override final;
+
 		virtual ETrainerResponse SendStop(const float Timeout = Trainer::DefaultTimeout) override final;
 
 		virtual ETrainerResponse SendPolicy(
@@ -278,6 +361,12 @@ namespace UE::Learning
 			FRWLock* NetworkLock = nullptr,
 			const ELogSetting LogSettings = Trainer::DefaultLogSettings) override final;
 
+		virtual ETrainerResponse SendCritic(
+			const FNeuralNetwork& Network,
+			const float Timeout = Trainer::DefaultTimeout,
+			FRWLock* NetworkLock = nullptr,
+			const ELogSetting LogSettings = Trainer::DefaultLogSettings) override final;
+
 		virtual ETrainerResponse SendExperience(
 			const FReplayBuffer& ReplayBuffer,
 			const float Timeout = Trainer::DefaultTimeout,
@@ -293,6 +382,7 @@ namespace UE::Learning
 		// Shared memory
 
 		UE::Learning::TSharedMemoryArrayView<1, uint8> Policy;
+		UE::Learning::TSharedMemoryArrayView<1, uint8> Critic;
 		UE::Learning::TSharedMemoryArrayView<2, volatile int32> Controls; // Mark as volatile to avoid compiler optimizing away reads without writes etc.
 		UE::Learning::TSharedMemoryArrayView<2, int32> EpisodeStarts;
 		UE::Learning::TSharedMemoryArrayView<2, int32> EpisodeLengths;
@@ -332,6 +422,7 @@ namespace UE::Learning
 		* @param IntermediatePath			Path to the intermediate folder to write temporary files, logs, and snapshots to
 		* @param IpAddress					Ip address to bind the listening socket to. For a local server you will want to use 127.0.0.1
 		* @param Port						Port to use for the listening socket.
+		* @param TrainingProcessFlags		Training server subprocess flags
 		* @param LogSettings				Logging settings to use
 		*/
 		FSocketPPOTrainerServerProcess(
@@ -341,6 +432,7 @@ namespace UE::Learning
 			const FString& IntermediatePath,
 			const TCHAR* IpAddress = Trainer::DefaultIp,
 			const uint32 Port = Trainer::DefaultPort,
+			const ESubprocessFlags TrainingProcessFlags = ESubprocessFlags::None,
 			const ELogSetting LogSettings = ELogSetting::Normal);
 
 		~FSocketPPOTrainerServerProcess();
@@ -385,26 +477,24 @@ namespace UE::Learning
 		*
 		* @param OutResponse				Response to the initial connection
 		* @param TaskName					Name of the training task - used to help identify the logs, snapshots, and other files generated by training
-		* @param Network					Neural Network object
-		* @param ActionNoiseMin				Minimum Action Noise used by Policy during experience gathering
-		* @param ActionNoiseMax				Maximum Action Noise used by Policy during experience gathering
 		* @param ReplayBuffer				Replay buffer used to collect experience
 		* @param IpAddress					Server Ip address
 		* @param Port						Server Port
 		* @param Timeout					Timeout to wait in seconds for connection and initial data transfer
-		* @param Settings					Trainer settings
+		* @param TrainingSettings			Trainer Training settings
+		* @param NetworkSettings			Trainer Network settings
+		* @param TrainerFlags				Flags for the trainer
 		*/
 		FSocketPPOTrainer(
 			ETrainerResponse& OutResponse,
 			const FString& TaskName,
-			const FNeuralNetwork& Network,
-			const float ActionNoiseMin,
-			const float ActionNoiseMax,
 			const FReplayBuffer& ReplayBuffer,
 			const TCHAR* IpAddress = Trainer::DefaultIp,
 			const uint32 Port = Trainer::DefaultPort,
 			const float Timeout = Trainer::DefaultTimeout,
-			const FPPOTrainerSettings& Settings = FPPOTrainerSettings());
+			const FPPOTrainerTrainingSettings& TrainingSettings = FPPOTrainerTrainingSettings(),
+			const FPPOTrainerNetworkSettings& NetworkSettings = FPPOTrainerNetworkSettings(),
+			const EPPOTrainerFlags TrainerFlags = EPPOTrainerFlags::None);
 
 		~FSocketPPOTrainer();
 
@@ -418,6 +508,12 @@ namespace UE::Learning
 			FRWLock* NetworkLock = nullptr,
 			const ELogSetting LogSettings = Trainer::DefaultLogSettings) override final;
 
+		virtual ETrainerResponse RecvCritic(
+			FNeuralNetwork& OutNetwork,
+			const float Timeout = Trainer::DefaultTimeout,
+			FRWLock* NetworkLock = nullptr,
+			const ELogSetting LogSettings = Trainer::DefaultLogSettings) override final;
+
 		virtual ETrainerResponse SendStop(const float Timeout = Trainer::DefaultTimeout) override final;
 
 		virtual ETrainerResponse SendPolicy(
@@ -426,6 +522,12 @@ namespace UE::Learning
 			FRWLock* NetworkLock = nullptr,
 			const ELogSetting LogSettings = Trainer::DefaultLogSettings) override final;
 
+		virtual ETrainerResponse SendCritic(
+			const FNeuralNetwork& Network,
+			const float Timeout = Trainer::DefaultTimeout,
+			FRWLock* NetworkLock = nullptr,
+			const ELogSetting LogSettings = Trainer::DefaultLogSettings) override final;
+
 		virtual ETrainerResponse SendExperience(
 			const FReplayBuffer& ReplayBuffer,
 			const float Timeout = Trainer::DefaultTimeout,
@@ -433,7 +535,8 @@ namespace UE::Learning
 
 	private:
 
-		TLearningArray<1, uint8> NetworkBuffer;
+		TLearningArray<1, uint8> PolicyNetworkBuffer;
+		TLearningArray<1, uint8> CriticNetworkBuffer;
 		FSocket* Socket = nullptr;
 	};
 
@@ -447,7 +550,8 @@ namespace UE::Learning
 		* @param ReplayBuffer					Replay Buffer
 		* @param EpisodeBuffer					Episode Buffer
 		* @param ResetBuffer					Reset Buffer
-		* @param Network						Network to use
+		* @param PolicyNetwork					Policy Network to use
+		* @param CriticNetwork					Optional Critic Network to use
 		* @param ObservationVectorBuffer		Buffer to read/write observation vectors into
 		* @param ActionVectorBuffer				Buffer to read/write action vectors into
 		* @param RewardBuffer					Buffer to read/write rewards into
@@ -461,10 +565,12 @@ namespace UE::Learning
 		* @param RewardFunction					Function to run for evaluating rewards
 		* @param CompletionFunction				Function to run for evaluating completions
 		* @param Instances						Set of instances to run training for
-		* @param bReinitializeNetwork			Was the network reinitialized to random weights or used as-is
+		* @param TrainerFlags					Flags for the trainer, should match what was used to initialize the Trainer object.
 		* @param bRequestTrainingStopSignal		Optional signal that can be set to indicate training should be stopped
-		* @param NetworkLock					Optional Lock to use when updating the policy network
-		* @param bNetworkUpdatedSignal			Optional signal that will be set when the policy network is updated
+		* @param PolicyNetworkLock				Optional Lock to use when updating the policy network
+		* @param CriticNetworkLock				Optional Lock to use when updating the critic network
+		* @param bPolicyNetworkUpdatedSignal	Optional signal that will be set when the policy network is updated
+		* @param bCriticNetworkUpdatedSignal	Optional signal that will be set when the critic network is updated
 		* @param LogSettings					Logging settings
 		* @returns								Trainer response in case of errors during communication otherwise Success
 		*/
@@ -473,7 +579,8 @@ namespace UE::Learning
 			FReplayBuffer& ReplayBuffer,
 			FEpisodeBuffer& EpisodeBuffer,
 			FResetInstanceBuffer& ResetBuffer,
-			FNeuralNetwork& Network,
+			FNeuralNetwork& PolicyNetwork,
+			FNeuralNetwork* CriticNetwork,
 			TLearningArrayView<2, float> ObservationVectorBuffer,
 			TLearningArrayView<2, float> ActionVectorBuffer,
 			TLearningArrayView<1, float> RewardBuffer,
@@ -487,10 +594,12 @@ namespace UE::Learning
 			const TFunctionRef<void(const FIndexSet Instances)> RewardFunction,
 			const TFunctionRef<void(const FIndexSet Instances)> CompletionFunction,
 			const FIndexSet Instances,
-			const bool bReinitializeNetwork = true,
+			const EPPOTrainerFlags TrainerFlags = EPPOTrainerFlags::None,
 			TAtomic<bool>* bRequestTrainingStopSignal = nullptr,
-			FRWLock* NetworkLock = nullptr,
-			TAtomic<bool>* bNetworkUpdatedSignal = nullptr,
+			FRWLock* PolicyNetworkLock = nullptr,
+			FRWLock* CriticNetworkLock = nullptr,
+			TAtomic<bool>* bPolicyNetworkUpdatedSignal = nullptr,
+			TAtomic<bool>* bCriticNetworkUpdatedSignal = nullptr,
 			const ELogSetting LogSettings = ELogSetting::Normal);
 	}
 
diff --git a/Engine/Plugins/Experimental/LearningAgents/Source/LearningTraining/Public/LearningTrainer.h b/Engine/Plugins/Experimental/LearningAgents/Source/LearningTraining/Public/LearningTrainer.h
--- a/Engine/Plugins/Experimental/LearningAgents/Source/LearningTraining/Public/LearningTrainer.h
+++ b/Engine/Plugins/Experimental/LearningAgents/Source/LearningTraining/Public/LearningTrainer.h
@@ -37,6 +37,21 @@ namespace UE::Learning
 		Timeout = 4,
 	};
 
+	/**
+	* Subprocess flags
+	*/
+	enum class ESubprocessFlags : uint8
+	{
+		None = 0,
+
+		// If to show the sub-process console window
+		ShowWindow = 1 << 0,
+
+		// If to avoid redirecting the sub-process output to the output log
+		NoRedirectOutput = 1 << 1,
+	};
+	ENUM_CLASS_FLAGS(ESubprocessFlags)
+
 	namespace Trainer
 	{
 		/**
@@ -67,7 +82,7 @@ namespace UE::Learning
 		/**
 		* Converts a ETrainerResponse into a string for use in logging and error messages.
 		*/
-		LEARNINGTRAINING_API const TCHAR* ResponseString(const ETrainerResponse Response);
+		LEARNINGTRAINING_API const TCHAR* GetResponseString(const ETrainerResponse Response);
 
 		/**
 		* Compute the discount factor that corresponds to a particular HalfLife and DeltaTime.
