LearningAgents: Updated MultiArray - removed some unused functions and made TConstMultiArrayView.

diff --git a/Engine/Plugins/Experimental/LearningAgents/Source/Learning/Public/LearningArray.h b/Engine/Plugins/Experimental/LearningAgents/Source/Learning/Public/LearningArray.h
--- a/Engine/Plugins/Experimental/LearningAgents/Source/Learning/Public/LearningArray.h
+++ b/Engine/Plugins/Experimental/LearningAgents/Source/Learning/Public/LearningArray.h
@@ -116,6 +116,13 @@ using TLearningArrayView = TMultiArrayView<
 	UE_LEARNING_ARRAY_CHECK_INDEX, 
 	UE_LEARNING_ARRAY_RESTRICT>;
 
+template<uint32 DimNum, typename ElementType>
+using TLearningConstArrayView = TConstMultiArrayView<
+	DimNum,
+	ElementType,
+	UE_LEARNING_ARRAY_CHECK_INDEX,
+	UE_LEARNING_ARRAY_RESTRICT>;
+
 template<uint32 DimNum, typename ElementType, typename Allocator = FDefaultAllocator>
 using TLearningArray = TMultiArray<
 	DimNum, 
@@ -125,6 +132,7 @@ using TLearningArray = TMultiArray<
 	UE_LEARNING_ARRAY_RESTRICT>;
 
 
+
 // Learning Array Functions
 
 namespace UE::Learning
diff --git a/Engine/Plugins/Experimental/LearningAgents/Source/Learning/Public/MultiArray.h b/Engine/Plugins/Experimental/LearningAgents/Source/Learning/Public/MultiArray.h
--- a/Engine/Plugins/Experimental/LearningAgents/Source/Learning/Public/MultiArray.h
+++ b/Engine/Plugins/Experimental/LearningAgents/Source/Learning/Public/MultiArray.h
@@ -521,33 +521,6 @@ public:
 		}
 	}
 
-	/**
-	 * Tests if dimension is valid, i.e. less than the number of dimensions.
-	 *
-	 * @param Dimension Dimension of the array.
-	 *
-	 * @returns True if dimension is valid. False otherwise.
-	 */
-	FORCEINLINE bool IsValidDimension(uint8 Dimension) const
-	{
-		return Dimension < DimNum;
-	}
-
-	/**
-	 * Tests if index is valid, i.e. greater than or equal to zero, and less than the number of elements in the dimension.
-	 *
-	 * @param Dimension Dimension of the array.
-	 * @param Index Index to test.
-	 *
-	 * @returns True if index is valid. False otherwise.
-	 */
-	FORCEINLINE bool IsValidIndex(uint8 Dimension, SizeType Index) const
-	{
-		DimensionCheck(Dimension);
-
-		return (Index >= 0) && (Index < ArrayShape[Dimension]);
-	}
-
 	/**
 	 * Returns true if the array is empty and contains no elements.
 	 *
@@ -1204,33 +1177,6 @@ public:
 		}
 	}
 
-	/**
-	 * Tests if dimension is valid, i.e. less than the number of dimensions.
-	 *
-	 * @param Dimension Dimension of the array.
-	 *
-	 * @returns True if dimension is valid. False otherwise.
-	 */
-	FORCEINLINE bool IsValidDimension(uint8 Dimension) const
-	{
-		return Dimension < DimNum;
-	}
-
-	/**
-	 * Tests if index is valid, i.e. greater than or equal to zero, and less than the number of elements in the dimension.
-	 *
-	 * @param Dimension Dimension of the array.
-	 * @param Index Index to test.
-	 *
-	 * @returns True if index is valid. False otherwise.
-	 */
-	FORCEINLINE bool IsValidIndex(uint8 Dimension, SizeType Index) const
-	{
-		DimensionCheck(Dimension);
-
-		return (Index >= 0) && (Index < ArrayShape[Dimension]);
-	}
-
 	/**
 	 * Returns true if the array is empty and contains no elements.
 	 *
diff --git a/Engine/Plugins/Experimental/LearningAgents/Source/Learning/Public/MultiArrayView.h b/Engine/Plugins/Experimental/LearningAgents/Source/Learning/Public/MultiArrayView.h
--- a/Engine/Plugins/Experimental/LearningAgents/Source/Learning/Public/MultiArrayView.h
+++ b/Engine/Plugins/Experimental/LearningAgents/Source/Learning/Public/MultiArrayView.h
@@ -410,33 +410,6 @@ public:
 		}
 	}
 
-	/**
-	 * Tests if dimension is valid, i.e. less than the number of dimensions.
-	 *
-	 * @param Dimension Dimension of the array.
-	 *
-	 * @returns True if dimension is valid. False otherwise.
-	 */
-	FORCEINLINE bool IsValidDimension(uint8 Dimension) const
-	{
-		return Dimension < DimNum;
-	}
-
-	/**
-	 * Tests if index is valid, i.e. greater than or equal to zero, and less than the number of elements in the dimension.
-	 *
-	 * @param Dimension Dimension of the array.
-	 * @param Index Index to test.
-	 *
-	 * @returns True if index is valid. False otherwise.
-	 */
-	FORCEINLINE bool IsValidIndex(uint8 Dimension, SizeType Index) const
-	{
-		DimensionCheck(Dimension);
-
-		return (Index >= 0) && (Index < ArrayShape[Dimension]);
-	}
-
 	/**
 	 * Returns true if the array is empty and contains no elements.
 	 *
@@ -550,93 +523,6 @@ public:
 		return TMultiArrayView(DataPtr + Index * Stride(0), NewShape);
 	}
 
-	/** Returns the left-most part of the view by taking the given number of elements from the left. */
-	[[nodiscard]] inline TMultiArrayView Left(SizeType Count) const
-	{
-		TMultiArrayShape NewShape = ArrayShape;
-		NewShape[0] = FMath::Clamp(Count, 0, ArrayShape[0]);
-		return TMultiArrayView(DataPtr, NewShape);
-	}
-
-	/** Returns the left-most part of the view by chopping the given number of elements from the right. */
-	[[nodiscard]] inline TMultiArrayView LeftChop(SizeType Count) const
-	{
-		TMultiArrayShape NewShape = ArrayShape;
-		NewShape[0] = FMath::Clamp(ArrayShape[0] - Count, 0, ArrayShape[0]);
-		return TMultiArrayView(DataPtr, NewShape);
-	}
-
-	/** Returns the right-most part of the view by taking the given number of elements from the right. */
-	[[nodiscard]] inline TMultiArrayView Right(SizeType Count) const
-	{
-		const SizeType OutLen = FMath::Clamp(Count, 0, ArrayShape[0]);
-		TMultiArrayShape NewShape = ArrayShape;
-		NewShape[0] = OutLen;
-		return TMultiArrayView(DataPtr + (ArrayShape[0] - OutLen) * Stride(0), NewShape);
-	}
-
-	/** Returns the right-most part of the view by chopping the given number of elements from the left. */
-	[[nodiscard]] inline TMultiArrayView RightChop(SizeType Count) const
-	{
-		const SizeType OutLen = FMath::Clamp(ArrayShape[0] - Count, 0, ArrayShape[0]);
-		TMultiArrayShape NewShape = ArrayShape;
-		NewShape[0] = OutLen;
-		return TMultiArrayView(DataPtr + (ArrayShape[0] - OutLen) * Stride(0), NewShape);
-	}
-
-	/** Returns the middle part of the view by taking up to the given number of elements from the given position. */
-	[[nodiscard]] inline TMultiArrayView Mid(SizeType Index, SizeType Count = TNumericLimits<SizeType>::Max()) const
-	{
-		ElementType* const CurrentStart = DataPtr;
-		const SizeType     CurrentLength = Num(0);
-
-		// Clamp minimum index at the start of the range, adjusting the length down if necessary
-		const SizeType NegativeIndexOffset = (Index < 0) ? Index : 0;
-		Count += NegativeIndexOffset;
-		Index -= NegativeIndexOffset;
-
-		// Clamp maximum index at the end of the range
-		Index = (Index > CurrentLength) ? CurrentLength : Index;
-
-		// Clamp count between 0 and the distance to the end of the range
-		Count = FMath::Clamp(Count, 0, (CurrentLength - Index));
-
-		TMultiArrayShape NewShape = ArrayShape;
-		NewShape[0] = Count;
-		TMultiArrayView Result = TMultiArrayView(CurrentStart + Index * Stride(0), NewShape);
-		return Result;
-	}
-
-	/** Modifies the view to be the given number of elements from the left. */
-	inline void LeftInline(SizeType CharCount)
-	{
-		*this = Left(CharCount);
-	}
-
-	/** Modifies the view by chopping the given number of elements from the right. */
-	inline void LeftChopInline(SizeType CharCount)
-	{
-		*this = LeftChop(CharCount);
-	}
-
-	/** Modifies the view to be the given number of elements from the right. */
-	inline void RightInline(SizeType CharCount)
-	{
-		*this = Right(CharCount);
-	}
-
-	/** Modifies the view by chopping the given number of elements from the left. */
-	inline void RightChopInline(SizeType CharCount)
-	{
-		*this = RightChop(CharCount);
-	}
-
-	/** Modifies the view to be the middle part by taking up to the given number of elements from the given position. */
-	inline void MidInline(SizeType Position, SizeType CharCount = TNumericLimits<SizeType>::Max())
-	{
-		*this = Mid(Position, CharCount);
-	}
-
 	/**
 	 * Checks if this array contains the element.
 	 *
@@ -949,33 +835,6 @@ public:
 		}
 	}
 
-	/**
-	 * Tests if dimension is valid, i.e. less than the number of dimensions.
-	 *
-	 * @param Dimension Dimension of the array.
-	 *
-	 * @returns True if dimension is valid. False otherwise.
-	 */
-	FORCEINLINE bool IsValidDimension(uint8 Dimension) const
-	{
-		return Dimension < DimNum;
-	}
-
-	/**
-	 * Tests if index is valid, i.e. greater than or equal to zero, and less than the number of elements in the dimension.
-	 *
-	 * @param Dimension Dimension of the array.
-	 * @param Index Index to test.
-	 *
-	 * @returns True if index is valid. False otherwise.
-	 */
-	FORCEINLINE bool IsValidIndex(uint8 Dimension, SizeType Index) const
-	{
-		DimensionCheck(Dimension);
-
-		return (Index >= 0) && (Index < ArrayShape[Dimension]);
-	}
-
 	/**
 	 * Returns true if the array is empty and contains no elements.
 	 *
@@ -1089,93 +948,6 @@ public:
 		return TMultiArrayView(DataPtr + Index * Stride(0), NewShape);
 	}
 
-	/** Returns the left-most part of the view by taking the given number of elements from the left. */
-	[[nodiscard]] inline TMultiArrayView Left(SizeType Count) const
-	{
-		TMultiArrayShape NewShape = ArrayShape;
-		NewShape[0] = FMath::Clamp(Count, 0, ArrayShape[0]);
-		return TMultiArrayView(DataPtr, NewShape);
-	}
-
-	/** Returns the left-most part of the view by chopping the given number of elements from the right. */
-	[[nodiscard]] inline TMultiArrayView LeftChop(SizeType Count) const
-	{
-		TMultiArrayShape NewShape = ArrayShape;
-		NewShape[0] = FMath::Clamp(ArrayShape[0] - Count, 0, ArrayShape[0]);
-		return TMultiArrayView(DataPtr, NewShape);
-	}
-
-	/** Returns the right-most part of the view by taking the given number of elements from the right. */
-	[[nodiscard]] inline TMultiArrayView Right(SizeType Count) const
-	{
-		const SizeType OutLen = FMath::Clamp(Count, 0, ArrayShape[0]);
-		TMultiArrayShape NewShape = ArrayShape;
-		NewShape[0] = OutLen;
-		return TMultiArrayView(DataPtr + (ArrayShape[0] - OutLen) * Stride(0), NewShape);
-	}
-
-	/** Returns the right-most part of the view by chopping the given number of elements from the left. */
-	[[nodiscard]] inline TMultiArrayView RightChop(SizeType Count) const
-	{
-		const SizeType OutLen = FMath::Clamp(ArrayShape[0] - Count, 0, ArrayShape[0]);
-		TMultiArrayShape NewShape = ArrayShape;
-		NewShape[0] = OutLen;
-		return TMultiArrayView(DataPtr + (ArrayShape[0] - OutLen) * Stride(0), NewShape);
-	}
-
-	/** Returns the middle part of the view by taking up to the given number of elements from the given position. */
-	[[nodiscard]] inline TMultiArrayView Mid(SizeType Index, SizeType Count = TNumericLimits<SizeType>::Max()) const
-	{
-		ElementType* const CurrentStart = DataPtr;
-		const SizeType     CurrentLength = Num(0);
-
-		// Clamp minimum index at the start of the range, adjusting the length down if necessary
-		const SizeType NegativeIndexOffset = (Index < 0) ? Index : 0;
-		Count += NegativeIndexOffset;
-		Index -= NegativeIndexOffset;
-
-		// Clamp maximum index at the end of the range
-		Index = (Index > CurrentLength) ? CurrentLength : Index;
-
-		// Clamp count between 0 and the distance to the end of the range
-		Count = FMath::Clamp(Count, 0, (CurrentLength - Index));
-
-		TMultiArrayShape NewShape = ArrayShape;
-		NewShape[0] = Count;
-		TMultiArrayView Result = TMultiArrayView(CurrentStart + Index * Stride(0), NewShape);
-		return Result;
-	}
-
-	/** Modifies the view to be the given number of elements from the left. */
-	inline void LeftInline(SizeType CharCount)
-	{
-		*this = Left(CharCount);
-	}
-
-	/** Modifies the view by chopping the given number of elements from the right. */
-	inline void LeftChopInline(SizeType CharCount)
-	{
-		*this = LeftChop(CharCount);
-	}
-
-	/** Modifies the view to be the given number of elements from the right. */
-	inline void RightInline(SizeType CharCount)
-	{
-		*this = Right(CharCount);
-	}
-
-	/** Modifies the view by chopping the given number of elements from the left. */
-	inline void RightChopInline(SizeType CharCount)
-	{
-		*this = RightChop(CharCount);
-	}
-
-	/** Modifies the view to be the middle part by taking up to the given number of elements from the given position. */
-	inline void MidInline(SizeType Position, SizeType CharCount = TNumericLimits<SizeType>::Max())
-	{
-		*this = Mid(Position, CharCount);
-	}
-
 	/**
 	 * Checks if this array contains the element.
 	 *
@@ -1249,3 +1021,6 @@ private:
 	PointerType DataPtr;
 	TMultiArrayShape<DimNum, SizeType> ArrayShape;
 };
+
+template<uint8 InDimNum, typename InElementType, bool bInIsChecked = true, bool bInIsRestrict = false, typename InSizeType = FDefaultAllocator::SizeType>
+using TConstMultiArrayView = TMultiArrayView<InDimNum, const InElementType, bInIsChecked, bInIsRestrict, InSizeType>;
