LearningAgents: Added ability to add partial, truncated episodes when replay buffer is full.

diff --git a/Engine/Plugins/Experimental/LearningAgents/Source/LearningTraining/Private/LearningExperience.cpp b/Engine/Plugins/Experimental/LearningAgents/Source/LearningTraining/Private/LearningExperience.cpp
--- a/Engine/Plugins/Experimental/LearningAgents/Source/LearningTraining/Private/LearningExperience.cpp
+++ b/Engine/Plugins/Experimental/LearningAgents/Source/LearningTraining/Private/LearningExperience.cpp
@@ -108,20 +108,21 @@ namespace UE::Learning
 		const TLearningArrayView<1, const ECompletionMode> InEpisodeCompletionModes,
 		const TLearningArrayView<2, const float> InEpisodeFinalObservations,
 		const FEpisodeBuffer& EpisodeBuffer,
-		const FIndexSet Instances)
+		const FIndexSet Instances,
+		const bool bAddTruncatedEpisodeWhenFull)
 	{
 		UE_LEARNING_TRACE_CPUPROFILER_EVENT_SCOPE(Learning::FReplayBuffer::AddEpisodes);
 
 		for (const int32 InstanceIdx : Instances)
 		{
-			const ECompletionMode EpisodeCompletion = InEpisodeCompletionModes[InstanceIdx];
+			UE_LEARNING_CHECKF(InEpisodeCompletionModes[InstanceIdx] != ECompletionMode::Running,
+				TEXT("Tried to add experience from episode that is still running..."));
+
 			const int32 EpisodeStepNum = EpisodeBuffer.GetEpisodeStepNums()[InstanceIdx];
 
-			UE_LEARNING_CHECKF(EpisodeCompletion != ECompletionMode::Running,
-				TEXT("Tried to add experience from episode that is still running..."));
+			// Is there space for the full episode in the buffer?
 
-			// If there is still space in the buffer...
-			if (EpisodeNum + 1 <= MaxEpisodeNum && StepNum + EpisodeStepNum <= MaxStepNum)
+			if (EpisodeNum < MaxEpisodeNum && StepNum + EpisodeStepNum <= MaxStepNum)
 			{
 				// Copy the data into the replay buffer
 				Array::Copy(Observations.Slice(StepNum, EpisodeStepNum), EpisodeBuffer.GetObservations(InstanceIdx));
@@ -131,22 +132,45 @@ namespace UE::Learning
 				// Write the Episode start, length, completion, and final observation
 				EpisodeStarts[EpisodeNum] = StepNum;
 				EpisodeLengths[EpisodeNum] = EpisodeStepNum;
-				EpisodeCompletionModes[EpisodeNum] = EpisodeCompletion;
+				EpisodeCompletionModes[EpisodeNum] = InEpisodeCompletionModes[InstanceIdx];
 				Array::Copy(EpisodeFinalObservations[EpisodeNum], InEpisodeFinalObservations[InstanceIdx]);
 
 				// Increment the Counts
 				EpisodeNum++;
 				StepNum += EpisodeStepNum;
+
+				// Continue onto next Episode
+				continue;
 			}
-			else
+
+			// Is there space for a partial episode in the buffer?
+
+			if (bAddTruncatedEpisodeWhenFull && EpisodeNum < MaxEpisodeNum && StepNum < MaxStepNum)
 			{
-				// Buffer is full
-				return true;
+				const int32 PartialStepNum = MaxStepNum - StepNum;
+				UE_LEARNING_CHECK(PartialStepNum > 0 && PartialStepNum < EpisodeStepNum);
+
+				// Copy the data into the replay buffer
+				Array::Copy(Observations.Slice(StepNum, PartialStepNum), EpisodeBuffer.GetObservations(InstanceIdx).Slice(0, PartialStepNum));
+				Array::Copy(Actions.Slice(StepNum, PartialStepNum), EpisodeBuffer.GetActions(InstanceIdx).Slice(0, PartialStepNum));
+				Array::Copy(Rewards.Slice(StepNum, PartialStepNum), EpisodeBuffer.GetRewards(InstanceIdx).Slice(0, PartialStepNum));
+
+				// Write the Episode start, length, completion, and final observation
+				EpisodeStarts[EpisodeNum] = StepNum;
+				EpisodeLengths[EpisodeNum] = PartialStepNum;
+				EpisodeCompletionModes[EpisodeNum] = ECompletionMode::Truncated;
+				Array::Copy(EpisodeFinalObservations[EpisodeNum], InEpisodeFinalObservations[InstanceIdx]);
+
+				// Increment the Counts
+				EpisodeNum++;
+				StepNum += PartialStepNum;
 			}
+
+			// Otherwise buffer is full
+			return true;
 		}
 
-		// Buffer is not full
-		return false;
+		return (EpisodeNum == MaxEpisodeNum) || (StepNum == MaxStepNum);
 	}
 
 	const int32 FReplayBuffer::GetMaxEpisodeNum() const
diff --git a/Engine/Plugins/Experimental/LearningAgents/Source/LearningTraining/Public/LearningExperience.h b/Engine/Plugins/Experimental/LearningAgents/Source/LearningTraining/Public/LearningExperience.h
--- a/Engine/Plugins/Experimental/LearningAgents/Source/LearningTraining/Public/LearningExperience.h
+++ b/Engine/Plugins/Experimental/LearningAgents/Source/LearningTraining/Public/LearningExperience.h
@@ -98,6 +98,7 @@ namespace UE::Learning
 		* @param InEpisodeFinalObservations		Array of final observations for each instance of shape (InstanceNum, ObservationVectorDimNum)
 		* @param EpisodeBuffer					Episode buffer to add experience from
 		* @param Instances						Instances to add
+		* @param bAddTruncatedEpisodeWhenFull	When enabled, this will add a truncated, partial episode to the buffer when full
 		*
 		* @returns								True when the replay buffer is full
 		*/
@@ -105,7 +106,8 @@ namespace UE::Learning
 			const TLearningArrayView<1, const ECompletionMode> InEpisodeCompletionModes,
 			const TLearningArrayView<2, const float> InEpisodeFinalObservations,
 			const FEpisodeBuffer& EpisodeBuffer,
-			const FIndexSet Instances);
+			const FIndexSet Instances,
+			const bool bAddTruncatedEpisodeWhenFull = true);
 
 		const int32 GetMaxEpisodeNum() const;
 		const int32 GetMaxStepNum() const;
