Learning Agents: Add code comments to all major items and const everything per the coding standards.

diff --git a/Engine/Plugins/Experimental/LearningAgents/Source/LearningAgents/Private/LearningAgentsActions.cpp b/Engine/Plugins/Experimental/LearningAgents/Source/LearningAgents/Private/LearningAgentsActions.cpp
--- a/Engine/Plugins/Experimental/LearningAgents/Source/LearningAgents/Private/LearningAgentsActions.cpp
+++ b/Engine/Plugins/Experimental/LearningAgents/Source/LearningAgents/Private/LearningAgentsActions.cpp
@@ -27,7 +27,7 @@
 namespace UE::Learning::Agents::Private
 {
 	template<typename ActionUObject, typename ActionFObject, typename... InArgTypes>
-	ActionUObject* AddAction(ULearningAgentsType* AgentType, FName Name, InArgTypes&& ...Args)
+	ActionUObject* AddAction(ULearningAgentsType* AgentType, const FName Name, InArgTypes&& ...Args)
 	{
 		if (!AgentType)
 		{
@@ -51,12 +51,12 @@ namespace UE::Learning::Agents::Private
 
 //------------------------------------------------------------------
 
-UFloatAction* UFloatAction::AddFloatAction(ULearningAgentsType* AgentType, FName Name, float Scale)
+UFloatAction* UFloatAction::AddFloatAction(ULearningAgentsType* AgentType, const FName Name, const float Scale)
 {
 	return UE::Learning::Agents::Private::AddAction<UFloatAction, UE::Learning::FFloatFeature>(AgentType, Name, 1, Scale);
 }
 
-float UFloatAction::GetFloatAction(int32 AgentId)
+float UFloatAction::GetFloatAction(const int32 AgentId)
 {
 	const TLearningArrayView<2, const float> View = FeatureObject->InstanceData->ConstView(FeatureObject->ValueHandle);
 
@@ -75,9 +75,9 @@ float UFloatAction::GetFloatAction(int32 AgentId)
 	return View[AgentId][0];
 }
 
-void UFloatAction::SetFloatAction(int32 AgentId, float Value)
+void UFloatAction::SetFloatAction(const int32 AgentId, const float Value)
 {
-	TLearningArrayView<2, float> View = FeatureObject->InstanceData->View(FeatureObject->ValueHandle);
+	const TLearningArrayView<2, float> View = FeatureObject->InstanceData->View(FeatureObject->ValueHandle);
 
 	if (AgentId == INDEX_NONE)
 	{
@@ -124,12 +124,12 @@ void UFloatAction::VisualLog(const UE::Learning::FIndexSet Instances) const
 
 //------------------------------------------------------------------
 
-UVectorAction* UVectorAction::AddVectorAction(ULearningAgentsType* AgentType, FName Name, float Scale)
+UVectorAction* UVectorAction::AddVectorAction(ULearningAgentsType* AgentType, const FName Name, const float Scale)
 {
 	return UE::Learning::Agents::Private::AddAction<UVectorAction, UE::Learning::FFloatFeature>(AgentType, Name, 3, Scale);
 }
 
-FVector UVectorAction::GetVectorAction(int32 AgentId)
+FVector UVectorAction::GetVectorAction(const int32 AgentId)
 {
 	const TLearningArrayView<2, const float> View = FeatureObject->InstanceData->ConstView(FeatureObject->ValueHandle);
 
@@ -148,9 +148,9 @@ FVector UVectorAction::GetVectorAction(int32 AgentId)
 	return FVector(View[AgentId][0], View[AgentId][1], View[AgentId][2]);
 }
 
-void UVectorAction::SetVectorAction(int32 AgentId, FVector InAction)
+void UVectorAction::SetVectorAction(const int32 AgentId, const FVector InAction)
 {
-	TLearningArrayView<2, float> View = FeatureObject->InstanceData->View(FeatureObject->ValueHandle);
+	const TLearningArrayView<2, float> View = FeatureObject->InstanceData->View(FeatureObject->ValueHandle);
 
 	if (AgentId == INDEX_NONE)
 	{
@@ -212,12 +212,12 @@ void UVectorAction::VisualLog(const UE::Learning::FIndexSet Instances) const
 
 //------------------------------------------------------------------
 
-URotationVectorArrayAction* URotationVectorArrayAction::AddRotationVectorArrayAction(ULearningAgentsType* AgentType, FName Name, int32 RotationVectorNum, float Scale)
+URotationVectorArrayAction* URotationVectorArrayAction::AddRotationVectorArrayAction(ULearningAgentsType* AgentType, const FName Name, const int32 RotationVectorNum, const float Scale)
 {
 	return UE::Learning::Agents::Private::AddAction<URotationVectorArrayAction, UE::Learning::FRotationVectorFeature>(AgentType, Name, RotationVectorNum, FMath::DegreesToRadians(Scale));
 }
 
-void URotationVectorArrayAction::GetRotationVectorArrayAction(int32 AgentId, TArray<FVector>& OutRotationVectors)
+void URotationVectorArrayAction::GetRotationVectorArrayAction(const int32 AgentId, TArray<FVector>& OutRotationVectors)
 {
 	const TLearningArrayView<2, const FVector> View = FeatureObject->InstanceData->ConstView(FeatureObject->RotationVectorsHandle);
 
@@ -240,7 +240,7 @@ void URotationVectorArrayAction::GetRotationVectorArrayAction(int32 AgentId, TAr
 	}
 }
 
-void URotationVectorArrayAction::GetRotationVectorArrayActionAsQuats(int32 AgentId, TArray<FQuat>& OutRotations)
+void URotationVectorArrayAction::GetRotationVectorArrayActionAsQuats(const int32 AgentId, TArray<FQuat>& OutRotations)
 {
 	const TLearningArrayView<2, const FVector> View = FeatureObject->InstanceData->ConstView(FeatureObject->RotationVectorsHandle);
 
@@ -263,7 +263,7 @@ void URotationVectorArrayAction::GetRotationVectorArrayActionAsQuats(int32 Agent
 	}
 }
 
-void URotationVectorArrayAction::GetRotationVectorArrayActionAsRotators(int32 AgentId, TArray<FRotator>& OutRotations)
+void URotationVectorArrayAction::GetRotationVectorArrayActionAsRotators(const int32 AgentId, TArray<FRotator>& OutRotations)
 {
 	const TLearningArrayView<2, const FVector> View = FeatureObject->InstanceData->ConstView(FeatureObject->RotationVectorsHandle);
 
@@ -356,4 +356,4 @@ void URotationVectorArrayAction::VisualLog(const UE::Learning::FIndexSet Instanc
 
 #undef UE_LEARNING_AGENTS_VLOG_STRING
 #undef UE_LEARNING_AGENTS_VLOG_TRANSFORM
-#undef UE_LEARNING_AGENTS_VLOG_PLANE
\ No newline at end of file
+#undef UE_LEARNING_AGENTS_VLOG_PLANE
diff --git a/Engine/Plugins/Experimental/LearningAgents/Source/LearningAgents/Private/LearningAgentsController.cpp b/Engine/Plugins/Experimental/LearningAgents/Source/LearningAgents/Private/LearningAgentsController.cpp
--- a/Engine/Plugins/Experimental/LearningAgents/Source/LearningAgents/Private/LearningAgentsController.cpp
+++ b/Engine/Plugins/Experimental/LearningAgents/Source/LearningAgents/Private/LearningAgentsController.cpp
@@ -46,7 +46,7 @@ bool ULearningAgentsController::IsControllerSetupPerformed() const
 	return AgentType ? true : false;
 }
 
-void ULearningAgentsController::AddAgent(int32 AgentId)
+void ULearningAgentsController::AddAgent(const int32 AgentId)
 {
 	if (!IsControllerSetupPerformed())
 	{
@@ -71,7 +71,7 @@ void ULearningAgentsController::AddAgent(int32 AgentId)
 	SelectedAgentsSet.TryMakeSlice();
 }
 
-void ULearningAgentsController::RemoveAgent(int32 AgentId)
+void ULearningAgentsController::RemoveAgent(const int32 AgentId)
 {
 	if (!IsControllerSetupPerformed())
 	{
@@ -89,7 +89,7 @@ void ULearningAgentsController::RemoveAgent(int32 AgentId)
 	SelectedAgentsSet.TryMakeSlice();
 }
 
-bool ULearningAgentsController::HasAgent(int32 AgentId) const
+bool ULearningAgentsController::HasAgent(const int32 AgentId) const
 {
 	return SelectedAgentsSet.Contains(AgentId);
 }
diff --git a/Engine/Plugins/Experimental/LearningAgents/Source/LearningAgents/Private/LearningAgentsCritic.cpp b/Engine/Plugins/Experimental/LearningAgents/Source/LearningAgents/Private/LearningAgentsCritic.cpp
--- a/Engine/Plugins/Experimental/LearningAgents/Source/LearningAgents/Private/LearningAgentsCritic.cpp
+++ b/Engine/Plugins/Experimental/LearningAgents/Source/LearningAgents/Private/LearningAgentsCritic.cpp
@@ -130,7 +130,7 @@ bool ULearningAgentsCritic::IsCriticSetupPerformed() const
 	return bCriticSetupPerformed;
 }
 
-void ULearningAgentsCritic::AddAgent(int32 AgentId)
+void ULearningAgentsCritic::AddAgent(const int32 AgentId)
 {
 	if (!IsCriticSetupPerformed())
 	{
@@ -155,7 +155,7 @@ void ULearningAgentsCritic::AddAgent(int32 AgentId)
 	SelectedAgentsSet.TryMakeSlice();
 }
 
-void ULearningAgentsCritic::RemoveAgent(int32 AgentId)
+void ULearningAgentsCritic::RemoveAgent(const int32 AgentId)
 {
 	if (!IsCriticSetupPerformed())
 	{
@@ -173,12 +173,12 @@ void ULearningAgentsCritic::RemoveAgent(int32 AgentId)
 	SelectedAgentsSet.TryMakeSlice();
 }
 
-bool ULearningAgentsCritic::HasAgent(int32 AgentId) const
+bool ULearningAgentsCritic::HasAgent(const int32 AgentId) const
 {
 	return SelectedAgentsSet.Contains(AgentId);
 }
 
-ULearningAgentsType* ULearningAgentsCritic::GetAgentType(TSubclassOf<ULearningAgentsType> AgentClass)
+ULearningAgentsType* ULearningAgentsCritic::GetAgentType(const TSubclassOf<ULearningAgentsType> AgentClass)
 {
 	if (!IsCriticSetupPerformed())
 	{
@@ -208,7 +208,7 @@ void ULearningAgentsCritic::LoadCriticFromSnapshot(const FDirectoryPath& Directo
 	}
 
 	TArray64<uint8> NetworkData;
-	FString FilePath = Directory.Path + FGenericPlatformMisc::GetDefaultPathSeparator() + Filename;
+	const FString FilePath = Directory.Path + FGenericPlatformMisc::GetDefaultPathSeparator() + Filename;
 	if (FFileHelper::LoadFileToArray(NetworkData, *FilePath))
 	{
 		const int32 TotalByteNum = UE::Learning::FNeuralNetwork::GetSerializationByteNum(
@@ -248,7 +248,7 @@ void ULearningAgentsCritic::SaveCriticToSnapshot(const FDirectoryPath& Directory
 
 	Network->NeuralNetwork->SerializeToBytes(NetworkData);
 
-	FString FilePath = Directory.Path + FGenericPlatformMisc::GetDefaultPathSeparator() + Filename;
+	const FString FilePath = Directory.Path + FGenericPlatformMisc::GetDefaultPathSeparator() + Filename;
 	if (!FFileHelper::SaveArrayToFile(NetworkData, *FilePath))
 	{
 		UE_LOG(LogLearning, Error, TEXT("Failed to save network to file: %s"), *FilePath);
@@ -361,4 +361,4 @@ void ULearningAgentsCritic::VisualLog(const UE::Learning::FIndexSet Instances) c
 }
 #endif
 
-#undef UE_LEARNING_AGENTS_VLOG_STRING
\ No newline at end of file
+#undef UE_LEARNING_AGENTS_VLOG_STRING
diff --git a/Engine/Plugins/Experimental/LearningAgents/Source/LearningAgents/Private/LearningAgentsObservations.cpp b/Engine/Plugins/Experimental/LearningAgents/Source/LearningAgents/Private/LearningAgentsObservations.cpp
--- a/Engine/Plugins/Experimental/LearningAgents/Source/LearningAgents/Private/LearningAgentsObservations.cpp
+++ b/Engine/Plugins/Experimental/LearningAgents/Source/LearningAgents/Private/LearningAgentsObservations.cpp
@@ -27,7 +27,7 @@
 namespace UE::Learning::Agents::Private
 {
 	template<typename ObservationUObject, typename ObservationFObject, typename... InArgTypes>
-	ObservationUObject* AddObservation(ULearningAgentsType* AgentType, FName Name, InArgTypes&& ...Args)
+	ObservationUObject* AddObservation(ULearningAgentsType* AgentType, const FName Name, InArgTypes&& ...Args)
 	{
 		if (!AgentType)
 		{
@@ -51,14 +51,14 @@ namespace UE::Learning::Agents::Private
 
 //------------------------------------------------------------------
 
-UFloatObservation* UFloatObservation::AddFloatObservation(ULearningAgentsType* AgentType, FName Name, float Scale)
+UFloatObservation* UFloatObservation::AddFloatObservation(ULearningAgentsType* AgentType, const FName Name, const float Scale)
 {
 	return UE::Learning::Agents::Private::AddObservation<UFloatObservation, UE::Learning::FFloatFeature>(AgentType, Name, 1, Scale);
 }
 
-void UFloatObservation::SetFloatObservation(int32 AgentId, float Value)
+void UFloatObservation::SetFloatObservation(const int32 AgentId, const float Value)
 {
-	TLearningArrayView<2, float> View = FeatureObject->InstanceData->View(FeatureObject->ValueHandle);
+	const TLearningArrayView<2, float> View = FeatureObject->InstanceData->View(FeatureObject->ValueHandle);
 
 	if (AgentId == INDEX_NONE)
 	{
@@ -105,14 +105,14 @@ void UFloatObservation::VisualLog(const UE::Learning::FIndexSet Instances) const
 
 //------------------------------------------------------------------
 
-UVectorObservation* UVectorObservation::AddVectorObservation(ULearningAgentsType* AgentType, FName Name, float Scale)
+UVectorObservation* UVectorObservation::AddVectorObservation(ULearningAgentsType* AgentType, const FName Name, const float Scale)
 {
 	return UE::Learning::Agents::Private::AddObservation<UVectorObservation, UE::Learning::FFloatFeature>(AgentType, Name, 3, Scale);
 }
 
-void UVectorObservation::SetVectorObservation(int32 AgentId, FVector Vector)
+void UVectorObservation::SetVectorObservation(const int32 AgentId, const FVector Vector)
 {
-	TLearningArrayView<2, float> View = FeatureObject->InstanceData->View(FeatureObject->ValueHandle);
+	const TLearningArrayView<2, float> View = FeatureObject->InstanceData->View(FeatureObject->ValueHandle);
 
 	if (AgentId == INDEX_NONE)
 	{
@@ -174,15 +174,15 @@ void UVectorObservation::VisualLog(const UE::Learning::FIndexSet Instances) cons
 
 //------------------------------------------------------------------
 
-UAngleObservation* UAngleObservation::AddAngleObservation(ULearningAgentsType* AgentType, FName Name, float Scale)
+UAngleObservation* UAngleObservation::AddAngleObservation(ULearningAgentsType* AgentType, const FName Name, const float Scale)
 {
 	return UE::Learning::Agents::Private::AddObservation<UAngleObservation, UE::Learning::FAngleFeature>(AgentType, Name, 1, Scale);
 }
 
-void UAngleObservation::SetAngleObservation(int32 AgentId, float Angle, float RelativeAngle)
+void UAngleObservation::SetAngleObservation(const int32 AgentId, const float Angle, const float RelativeAngle)
 {
-	TLearningArrayView<2, float> AngleView = FeatureObject->InstanceData->View(FeatureObject->AngleHandle);
-	TLearningArrayView<1, float> RelativeAngleView = FeatureObject->InstanceData->View(FeatureObject->RelativeAngleHandle);
+	const TLearningArrayView<2, float> AngleView = FeatureObject->InstanceData->View(FeatureObject->AngleHandle);
+	const TLearningArrayView<1, float> RelativeAngleView = FeatureObject->InstanceData->View(FeatureObject->RelativeAngleHandle);
 
 	if (AgentId == INDEX_NONE)
 	{
@@ -267,15 +267,15 @@ void UAngleObservation::VisualLog(const UE::Learning::FIndexSet Instances) const
 
 //------------------------------------------------------------------
 
-UPlanarDirectionObservation* UPlanarDirectionObservation::AddPlanarDirectionObservation(ULearningAgentsType* AgentType, FName Name, float Scale, FVector Axis0, FVector Axis1)
+UPlanarDirectionObservation* UPlanarDirectionObservation::AddPlanarDirectionObservation(ULearningAgentsType* AgentType, const FName Name, const float Scale, const FVector Axis0, const FVector Axis1)
 {
 	return UE::Learning::Agents::Private::AddObservation<UPlanarDirectionObservation, UE::Learning::FPlanarDirectionFeature>(AgentType, Name, 1, Scale, Axis0, Axis1);
 }
 
-void UPlanarDirectionObservation::SetPlanarDirectionObservation(int32 AgentId, FVector Direction, FRotator RelativeRotation)
+void UPlanarDirectionObservation::SetPlanarDirectionObservation(const int32 AgentId, const FVector Direction, const FRotator RelativeRotation)
 {
-	TLearningArrayView<2, FVector> DirectionView = FeatureObject->InstanceData->View(FeatureObject->DirectionHandle);
-	TLearningArrayView<1, FQuat> RelativeRotationView = FeatureObject->InstanceData->View(FeatureObject->RelativeRotationHandle);
+	const TLearningArrayView<2, FVector> DirectionView = FeatureObject->InstanceData->View(FeatureObject->DirectionHandle);
+	const TLearningArrayView<1, FQuat> RelativeRotationView = FeatureObject->InstanceData->View(FeatureObject->RelativeRotationHandle);
 
 	if (AgentId == INDEX_NONE)
 	{
@@ -347,15 +347,15 @@ void UPlanarDirectionObservation::VisualLog(const UE::Learning::FIndexSet Instan
 }
 #endif
 
-UDirectionObservation* UDirectionObservation::AddDirectionObservation(ULearningAgentsType* AgentType, FName Name, float Scale)
+UDirectionObservation* UDirectionObservation::AddDirectionObservation(ULearningAgentsType* AgentType, const FName Name, const float Scale)
 {
 	return UE::Learning::Agents::Private::AddObservation<UDirectionObservation, UE::Learning::FDirectionFeature>(AgentType, Name, 1, Scale);
 }
 
-void UDirectionObservation::SetDirectionObservation(int32 AgentId, FVector Direction, FRotator RelativeRotation)
+void UDirectionObservation::SetDirectionObservation(const int32 AgentId, const FVector Direction, const FRotator RelativeRotation)
 {
-	TLearningArrayView<2, FVector> DirectionView = FeatureObject->InstanceData->View(FeatureObject->DirectionHandle);
-	TLearningArrayView<1, FQuat> RelativeRotationView = FeatureObject->InstanceData->View(FeatureObject->RelativeRotationHandle);
+	const TLearningArrayView<2, FVector> DirectionView = FeatureObject->InstanceData->View(FeatureObject->DirectionHandle);
+	const TLearningArrayView<1, FQuat> RelativeRotationView = FeatureObject->InstanceData->View(FeatureObject->RelativeRotationHandle);
 
 	if (AgentId == INDEX_NONE)
 	{
@@ -421,16 +421,16 @@ void UDirectionObservation::VisualLog(const UE::Learning::FIndexSet Instances) c
 
 //------------------------------------------------------------------
 
-UPlanarPositionObservation* UPlanarPositionObservation::AddPlanarPositionObservation(ULearningAgentsType* AgentType, FName Name, float Scale, FVector Axis0, FVector Axis1)
+UPlanarPositionObservation* UPlanarPositionObservation::AddPlanarPositionObservation(ULearningAgentsType* AgentType, const FName Name, const float Scale, const FVector Axis0, const FVector Axis1)
 {
 	return UE::Learning::Agents::Private::AddObservation<UPlanarPositionObservation, UE::Learning::FPlanarPositionFeature>(AgentType, Name, 1, Scale, Axis0, Axis1);
 }
 
-void UPlanarPositionObservation::SetPlanarPositionObservation(int32 AgentId, FVector Position, FVector RelativePosition, FRotator RelativeRotation)
+void UPlanarPositionObservation::SetPlanarPositionObservation(const int32 AgentId, const FVector Position, const FVector RelativePosition, const FRotator RelativeRotation)
 {
-	TLearningArrayView<2, FVector> PositionView = FeatureObject->InstanceData->View(FeatureObject->PositionHandle);
-	TLearningArrayView<1, FVector> RelativePositionView = FeatureObject->InstanceData->View(FeatureObject->RelativePositionHandle);
-	TLearningArrayView<1, FQuat> RelativeRotationView = FeatureObject->InstanceData->View(FeatureObject->RelativeRotationHandle);
+	const TLearningArrayView<2, FVector> PositionView = FeatureObject->InstanceData->View(FeatureObject->PositionHandle);
+	const TLearningArrayView<1, FVector> RelativePositionView = FeatureObject->InstanceData->View(FeatureObject->RelativePositionHandle);
+	const TLearningArrayView<1, FQuat> RelativeRotationView = FeatureObject->InstanceData->View(FeatureObject->RelativeRotationHandle);
 
 	if (AgentId == INDEX_NONE)
 	{
@@ -506,16 +506,16 @@ void UPlanarPositionObservation::VisualLog(const UE::Learning::FIndexSet Instanc
 }
 #endif
 
-UPositionObservation* UPositionObservation::AddPositionObservation(ULearningAgentsType* AgentType, FName Name, float Scale)
+UPositionObservation* UPositionObservation::AddPositionObservation(ULearningAgentsType* AgentType, const FName Name, const float Scale)
 {
 	return UE::Learning::Agents::Private::AddObservation<UPositionObservation, UE::Learning::FPositionFeature>(AgentType, Name, 1, Scale);
 }
 
-void UPositionObservation::SetPositionObservation(int32 AgentId, FVector Position, FVector RelativePosition, FRotator RelativeRotation)
+void UPositionObservation::SetPositionObservation(const int32 AgentId, const FVector Position, const FVector RelativePosition, const FRotator RelativeRotation)
 {
-	TLearningArrayView<2, FVector> PositionView = FeatureObject->InstanceData->View(FeatureObject->PositionHandle);
-	TLearningArrayView<1, FVector> RelativePositionView = FeatureObject->InstanceData->View(FeatureObject->RelativePositionHandle);
-	TLearningArrayView<1, FQuat> RelativeRotationView = FeatureObject->InstanceData->View(FeatureObject->RelativeRotationHandle);
+	const TLearningArrayView<2, FVector> PositionView = FeatureObject->InstanceData->View(FeatureObject->PositionHandle);
+	const TLearningArrayView<1, FVector> RelativePositionView = FeatureObject->InstanceData->View(FeatureObject->RelativePositionHandle);
+	const TLearningArrayView<1, FQuat> RelativeRotationView = FeatureObject->InstanceData->View(FeatureObject->RelativeRotationHandle);
 
 	if (AgentId == INDEX_NONE)
 	{
@@ -583,16 +583,16 @@ void UPositionObservation::VisualLog(const UE::Learning::FIndexSet Instances) co
 }
 #endif
 
-UPlanarPositionArrayObservation* UPlanarPositionArrayObservation::AddPlanarPositionArrayObservation(ULearningAgentsType* AgentType, FName Name, int32 PositionNum, float Scale, FVector Axis0, FVector Axis1)
+UPlanarPositionArrayObservation* UPlanarPositionArrayObservation::AddPlanarPositionArrayObservation(ULearningAgentsType* AgentType, const FName Name, const int32 PositionNum, const float Scale, const FVector Axis0, const FVector Axis1)
 {
 	return UE::Learning::Agents::Private::AddObservation<UPlanarPositionArrayObservation, UE::Learning::FPlanarPositionFeature>(AgentType, Name, PositionNum, Scale, Axis0, Axis1);
 }
 
-void UPlanarPositionArrayObservation::SetPlanarPositionArrayObservation(int32 AgentId, const TArray<FVector>& Positions, FVector RelativePosition, FRotator RelativeRotation)
+void UPlanarPositionArrayObservation::SetPlanarPositionArrayObservation(const int32 AgentId, const TArray<FVector>& Positions, const FVector RelativePosition, const FRotator RelativeRotation)
 {
-	TLearningArrayView<2, FVector> PositionView = FeatureObject->InstanceData->View(FeatureObject->PositionHandle);
-	TLearningArrayView<1, FVector> RelativePositionView = FeatureObject->InstanceData->View(FeatureObject->RelativePositionHandle);
-	TLearningArrayView<1, FQuat> RelativeRotationView = FeatureObject->InstanceData->View(FeatureObject->RelativeRotationHandle);
+	const TLearningArrayView<2, FVector> PositionView = FeatureObject->InstanceData->View(FeatureObject->PositionHandle);
+	const TLearningArrayView<1, FVector> RelativePositionView = FeatureObject->InstanceData->View(FeatureObject->RelativePositionHandle);
+	const TLearningArrayView<1, FQuat> RelativeRotationView = FeatureObject->InstanceData->View(FeatureObject->RelativeRotationHandle);
 
 	if (AgentId == INDEX_NONE)
 	{
@@ -696,16 +696,16 @@ void UPlanarPositionArrayObservation::VisualLog(const UE::Learning::FIndexSet In
 }
 #endif
 
-UPositionArrayObservation* UPositionArrayObservation::AddPositionArrayObservation(ULearningAgentsType* AgentType, FName Name, int32 PositionNum, float Scale)
+UPositionArrayObservation* UPositionArrayObservation::AddPositionArrayObservation(ULearningAgentsType* AgentType, const FName Name, const int32 PositionNum, const float Scale)
 {
 	return UE::Learning::Agents::Private::AddObservation<UPositionArrayObservation, UE::Learning::FPositionFeature>(AgentType, Name, PositionNum, Scale);
 }
 
-void UPositionArrayObservation::SetPositionArrayObservation(int32 AgentId, const TArray<FVector>& Positions, FVector RelativePosition, FRotator RelativeRotation)
+void UPositionArrayObservation::SetPositionArrayObservation(const int32 AgentId, const TArray<FVector>& Positions, const FVector RelativePosition, const FRotator RelativeRotation)
 {
-	TLearningArrayView<2, FVector> PositionView = FeatureObject->InstanceData->View(FeatureObject->PositionHandle);
-	TLearningArrayView<1, FVector> RelativePositionView = FeatureObject->InstanceData->View(FeatureObject->RelativePositionHandle);
-	TLearningArrayView<1, FQuat> RelativeRotationView = FeatureObject->InstanceData->View(FeatureObject->RelativeRotationHandle);
+	const TLearningArrayView<2, FVector> PositionView = FeatureObject->InstanceData->View(FeatureObject->PositionHandle);
+	const TLearningArrayView<1, FVector> RelativePositionView = FeatureObject->InstanceData->View(FeatureObject->RelativePositionHandle);
+	const TLearningArrayView<1, FQuat> RelativeRotationView = FeatureObject->InstanceData->View(FeatureObject->RelativeRotationHandle);
 	
 	if (AgentId == INDEX_NONE)
 	{
@@ -803,15 +803,15 @@ void UPositionArrayObservation::VisualLog(const UE::Learning::FIndexSet Instance
 
 //------------------------------------------------------------------
 
-UPlanarVelocityObservation* UPlanarVelocityObservation::AddPlanarVelocityObservation(ULearningAgentsType* AgentType, FName Name, float Scale, FVector Axis0, FVector Axis1)
+UPlanarVelocityObservation* UPlanarVelocityObservation::AddPlanarVelocityObservation(ULearningAgentsType* AgentType, const FName Name, const float Scale, const FVector Axis0, const FVector Axis1)
 {
 	return UE::Learning::Agents::Private::AddObservation<UPlanarVelocityObservation, UE::Learning::FPlanarVelocityFeature>(AgentType, Name, 1, Scale, Axis0, Axis1);
 }
 
-void UPlanarVelocityObservation::SetPlanarVelocityObservation(int32 AgentId, FVector Velocity, FRotator RelativeRotation)
+void UPlanarVelocityObservation::SetPlanarVelocityObservation(const int32 AgentId, const FVector Velocity, const FRotator RelativeRotation)
 {
-	TLearningArrayView<2, FVector> VelocityView = FeatureObject->InstanceData->View(FeatureObject->VelocityHandle);
-	TLearningArrayView<1, FQuat> RelativeRotationView = FeatureObject->InstanceData->View(FeatureObject->RelativeRotationHandle);
+	const TLearningArrayView<2, FVector> VelocityView = FeatureObject->InstanceData->View(FeatureObject->VelocityHandle);
+	const TLearningArrayView<1, FQuat> RelativeRotationView = FeatureObject->InstanceData->View(FeatureObject->RelativeRotationHandle);
 
 	if (AgentId == INDEX_NONE)
 	{
@@ -883,15 +883,15 @@ void UPlanarVelocityObservation::VisualLog(const UE::Learning::FIndexSet Instanc
 }
 #endif
 
-UVelocityObservation* UVelocityObservation::AddVelocityObservation(ULearningAgentsType* AgentType, FName Name, float Scale)
+UVelocityObservation* UVelocityObservation::AddVelocityObservation(ULearningAgentsType* AgentType, const FName Name, const float Scale)
 {
 	return UE::Learning::Agents::Private::AddObservation<UVelocityObservation, UE::Learning::FVelocityFeature>(AgentType, Name, 1, Scale);
 }
 
-void UVelocityObservation::SetVelocityObservation(int32 AgentId, FVector Velocity, FRotator RelativeRotation)
+void UVelocityObservation::SetVelocityObservation(const int32 AgentId, const FVector Velocity, const FRotator RelativeRotation)
 {
-	TLearningArrayView<2, FVector> VelocityView = FeatureObject->InstanceData->View(FeatureObject->VelocityHandle);
-	TLearningArrayView<1, FQuat> RelativeRotationView = FeatureObject->InstanceData->View(FeatureObject->RelativeRotationHandle);
+	const TLearningArrayView<2, FVector> VelocityView = FeatureObject->InstanceData->View(FeatureObject->VelocityHandle);
+	const TLearningArrayView<1, FQuat> RelativeRotationView = FeatureObject->InstanceData->View(FeatureObject->RelativeRotationHandle);
 
 	if (AgentId == INDEX_NONE)
 	{
@@ -960,4 +960,3 @@ void UVelocityObservation::VisualLog(const UE::Learning::FIndexSet Instances) co
 #undef UE_LEARNING_AGENTS_VLOG_STRING
 #undef UE_LEARNING_AGENTS_VLOG_TRANSFORM
 #undef UE_LEARNING_AGENTS_VLOG_PLANE
-
diff --git a/Engine/Plugins/Experimental/LearningAgents/Source/LearningAgents/Private/LearningAgentsPolicy.cpp b/Engine/Plugins/Experimental/LearningAgents/Source/LearningAgents/Private/LearningAgentsPolicy.cpp
--- a/Engine/Plugins/Experimental/LearningAgents/Source/LearningAgents/Private/LearningAgentsPolicy.cpp
+++ b/Engine/Plugins/Experimental/LearningAgents/Source/LearningAgents/Private/LearningAgentsPolicy.cpp
@@ -144,7 +144,7 @@ bool ULearningAgentsPolicy::IsPolicySetupPerformed() const
 	return bPolicySetupPerformed;
 }
 
-void ULearningAgentsPolicy::AddAgent(int32 AgentId)
+void ULearningAgentsPolicy::AddAgent(const int32 AgentId)
 {
 	if (!IsPolicySetupPerformed())
 	{
@@ -169,7 +169,7 @@ void ULearningAgentsPolicy::AddAgent(int32 AgentId)
 	SelectedAgentsSet.TryMakeSlice();
 }
 
-void ULearningAgentsPolicy::RemoveAgent(int32 AgentId)
+void ULearningAgentsPolicy::RemoveAgent(const int32 AgentId)
 {
 	if (!IsPolicySetupPerformed())
 	{
@@ -187,12 +187,12 @@ void ULearningAgentsPolicy::RemoveAgent(int32 AgentId)
 	SelectedAgentsSet.TryMakeSlice();
 }
 
-bool ULearningAgentsPolicy::HasAgent(int32 AgentId) const
+bool ULearningAgentsPolicy::HasAgent(const int32 AgentId) const
 {
 	return SelectedAgentsSet.Contains(AgentId);
 }
 
-ULearningAgentsType* ULearningAgentsPolicy::GetAgentType(TSubclassOf<ULearningAgentsType> AgentClass)
+ULearningAgentsType* ULearningAgentsPolicy::GetAgentType(const TSubclassOf<ULearningAgentsType> AgentClass)
 {
 	if (!IsPolicySetupPerformed())
 	{
@@ -222,7 +222,7 @@ void ULearningAgentsPolicy::LoadPolicyFromSnapshot(const FDirectoryPath& Directo
 	}
 
 	TArray64<uint8> NetworkData;
-	FString FilePath = Directory.Path + FGenericPlatformMisc::GetDefaultPathSeparator() + Filename;
+	const FString FilePath = Directory.Path + FGenericPlatformMisc::GetDefaultPathSeparator() + Filename;
 	if (FFileHelper::LoadFileToArray(NetworkData, *FilePath))
 	{
 		const int32 TotalByteNum = UE::Learning::FNeuralNetwork::GetSerializationByteNum(
@@ -262,7 +262,7 @@ void ULearningAgentsPolicy::SavePolicyToSnapshot(const FDirectoryPath& Directory
 
 	Network->NeuralNetwork->SerializeToBytes(NetworkData);
 
-	FString FilePath = Directory.Path + FGenericPlatformMisc::GetDefaultPathSeparator() + Filename;
+	const FString FilePath = Directory.Path + FGenericPlatformMisc::GetDefaultPathSeparator() + Filename;
 	if (!FFileHelper::SaveArrayToFile(NetworkData, *FilePath))
 	{
 		UE_LOG(LogLearning, Error, TEXT("Failed to save network to file: %s"), *FilePath);
diff --git a/Engine/Plugins/Experimental/LearningAgents/Source/LearningAgents/Private/LearningAgentsType.cpp b/Engine/Plugins/Experimental/LearningAgents/Source/LearningAgents/Private/LearningAgentsType.cpp
--- a/Engine/Plugins/Experimental/LearningAgents/Source/LearningAgents/Private/LearningAgentsType.cpp
+++ b/Engine/Plugins/Experimental/LearningAgents/Source/LearningAgents/Private/LearningAgentsType.cpp
@@ -68,12 +68,12 @@ void ULearningAgentsType::SetupAgentType()
 	bSetupPerformed = true;
 }
 
-const bool ULearningAgentsType::IsSetupPerformed() const
+bool ULearningAgentsType::IsSetupPerformed() const
 {
 	return bSetupPerformed;
 }
 
-const int32 ULearningAgentsType::GetMaxInstanceNum() const
+int32 ULearningAgentsType::GetMaxInstanceNum() const
 {
 	return MaxInstanceNum;
 }
@@ -103,17 +103,17 @@ TConstArrayView<ULearningAgentsAction*> ULearningAgentsType::GetActionObjects()
 	return ActionObjects;
 }
 
-const TConstArrayView<TObjectPtr<UObject>> ULearningAgentsType::GetAgents() const
+TConstArrayView<TObjectPtr<UObject>> ULearningAgentsType::GetAgents() const
 {
 	return Agents;
 }
 
-const UE::Learning::FIndexSet ULearningAgentsType::GetOccupiedAgentSet() const
+UE::Learning::FIndexSet ULearningAgentsType::GetOccupiedAgentSet() const
 {
 	return OccupiedAgentSet;
 }
 
-const UE::Learning::FIndexSet ULearningAgentsType::GetVacantAgentSet() const
+UE::Learning::FIndexSet ULearningAgentsType::GetVacantAgentSet() const
 {
 	return VacantAgentSet;
 }
@@ -143,7 +143,7 @@ int32 ULearningAgentsType::AddAgent(UObject* Agent)
 	return NewAgentId;
 }
 
-void ULearningAgentsType::RemoveAgentById(int32 AgentId)
+void ULearningAgentsType::RemoveAgentById(const int32 AgentId)
 {
 	if (AgentId == INDEX_NONE)
 	{
@@ -151,7 +151,7 @@ void ULearningAgentsType::RemoveAgentById(int32 AgentId)
 		return;
 	}
 
-	int32 RemovedCount = OccupiedAgentIds.RemoveSingleSwap(AgentId, false);
+	const int32 RemovedCount = OccupiedAgentIds.RemoveSingleSwap(AgentId, false);
 
 	if (RemovedCount == 0)
 	{
@@ -159,8 +159,6 @@ void ULearningAgentsType::RemoveAgentById(int32 AgentId)
 		return;
 	}
 
-	UObject* RemovedAgent = Agents[AgentId];
-
 	// Remove Agent
 	VacantAgentIds.Push(AgentId);
 	Agents[AgentId] = nullptr;
@@ -177,7 +175,7 @@ void ULearningAgentsType::RemoveAgent(UObject* Agent)
 	}
 
 	int32 AgentId = INDEX_NONE;
-	bool bIsFound = Agents.Find(Agent, AgentId);
+	const bool bIsFound = Agents.Find(Agent, AgentId);
 
 	if (!bIsFound)
 	{
@@ -193,12 +191,12 @@ bool ULearningAgentsType::HasAgent(UObject* Agent) const
 	return Agents.Find(Agent) ? true : false;
 }
 
-bool ULearningAgentsType::HasAgentById(int32 AgentId) const
+bool ULearningAgentsType::HasAgentById(const int32 AgentId) const
 {
 	return OccupiedAgentSet.Contains(AgentId);
 }
 
-UObject* ULearningAgentsType::GetAgent(int32 AgentId, TSubclassOf<UObject> AgentClass)
+UObject* ULearningAgentsType::GetAgent(const int32 AgentId, const TSubclassOf<UObject> AgentClass)
 {
 	if (AgentId < 0 || AgentId >= Agents.Num())
 	{
@@ -209,12 +207,12 @@ UObject* ULearningAgentsType::GetAgent(int32 AgentId, TSubclassOf<UObject> Agent
 	return Agents[AgentId];
 }
 
-const UObject* ULearningAgentsType::GetAgent(int32 AgentId) const
+const UObject* ULearningAgentsType::GetAgent(const int32 AgentId) const
 {
 	return Agents[AgentId];
 }
 
-UObject* ULearningAgentsType::GetAgent(int32 AgentId)
+UObject* ULearningAgentsType::GetAgent(const int32 AgentId)
 {
 	return Agents[AgentId];
 }
diff --git a/Engine/Plugins/Experimental/LearningAgents/Source/LearningAgents/Public/LearningAgentsActions.h b/Engine/Plugins/Experimental/LearningAgents/Source/LearningAgents/Public/LearningAgentsActions.h
--- a/Engine/Plugins/Experimental/LearningAgents/Source/LearningAgents/Public/LearningAgentsActions.h
+++ b/Engine/Plugins/Experimental/LearningAgents/Source/LearningAgents/Public/LearningAgentsActions.h
@@ -22,6 +22,7 @@ class ULearningAgentsType;
 // For functions in this file, we are favoring having more verbose names such as "AddFloatAction" vs simply "Add" in 
 // order to keep it easy to find the correct function in blueprints.
 
+/** The base class for all actions. Actions define the outputs from your agents. */
 UCLASS(Abstract, BlueprintType)
 class LEARNINGAGENTS_API ULearningAgentsAction : public UObject
 {
@@ -33,83 +34,142 @@ public:
 	/** Color used to draw this action in the visual log */
 	FLinearColor VisualLogColor = FColor::Blue;
 
+	/** Describes this observation to the visual logger for debugging purposes. */
 	virtual void VisualLog(const UE::Learning::FIndexSet Instances) const {}
 #endif
 };
 
-//------------------------------------------------------------------
-
+/** A simple float action. Used as a catch-all for situations where a more type-specific action does not exist yet. */
 UCLASS()
 class LEARNINGAGENTS_API UFloatAction : public ULearningAgentsAction
 {
 	GENERATED_BODY()
 
 public:
+
+	/**
+	* Adds a new float action to the given agent type. Call during ULearningAgentsType::SetupActions event.
+	* @param AgentType The agent type to add this action to.
+	* @param Name The name of this new action. Used for debugging.
+	* @param Scale Used to normalize the data for the action.
+	* @return The newly created action.
+	*/
 	UFUNCTION(BlueprintCallable, Category = "LearningAgents")
-	static UFloatAction* AddFloatAction(ULearningAgentsType* AgentType, FName Name = NAME_None, float Scale = 1.0f);
+	static UFloatAction* AddFloatAction(ULearningAgentsType* AgentType, const FName Name = NAME_None, const float Scale = 1.0f);
 
+	/**
+	* Gets the data for this action. Call during ULearningAgentsType::GetActions event.
+	* @param AgentId The agent id to get data for.
+	* @return The current action value.
+	*/
 	UFUNCTION(BlueprintPure, Category = "LearningAgents")
-	float GetFloatAction(int32 AgentId);
+	float GetFloatAction(const int32 AgentId);
 
+	/**
+	* Sets the data for this action. Call during ULearningAgentsController::EncodeActions event.
+	* @param AgentId The agent id to set data for.
+	* @param Value The current action value.
+	*/	
 	UFUNCTION(BlueprintCallable, Category = "LearningAgents")
-	void SetFloatAction(int32 AgentId, float Value);
+	void SetFloatAction(const int32 AgentId, const float Value);
 
 #if ENABLE_VISUAL_LOG
+	/** Describes this observation to the visual logger for debugging purposes. */
 	virtual void VisualLog(const UE::Learning::FIndexSet Instances) const override;
 #endif
 
 	TSharedPtr<UE::Learning::FFloatFeature> FeatureObject;
 };
 
-//------------------------------------------------------------------
-
+/** A simple FVector action. */
 UCLASS()
 class LEARNINGAGENTS_API UVectorAction : public ULearningAgentsAction
 {
 	GENERATED_BODY()
 
 public:
+
+	/**
+	* Adds a new vector action to the given agent type. Call during ULearningAgentsType::SetupActions event.
+	* @param AgentType The agent type to add this action to.
+	* @param Name The name of this new action. Used for debugging.
+	* @param Scale Used to normalize the data for the action.
+	* @return The newly created action.
+	*/
 	UFUNCTION(BlueprintCallable, Category = "LearningAgents")
-	static UVectorAction* AddVectorAction(ULearningAgentsType* AgentType, FName Name = NAME_None, float Scale = 1.0f);
+	static UVectorAction* AddVectorAction(ULearningAgentsType* AgentType, const FName Name = NAME_None, const float Scale = 1.0f);
 
+	/**
+	* Gets the data for this action. Call during ULearningAgentsType::GetActions event.
+	* @param AgentId The agent id to get data for.
+	* @return The current action values.
+	*/
 	UFUNCTION(BlueprintPure, Category = "LearningAgents")
-	FVector GetVectorAction(int32 AgentId);
+	FVector GetVectorAction(const int32 AgentId);
 
+	/**
+	* Sets the data for this action. Call during ULearningAgentsController::EncodeActions event.
+	* @param AgentId The agent id to set data for.
+	* @param Value The current action values.
+	*/
 	UFUNCTION(BlueprintCallable, Category = "LearningAgents")
-	void SetVectorAction(int32 AgentId, FVector Action);
+	void SetVectorAction(const int32 AgentId, const FVector Action);
 
 #if ENABLE_VISUAL_LOG
+	/** Describes this observation to the visual logger for debugging purposes. */
 	virtual void VisualLog(const UE::Learning::FIndexSet Instances) const override;
 #endif
 
 	TSharedPtr<UE::Learning::FFloatFeature> FeatureObject;
 };
 
-//------------------------------------------------------------------
-
+/** An array of rotation vector actions. */
 UCLASS()
 class LEARNINGAGENTS_API URotationVectorArrayAction : public ULearningAgentsAction
 {
 	GENERATED_BODY()
 
 public:
+
+	/**
+	* Adds a new rotation vector array action to the given agent type. Call during ULearningAgentsType::SetupActions event.
+	* @param AgentType The agent type to add this action to.
+	* @param Name The name of this new action. Used for debugging.
+	* @param RotationVectorNum The number of rotations in the array.
+	* @param Scale Used to normalize the data for the action.
+	* @return The newly created action.
+	*/
 	UFUNCTION(BlueprintCallable, Category = "LearningAgents")
-	static URotationVectorArrayAction* AddRotationVectorArrayAction(ULearningAgentsType* AgentType, FName Name = NAME_None, int32 RotationVectorNum = 0, float Scale = 180.0f);
+	static URotationVectorArrayAction* AddRotationVectorArrayAction(ULearningAgentsType* AgentType, const FName Name = NAME_None, const int32 RotationVectorNum = 0, const float Scale = 180.0f);
 
+	/**
+	* Gets the data for this action as rotation vectors. Call during ULearningAgentsType::GetActions event.
+	* @param AgentId The agent id to get data for.
+	* @param OutRotationVectors The current action values.
+	*/
 	UFUNCTION(BlueprintPure, Category = "LearningAgents")
-	void GetRotationVectorArrayAction(int32 AgentId, TArray<FVector>& OutRotationVectors);
+	void GetRotationVectorArrayAction(const int32 AgentId, TArray<FVector>& OutRotationVectors);
 
+	/**
+	* Gets the data for this action as quaternions. Call during ULearningAgentsType::GetActions event.
+	* @param AgentId The agent id to get data for.
+	* @param OutRotations The current action values.
+	*/
 	UFUNCTION(BlueprintPure, Category = "LearningAgents")
-	void GetRotationVectorArrayActionAsQuats(int32 AgentId, TArray<FQuat>& OutRotations);
+	void GetRotationVectorArrayActionAsQuats(const int32 AgentId, TArray<FQuat>& OutRotations);
 
+	/**
+	* Gets the data for this action as rotators. Call during ULearningAgentsType::GetActions event.
+	* @param AgentId The agent id to get data for.
+	* @param OutRotations The current action values.
+	*/
 	UFUNCTION(BlueprintPure, Category = "LearningAgents")
-	void GetRotationVectorArrayActionAsRotators(int32 AgentId, TArray<FRotator>& OutRotations);
+	void GetRotationVectorArrayActionAsRotators(const int32 AgentId, TArray<FRotator>& OutRotations);
 
 #if ENABLE_VISUAL_LOG
+	/** Describes this observation to the visual logger for debugging purposes. */
 	virtual void VisualLog(const UE::Learning::FIndexSet Instances) const override;
 #endif
 
 	TSharedPtr<UE::Learning::FRotationVectorFeature> FeatureObject;
 };
-
-//------------------------------------------------------------------
\ No newline at end of file
diff --git a/Engine/Plugins/Experimental/LearningAgents/Source/LearningAgents/Public/LearningAgentsController.h b/Engine/Plugins/Experimental/LearningAgents/Source/LearningAgents/Public/LearningAgentsController.h
--- a/Engine/Plugins/Experimental/LearningAgents/Source/LearningAgents/Public/LearningAgentsController.h
+++ b/Engine/Plugins/Experimental/LearningAgents/Source/LearningAgents/Public/LearningAgentsController.h
@@ -8,6 +8,11 @@
 
 #include "LearningAgentsController.generated.h"
 
+/**
+* A controller provides a method for injecting actions into the learning agents system from some other existing
+* behavior, e.g. we may want to gather demonstrations from a human or AI behavior tree controlling our agent(s)
+* for imitation learning purposes.
+*/
 UCLASS(Abstract, BlueprintType, Blueprintable)
 class LEARNINGAGENTS_API ULearningAgentsController : public UActorComponent
 {
@@ -21,42 +26,71 @@ public:
 	ULearningAgentsController(FVTableHelper& Helper);
 	virtual ~ULearningAgentsController();
 
+	/** Initializes this object to be used with the given agent type. */
 	UFUNCTION(BlueprintCallable, Category = "LearningAgents")
 	void SetupController(ULearningAgentsType* InAgentType);
 
+	/** Returns true if SetupController has been run successfully; Otherwise, false. */
 	UFUNCTION(BlueprintPure, Category = "LearningAgents")
 	bool IsControllerSetupPerformed() const;
 
 // ----- Agent Management -----
 public:
 
+	/**
+	 * Adds an agent to this controller.
+	 * @param AgentId The id of the agent to be added.
+	 * @warning The agent id must exist for the agent type.
+	 */
 	UFUNCTION(BlueprintCallable, Category = "LearningAgents")
-	void AddAgent(int32 AgentId);
+	void AddAgent(const int32 AgentId);
 
+	/**
+	* Removes an agent from this controller.
+	* @param AgentId The id of the agent to be removed.
+	* @warning The agent id must exist for this controller already.
+	*/
 	UFUNCTION(BlueprintCallable, Category = "LearningAgents")
-	void RemoveAgent(int32 AgentId);
+	void RemoveAgent(const int32 AgentId);
 
+	/** Returns true if the given id has been previously added to this controller; Otherwise, false. */
 	UFUNCTION(BlueprintPure, Category = "LearningAgents")
-	bool HasAgent(int32 AgentId) const;
+	bool HasAgent(const int32 AgentId) const;
 
+	/**
+	* Gets the agent type this controller is associated with.
+	* @param AgentClass The class to cast the agent type to (in blueprint).
+	*/
 	UFUNCTION(BlueprintPure, Category = "LearningAgents", meta = (DeterminesOutputType = "AgentClass"))
-	ULearningAgentsType* GetAgentType(TSubclassOf<ULearningAgentsType> AgentClass);
+	ULearningAgentsType* GetAgentType(const TSubclassOf<ULearningAgentsType> AgentClass);
 
 // ----- Actions -----
 public:
 
+	/**
+	* During this event, you should set the actions of your agents.
+	* @param AgentIds The list of agent ids to set actions for.
+	* @see LearningAgentsActions.h for the list of available actions.
+	* @see ULearningAgentsType::GetAgent to get the agent corresponding to each id.
+	*/
 	UFUNCTION(BlueprintNativeEvent, Category = "LearningAgents")
 	void SetActions(const TArray<int32>& AgentIds);
 
+	/**
+	* Call this function when it is time to gather all the actions for your agents. This should be called roughly 
+	* whenever you are calling ULearningAgentsType::EncodeObservations. This will call this controller's SetActions event.
+	*/
 	UFUNCTION(BlueprintCallable, Category = "LearningAgents")
 	void EncodeActions();
 
 // ----- Private Data -----
 private:
 
+	/** The agent type this controller is associated with. */
 	UPROPERTY(VisibleAnywhere, Category = "LearningAgents")
 	TObjectPtr<ULearningAgentsType> AgentType;
 
+	/** The agent ids this controller is managing. */
 	UPROPERTY(VisibleAnywhere, Category = "LearningAgents")
 	TArray<int32> SelectedAgentIds;
 
diff --git a/Engine/Plugins/Experimental/LearningAgents/Source/LearningAgents/Public/LearningAgentsCritic.h b/Engine/Plugins/Experimental/LearningAgents/Source/LearningAgents/Public/LearningAgentsCritic.h
--- a/Engine/Plugins/Experimental/LearningAgents/Source/LearningAgents/Public/LearningAgentsCritic.h
+++ b/Engine/Plugins/Experimental/LearningAgents/Source/LearningAgents/Public/LearningAgentsCritic.h
@@ -20,8 +20,7 @@ namespace UE::Learning
 struct FDirectoryPath;
 class ULearningAgentsNeuralNetwork;
 
-//------------------------------------------------------------------
-
+/** The configurable settings for a ULearningAgentsCritic. */
 USTRUCT(BlueprintType, Category = "LearningAgents")
 struct LEARNINGAGENTS_API FLearningAgentsCriticSettings
 {
@@ -42,8 +41,7 @@ public:
 	ELearningAgentsActivationFunction ActivationFunction = ELearningAgentsActivationFunction::ELU;
 };
 
-//------------------------------------------------------------------
-
+/** A critic used by some algorithms for training the managed agents. */
 UCLASS(BlueprintType, Blueprintable)
 class LEARNINGAGENTS_API ULearningAgentsCritic : public UActorComponent
 {
@@ -57,66 +55,112 @@ public:
 	ULearningAgentsCritic(FVTableHelper& Helper);
 	virtual ~ULearningAgentsCritic();
 
+	/** Initializes this object to be used with the given agent type and critic settings. */
 	UFUNCTION(BlueprintCallable, Category = "LearningAgents")
 	void SetupCritic(ULearningAgentsType* InAgentType, const FLearningAgentsCriticSettings& CriticSettings = FLearningAgentsCriticSettings());
 
+	/** Returns true if SetupCritic has been run successfully; Otherwise, false. */
 	UFUNCTION(BlueprintPure, Category = "LearningAgents")
 	bool IsCriticSetupPerformed() const;
 
 // ----- Agent Management -----
 public:
 
+	/**
+	 * Adds an agent to this critic.
+	 * @param AgentId The id of the agent to be added.
+	 * @warning The agent id must exist for the agent type.
+	 */
 	UFUNCTION(BlueprintCallable, Category = "LearningAgents")
-	void AddAgent(int32 AgentId);
+	void AddAgent(const int32 AgentId);
 
+	/**
+	* Removes an agent from this critic.
+	* @param AgentId The id of the agent to be removed.
+	* @warning The agent id must exist for this critic already.
+	*/
 	UFUNCTION(BlueprintCallable, Category = "LearningAgents")
-	void RemoveAgent(int32 AgentId);
+	void RemoveAgent(const int32 AgentId);
 
+	/** Returns true if the given id has been previously added to this critic; Otherwise, false. */
 	UFUNCTION(BlueprintPure, Category = "LearningAgents")
-	bool HasAgent(int32 AgentId) const;
+	bool HasAgent(const int32 AgentId) const;
 
+	/**
+	* Gets the agent type this critic is associated with.
+	* @param AgentClass The class to cast the agent type to (in blueprint).
+	*/
 	UFUNCTION(BlueprintPure, Category = "LearningAgents", meta = (DeterminesOutputType = "AgentClass"))
-	ULearningAgentsType* GetAgentType(TSubclassOf<ULearningAgentsType> AgentClass);
+	ULearningAgentsType* GetAgentType(const TSubclassOf<ULearningAgentsType> AgentClass);
 
 // ----- Load / Save -----
 public:
 
+	/**
+	* Load a snapshot's weights into this critic.
+	* @param Directory The directory the snapshot file is in.
+	* @param Filename The filename of the snapshot, including the file extension.
+	*/
 	UFUNCTION(BlueprintCallable, Category = "LearningAgents")
 	void LoadCriticFromSnapshot(const FDirectoryPath& Directory, const FString Filename);
 
+	/**
+	* Save this critic's weights into a snapshot.
+	* @param Directory The directory to save the snapshot file in.
+	* @param Filename The filename of the snapshot, including the file extension.
+	*/
 	UFUNCTION(BlueprintCallable, Category = "LearningAgents")
 	void SaveCriticToSnapshot(const FDirectoryPath& Directory, const FString Filename) const;
 
+	/**
+	* Load a ULearningAgentsNeuralNetwork asset's weights into this critic.
+	* @param NeuralNetworkAsset The asset to load from.
+	*/
 	UFUNCTION(BlueprintCallable, Category = "LearningAgents")
 	void LoadCriticFromAsset(const ULearningAgentsNeuralNetwork* NeuralNetworkAsset);
 
+	/**
+	* Save this critic's weights to a ULearningAgentsNeuralNetwork asset.
+	* @param NeuralNetworkAsset The asset to save to.
+	*/
 	UFUNCTION(BlueprintCallable, Category = "LearningAgents", Meta=(DevelopmentOnly))
 	void SaveCriticToAsset(ULearningAgentsNeuralNetwork* NeuralNetworkAsset) const;
 
 // ----- Evaluation -----
 public:
 
+	/**
+	* Calling this function will run the underlying neural network on the previously buffered observations to populate
+	* the output value buffer. This should be called after the corresponding agent type's EncodeObservations.
+	*/
 	UFUNCTION(BlueprintCallable, Category = "LearningAgents")
 	void EvaluateCritic();
 
 // ----- Non-blueprint public interface -----
 public:
 
+	/** Get a reference to this critic's neural network. */
 	UE::Learning::FNeuralNetwork& GetCriticNetwork();
+	
+	/** Get a reference to this critic's critic function object. */
 	UE::Learning::FNeuralNetworkCriticFunction& GetCriticObject();
 
 // ----- Private Data -----
 private:
 
+	/** The agent type this critic is associated with. */
 	UPROPERTY(VisibleAnywhere, Category = "LearningAgents")
 	TObjectPtr<ULearningAgentsType> AgentType;
 
+	/** The agent ids this critic is managing. */
 	UPROPERTY(VisibleAnywhere, Category = "LearningAgents")
 	TArray<int32> SelectedAgentIds;
 
+	/** True if this critic's SetupCritic has been run; Otherwise, false. */
 	UPROPERTY(VisibleAnywhere, Category = "LearningAgents")
 	bool bCriticSetupPerformed = false;
 
+	/** The underlying neural network. */
 	UPROPERTY(VisibleAnywhere, Category = "LearningAgents")
 	TObjectPtr<ULearningAgentsNeuralNetwork> Network;
 
@@ -132,6 +176,7 @@ private:
 	/** Color used to draw this action in the visual log */
 	FLinearColor VisualLogColor = FColor::Orange;
 
+	/** Describes this critic to the visual logger for debugging purposes. */
 	void VisualLog(const UE::Learning::FIndexSet Instances) const;
 #endif
 };
diff --git a/Engine/Plugins/Experimental/LearningAgents/Source/LearningAgents/Public/LearningAgentsNeuralNetwork.h b/Engine/Plugins/Experimental/LearningAgents/Source/LearningAgents/Public/LearningAgentsNeuralNetwork.h
--- a/Engine/Plugins/Experimental/LearningAgents/Source/LearningAgents/Public/LearningAgentsNeuralNetwork.h
+++ b/Engine/Plugins/Experimental/LearningAgents/Source/LearningAgents/Public/LearningAgentsNeuralNetwork.h
@@ -12,8 +12,7 @@ namespace UE::Learning
 	enum class EActivationFunction : uint8;
 }
 
-//------------------------------------------------------------------
-
+/** Activation functions for neural networks. */
 UENUM(BlueprintType, Category = "LearningAgents")
 enum class ELearningAgentsActivationFunction : uint8
 {
@@ -27,16 +26,16 @@ enum class ELearningAgentsActivationFunction : uint8
 	TanH	UMETA(DisplayName = "TanH"),
 };
 
-//------------------------------------------------------------------
-
 namespace UE::Learning::Agents
 {
+	/** Get the learning agents activation function from the UE::Learning activation function. */
 	LEARNINGAGENTS_API ELearningAgentsActivationFunction GetLearningAgentsActivationFunction(const EActivationFunction ActivationFunction);
+
+	/** Get the UE::Learning activation function from the learning agents activation function. */
 	LEARNINGAGENTS_API EActivationFunction GetActivationFunction(const ELearningAgentsActivationFunction ActivationFunction);
 }
 
-//------------------------------------------------------------------
-
+/** A neural network data asset. */
 UCLASS(BlueprintType)
 class LEARNINGAGENTS_API ULearningAgentsNeuralNetwork : public UDataAsset
 {
@@ -48,6 +47,7 @@ public:
 	ULearningAgentsNeuralNetwork(FVTableHelper& Helper);
 	virtual ~ULearningAgentsNeuralNetwork();
 
+	/** Serialize this neural network to/from the given archive. */
 	virtual void Serialize(FArchive& Ar) override;
 
 	TSharedPtr<UE::Learning::FNeuralNetwork> NeuralNetwork;
diff --git a/Engine/Plugins/Experimental/LearningAgents/Source/LearningAgents/Public/LearningAgentsObservations.h b/Engine/Plugins/Experimental/LearningAgents/Source/LearningAgents/Public/LearningAgentsObservations.h
--- a/Engine/Plugins/Experimental/LearningAgents/Source/LearningAgents/Public/LearningAgentsObservations.h
+++ b/Engine/Plugins/Experimental/LearningAgents/Source/LearningAgents/Public/LearningAgentsObservations.h
@@ -28,6 +28,7 @@ class ULearningAgentsType;
 // For functions in this file, we are favoring having more verbose names such as "AddFloatObservation" vs simply "Add" in 
 // order to keep it easy to find the correct function in blueprints.
 
+/** The base class for all observations. Observations define the inputs to your agents. */
 UCLASS(Abstract, BlueprintType)
 class LEARNINGAGENTS_API ULearningAgentsObservation : public UObject
 {
@@ -39,229 +40,421 @@ public:
 	/** Color used to draw this observation in the visual log */
 	FLinearColor VisualLogColor = FColor::Red;
 
+	/** Describes this observation to the visual logger for debugging purposes. */
 	virtual void VisualLog(const UE::Learning::FIndexSet Instances) const {}
 #endif
 };
 
-//------------------------------------------------------------------
-
+/** A simple float observation. Used as a catch-all for situations where a more type-specific observation does not exist yet. */
 UCLASS()
 class LEARNINGAGENTS_API UFloatObservation : public ULearningAgentsObservation
 {
 	GENERATED_BODY()
 
 public:
+
+	/**
+	* Adds a new float observation to the given agent type. Call during ULearningAgentsType::SetupObservations event.
+	* @param AgentType The agent type to add this observation to.
+	* @param Name The name of this new observation. Used for debugging.
+	* @param Scale Used to normalize the data for the observation.
+	* @return The newly created observation.
+	*/
 	UFUNCTION(BlueprintCallable, Category = "LearningAgents")
-	static UFloatObservation* AddFloatObservation(ULearningAgentsType* AgentType, FName Name = NAME_None, float Scale = 1.0f);
+	static UFloatObservation* AddFloatObservation(ULearningAgentsType* AgentType, const FName Name = NAME_None, const float Scale = 1.0f);
 
+	/**
+	* Sets the data for this observation. Call during ULearningAgentsType::SetObservations event.
+	* @param AgentId The agent id this data corresponds to.
+	* @param Observation The value currently being observed.
+	*/
 	UFUNCTION(BlueprintCallable, Category = "LearningAgents")
-	void SetFloatObservation(int32 AgentId, float Observation);
+	void SetFloatObservation(const int32 AgentId, const float Observation);
 
 #if ENABLE_VISUAL_LOG
+	/** Describes this observation to the visual logger for debugging purposes. */
 	virtual void VisualLog(const UE::Learning::FIndexSet Instances) const override;
 #endif
 
 	TSharedPtr<UE::Learning::FFloatFeature> FeatureObject;
 };
 
-//------------------------------------------------------------------
-
+/** A simple observation for an FVector. */
 UCLASS()
 class LEARNINGAGENTS_API UVectorObservation : public ULearningAgentsObservation
 {
 	GENERATED_BODY()
 
 public:
+
+	/**
+	* Adds a new vector observation to the given agent type. Call during ULearningAgentsType::SetupObservations event.
+	* @param AgentType The agent type to add this observation to.
+	* @param Name The name of this new observation. Used for debugging.
+	* @param Scale Used to normalize the data for the observation.
+	* @return The newly created observation.
+	*/
 	UFUNCTION(BlueprintCallable, Category = "LearningAgents")
-	static UVectorObservation* AddVectorObservation(ULearningAgentsType* AgentType, FName Name = NAME_None, float Scale = 1.0f);
+	static UVectorObservation* AddVectorObservation(ULearningAgentsType* AgentType, const FName Name = NAME_None, const float Scale = 1.0f);
 
+	/**
+	* Sets the data for this observation. Call during ULearningAgentsType::SetObservations event.
+	* @param AgentId The agent id this data corresponds to.
+	* @param Observation The values currently being observed.
+	*/
 	UFUNCTION(BlueprintCallable, Category = "LearningAgents")
-	void SetVectorObservation(int32 AgentId, FVector Observation);
+	void SetVectorObservation(const int32 AgentId, const FVector Observation);
 
 #if ENABLE_VISUAL_LOG
+	/** Describes this observation to the visual logger for debugging purposes. */
 	virtual void VisualLog(const UE::Learning::FIndexSet Instances) const override;
 #endif
 
 	TSharedPtr<UE::Learning::FFloatFeature> FeatureObject;
 };
 
-//------------------------------------------------------------------
-
+/** An observation of an angle relative to another angle. */
 UCLASS()
 class LEARNINGAGENTS_API UAngleObservation : public ULearningAgentsObservation
 {
 	GENERATED_BODY()
 
 public:
-	UFUNCTION(BlueprintCallable, Category = "LearningAgents")
-	static UAngleObservation* AddAngleObservation(ULearningAgentsType* AgentType, FName Name = NAME_None, float Scale = 1.0f);
 
+	/**
+	* Adds a new angle observation to the given agent type. Call during ULearningAgentsType::SetupObservations event.
+	* @param AgentType The agent type to add this observation to.
+	* @param Name The name of this new observation. Used for debugging.
+	* @param Scale Used to normalize the data for the observation.
+	* @return The newly created observation.
+	*/
+	UFUNCTION(BlueprintCallable, Category = "LearningAgents")
+	static UAngleObservation* AddAngleObservation(ULearningAgentsType* AgentType, const FName Name = NAME_None, const float Scale = 1.0f);
+
+	/**
+	* Sets the data for this observation. Call during ULearningAgentsType::SetObservations event.
+	* @param AgentId The agent id this data corresponds to.
+	* @param Angle The angle currently being observed.
+	* @param RelativeAngle The frame of reference angle.
+	*/
 	UFUNCTION(BlueprintCallable, Category = "LearningAgents")
-	void SetAngleObservation(int32 AgentId, float Angle, float RelativeAngle = 0.0f);
+	void SetAngleObservation(const int32 AgentId, const float Angle, const float RelativeAngle = 0.0f);
 
 #if ENABLE_VISUAL_LOG
+	/** Describes this observation to the visual logger for debugging purposes. */
 	virtual void VisualLog(const UE::Learning::FIndexSet Instances) const override;
 #endif
 
 	TSharedPtr<UE::Learning::FAngleFeature> FeatureObject;
 };
 
-//------------------------------------------------------------------
-
+/** An observation of a direction vector projected onto a plane. */
 UCLASS()
 class LEARNINGAGENTS_API UPlanarDirectionObservation : public ULearningAgentsObservation
 {
 	GENERATED_BODY()
 
 public:
-	UFUNCTION(BlueprintCallable, Category = "LearningAgents")
-	static UPlanarDirectionObservation* AddPlanarDirectionObservation(ULearningAgentsType* AgentType, FName Name = NAME_None, float Scale = 1.0f, FVector Axis0 = FVector::ForwardVector, FVector Axis1 = FVector::RightVector);
 
+	/**
+	* Adds a new planar direction observation to the given agent type. The axis parameters define the plane. Call
+	* during ULearningAgentsType::SetupObservations event.
+	* @param AgentType The agent type to add this observation to.
+	* @param Name The name of this new observation. Used for debugging.
+	* @param Scale Used to normalize the data for the observation.
+	* @param Axis0 The forward axis of the plane.
+	* @param Axis1 The right axis of the plane.
+	* @return The newly created observation.
+	*/
 	UFUNCTION(BlueprintCallable, Category = "LearningAgents")
-	void SetPlanarDirectionObservation(int32 AgentId, FVector Direction, FRotator RelativeRotation = FRotator::ZeroRotator);
+	static UPlanarDirectionObservation* AddPlanarDirectionObservation(ULearningAgentsType* AgentType, const FName Name = NAME_None, const float Scale = 1.0f, const FVector Axis0 = FVector::ForwardVector, const FVector Axis1 = FVector::RightVector);
+
+	/**
+	* Sets the data for this observation. The relative rotation can be used to make this observation relative to the
+	* agent's perspective, e.g. by passing the agent's forward rotation. Call during ULearningAgentsType::SetObservations event.
+	* @param AgentId The agent id this data corresponds to.
+	* @param Direction The direction currently being observed.
+	* @param RelativeRotation The frame of reference rotation.
+	*/
+	UFUNCTION(BlueprintCallable, Category = "LearningAgents")
+	void SetPlanarDirectionObservation(const int32 AgentId, const FVector Direction, const FRotator RelativeRotation = FRotator::ZeroRotator);
 
 #if ENABLE_VISUAL_LOG
+	/** Describes this observation to the visual logger for debugging purposes. */
 	virtual void VisualLog(const UE::Learning::FIndexSet Instances) const override;
 #endif
 
 	TSharedPtr<UE::Learning::FPlanarDirectionFeature> FeatureObject;
 };
 
+/** An observation of a direction vector. */
 UCLASS()
 class LEARNINGAGENTS_API UDirectionObservation : public ULearningAgentsObservation
 {
 	GENERATED_BODY()
 
 public:
-	UFUNCTION(BlueprintCallable, Category = "LearningAgents")
-	static UDirectionObservation* AddDirectionObservation(ULearningAgentsType* AgentType, FName Name = NAME_None, float Scale = 100.0f);
 
+	/**
+	* Adds a new direction observation to the given agent type. Call during ULearningAgentsType::SetupObservations event.
+	* @param AgentType The agent type to add this observation to.
+	* @param Name The name of this new observation. Used for debugging.
+	* @param Scale Used to normalize the data for the observation.
+	* @return The newly created observation.
+	*/
 	UFUNCTION(BlueprintCallable, Category = "LearningAgents")
-	void SetDirectionObservation(int32 AgentId, FVector Direction, FRotator RelativeRotation = FRotator::ZeroRotator);
+	static UDirectionObservation* AddDirectionObservation(ULearningAgentsType* AgentType, const FName Name = NAME_None, const float Scale = 1.0f);
+
+	/**
+	* Sets the data for this observation. The relative rotation can be used to make this observation relative to the
+	* agent's perspective, e.g. by passing the agent's forward rotation. Call during ULearningAgentsType::SetObservations event.
+	* @param AgentId The agent id this data corresponds to.
+	* @param Direction The direction currently being observed.
+	* @param RelativeRotation The frame of reference rotation.
+	*/
+	UFUNCTION(BlueprintCallable, Category = "LearningAgents")
+	void SetDirectionObservation(const int32 AgentId, const FVector Direction, const FRotator RelativeRotation = FRotator::ZeroRotator);
 
 #if ENABLE_VISUAL_LOG
+	/** Describes this observation to the visual logger for debugging purposes. */
 	virtual void VisualLog(const UE::Learning::FIndexSet Instances) const override;
 #endif
 
 	TSharedPtr<UE::Learning::FDirectionFeature> FeatureObject;
 };
 
-//------------------------------------------------------------------
-
+/** An observation of a position projected onto a plane. */
 UCLASS()
 class LEARNINGAGENTS_API UPlanarPositionObservation : public ULearningAgentsObservation
 {
 	GENERATED_BODY()
 
 public:
-	UFUNCTION(BlueprintCallable, Category = "LearningAgents")
-	static UPlanarPositionObservation* AddPlanarPositionObservation(ULearningAgentsType* AgentType, FName Name = NAME_None, float Scale = 100.0f, FVector Axis0 = FVector::ForwardVector, FVector Axis1 = FVector::RightVector);
 
+	/**
+	* Adds a new planar position observation to the given agent type. The axis parameters define the plane. Call
+	* during ULearningAgentsType::SetupObservations event.
+	* @param AgentType The agent type to add this observation to.
+	* @param Name The name of this new observation. Used for debugging.
+	* @param Scale Used to normalize the data for the observation.
+	* @param Axis0 The forward axis of the plane.
+	* @param Axis1 The right axis of the plane.
+	* @return The newly created observation.
+	*/
 	UFUNCTION(BlueprintCallable, Category = "LearningAgents")
-	void SetPlanarPositionObservation(int32 AgentId, FVector Position, FVector RelativePosition = FVector::ZeroVector, FRotator RelativeRotation = FRotator::ZeroRotator);
+	static UPlanarPositionObservation* AddPlanarPositionObservation(ULearningAgentsType* AgentType, const FName Name = NAME_None, const float Scale = 100.0f, const FVector Axis0 = FVector::ForwardVector, const FVector Axis1 = FVector::RightVector);
+
+	/**
+	* Sets the data for this observation. The relative position & rotation can be used to make this observation
+	* relative to the agent's perspective, e.g. by passing the agent's position & forward rotation. Call during
+	* ULearningAgentsType::SetObservations event.
+	* @param AgentId The agent id this data corresponds to.
+	* @param Position The position currently being observed.
+	* @param RelativePosition The vector Position will be offset from.
+	* @param RelativeRotation The frame of reference rotation.
+	*/
+	UFUNCTION(BlueprintCallable, Category = "LearningAgents")
+	void SetPlanarPositionObservation(const int32 AgentId, const FVector Position, const FVector RelativePosition = FVector::ZeroVector, const FRotator RelativeRotation = FRotator::ZeroRotator);
 
 #if ENABLE_VISUAL_LOG
+	/** Describes this observation to the visual logger for debugging purposes. */
 	virtual void VisualLog(const UE::Learning::FIndexSet Instances) const override;
 #endif
 
 	TSharedPtr<UE::Learning::FPlanarPositionFeature> FeatureObject;
 };
 
+/** An observation of a position vector. */
 UCLASS()
 class LEARNINGAGENTS_API UPositionObservation : public ULearningAgentsObservation
 {
 	GENERATED_BODY()
 
 public:
-	UFUNCTION(BlueprintCallable, Category = "LearningAgents")
-	static UPositionObservation* AddPositionObservation(ULearningAgentsType* AgentType, FName Name = NAME_None, float Scale = 100.0f);
 
+	/**
+	* Adds a new position observation to the given agent type. Call during ULearningAgentsType::SetupObservations event.
+	* @param AgentType The agent type to add this observation to.
+	* @param Name The name of this new observation. Used for debugging.
+	* @param Scale Used to normalize the data for the observation.
+	* @return The newly created observation.
+	*/
 	UFUNCTION(BlueprintCallable, Category = "LearningAgents")
-	void SetPositionObservation(int32 AgentId, FVector Position, FVector RelativePosition = FVector::ZeroVector, FRotator RelativeRotation = FRotator::ZeroRotator);
+	static UPositionObservation* AddPositionObservation(ULearningAgentsType* AgentType, const FName Name = NAME_None, const float Scale = 100.0f);
+
+	/**
+	* Sets the data for this observation. The relative position & rotation can be used to make this observation
+	* relative to the agent's perspective, e.g. by passing the agent's position & forward rotation. Call during
+	* ULearningAgentsType::SetObservations event.
+	* @param AgentId The agent id this data corresponds to.
+	* @param Position The position currently being observed.
+	* @param RelativePosition The vector Position will be offset from.
+	* @param RelativeRotation The frame of reference rotation.
+	*/
+	UFUNCTION(BlueprintCallable, Category = "LearningAgents")
+	void SetPositionObservation(const int32 AgentId, const FVector Position, const FVector RelativePosition = FVector::ZeroVector, const FRotator RelativeRotation = FRotator::ZeroRotator);
 
 #if ENABLE_VISUAL_LOG
+	/** Describes this observation to the visual logger for debugging purposes. */
 	virtual void VisualLog(const UE::Learning::FIndexSet Instances) const override;
 #endif
 
 	TSharedPtr<UE::Learning::FPositionFeature> FeatureObject;
 };
 
+/** An observation of an array of positions projected onto a plane. */
 UCLASS()
 class LEARNINGAGENTS_API UPlanarPositionArrayObservation : public ULearningAgentsObservation
 {
 	GENERATED_BODY()
 
 public:
-	UFUNCTION(BlueprintCallable, Category = "LearningAgents")
-	static UPlanarPositionArrayObservation* AddPlanarPositionArrayObservation(ULearningAgentsType* AgentType, FName Name = NAME_None, int32 PositionNum = 0, float Scale = 100.0f, FVector Axis0 = FVector::ForwardVector, FVector Axis1 = FVector::RightVector);
 
+	/**
+	* Adds a new planar position array observation to the given agent type. The axis parameters define the plane.
+	* Call during ULearningAgentsType::SetupObservations event.
+	* @param AgentType The agent type to add this observation to.
+	* @param Name The name of this new observation. Used for debugging.
+	* @param PositionNum The number of positions in the array.
+	* @param Scale Used to normalize the data for the observation.
+	* @param Axis0 The forward axis of the plane.
+	* @param Axis1 The right axis of the plane.
+	* @return The newly created observation.
+	*/
 	UFUNCTION(BlueprintCallable, Category = "LearningAgents")
-	void SetPlanarPositionArrayObservation(int32 AgentId, const TArray<FVector>& Positions, FVector RelativePosition = FVector::ZeroVector, FRotator RelativeRotation = FRotator::ZeroRotator);
+	static UPlanarPositionArrayObservation* AddPlanarPositionArrayObservation(ULearningAgentsType* AgentType, const FName Name = NAME_None, const int32 PositionNum = 0, const float Scale = 100.0f, const FVector Axis0 = FVector::ForwardVector, const FVector Axis1 = FVector::RightVector);
+
+
+	/**
+	* Sets the data for this observation. The relative position & rotation can be used to make this observation
+	* relative to the agent's perspective, e.g. by passing the agent's position & forward rotation. Call during
+	* ULearningAgentsType::SetObservations event.
+	* @param AgentId The agent id this data corresponds to.
+	* @param Positions The positions currently being observed.
+	* @param RelativePosition The vector Positions will be offset from.
+	* @param RelativeRotation The frame of reference rotation.
+	*/
+	UFUNCTION(BlueprintCallable, Category = "LearningAgents")
+	void SetPlanarPositionArrayObservation(const int32 AgentId, const TArray<FVector>& Positions, const FVector RelativePosition = FVector::ZeroVector, const FRotator RelativeRotation = FRotator::ZeroRotator);
 
 #if ENABLE_VISUAL_LOG
+	/** Describes this observation to the visual logger for debugging purposes. */
 	virtual void VisualLog(const UE::Learning::FIndexSet Instances) const override;
 #endif
 
 	TSharedPtr<UE::Learning::FPlanarPositionFeature> FeatureObject;
 };
 
+/** An observation of an array of positions. */
 UCLASS()
 class LEARNINGAGENTS_API UPositionArrayObservation : public ULearningAgentsObservation
 {
 	GENERATED_BODY()
 
 public:
-	UFUNCTION(BlueprintCallable, Category = "LearningAgents")
-	static UPositionArrayObservation* AddPositionArrayObservation(ULearningAgentsType* AgentType, FName Name = NAME_None, int32 PositionNum = 0, float Scale = 100.0f);
 
+	/**
+	* Adds a new position array observation to the given agent type. Call during ULearningAgentsType::SetupObservations event.
+	* @param AgentType The agent type to add this observation to.
+	* @param Name The name of this new observation. Used for debugging.
+	* @param PositionNum The number of positions in the array.
+	* @param Scale Used to normalize the data for the observation.
+	* @return The newly created observation.
+	*/
+	UFUNCTION(BlueprintCallable, Category = "LearningAgents")
+	static UPositionArrayObservation* AddPositionArrayObservation(ULearningAgentsType* AgentType, const FName Name = NAME_None, const int32 PositionNum = 0, const float Scale = 100.0f);
+
+	/**
+	* Sets the data for this observation. The relative position & rotation can be used to make this observation
+	* relative to the agent's perspective, e.g. by passing the agent's position & forward rotation. Call during
+	* ULearningAgentsType::SetObservations event.
+	* @param AgentId The agent id this data corresponds to.
+	* @param Positions The positions currently being observed.
+	* @param RelativePosition The vector Positions will be offset from.
+	* @param RelativeRotation The frame of reference rotation.
+	*/
 	UFUNCTION(BlueprintCallable, Category = "LearningAgents")
-	void SetPositionArrayObservation(int32 AgentId, const TArray<FVector>& Positions, FVector RelativePosition = FVector::ZeroVector, FRotator RelativeRotation = FRotator::ZeroRotator);
+	void SetPositionArrayObservation(int32 AgentId, const TArray<FVector>& Positions, const FVector RelativePosition = FVector::ZeroVector, const FRotator RelativeRotation = FRotator::ZeroRotator);
 
 #if ENABLE_VISUAL_LOG
+	/** Describes this observation to the visual logger for debugging purposes. */
 	virtual void VisualLog(const UE::Learning::FIndexSet Instances) const override;
 #endif
 
 	TSharedPtr<UE::Learning::FPositionFeature> FeatureObject;
 };
 
-//------------------------------------------------------------------
-
+/** An observation of a velocity projected onto a plane. */
 UCLASS()
 class LEARNINGAGENTS_API UPlanarVelocityObservation : public ULearningAgentsObservation
 {
 	GENERATED_BODY()
 
 public:
-	UFUNCTION(BlueprintCallable, Category = "LearningAgents")
-	static UPlanarVelocityObservation* AddPlanarVelocityObservation(ULearningAgentsType* AgentType, FName Name = NAME_None, float Scale = 200.0f, FVector Axis0 = FVector::ForwardVector, FVector Axis1 = FVector::RightVector);
 
+	/**
+	* Adds a new planar velocity observation to the given agent type. The axis parameters define the plane.
+	* Call during ULearningAgentsType::SetupObservations event.
+	* @param AgentType The agent type to add this observation to.
+	* @param Name The name of this new observation. Used for debugging.
+	* @param Scale Used to normalize the data for the observation.
+	* @param Axis0 The forward axis of the plane.
+	* @param Axis1 The right axis of the plane.
+	* @return The newly created observation.
+	*/
+	UFUNCTION(BlueprintCallable, Category = "LearningAgents")
+	static UPlanarVelocityObservation* AddPlanarVelocityObservation(ULearningAgentsType* AgentType, const FName Name = NAME_None, const float Scale = 200.0f, const FVector Axis0 = FVector::ForwardVector, const FVector Axis1 = FVector::RightVector);
+
+	/**
+	* Sets the data for this observation. The relative rotation can be used to make this observation relative to the
+	* agent's perspective, e.g. by passing the agent's forward rotation. Call during ULearningAgentsType::SetObservations event.
+	* @param AgentId The agent id this data corresponds to.
+	* @param Velocity The velocity currently being observed.
+	* @param RelativeRotation The frame of reference rotation.
+	*/
 	UFUNCTION(BlueprintCallable, Category = "LearningAgents")
-	void SetPlanarVelocityObservation(int32 AgentId, FVector Velocity, FRotator RelativeRotation = FRotator::ZeroRotator);
+	void SetPlanarVelocityObservation(const int32 AgentId, const FVector Velocity, const FRotator RelativeRotation = FRotator::ZeroRotator);
 
 #if ENABLE_VISUAL_LOG
+	/** Describes this observation to the visual logger for debugging purposes. */
 	virtual void VisualLog(const UE::Learning::FIndexSet Instances) const override;
 #endif
 
 	TSharedPtr<UE::Learning::FPlanarVelocityFeature> FeatureObject;
 };
 
+/** An observation of a velocity. */
 UCLASS()
 class LEARNINGAGENTS_API UVelocityObservation : public ULearningAgentsObservation
 {
 	GENERATED_BODY()
 
 public:
-	UFUNCTION(BlueprintCallable, Category = "LearningAgents")
-	static UVelocityObservation* AddVelocityObservation(ULearningAgentsType* AgentType, FName Name = NAME_None, float Scale = 200.0f);
 
+	/**
+	* Adds a new velocity observation to the given agent type. Call during ULearningAgentsType::SetupObservations event.
+	* @param AgentType The agent type to add this observation to.
+	* @param Name The name of this new observation. Used for debugging.
+	* @param Scale Used to normalize the data for the observation.
+	* @return The newly created observation.
+	*/
+	UFUNCTION(BlueprintCallable, Category = "LearningAgents")
+	static UVelocityObservation* AddVelocityObservation(ULearningAgentsType* AgentType, const FName Name = NAME_None, const float Scale = 200.0f);
+
+	/**
+	* Sets the data for this observation. The relative rotation can be used to make this observation relative to the
+	* agent's perspective, e.g. by passing the agent's forward rotation. Call during ULearningAgentsType::SetObservations event.
+	* @param AgentId The agent id this data corresponds to.
+	* @param Velocity The velocity currently being observed.
+	* @param RelativeRotation The frame of reference rotation.
+	*/
 	UFUNCTION(BlueprintCallable, Category = "LearningAgents")
-	void SetVelocityObservation(int32 AgentId, FVector Velocity, FRotator RelativeRotation = FRotator::ZeroRotator);
+	void SetVelocityObservation(const int32 AgentId, const FVector Velocity, const FRotator RelativeRotation = FRotator::ZeroRotator);
 
 #if ENABLE_VISUAL_LOG
+	/** Describes this observation to the visual logger for debugging purposes. */
 	virtual void VisualLog(const UE::Learning::FIndexSet Instances) const override;
 #endif
 
 	TSharedPtr<UE::Learning::FVelocityFeature> FeatureObject;
 };
-
-//------------------------------------------------------------------
\ No newline at end of file
diff --git a/Engine/Plugins/Experimental/LearningAgents/Source/LearningAgents/Public/LearningAgentsPolicy.h b/Engine/Plugins/Experimental/LearningAgents/Source/LearningAgents/Public/LearningAgentsPolicy.h
--- a/Engine/Plugins/Experimental/LearningAgents/Source/LearningAgents/Public/LearningAgentsPolicy.h
+++ b/Engine/Plugins/Experimental/LearningAgents/Source/LearningAgents/Public/LearningAgentsPolicy.h
@@ -20,8 +20,7 @@ namespace UE::Learning
 struct FDirectoryPath;
 class ULearningAgentsNeuralNetwork;
 
-//------------------------------------------------------------------
-
+/** The configurable settings for a ULearningAgentsPolicy. */
 USTRUCT(BlueprintType, Category = "LearningAgents")
 struct LEARNINGAGENTS_API FLearningAgentsPolicySettings
 {
@@ -54,8 +53,7 @@ public:
 	ELearningAgentsActivationFunction ActivationFunction = ELearningAgentsActivationFunction::ELU;
 };
 
-//------------------------------------------------------------------
-
+/** A policy that maps from observations to actions for the managed agents. */
 UCLASS(BlueprintType, Blueprintable)
 class LEARNINGAGENTS_API ULearningAgentsPolicy : public UActorComponent
 {
@@ -69,66 +67,113 @@ public:
 	ULearningAgentsPolicy(FVTableHelper& Helper);
 	virtual ~ULearningAgentsPolicy();
 
+	/** Initializes this object to be used with the given agent type and policy settings. */
 	UFUNCTION(BlueprintCallable, Category = "LearningAgents")
 	void SetupPolicy(ULearningAgentsType* InAgentType, const FLearningAgentsPolicySettings& PolicySettings = FLearningAgentsPolicySettings());
 
+	/** Returns true if SetupPolicy has been run successfully; Otherwise, false. */
 	UFUNCTION(BlueprintPure, Category = "LearningAgents")
 	bool IsPolicySetupPerformed() const;
 
 // ----- Agent Management -----
 public:
 
+	/**
+	 * Adds an agent to this policy.
+	 * @param AgentId The id of the agent to be added.
+	 * @warning The agent id must exist for the agent type.
+	 */
 	UFUNCTION(BlueprintCallable, Category = "LearningAgents")
-	void AddAgent(int32 AgentId);
+	void AddAgent(const int32 AgentId);
 
+	/**
+	* Removes an agent from this policy.
+	* @param AgentId The id of the agent to be removed.
+	* @warning The agent id must exist for this policy already.
+	*/
 	UFUNCTION(BlueprintCallable, Category = "LearningAgents")
-	void RemoveAgent(int32 AgentId);
+	void RemoveAgent(const int32 AgentId);
 
+	/** Returns true if the given id has been previously added to this policy; Otherwise, false. */
 	UFUNCTION(BlueprintPure, Category = "LearningAgents")
-	bool HasAgent(int32 AgentId) const;
+	bool HasAgent(const int32 AgentId) const;
 
+	/**
+	* Gets the agent type this policy is associated with.
+	* @param AgentClass The class to cast the agent type to (in blueprint).
+	*/
 	UFUNCTION(BlueprintPure, Category = "LearningAgents", meta = (DeterminesOutputType = "AgentClass"))
-	ULearningAgentsType* GetAgentType(TSubclassOf<ULearningAgentsType> AgentClass);
+	ULearningAgentsType* GetAgentType(const TSubclassOf<ULearningAgentsType> AgentClass);
 
 // ----- Load / Save -----
 public:
 
+	/**
+	* Load a snapshot's weights into this policy.
+	* @param Directory The directory the snapshot file is in.
+	* @param Filename The filename of the snapshot, including the file extension.
+	*/
 	UFUNCTION(BlueprintCallable, Category = "LearningAgents")
 	void LoadPolicyFromSnapshot(const FDirectoryPath& Directory, const FString Filename);
 
+	/**
+	* Save this policy's weights into a snapshot.
+	* @param Directory The directory to save the snapshot file in.
+	* @param Filename The filename of the snapshot, including the file extension.
+	*/
 	UFUNCTION(BlueprintCallable, Category = "LearningAgents")
 	void SavePolicyToSnapshot(const FDirectoryPath& Directory, const FString Filename) const;
 
+	/**
+	* Load a ULearningAgentsNeuralNetwork asset's weights into this policy.
+	* @param NeuralNetworkAsset The asset to load from.
+	*/
 	UFUNCTION(BlueprintCallable, Category = "LearningAgents")
 	void LoadPolicyFromAsset(const ULearningAgentsNeuralNetwork* NeuralNetworkAsset);
 
+	/**
+	* Save this policy's weights to a ULearningAgentsNeuralNetwork asset.
+	* @param NeuralNetworkAsset The asset to save to.
+	*/
 	UFUNCTION(BlueprintCallable, Category = "LearningAgents", Meta = (DevelopmentOnly))
 	void SavePolicyToAsset(ULearningAgentsNeuralNetwork* NeuralNetworkAsset) const;
 
 // ----- Evaluation -----
 public:
 
+	/**
+	* Calling this function will run the underlying neural network on the previously buffered observations to populate
+	* the output action buffer. This should be called after the associated agent type's EncodeObservations and
+	* before its DecodeActions.
+	*/
 	UFUNCTION(BlueprintCallable, Category = "LearningAgents")
 	void EvaluatePolicy();
 
 // ----- Non-blueprint public interface -----
 public:
 
+	/** Get a reference to this policy's neural network. */
 	UE::Learning::FNeuralNetwork& GetPolicyNetwork();
+
+	/** Get a reference to this policy's policy function object. */
 	UE::Learning::FNeuralNetworkPolicyFunction& GetPolicyObject();
 
 // ----- Private Data -----
 private:
 
+	/** The agent type this policy is associated with. */
 	UPROPERTY(VisibleAnywhere, Category = "LearningAgents")
 	TObjectPtr<ULearningAgentsType> AgentType;
 
+	/** The agent ids this policy is managing. */
 	UPROPERTY(VisibleAnywhere, Category = "LearningAgents")
 	TArray<int32> SelectedAgentIds;
 
+	/** True if this policy's SetupPolicy has been run; Otherwise, false. */
 	UPROPERTY(VisibleAnywhere, Category = "LearningAgents")
 	bool bPolicySetupPerformed = false;
 
+	/** The underlying neural network. */
 	UPROPERTY(VisibleAnywhere, Category = "LearningAgents")
 	TObjectPtr<ULearningAgentsNeuralNetwork> Network;
 
@@ -144,6 +189,7 @@ private:
 	/** Color used to draw this action in the visual log */
 	FLinearColor VisualLogColor = FColor::Purple;
 
+	/** Describes this policy to the visual logger for debugging purposes. */
 	void VisualLog(const UE::Learning::FIndexSet Instances) const;
 #endif
 };
diff --git a/Engine/Plugins/Experimental/LearningAgents/Source/LearningAgents/Public/LearningAgentsType.h b/Engine/Plugins/Experimental/LearningAgents/Source/LearningAgents/Public/LearningAgentsType.h
--- a/Engine/Plugins/Experimental/LearningAgents/Source/LearningAgents/Public/LearningAgentsType.h
+++ b/Engine/Plugins/Experimental/LearningAgents/Source/LearningAgents/Public/LearningAgentsType.h
@@ -24,6 +24,22 @@ namespace UE::Learning
 class ULearningAgentsAction;
 class ULearningAgentsObservation;
 
+/**
+* The agent type is the core class around which the rest of Learning Agents is built. It has a few responsibilities:
+*   1) It keeps track of which objects are agents.
+*   2) It defines how those agents' observations and actions are implemented.
+*   3) It provides methods that need to be called during the inference process of those agents.
+*
+* To use this class, you need to implement the SetupObservations and SetupActions functions (as well as their
+* corresponding SetObservations and SetActions functions), which will define the size of inputs and outputs to your
+* policy. Before you can do inference, you need to call SetupAgentType, which will initialize the underlying data
+* structure, and you need to call AddAgent for each object you want controlled by this agent type.
+*
+* If you have multiple different types of objects you want controlled by Learning Agents, you should create
+* one agent type per object type, rather than trying to share an agent type.
+*
+* @see ULearningAgentsTrainer to understand how training works.
+*/
 UCLASS(Abstract, BlueprintType, Blueprintable)
 class LEARNINGAGENTS_API ULearningAgentsType : public UActorComponent
 {
@@ -36,105 +52,178 @@ public:
 	ULearningAgentsType(FVTableHelper& Helper);
 	virtual ~ULearningAgentsType();
 
+	/** Sets up the agent ids so that agents can be added prior to calling SetupAgentType. */
 	virtual void PostInitProperties() override;
 
 // ----- Setup -----
 public:
 
-	/**
-	 * Initializes this object and runs the setup functions for observations and actions
-	 */
+	/** Initializes this object and runs the setup events for observations and actions. */
 	UFUNCTION(BlueprintCallable, Category = "LearningAgents")
 	void SetupAgentType();
 
+	/** Returns true if SetupAgentType has been run successfully; Otherwise, false. */
 	UFUNCTION(BlueprintPure, Category = "LearningAgents")
-	const bool IsSetupPerformed() const;
+	bool IsSetupPerformed() const;
 
+	/** Returns the maximum number of agents that this agent type is configured to handle. */
 	UFUNCTION(BlueprintPure, Category = "LearningAgents")
-	const int32 GetMaxInstanceNum() const;
+	int32 GetMaxInstanceNum() const;
 
 // ----- Agent Management -----
 public:
 
+	/**
+	 * Adds the given object as an agent to this agent type. This can be called before or after SetupAgentType.
+	 * @param Agent The object to be added.
+	 * @return The agent's newly assigned id.
+	 */
 	UFUNCTION(BlueprintCallable, Category = "LearningAgents")
 	int32 AddAgent(UObject* Agent);
 	
+	/**
+	 * Removes the agent with the given id from this agent type.
+	 * @param AgentId The id of the agent to remove.
+	 */
 	UFUNCTION(BlueprintCallable, Category = "LearningAgents")
-	void RemoveAgentById(int32 AgentId);
+	void RemoveAgentById(const int32 AgentId);
 
+	/**
+	* Removes the given agent from this agent type. Use RemoveAgentById if you have the id available as this function
+	* is slower.
+	* @param Agent The agent to be removed.
+	*/
 	UFUNCTION(BlueprintCallable, Category = "LearningAgents")
 	void RemoveAgent(UObject* Agent);
 
+	/**
+	 * Gets the agent with the given id. Calling this from blueprint with the appropriate AgentClass will automatically
+	 * cast the object to the given type. If not in a blueprint, you should use one of the other GetAgent overloads.
+	 * @param AgentId The id of the agent to get.
+	 * @param AgentClass The class to cast the agent object to (in blueprint).
+	 * @return The agent object.
+	 */
 	UFUNCTION(BlueprintPure, Category = "LearningAgents", meta = (DeterminesOutputType = "AgentClass"))
-	UObject* GetAgent(int32 AgentId, TSubclassOf<UObject> AgentClass);
+	UObject* GetAgent(const int32 AgentId, const TSubclassOf<UObject> AgentClass);
 
+	/** Returns true if the given object is an agent in this agent type; Otherwise, false. */
 	UFUNCTION(BlueprintPure, Category = "LearningAgents")
 	bool HasAgent(UObject* Agent) const;
 
+	/** Returns true if the given id is used by an agent in this agent type; Otherwise, false. */
 	UFUNCTION(BlueprintPure, Category = "LearningAgents")
-	bool HasAgentById(int32 AgentId) const;
+	bool HasAgentById(const int32 AgentId) const;
 
 public:
 
-	const UObject* GetAgent(int32 AgentId) const;
-	UObject* GetAgent(int32 AgentId);
+	/** Gets the agent corresponding to the given id. */
+	const UObject* GetAgent(const int32 AgentId) const;
+
+	/** Gets the agent corresponding to the given id. */
+	UObject* GetAgent(const int32 AgentId);
 
 // ----- Observations -----
 public:
 
 	/**
-	 * Event where all observations should be set up for this agent type.
-	 * @param AgentType This AgentType. This is a convenience for blueprints so you don't have to find the Self pin.
+	 * During this event, all observations should be added to this agent type.
+	 * @param AgentType This agent type. This is a convenience for blueprints so you don't have to find the Self pin.
+	 * @see LearningAgentsObservations.h for the list of available observations.
 	 */
 	UFUNCTION(BlueprintNativeEvent, Category = "LearningAgents")
 	void SetupObservations(ULearningAgentsType* AgentType);
 
+	/**
+	* During this event, all observations should be set for each agent.
+	* @param AgentIds The list of agent ids to set observations for.
+	* @see LearningAgentsObservations.h for the list of available observations.
+	* @see GetAgent to get the agent corresponding to each id.
+	*/
 	UFUNCTION(BlueprintNativeEvent, Category = "LearningAgents")
 	void SetObservations(const TArray<int32>& AgentIds);
 
+	/**
+	* Used by objects derived from ULearningAgentsObservation to add themselves to this agent type during their
+	* creation. You shouldn't need to call this directly.
+	*/
 	void AddObservation(TObjectPtr<ULearningAgentsObservation> Object, const TSharedRef<UE::Learning::FFeatureObject>& Feature);
 
 // ----- Actions -----
 public:
 
 	/**
-	 * Event where all actions should be set up for this agent type.
+	 * During this event, all actions should be added to this agent type.
 	 * @param AgentType This AgentType. This is a convenience for blueprints so you don't have to find the Self pin.
+	 * @see LearningAgentsActions.h for the list of available actions.
 	 */
 	UFUNCTION(BlueprintNativeEvent, Category = "LearningAgents")
 	void SetupActions(ULearningAgentsType* AgentType);
 
+	/**
+	* During this event, you should retrieve the actions and apply them to your agents.
+	* @param AgentIds The list of agent ids to get actions for.
+	* @see LearningAgentsActions.h for the list of available actions.
+	* @see GetAgent to get the agent corresponding to each id.
+	*/
 	UFUNCTION(BlueprintNativeEvent, Category = "LearningAgents")
 	void GetActions(const TArray<int32>& AgentIds);
 
+	/**
+	* Used by objects derived from ULearningAgentsAction to add themselves to this agent type during their creation.
+	* You shouldn't need to call this directly.
+	*/
 	void AddAction(TObjectPtr<ULearningAgentsAction> Object, const TSharedRef<UE::Learning::FFeatureObject>& Feature);
 
 // ----- Encoding / Decoding -----
 public:
 
+	/**
+	* Call this function when it is time to gather all the observations for your agents. This can be done each frame or
+	* you can consider wiring it up to some kind of meaningful event, e.g. a hypothetical "OnAITurnStarted" if you have
+	* a turn-based game. This will call this agent type's SetObservations event.
+	*/
 	UFUNCTION(BlueprintCallable, Category = "LearningAgents")
 	void EncodeObservations();
 
+	/**
+	* Call this function when it is time for your agents to take their actions. You most likely want to call this after
+	* your policy's EvaluatePolicy function to ensure that you are receiving the latest actions. This will call this
+	* agent type's GetActions event.
+	*/
 	UFUNCTION(BlueprintCallable, Category = "LearningAgents")
 	void DecodeActions();
 
 // ----- Non-blueprint public interface -----
 public:
 
+	/** Get a const reference to this agent type's underlying instance data. */
 	const TSharedPtr<UE::Learning::FArrayMap>& GetInstanceData() const;
 
+	/** Get a reference to this agent type's observation feature. */
 	UE::Learning::FFeatureObject& GetObservationFeature() const;
+
+	/** Get a reference to this agent type's action feature. */
 	UE::Learning::FFeatureObject& GetActionFeature() const;
+
+	/** Get a const array view of this agent type's observation objects. */
 	TConstArrayView<ULearningAgentsObservation*> GetObservationObjects() const;
+
+	/** Get a const array view of this agent type's action objects. */
 	TConstArrayView<ULearningAgentsAction*> GetActionObjects() const;
 
-	const TConstArrayView<TObjectPtr<UObject>> GetAgents() const;
-	const UE::Learning::FIndexSet GetOccupiedAgentSet() const;
-	const UE::Learning::FIndexSet GetVacantAgentSet() const;
+	/** Get a const array view of this agent type's agent objects. */
+	TConstArrayView<TObjectPtr<UObject>> GetAgents() const;
+
+	/** Get an FIndexSet with this agent type's occupied agent ids. */
+	UE::Learning::FIndexSet GetOccupiedAgentSet() const;
+
+	/** Get an FIndexSet with this agent type's vacant agent ids. */
+	UE::Learning::FIndexSet GetVacantAgentSet() const;
 
 // ----- Private Data -----
 private:
 
+	/** Update the agent sets to keep them in sync with the id lists. */
 	void UpdateAgentSets();
 
 	TArray<int32> OccupiedAgentIds;
@@ -156,15 +245,19 @@ private:
 	UPROPERTY(EditDefaultsOnly, Category = "LearningAgents", meta = (ClampMin = "1", UIMin = "1"))
 	int32 MaxInstanceNum = 1;
 
+	/** True if SetupAgentType has been performed; Otherwise, false. */
 	UPROPERTY(VisibleAnywhere, Category = "LearningAgents")
 	bool bSetupPerformed = false;
 
+	/** The list of current agents. */
 	UPROPERTY(VisibleAnywhere, Transient, Category = "LearningAgents")
 	TArray<TObjectPtr<UObject>> Agents;
 
+	/** The list of current observation objects. */
 	UPROPERTY(VisibleAnywhere, Category = "LearningAgents")
 	TArray<TObjectPtr<ULearningAgentsObservation>> ObservationObjects;
 
+	/** The list of current action objects. */
 	UPROPERTY(VisibleAnywhere, Category = "LearningAgents")
 	TArray<TObjectPtr<ULearningAgentsAction>> ActionObjects;
 };
diff --git a/Engine/Plugins/Experimental/LearningAgents/Source/LearningAgentsTraining/Private/LearningAgentsCompletions.cpp b/Engine/Plugins/Experimental/LearningAgents/Source/LearningAgentsTraining/Private/LearningAgentsCompletions.cpp
--- a/Engine/Plugins/Experimental/LearningAgents/Source/LearningAgentsTraining/Private/LearningAgentsCompletions.cpp
+++ b/Engine/Plugins/Experimental/LearningAgents/Source/LearningAgentsTraining/Private/LearningAgentsCompletions.cpp
@@ -26,11 +26,10 @@
 #define UE_LEARNING_AGENTS_VLOG_PLANE(Owner, Category, Verbosity, Location, Rotation, Axis0, Axis1, Color, Format, ...) \
 	UE_VLOG_OBOX(Owner, Category, Verbosity, FBox(-25.0f * (Axis0 + Axis1), 25.0f * (Axis0 + Axis1)), FTransform(Rotation, Location, FVector::OneVector).ToMatrixNoScale(), Color, Format, ##__VA_ARGS__)
 
-
 namespace UE::Learning::Agents::Private
 {
 	template<typename CompletionUObject, typename CompletionFObject, typename... InArgTypes>
-	CompletionUObject* AddCompletion(ULearningAgentsTrainer* AgentTrainer, FName Name, InArgTypes&& ...Args)
+	CompletionUObject* AddCompletion(ULearningAgentsTrainer* AgentTrainer, const FName Name, InArgTypes&& ...Args)
 	{
 		if (!AgentTrainer)
 		{
@@ -53,7 +52,7 @@ namespace UE::Learning::Agents::Private
 
 //------------------------------------------------------------------
 
-UConditionalCompletion* UConditionalCompletion::AddConditionalCompletion(ULearningAgentsTrainer* AgentTrainer, FName Name, ELearningAgentsCompletion InCompletionMode)
+UConditionalCompletion* UConditionalCompletion::AddConditionalCompletion(ULearningAgentsTrainer* AgentTrainer, const FName Name, const ELearningAgentsCompletion InCompletionMode)
 {
 	return UE::Learning::Agents::Private::AddCompletion<UConditionalCompletion, UE::Learning::FConditionalCompletion>(
 		AgentTrainer,
@@ -63,9 +62,9 @@ UConditionalCompletion* UConditionalCompletion::AddConditionalCompletion(ULearni
 		UE::Learning::ECompletionMode::Truncated);
 }
 
-void UConditionalCompletion::SetConditionalCompletion(int32 AgentId, bool bIsComplete)
+void UConditionalCompletion::SetConditionalCompletion(const int32 AgentId, const bool bIsComplete)
 {
-	TLearningArrayView<1, bool> View = CompletionObject->InstanceData->View(CompletionObject->ConditionHandle);
+	const TLearningArrayView<1, bool> View = CompletionObject->InstanceData->View(CompletionObject->ConditionHandle);
 
 	if (AgentId == INDEX_NONE)
 	{
@@ -111,7 +110,13 @@ void UConditionalCompletion::VisualLog(const UE::Learning::FIndexSet Instances)
 
 //------------------------------------------------------------------
 
-UPlanarPositionDifferenceCompletion* UPlanarPositionDifferenceCompletion::AddPlanarPositionDifferenceCompletion(ULearningAgentsTrainer* AgentTrainer, FName Name, float Threshold, ELearningAgentsCompletion InCompletionMode, FVector Axis0, FVector Axis1)
+UPlanarPositionDifferenceCompletion* UPlanarPositionDifferenceCompletion::AddPlanarPositionDifferenceCompletion(
+	ULearningAgentsTrainer* AgentTrainer,
+	const FName Name,
+	const float Threshold,
+	const ELearningAgentsCompletion InCompletionMode,
+	const FVector Axis0,
+	const FVector Axis1)
 {
 	return UE::Learning::Agents::Private::AddCompletion<UPlanarPositionDifferenceCompletion, UE::Learning::FPlanarPositionDifferenceCompletion>(
 		AgentTrainer,
@@ -125,10 +130,10 @@ UPlanarPositionDifferenceCompletion* UPlanarPositionDifferenceCompletion::AddPla
 		Axis1);
 }
 
-void UPlanarPositionDifferenceCompletion::SetPlanarPositionDifferenceCompletion(int32 AgentId, FVector Position0, FVector Position1)
+void UPlanarPositionDifferenceCompletion::SetPlanarPositionDifferenceCompletion(const int32 AgentId, const FVector Position0, const FVector Position1)
 {
-	TLearningArrayView<2, FVector> Position0View = CompletionObject->InstanceData->View(CompletionObject->Position0Handle);
-	TLearningArrayView<2, FVector> Position1View = CompletionObject->InstanceData->View(CompletionObject->Position1Handle);
+	const TLearningArrayView<2, FVector> Position0View = CompletionObject->InstanceData->View(CompletionObject->Position0Handle);
+	const TLearningArrayView<2, FVector> Position1View = CompletionObject->InstanceData->View(CompletionObject->Position1Handle);
 
 	if (AgentId == INDEX_NONE)
 	{
@@ -220,4 +225,4 @@ void UPlanarPositionDifferenceCompletion::VisualLog(const UE::Learning::FIndexSe
 
 #undef UE_LEARNING_AGENTS_VLOG_STRING
 #undef UE_LEARNING_AGENTS_VLOG_TRANSFORM
-#undef UE_LEARNING_AGENTS_VLOG_PLANE
\ No newline at end of file
+#undef UE_LEARNING_AGENTS_VLOG_PLANE
diff --git a/Engine/Plugins/Experimental/LearningAgents/Source/LearningAgentsTraining/Private/LearningAgentsDataStorage.cpp b/Engine/Plugins/Experimental/LearningAgents/Source/LearningAgentsTraining/Private/LearningAgentsDataStorage.cpp
--- a/Engine/Plugins/Experimental/LearningAgents/Source/LearningAgentsTraining/Private/LearningAgentsDataStorage.cpp
+++ b/Engine/Plugins/Experimental/LearningAgents/Source/LearningAgentsTraining/Private/LearningAgentsDataStorage.cpp
@@ -16,7 +16,7 @@ ULearningAgentsRecord::ULearningAgentsRecord() {}
 ULearningAgentsRecord::ULearningAgentsRecord(FVTableHelper& Helper) {}
 ULearningAgentsRecord::~ULearningAgentsRecord() {}
 
-void ULearningAgentsRecord::Init(int32 ObsNum, int32 ActNum)
+void ULearningAgentsRecord::Init(const int32 ObsNum, const int32 ActNum)
 {
 	ObservationNum = ObsNum;
 	ActionNum = ActNum;
@@ -70,7 +70,7 @@ ULearningAgentsDataStorage::ULearningAgentsDataStorage() {}
 ULearningAgentsDataStorage::ULearningAgentsDataStorage(FVTableHelper& Helper) {}
 ULearningAgentsDataStorage::~ULearningAgentsDataStorage() {}
 
-ULearningAgentsRecord* ULearningAgentsDataStorage::CreateRecord(FName RecordName, ULearningAgentsType* AgentType)
+ULearningAgentsRecord* ULearningAgentsDataStorage::CreateRecord(const FName RecordName, ULearningAgentsType* AgentType)
 {
 	if (!AgentType)
 	{
@@ -142,7 +142,7 @@ ULearningAgentsRecord* ULearningAgentsDataStorage::LoadRecord(ULearningAgentsTyp
 
 	TArray<uint8> Bytes;
 
-	FString InputFilePath = Directory.Path + FGenericPlatformMisc::GetDefaultPathSeparator() + Filename;
+	const FString InputFilePath = Directory.Path + FGenericPlatformMisc::GetDefaultPathSeparator() + Filename;
 
 	if (!FPaths::FileExists(InputFilePath))
 	{
diff --git a/Engine/Plugins/Experimental/LearningAgents/Source/LearningAgentsTraining/Private/LearningAgentsRecorder.cpp b/Engine/Plugins/Experimental/LearningAgents/Source/LearningAgentsTraining/Private/LearningAgentsRecorder.cpp
--- a/Engine/Plugins/Experimental/LearningAgents/Source/LearningAgentsTraining/Private/LearningAgentsRecorder.cpp
+++ b/Engine/Plugins/Experimental/LearningAgents/Source/LearningAgentsTraining/Private/LearningAgentsRecorder.cpp
@@ -60,7 +60,7 @@ bool ULearningAgentsRecorder::IsRecorderSetupPerformed() const
 	return AgentType ? true : false;
 }
 
-void ULearningAgentsRecorder::AddAgent(int32 AgentId)
+void ULearningAgentsRecorder::AddAgent(const int32 AgentId)
 {
 	if (!IsRecorderSetupPerformed())
 	{
@@ -90,7 +90,7 @@ void ULearningAgentsRecorder::AddAgent(int32 AgentId)
 	}
 }
 
-void ULearningAgentsRecorder::RemoveAgent(int32 AgentId)
+void ULearningAgentsRecorder::RemoveAgent(const int32 AgentId)
 {
 	if (!IsRecorderSetupPerformed())
 	{
@@ -114,7 +114,7 @@ void ULearningAgentsRecorder::RemoveAgent(int32 AgentId)
 	}
 }
 
-bool ULearningAgentsRecorder::HasAgent(int32 AgentId) const
+bool ULearningAgentsRecorder::HasAgent(const int32 AgentId) const
 {
 	return SelectedAgentsSet.Contains(AgentId);
 }
diff --git a/Engine/Plugins/Experimental/LearningAgents/Source/LearningAgentsTraining/Private/LearningAgentsRewards.cpp b/Engine/Plugins/Experimental/LearningAgents/Source/LearningAgentsTraining/Private/LearningAgentsRewards.cpp
--- a/Engine/Plugins/Experimental/LearningAgents/Source/LearningAgentsTraining/Private/LearningAgentsRewards.cpp
+++ b/Engine/Plugins/Experimental/LearningAgents/Source/LearningAgentsTraining/Private/LearningAgentsRewards.cpp
@@ -28,7 +28,7 @@
 namespace UE::Learning::Agents::Private
 {
 	template<typename RewardUObject, typename RewardFObject, typename... InArgTypes>
-	RewardUObject* AddReward(ULearningAgentsTrainer* AgentTrainer, FName Name, InArgTypes&& ...Args)
+	RewardUObject* AddReward(ULearningAgentsTrainer* AgentTrainer, const FName Name, InArgTypes&& ...Args)
 	{
 		if (!AgentTrainer)
 		{
@@ -52,14 +52,14 @@ namespace UE::Learning::Agents::Private
 
 //------------------------------------------------------------------
 
-UFloatReward* UFloatReward::AddFloatReward(ULearningAgentsTrainer* AgentTrainer, FName Name, float Weight)
+UFloatReward* UFloatReward::AddFloatReward(ULearningAgentsTrainer* AgentTrainer, const FName Name, const float Weight)
 {
 	return UE::Learning::Agents::Private::AddReward<UFloatReward, UE::Learning::FFloatReward>(AgentTrainer, Name, Weight);
 }
 
-void UFloatReward::SetFloatReward(int32 AgentId, float Reward)
+void UFloatReward::SetFloatReward(const int32 AgentId, const float Reward)
 {
-	TLearningArrayView<1, float> View = RewardObject->InstanceData->View(RewardObject->ValueHandle);
+	const TLearningArrayView<1, float> View = RewardObject->InstanceData->View(RewardObject->ValueHandle);
 
 	if (AgentId == INDEX_NONE)
 	{
@@ -107,14 +107,14 @@ void UFloatReward::VisualLog(const UE::Learning::FIndexSet Instances) const
 
 //------------------------------------------------------------------
 
-UScalarVelocityReward* UScalarVelocityReward::AddScalarVelocityReward(ULearningAgentsTrainer* AgentTrainer, FName Name, float Weight, float Scale)
+UScalarVelocityReward* UScalarVelocityReward::AddScalarVelocityReward(ULearningAgentsTrainer* AgentTrainer, const FName Name, const float Weight, const float Scale)
 {
 	return UE::Learning::Agents::Private::AddReward<UScalarVelocityReward, UE::Learning::FScalarVelocityReward>(AgentTrainer, Name, Weight, Scale);
 }
 
 void UScalarVelocityReward::SetScalarVelocityReward(int32 AgentId, float Velocity)
 {
-	TLearningArrayView<1, float> VelocityView = RewardObject->InstanceData->View(RewardObject->VelocityHandle);
+	const TLearningArrayView<1, float> VelocityView = RewardObject->InstanceData->View(RewardObject->VelocityHandle);
 
 	if (AgentId == INDEX_NONE)
 	{
@@ -162,15 +162,15 @@ void UScalarVelocityReward::VisualLog(const UE::Learning::FIndexSet Instances) c
 }
 #endif
 
-ULocalDirectionalVelocityReward* ULocalDirectionalVelocityReward::AddLocalDirectionalVelocityReward(ULearningAgentsTrainer* AgentTrainer, FName Name, float Weight, float Scale, FVector Axis)
+ULocalDirectionalVelocityReward* ULocalDirectionalVelocityReward::AddLocalDirectionalVelocityReward(ULearningAgentsTrainer* AgentTrainer, const FName Name, const float Weight, const float Scale, const FVector Axis)
 {
 	return UE::Learning::Agents::Private::AddReward<ULocalDirectionalVelocityReward, UE::Learning::FLocalDirectionalVelocityReward>(AgentTrainer, Name, Weight, Scale, Axis);
 }
 
-void ULocalDirectionalVelocityReward::SetLocalDirectionalVelocityReward(int32 AgentId, FVector Velocity, FRotator RelativeRotation)
+void ULocalDirectionalVelocityReward::SetLocalDirectionalVelocityReward(const int32 AgentId, const FVector Velocity, const FRotator RelativeRotation)
 {
-	TLearningArrayView<1, FVector> VelocityView = RewardObject->InstanceData->View(RewardObject->VelocityHandle);
-	TLearningArrayView<1, FQuat> RelativeRotationView = RewardObject->InstanceData->View(RewardObject->RelativeRotationHandle);
+	const TLearningArrayView<1, FVector> VelocityView = RewardObject->InstanceData->View(RewardObject->VelocityHandle);
+	const TLearningArrayView<1, FQuat> RelativeRotationView = RewardObject->InstanceData->View(RewardObject->RelativeRotationHandle);
 
 	if (AgentId == INDEX_NONE)
 	{
@@ -254,15 +254,22 @@ void ULocalDirectionalVelocityReward::VisualLog(const UE::Learning::FIndexSet In
 
 //------------------------------------------------------------------
 
-UPlanarPositionDifferencePenalty* UPlanarPositionDifferencePenalty::AddPlanarPositionDifferencePenalty(ULearningAgentsTrainer* AgentTrainer, FName Name, float Weight, float Scale, float Threshold, FVector Axis0, FVector Axis1)
+UPlanarPositionDifferencePenalty* UPlanarPositionDifferencePenalty::AddPlanarPositionDifferencePenalty(
+	ULearningAgentsTrainer* AgentTrainer,
+	const FName Name,
+	const float Weight,
+	const float Scale,
+	const float Threshold,
+	const FVector Axis0,
+	const FVector Axis1)
 {
 	return UE::Learning::Agents::Private::AddReward<UPlanarPositionDifferencePenalty, UE::Learning::FPlanarPositionDifferencePenalty>(AgentTrainer, Name, Weight, Scale, Threshold, Axis0, Axis1);
 }
 
-void UPlanarPositionDifferencePenalty::SetPlanarPositionDifferencePenalty(int32 AgentId, FVector Position0, FVector Position1)
+void UPlanarPositionDifferencePenalty::SetPlanarPositionDifferencePenalty(const int32 AgentId, const FVector Position0, const FVector Position1)
 {
-	TLearningArrayView<1, FVector> Position0View = RewardObject->InstanceData->View(RewardObject->Position0Handle);
-	TLearningArrayView<1, FVector> Position1View = RewardObject->InstanceData->View(RewardObject->Position1Handle);
+	const TLearningArrayView<1, FVector> Position0View = RewardObject->InstanceData->View(RewardObject->Position0Handle);
+	const TLearningArrayView<1, FVector> Position1View = RewardObject->InstanceData->View(RewardObject->Position1Handle);
 
 	if (AgentId == INDEX_NONE)
 	{
@@ -356,26 +363,31 @@ void UPlanarPositionDifferencePenalty::VisualLog(const UE::Learning::FIndexSet I
 
 //------------------------------------------------------------------
 
-UPositionArraySimilarityReward* UPositionArraySimilarityReward::AddPositionArraySimilarityReward(ULearningAgentsTrainer* AgentTrainer, FName Name, int32 PositionNum, float Weight, float Scale)
+UPositionArraySimilarityReward* UPositionArraySimilarityReward::AddPositionArraySimilarityReward(
+	ULearningAgentsTrainer* AgentTrainer,
+	const FName Name,
+	const int32 PositionNum,
+	const float Weight,
+	const float Scale)
 {
 	return UE::Learning::Agents::Private::AddReward<UPositionArraySimilarityReward, UE::Learning::FPositionArraySimilarityReward>(AgentTrainer, Name, PositionNum, Weight, Scale);
 }
 
 void UPositionArraySimilarityReward::SetPositionArraySimilarityReward(
-	int32 AgentId,
+	const int32 AgentId,
 	const TArray<FVector>& Positions0,
 	const TArray<FVector>& Positions1,
-	FVector RelativePosition0,
-	FVector RelativePosition1,
-	FRotator RelativeRotation0,
-	FRotator RelativeRotation1)
+	const FVector RelativePosition0,
+	const FVector RelativePosition1,
+	const FRotator RelativeRotation0,
+	const FRotator RelativeRotation1)
 {
-	TLearningArrayView<2, FVector> Positions0View = RewardObject->InstanceData->View(RewardObject->Positions0Handle);
-	TLearningArrayView<2, FVector> Positions1View = RewardObject->InstanceData->View(RewardObject->Positions1Handle);
-	TLearningArrayView<1, FVector> RelativePosition0View = RewardObject->InstanceData->View(RewardObject->RelativePosition0Handle);
-	TLearningArrayView<1, FVector> RelativePosition1View = RewardObject->InstanceData->View(RewardObject->RelativePosition1Handle);
-	TLearningArrayView<1, FQuat> RelativeRotation0View = RewardObject->InstanceData->View(RewardObject->RelativeRotation0Handle);
-	TLearningArrayView<1, FQuat> RelativeRotation1View = RewardObject->InstanceData->View(RewardObject->RelativeRotation1Handle);
+	const TLearningArrayView<2, FVector> Positions0View = RewardObject->InstanceData->View(RewardObject->Positions0Handle);
+	const TLearningArrayView<2, FVector> Positions1View = RewardObject->InstanceData->View(RewardObject->Positions1Handle);
+	const TLearningArrayView<1, FVector> RelativePosition0View = RewardObject->InstanceData->View(RewardObject->RelativePosition0Handle);
+	const TLearningArrayView<1, FVector> RelativePosition1View = RewardObject->InstanceData->View(RewardObject->RelativePosition1Handle);
+	const TLearningArrayView<1, FQuat> RelativeRotation0View = RewardObject->InstanceData->View(RewardObject->RelativeRotation0Handle);
+	const TLearningArrayView<1, FQuat> RelativeRotation1View = RewardObject->InstanceData->View(RewardObject->RelativeRotation1Handle);
 
 	if (AgentId == INDEX_NONE)
 	{
@@ -498,4 +510,4 @@ void UPositionArraySimilarityReward::VisualLog(const UE::Learning::FIndexSet Ins
 
 #undef UE_LEARNING_AGENTS_VLOG_STRING
 #undef UE_LEARNING_AGENTS_VLOG_TRANSFORM
-#undef UE_LEARNING_AGENTS_VLOG_PLANE
\ No newline at end of file
+#undef UE_LEARNING_AGENTS_VLOG_PLANE
diff --git a/Engine/Plugins/Experimental/LearningAgents/Source/LearningAgentsTraining/Private/LearningAgentsTrainer.cpp b/Engine/Plugins/Experimental/LearningAgents/Source/LearningAgentsTraining/Private/LearningAgentsTrainer.cpp
--- a/Engine/Plugins/Experimental/LearningAgents/Source/LearningAgentsTraining/Private/LearningAgentsTrainer.cpp
+++ b/Engine/Plugins/Experimental/LearningAgents/Source/LearningAgentsTraining/Private/LearningAgentsTrainer.cpp
@@ -144,7 +144,7 @@ bool ULearningAgentsTrainer::IsTrainerSetupPerformed() const
 	return bTrainerSetupPerformed;
 }
 
-void ULearningAgentsTrainer::AddAgent(int32 AgentId)
+void ULearningAgentsTrainer::AddAgent(const int32 AgentId)
 {
 	if (!IsTrainerSetupPerformed())
 	{
@@ -177,7 +177,7 @@ void ULearningAgentsTrainer::AddAgent(int32 AgentId)
 	}
 }
 
-void ULearningAgentsTrainer::RemoveAgent(int32 AgentId)
+void ULearningAgentsTrainer::RemoveAgent(const int32 AgentId)
 {
 	if (!IsTrainerSetupPerformed())
 	{
@@ -195,7 +195,7 @@ void ULearningAgentsTrainer::RemoveAgent(int32 AgentId)
 	SelectedAgentsSet.TryMakeSlice();
 }
 
-bool ULearningAgentsTrainer::HasAgent(int32 AgentId) const
+bool ULearningAgentsTrainer::HasAgent(const int32 AgentId) const
 {
 	return SelectedAgentsSet.Contains(AgentId);
 }
@@ -211,13 +211,13 @@ ULearningAgentsType* ULearningAgentsTrainer::GetAgentType(TSubclassOf<ULearningA
 	return AgentType;
 }
 
-const UObject* ULearningAgentsTrainer::GetAgent(int32 AgentId) const
+const UObject* ULearningAgentsTrainer::GetAgent(const int32 AgentId) const
 {
 	UE_LEARNING_CHECK(AgentType);
 	return AgentType->GetAgent(AgentId);
 }
 
-UObject* ULearningAgentsTrainer::GetAgent(int32 AgentId)
+UObject* ULearningAgentsTrainer::GetAgent(const int32 AgentId)
 {
 	UE_LEARNING_CHECK(AgentType);
 	return AgentType->GetAgent(AgentId);
@@ -645,7 +645,7 @@ void ULearningAgentsTrainer::IterateTraining()
 
 		// Push Experience to Replay Buffer
 
-		bool bReplayBufferFull = ReplayBuffer->AddEpisodes(
+		const bool bReplayBufferFull = ReplayBuffer->AddEpisodes(
 			Completions->CompletionBuffer(),
 			Observations.FeatureBuffer(),
 			*EpisodeBuffer,
@@ -655,9 +655,7 @@ void ULearningAgentsTrainer::IterateTraining()
 		{
 			// Send Experience
 
-			UE::Learning::ETrainerResponse Response = UE::Learning::ETrainerResponse::Success;
-
-			Response = Trainer->SendExperience(*ReplayBuffer, TrainerTimeout);
+			UE::Learning::ETrainerResponse Response = Trainer->SendExperience(*ReplayBuffer, TrainerTimeout);
 
 			if (Response != UE::Learning::ETrainerResponse::Success)
 			{
diff --git a/Engine/Plugins/Experimental/LearningAgents/Source/LearningAgentsTraining/Public/LearningAgentsCompletions.h b/Engine/Plugins/Experimental/LearningAgents/Source/LearningAgentsTraining/Public/LearningAgentsCompletions.h
--- a/Engine/Plugins/Experimental/LearningAgents/Source/LearningAgentsTraining/Public/LearningAgentsCompletions.h
+++ b/Engine/Plugins/Experimental/LearningAgents/Source/LearningAgentsTraining/Public/LearningAgentsCompletions.h
@@ -21,6 +21,17 @@ namespace UE::Learning
 // For functions in this file, we are favoring having more verbose names such as "AddConditionalCompletion" vs simply "Add" in 
 // order to keep it easy to find the correct function in blueprints.
 
+/**
+ * The base class for all completions. Completions contain logic that determines if an agent's current episode should
+ * end, e.g. because the agent achieved the normal win/loss condition for the game. Additionally, completions can speed
+ * up training by ending episodes early if the agent has gotten into a state where training data is no longer useful,
+ * e.g. the agent is stuck somewhere. These two modes of completions are expressed with the following enum values:
+ *   ELearningAgentsCompletion::Termination - used when the episode ends in an expected way and no further rewards
+ *       should be expected, i.e. do not use the value function to estimate future rewards.
+ *   ELearningAgentsCompletion::Truncation - used when the episode ends in an unexpected way, mainly to speed up the 
+ *       training process. The agent should expect additional rewards if training were to continue, so it should use
+ *       its value function to estimate future rewards.
+ */
 UCLASS(Abstract, BlueprintType)
 class LEARNINGAGENTSTRAINING_API ULearningAgentsCompletion : public UObject
 {
@@ -32,52 +43,87 @@ public:
 	/** Color used to draw this completion in the visual log */
 	FLinearColor VisualLogColor = FColor::Yellow;
 
+	/** Describes this observation to the visual logger for debugging purposes. */
 	virtual void VisualLog(const UE::Learning::FIndexSet Instances) const {}
 #endif
 };
 
-//------------------------------------------------------------------
-
+/** A simple boolean completion. Used as a catch-all for situations where a more type-specific completion does not exist yet. */
 UCLASS()
 class LEARNINGAGENTSTRAINING_API UConditionalCompletion : public ULearningAgentsCompletion
 {
 	GENERATED_BODY()
 
 public:
+
+	/**
+	* Adds a new conditional completion to the given trainer. Call during ULearningAgentsTrainer::SetupCompletions event.
+	* @param AgentTrainer The trainer to add this completion to.
+	* @param Name The name of this new completion. Used for debugging.
+	* @param InCompletionMode The completion mode.
+	* @return The newly created completion.
+	*/
 	UFUNCTION(BlueprintCallable, Category = "LearningAgents")
-	static UConditionalCompletion* AddConditionalCompletion(ULearningAgentsTrainer* AgentTrainer, FName Name = NAME_None, ELearningAgentsCompletion InCompletionMode = ELearningAgentsCompletion::Termination);
+	static UConditionalCompletion* AddConditionalCompletion(ULearningAgentsTrainer* AgentTrainer, const FName Name = NAME_None, const ELearningAgentsCompletion InCompletionMode = ELearningAgentsCompletion::Termination);
 
+	/**
+	* Sets the data for this completion. Call during ULearningAgentsTrainer::SetCompletions event.
+	* @param AgentId The agent id this data corresponds to.
+	* @param bIsCompleted Pass in true if condition is met. Otherwise, false.
+	*/
 	UFUNCTION(BlueprintCallable, Category = "LearningAgents")
-	void SetConditionalCompletion(int32 AgentId, bool bIsCompleted);
+	void SetConditionalCompletion(const int32 AgentId, const bool bIsCompleted);
 
 #if ENABLE_VISUAL_LOG
+	/** Describes this observation to the visual logger for debugging purposes. */
 	virtual void VisualLog(const UE::Learning::FIndexSet Instances) const override;
 #endif
 
 	TSharedPtr<UE::Learning::FConditionalCompletion> CompletionObject;
 };
 
-//------------------------------------------------------------------
-
+/**
+* A completion for if two positions differ by some threshold in a plane, e.g. if the agent gets too far from a
+* starting position.
+*/
 UCLASS()
 class LEARNINGAGENTSTRAINING_API UPlanarPositionDifferenceCompletion : public ULearningAgentsCompletion
 {
 	GENERATED_BODY()
 
 public:
+
+	/**
+	* Adds a new planar position difference completion to the given trainer. The axis parameters define the plane.
+	* Call during ULearningAgentsTrainer::SetupCompletions event.
+	* @param AgentTrainer The trainer to add this completion to.
+	* @param Name The name of this new completion. Used for debugging.
+	* @param Threshold If the distance becomes greater than this threshold, then the episode will complete.
+	* @param InCompletionMode The completion mode.
+	* @param Axis0 The forward axis of the plane.
+	* @param Axis1 The right axis of the plane.
+	* @return The newly created completion.
+	*/
 	UFUNCTION(BlueprintCallable, Category = "LearningAgents")
 	static UPlanarPositionDifferenceCompletion* AddPlanarPositionDifferenceCompletion(
 		ULearningAgentsTrainer* AgentTrainer, 
-		FName Name = NAME_None,
-		float Threshold = 100.0f, 
-		ELearningAgentsCompletion InCompletionMode = ELearningAgentsCompletion::Termination,
-		FVector Axis0 = FVector::ForwardVector,
-		FVector Axis1 = FVector::RightVector);
-
+		const FName Name = NAME_None,
+		const float Threshold = 100.0f,
+		const ELearningAgentsCompletion InCompletionMode = ELearningAgentsCompletion::Termination,
+		const FVector Axis0 = FVector::ForwardVector,
+		const FVector Axis1 = FVector::RightVector);
+
+	/**
+	* Sets the data for this completion. Call during ULearningAgentsTrainer::SetCompletions event.
+	* @param AgentId The agent id this data corresponds to.
+	* @param Position0 The first position.
+	* @param Position1 The second position.
+	*/
 	UFUNCTION(BlueprintCallable, Category = "LearningAgents")
-	void SetPlanarPositionDifferenceCompletion(int32 AgentId, FVector Position0, FVector Position1);
+	void SetPlanarPositionDifferenceCompletion(const int32 AgentId, const FVector Position0, const FVector Position1);
 
 #if ENABLE_VISUAL_LOG
+	/** Describes this observation to the visual logger for debugging purposes. */
 	virtual void VisualLog(const UE::Learning::FIndexSet Instances) const override;
 #endif
 
diff --git a/Engine/Plugins/Experimental/LearningAgents/Source/LearningAgentsTraining/Public/LearningAgentsDataStorage.h b/Engine/Plugins/Experimental/LearningAgents/Source/LearningAgentsTraining/Public/LearningAgentsDataStorage.h
--- a/Engine/Plugins/Experimental/LearningAgents/Source/LearningAgentsTraining/Public/LearningAgentsDataStorage.h
+++ b/Engine/Plugins/Experimental/LearningAgents/Source/LearningAgentsTraining/Public/LearningAgentsDataStorage.h
@@ -15,30 +15,38 @@
 struct FDirectoryPath;
 class ULearningAgentsType;
 
+/** Metadata for an agent record. */
 USTRUCT(BlueprintType)
 struct FRecordMetaData
 {
 	GENERATED_BODY()
 
+	/** The agent type the data was recorded from. Determine the shape of observations and actions. */
 	UPROPERTY(BlueprintReadOnly, VisibleAnywhere, Category = "LearningAgents")
 	TObjectPtr<ULearningAgentsType> AgentType;
 
+	/** The time the data was recorded (local). */
 	UPROPERTY(BlueprintReadOnly, VisibleAnywhere, Category = "LearningAgents")
 	FDateTime CreatedOn;
 
+	/** The time the data was recorded (universal). */
 	UPROPERTY(BlueprintReadOnly, VisibleAnywhere, Category = "LearningAgents")
 	FDateTime CreatedOnUtc;
 
+	/** The name of the record. */
 	UPROPERTY(BlueprintReadOnly, VisibleAnywhere, Category = "LearningAgents")
 	FName RecordName;
 
+	/** True if this record was loaded from a file. Otherwise, false. */
 	UPROPERTY(BlueprintReadOnly, VisibleAnywhere, Category = "LearningAgents")
 	bool bWasLoadedFromFile = false;
 
+	/** True if this record has unsaved changes. Otherwise, false. */
 	UPROPERTY(BlueprintReadOnly, VisibleAnywhere, Category = "LearningAgents")
 	bool bHasUnsavedChanges = false;
 };
 
+/** A recording of a human/AI demonstration from which we can learn. */
 UCLASS(BlueprintType)
 class LEARNINGAGENTSTRAINING_API ULearningAgentsRecord : public UObject
 {
@@ -52,15 +60,22 @@ public:
 	ULearningAgentsRecord(FVTableHelper& Helper);
 	virtual ~ULearningAgentsRecord();
 
-	void Init(int32 ObsNum, int32 ActNum);
+	/** Initialize this record with the given observation and action sizes. */
+	void Init(const int32 ObsNum, const int32 ActNum);
 
+	/** Add experience data to this record. */
 	void AddExperience(TLearningArrayView<1, const float> Observations, TLearningArrayView<1, const float> Actions);
 
-	 const TConstArrayView<TLearningArray<2, float>> GetObservations() const;
+	/** Get a const view of the observations. */
+	const TConstArrayView<TLearningArray<2, float>> GetObservations() const;
+
+	/** Get a const view of the actions. */
 	const TConstArrayView<TLearningArray<2, float>> GetActions() const;
 
+	/** Remove unused space from this record. Call after no more data will be written. */
 	void Trim();
 
+	/** The metadata for this record. */
 	UPROPERTY(BlueprintReadOnly, VisibleAnywhere, Category = "LearningAgents")
 	FRecordMetaData MetaData;
 
@@ -77,6 +92,7 @@ private:
 	TArray<TLearningArray<2, float>> Actions;
 };
 
+/** A manager for the saving/loading of recordings. */
 UCLASS(BlueprintType, Blueprintable)
 class LEARNINGAGENTSTRAINING_API ULearningAgentsDataStorage : public UObject
 {
@@ -88,31 +104,41 @@ public:
 	ULearningAgentsDataStorage(FVTableHelper& Helper);
 	virtual ~ULearningAgentsDataStorage();
 
+	/** Create a new record with the given name for the agent type. */
 	UFUNCTION(BlueprintCallable, Category = "LearningAgents")
-	ULearningAgentsRecord* CreateRecord(FName RecordName, ULearningAgentsType* AgentType);
+	ULearningAgentsRecord* CreateRecord(const FName RecordName, ULearningAgentsType* AgentType);
 
+	/** Get all currently loaded records. */
 	UFUNCTION(BlueprintPure, Category = "LearningAgents")
 	const TArray<ULearningAgentsRecord*>& GetAllRecords() const;
 
+	/** Load all records from a given directory, ensuring they are valid for the given agent type. */
 	UFUNCTION(BlueprintCallable, Category = "LearningAgents")
 	int32 LoadAllRecords(ULearningAgentsType* AgentType, const FDirectoryPath& Directory);
 
+	/** Load a record from a given directory and filename, ensuring the data is valid for the given agent type. */
 	UFUNCTION(BlueprintCallable, Category = "LearningAgents")
 	ULearningAgentsRecord* LoadRecord(ULearningAgentsType* AgentType, const FDirectoryPath& Directory, const FString& Filename);
 
+	/** Save all records which have unsaved changes to the given directory. */
 	UFUNCTION(BlueprintCallable, Category = "LearningAgents")
 	void SaveAllRecords(const FDirectoryPath& Directory) const;
 
+	/** Forcefully save a record to the given directory. */
 	UFUNCTION(BlueprintCallable, Category = "LearningAgents")
 	void SaveRecord(const FDirectoryPath& Directory, ULearningAgentsRecord* Record) const;
 
 private:
+	
+	/** All currently loaded records. */
 	UPROPERTY(VisibleAnywhere, Category = "LearningAgents")
 	TArray<TObjectPtr<ULearningAgentsRecord>> Records;
 
+	/** If true, prepends a timestamp to the saved file names. */
 	UPROPERTY(EditDefaultsOnly, Category = "LearningAgents")
 	bool bPrependUtcTimeStamp = true;
 
+	/** The file extension to use when searching for records to load or saving a new record. */
 	UPROPERTY(EditDefaultsOnly, Category = "LearningAgents")
 	FString FileExtension = TEXT("bin");
 };
diff --git a/Engine/Plugins/Experimental/LearningAgents/Source/LearningAgentsTraining/Public/LearningAgentsImitationTrainer.h b/Engine/Plugins/Experimental/LearningAgents/Source/LearningAgentsTraining/Public/LearningAgentsImitationTrainer.h
--- a/Engine/Plugins/Experimental/LearningAgents/Source/LearningAgentsTraining/Public/LearningAgentsImitationTrainer.h
+++ b/Engine/Plugins/Experimental/LearningAgents/Source/LearningAgentsTraining/Public/LearningAgentsImitationTrainer.h
@@ -19,8 +19,7 @@ namespace UE::Learning
 class ULearningAgentsType;
 class ULearningAgentsPolicy;
 
-//------------------------------------------------------------------
-
+/** The configurable settings for the training process. */
 USTRUCT(BlueprintType, Category = "LearningAgents")
 struct FLearningAgentsImitationTrainerTrainingSettings
 {
@@ -28,33 +27,54 @@ struct FLearningAgentsImitationTrainerTrainingSettings
 
 public:
 
+	/** The number of iterations to run before training is complete. */
 	UPROPERTY(EditAnywhere, Category = "LearningAgents", meta = (ClampMin = "0", UIMin = "0"))
 	int32 NumberOfIterations = 1000000;
 
+	/** Learning rate of the policy network. Typical values are between 0.001f and 0.0001f. */
 	UPROPERTY(EditAnywhere, Category = "LearningAgents", meta = (ClampMin = "0.0", UIMin = "0.0"))
 	float LearningRate = 0.0001f;
 
+	/** Ratio by which to decay the learning rate every 1000 iterations. */
 	UPROPERTY(EditAnywhere, Category = "LearningAgents", meta = (ClampMin = "0.0", UIMin = "0.0"))
 	float LearningRateDecay = 0.99f;
 
+	/**
+	* Amount of weight decay to apply to the network. Larger values encourage network weights to be smaller but too
+	* large a value can cause the network weights to collapse to all zeros.
+	*/
 	UPROPERTY(EditAnywhere, Category = "LearningAgents", meta = (ClampMin = "0.0", UIMin = "0.0"))
 	float WeightDecay = 0.001f;
 
+	/**
+	* Batch size to use for training. Smaller values tend to produce better results at the cost of slowing down
+	* training.
+	*/
 	UPROPERTY(EditAnywhere, Category = "LearningAgents", meta = (ClampMin = "1", UIMin = "1"))
 	uint32 BatchSize = 128;
 
+	/** The seed used for any random sampling the trainer will perform, e.g. for weight initialization. */
 	UPROPERTY(EditAnywhere, Category = "LearningAgents", meta = (ClampMin = "0", UIMin = "0"))
 	int32 RandomSeed = 1234;
 
+	/** The device to train on. */
 	UPROPERTY(EditAnywhere, Category = "LearningAgents")
 	ELearningAgentsTrainerDevice Device = ELearningAgentsTrainerDevice::CPU;
 
+	/** If true, TensorBoard logs will be emitted to Intermediate\TensorBoard. Otherwise, no logs will be emitted. */
 	UPROPERTY(EditAnywhere, Category = "LearningAgents")
 	bool bUseTensorboard = false;
 };
 
-//------------------------------------------------------------------
-
+/**
+* The ULearningAgentsImitationTrainer enable imitation learning, i.e. learning from human/AI demonstrations.
+* Imitation training is typically much faster than reinforcement learning, but requires gathering large amounts of
+* data in order to generalize. This can be used to initialize a reinforcement learning policy to speed up initial
+* exploration.
+* @see ULearningAgentsType to understand how observations and actions work.
+* @see ULearningAgentsRecorder to understand how to make new recordings.
+* @see ULearningAgentsDataStorage to understand how to retrieve previous recordings.
+*/
 UCLASS(BlueprintType, Blueprintable)
 class LEARNINGAGENTSTRAINING_API ULearningAgentsImitationTrainer : public UActorComponent
 {
@@ -68,8 +88,16 @@ public:
 	ULearningAgentsImitationTrainer(FVTableHelper& Helper);
 	virtual ~ULearningAgentsImitationTrainer();
 
+	/** Will automatically call EndTraining if training is still in-progress when play is ending. */
 	virtual void EndPlay(const EEndPlayReason::Type EndPlayReason) override;
 
+	/**
+	* Begins the training process with the provided settings.
+	* @param InPolicy The policy to train.
+	* @param Records The data to train on.
+	* @param TrainingSettings The settings for this training run.
+	* @param bReinitializePolicyNetwork If true, reinitialize the policy. Set this to false if your policy is pre-trained.
+	*/
 	UFUNCTION(BlueprintCallable, Category = "LearningAgents")
 	void BeginTraining(
 		ULearningAgentsPolicy* InPolicy,
@@ -77,18 +105,22 @@ public:
 		const FLearningAgentsImitationTrainerTrainingSettings& TrainingSettings = FLearningAgentsImitationTrainerTrainingSettings(),
 		const bool bReinitializePolicyNetwork = true);
 
+	/** Stops the training process. */
 	UFUNCTION(BlueprintCallable, Category = "LearningAgents")
 	void EndTraining();
 
+	/** Returns true if the trainer is currently training; Otherwise, false. */
 	UFUNCTION(BlueprintPure, Category = "LearningAgents")
 	bool IsTraining() const;
 
+	/** Returns true if the previously launched training has completed; Otherwise, false. */
 	UFUNCTION(BlueprintPure, Category = "LearningAgents")
 	bool IsTrainingComplete() const;
 
 // ----- Private Data -----
 private:
 
+	/** The policy being trained. */
 	UPROPERTY(VisibleAnywhere, Category = "LearningAgents")
 	TObjectPtr<ULearningAgentsPolicy> Policy;
 
@@ -96,6 +128,7 @@ private:
 	UPROPERTY(VisibleAnywhere, Category = "LearningAgents")
 	bool bIsTraining = false;
 
+	/** True if training is completed. Otherwise, false. */
 	UPROPERTY(VisibleAnywhere, Category = "LearningAgents")
 	bool bIsTrainingComplete = false;
 
diff --git a/Engine/Plugins/Experimental/LearningAgents/Source/LearningAgentsTraining/Public/LearningAgentsRecorder.h b/Engine/Plugins/Experimental/LearningAgents/Source/LearningAgentsTraining/Public/LearningAgentsRecorder.h
--- a/Engine/Plugins/Experimental/LearningAgents/Source/LearningAgentsTraining/Public/LearningAgentsRecorder.h
+++ b/Engine/Plugins/Experimental/LearningAgents/Source/LearningAgentsTraining/Public/LearningAgentsRecorder.h
@@ -13,6 +13,7 @@
 class ULearningAgentsDataStorage;
 class ULearningAgentsRecord;
 
+/** A component that can be used to create recordings of training data for imitation learning. */
 UCLASS(BlueprintType, Blueprintable)
 class LEARNINGAGENTSTRAINING_API ULearningAgentsRecorder : public UActorComponent
 {
@@ -26,68 +27,99 @@ public:
 	ULearningAgentsRecorder(FVTableHelper& Helper);
 	virtual ~ULearningAgentsRecorder();
 
+	/** Will automatically call EndRecording if recording is still in-progress when play is ending. */
 	virtual void EndPlay(const EEndPlayReason::Type EndPlayReason) override;
 
+	/**
+	* Initializes this object and runs the setup functions for the underlying data storage.
+	* @param InAgentType The agent type we are recording with.
+	*/
 	UFUNCTION(BlueprintCallable, Category = "LearningAgents")
 	void SetupRecorder(ULearningAgentsType* InAgentType);
 
+	/** Returns true if SetupRecorder has been run successfully; Otherwise, false. */
 	UFUNCTION(BlueprintPure, Category = "LearningAgents")
 	bool IsRecorderSetupPerformed() const;
 
 // ----- Agent Management -----
 public:
 
+	/**
+	 * Adds an agent to this recorder.
+	 * @param AgentId The id of the agent to be added.
+	 * @warning The agent id must exist for the agent type.
+	 */
 	UFUNCTION(BlueprintCallable, Category = "LearningAgents")
-	void AddAgent(int32 AgentId);
+	void AddAgent(const int32 AgentId);
 
+	/**
+	* Removes an agent from this recorder.
+	* @param AgentId The id of the agent to be removed.
+	* @warning The agent id must exist for this recorder already.
+	*/
 	UFUNCTION(BlueprintCallable, Category = "LearningAgents")
-	void RemoveAgent(int32 AgentId);
+	void RemoveAgent(const int32 AgentId);
 
+	/** Returns true if the given id has been previously added to this recorder; Otherwise, false. */
 	UFUNCTION(BlueprintPure, Category = "LearningAgents")
-	bool HasAgent(int32 AgentId) const;
+	bool HasAgent(const int32 AgentId) const;
 
+	/**
+	* Gets the agent type this recorder is associated with.
+	* @param AgentClass The class to cast the agent type to (in blueprint).
+	*/
 	UFUNCTION(BlueprintPure, Category = "LearningAgents", meta = (DeterminesOutputType = "AgentClass"))
-	ULearningAgentsType* GetAgentType(TSubclassOf<ULearningAgentsType> AgentClass);
+	ULearningAgentsType* GetAgentType(const TSubclassOf<ULearningAgentsType> AgentClass);
 
 // ----- Recording Process -----
 public:
 
-	/** Adds experience to tracked agents' records. Each agent id whose record is full will be returned, e.g. so you can reset them. */
+	/**
+	* Adds experience to the added agents' recordings. Call this after ULearningAgentsType::EncodeObservations and
+	* either ULearningAgentsController::EncodeActions (if recording a human/AI demonstration) or
+	* ULearningAgentsType::DecodeActions (if recording another policy).
+	*/
 	UFUNCTION(BlueprintCallable, Category = "LearningAgents")
 	void AddExperience();
 
+	/** Begin new recordings for each added agent. */
 	UFUNCTION(BlueprintCallable, Category = "LearningAgents")
 	void BeginRecording();
 
+	/** End all recordings. */
 	UFUNCTION(BlueprintCallable, Category = "LearningAgents")
 	void EndRecording();
 
+	/** Returns true if the recorder is currently recording; Otherwise, false. */
 	UFUNCTION(BlueprintCallable, Category = "LearningAgents")
 	bool IsRecording() const;
 
 // ----- Private Data ----- 
 private:
 
+	/** The agent type this recorder is associated with. */
 	UPROPERTY(VisibleAnywhere, Category = "LearningAgents")
 	TObjectPtr<ULearningAgentsType> AgentType;
 
+	/** The agent ids this recorder is managing. */
 	UPROPERTY(VisibleAnywhere, Category = "LearningAgents")
 	TArray<int32> SelectedAgentIds;
 
 // ----- Recorder Configuration -----
 private:
 
-	/** Directory where records will be saved. If not set, BeginPlay will automatically set this to the editor's default intermediate folder. */
+	/** Directory where records will be saved. If not set, SetupRecorder will automatically set this to the editor's default intermediate folder. */
 	UPROPERTY(EditDefaultsOnly, Category = "LearningAgents")
 	FDirectoryPath DataDirectory;
 
-	/** If true, recorder will automatically save all records on EndPlay. Set this to false if you want to manually save records. */
+	/** If true, recorder will automatically save all records on EndRecording. Set this to false if you want to manually save records. */
 	UPROPERTY(EditAnywhere, Category = "LearningAgents")
 	bool bSaveDataOnEndPlay = true;
 
 // ----- Recorder State -----
 private:
 
+	/** True if recording is currently in-progress. Otherwise, false. */
 	UPROPERTY(VisibleAnywhere, Category = "LearningAgents")
 	bool bIsRecording = false;
 
@@ -95,7 +127,7 @@ private:
 	UPROPERTY(VisibleAnywhere, Category = "LearningAgents")
 	TObjectPtr<ULearningAgentsDataStorage> DataStorage;
 
-	/** All records which are currently being written */
+	/** All records which are currently being written to. */
 	UPROPERTY(VisibleAnywhere, Category = "LearningAgents")
 	TMap<int32, TObjectPtr<ULearningAgentsRecord>> CurrentRecords;
 
diff --git a/Engine/Plugins/Experimental/LearningAgents/Source/LearningAgentsTraining/Public/LearningAgentsRewards.h b/Engine/Plugins/Experimental/LearningAgents/Source/LearningAgentsTraining/Public/LearningAgentsRewards.h
--- a/Engine/Plugins/Experimental/LearningAgents/Source/LearningAgentsTraining/Public/LearningAgentsRewards.h
+++ b/Engine/Plugins/Experimental/LearningAgents/Source/LearningAgentsTraining/Public/LearningAgentsRewards.h
@@ -26,7 +26,10 @@ class ULearningAgentsTrainer;
 // For functions in this file, we are favoring having more verbose names such as "AddFloatReward" vs simply "Add" in 
 // order to keep it easy to find the correct function in blueprints.
 
-/** Base class for all rewards */
+/**
+ * Base class for all rewards/penalties. Rewards are used during reinforcement learning to encourage/discourage
+ * certain behaviors from occurring.
+ */
 UCLASS(Abstract, BlueprintType)
 class LEARNINGAGENTSTRAINING_API ULearningAgentsReward : public UObject
 {
@@ -38,125 +41,228 @@ public:
 	/** Color used to draw this reward in the visual log */
 	FLinearColor VisualLogColor = FColor::Green;
 
+	/** Describes this reward to the visual logger for debugging purposes. */
 	virtual void VisualLog(const UE::Learning::FIndexSet Instances) const {}
 #endif
 };
 
-//------------------------------------------------------------------
-
+/** A simple float reward. Used as a catch-all for situations where a more type-specific reward does not exist yet. */
 UCLASS()
 class LEARNINGAGENTSTRAINING_API UFloatReward : public ULearningAgentsReward
 {
 	GENERATED_BODY()
 
 public:
-	UFUNCTION(BlueprintCallable, Category = "LearningAgents")
-	static UFloatReward* AddFloatReward(ULearningAgentsTrainer* AgentTrainer, FName Name = NAME_None, float Weight = 1.0f);
 
+	/**
+	* Adds a new float reward to the given trainer. Call during ULearningAgentsTrainer::SetupRewards event.
+	* @param AgentTrainer The trainer to add this reward to.
+	* @param Name The name of this new reward. Used for debugging.
+	* @param Weight Multiplier for this reward when being summed up for the total reward.
+	* @return The newly created reward.
+	*/
 	UFUNCTION(BlueprintCallable, Category = "LearningAgents")
-	void SetFloatReward(int32 AgentId, float Reward);
+	static UFloatReward* AddFloatReward(
+		ULearningAgentsTrainer* AgentTrainer,
+		const FName Name = NAME_None,
+		const float Weight = 1.0f);
+
+	/**
+	* Sets the data for this reward. Call during ULearningAgentsTrainer::SetRewards event.
+	* @param AgentId The agent id this data corresponds to.
+	* @param Reward The value currently being rewarded.
+	*/
+	UFUNCTION(BlueprintCallable, Category = "LearningAgents")
+	void SetFloatReward(const int32 AgentId, const float Reward);
 
 #if ENABLE_VISUAL_LOG
+	/** Describes this reward to the visual logger for debugging purposes. */
 	virtual void VisualLog(const UE::Learning::FIndexSet Instances) const override;
 #endif
 
 	TSharedPtr<UE::Learning::FFloatReward> RewardObject;
 };
 
-//------------------------------------------------------------------
-
+/** A reward for maximizing speed. */
 UCLASS()
 class LEARNINGAGENTSTRAINING_API UScalarVelocityReward : public ULearningAgentsReward
 {
 	GENERATED_BODY()
 
 public:
-	UFUNCTION(BlueprintCallable, Category = "LearningAgents")
-	static UScalarVelocityReward* AddScalarVelocityReward(ULearningAgentsTrainer* AgentTrainer, FName Name = NAME_None, float Weight = 1.0f, float Scale = 200.0f);
 
+	/**
+	* Adds a new scalar velocity reward to the given trainer. Call during ULearningAgentsTrainer::SetupRewards event.
+	* @param AgentTrainer The trainer to add this reward to.
+	* @param Name The name of this new reward. Used for debugging.
+	* @param Weight Multiplier for this reward when being summed up for the total reward.
+	* @param Scale Used to normalize the data for the reward.
+	* @return The newly created reward.
+	*/
+	UFUNCTION(BlueprintCallable, Category = "LearningAgents")
+	static UScalarVelocityReward* AddScalarVelocityReward(
+		ULearningAgentsTrainer* AgentTrainer,
+		const FName Name = NAME_None,
+		const float Weight = 1.0f,
+		const float Scale = 200.0f);
+
+	/**
+	* Sets the data for this reward. Call during ULearningAgentsTrainer::SetRewards event.
+	* @param AgentId The agent id this data corresponds to.
+	* @param Velocity The current scalar velocity.
+	*/
 	UFUNCTION(BlueprintCallable, Category = "LearningAgents")
-	void SetScalarVelocityReward(int32 AgentId, float Velocity);
+	void SetScalarVelocityReward(const int32 AgentId, const float Velocity);
 
 #if ENABLE_VISUAL_LOG
+	/** Describes this reward to the visual logger for debugging purposes. */
 	virtual void VisualLog(const UE::Learning::FIndexSet Instances) const override;
 #endif
 
 	TSharedPtr<UE::Learning::FScalarVelocityReward> RewardObject;
 };
 
+/** A reward for maximizing velocity along a given local axis. */
 UCLASS()
 class LEARNINGAGENTSTRAINING_API ULocalDirectionalVelocityReward : public ULearningAgentsReward
 {
 	GENERATED_BODY()
 
 public:
-	UFUNCTION(BlueprintCallable, Category = "LearningAgents")
-	static ULocalDirectionalVelocityReward* AddLocalDirectionalVelocityReward(ULearningAgentsTrainer* AgentTrainer, FName Name = NAME_None, float Weight = 1.0f, float Scale = 200.0f, FVector Axis = FVector::ForwardVector);
 
+	/**
+	* Adds a new directional velocity reward to the given trainer. Call during ULearningAgentsTrainer::SetupRewards event.
+	* @param AgentTrainer The trainer to add this reward to.
+	* @param Name The name of this new reward. Used for debugging.
+	* @param Weight Multiplier for this reward when being summed up for the total reward.
+	* @param Scale Used to normalize the data for the reward.
+	* @param Axis The local direction we want to maximize velocity in.
+	* @return The newly created reward.
+	*/
+	UFUNCTION(BlueprintCallable, Category = "LearningAgents")
+	static ULocalDirectionalVelocityReward* AddLocalDirectionalVelocityReward(
+		ULearningAgentsTrainer* AgentTrainer,
+		const FName Name = NAME_None,
+		const float Weight = 1.0f,
+		const float Scale = 200.0f,
+		const FVector Axis = FVector::ForwardVector);
+
+	/**
+	* Sets the data for this reward. Call during ULearningAgentsTrainer::SetRewards event.
+	* @param AgentId The agent id this data corresponds to.
+	* @param Velocity The current velocity.
+	* @param RelativeRotation The frame of reference rotation.
+	*/
 	UFUNCTION(BlueprintCallable, Category = "LearningAgents")
-	void SetLocalDirectionalVelocityReward(int32 AgentId, FVector Velocity, FRotator RelativeRotation = FRotator::ZeroRotator);
+	void SetLocalDirectionalVelocityReward(
+		const int32 AgentId,
+		const FVector Velocity,
+		const FRotator RelativeRotation = FRotator::ZeroRotator);
 
 #if ENABLE_VISUAL_LOG
+	/** Describes this reward to the visual logger for debugging purposes. */
 	virtual void VisualLog(const UE::Learning::FIndexSet Instances) const override;
 #endif
 
 	TSharedPtr<UE::Learning::FLocalDirectionalVelocityReward> RewardObject;
 };
 
-//------------------------------------------------------------------
-
+/** A penalty for being far from a goal position in a plane. */
 UCLASS()
 class LEARNINGAGENTSTRAINING_API UPlanarPositionDifferencePenalty : public ULearningAgentsReward
 {
 	GENERATED_BODY()
 
 public:
+
+	/**
+	* Adds a new planar difference penalty to the given trainer. The axis parameters define the plane.
+	* Call during ULearningAgentsTrainer::SetupRewards event.
+	* @param AgentTrainer The trainer to add this penalty to.
+	* @param Name The name of this new penalty. Used for debugging.
+	* @param Weight Multiplier for this penalty when being summed up for the total reward.
+	* @param Scale Used to normalize the data for the penalty.
+	* @param Threshold Minimal distance to apply this penalty.
+	* @param Axis0 The forward axis of the plane.
+	* @param Axis1 The right axis of the plane.
+	* @return The newly created reward.
+	*/
 	UFUNCTION(BlueprintCallable, Category = "LearningAgents")
 	static UPlanarPositionDifferencePenalty* AddPlanarPositionDifferencePenalty(
 		ULearningAgentsTrainer* AgentTrainer, 
-		FName Name = NAME_None,
-		float Weight = 1.0f, 
-		float Scale = 100.0f, 
-		float Threshold = 0.0f, 
-		FVector Axis0 = FVector::ForwardVector, 
-		FVector Axis1 = FVector::RightVector);
-
+		const FName Name = NAME_None,
+		const float Weight = 1.0f,
+		const float Scale = 100.0f,
+		const float Threshold = 0.0f,
+		const FVector Axis0 = FVector::ForwardVector,
+		const FVector Axis1 = FVector::RightVector);
+
+	/**
+	* Sets the data for this penalty. Call during ULearningAgentsTrainer::SetRewards event.
+	* @param AgentId The agent id this data corresponds to.
+	* @param Position0 The current position.
+	* @param Position1 The goal position.
+	*/
 	UFUNCTION(BlueprintCallable, Category = "LearningAgents")
-	void SetPlanarPositionDifferencePenalty(int32 AgentId, FVector Position0, FVector Position1);
+	void SetPlanarPositionDifferencePenalty(const int32 AgentId, const FVector Position0, const FVector Position1);
 
 #if ENABLE_VISUAL_LOG
+	/** Describes this reward to the visual logger for debugging purposes. */
 	virtual void VisualLog(const UE::Learning::FIndexSet Instances) const override;
 #endif
 
 	TSharedPtr<UE::Learning::FPlanarPositionDifferencePenalty> RewardObject;
 };
 
-//------------------------------------------------------------------
-
+/** A reward for minimizing the distances of positions in the given arrays. */
 UCLASS()
 class LEARNINGAGENTSTRAINING_API UPositionArraySimilarityReward : public ULearningAgentsReward
 {
 	GENERATED_BODY()
 
 public:
-	UFUNCTION(BlueprintCallable, Category = "LearningAgents")
-	static UPositionArraySimilarityReward* AddPositionArraySimilarityReward(ULearningAgentsTrainer* AgentTrainer, FName Name = NAME_None, int32 PositionNum = 0, float Scale = 100.0f, float Weight = 1.0f);
 
+	/**
+	* Adds a new position array similarity reward to the given trainer. Call during ULearningAgentsTrainer::SetupRewards event.
+	* @param AgentTrainer The trainer to add this reward to.
+	* @param Name The name of this new reward. Used for debugging.
+	* @param PositionNum The number of positions in the array.
+	* @param Scale Used to normalize the data for the reward.
+	* @param Weight Multiplier for this reward when being summed up for the total reward.
+	* @return The newly created reward.
+	*/
+	UFUNCTION(BlueprintCallable, Category = "LearningAgents")
+	static UPositionArraySimilarityReward* AddPositionArraySimilarityReward(
+		ULearningAgentsTrainer* AgentTrainer,
+		const FName Name = NAME_None,
+		const int32 PositionNum = 0,
+		const float Scale = 100.0f,
+		const float Weight = 1.0f);
+
+	/**
+	* Sets the data for this reward. Call during ULearningAgentsTrainer::SetRewards event.
+	* @param AgentId The agent id this data corresponds to.
+	* @param Positions0 The current positions.
+	* @param Positions1 The goal positions.
+	* @param RelativePosition0 The vector Positions0 will be offset from.
+	* @param RelativePosition1 The vector Positions1 will be offset from.
+	* @param RelativeRotation0 The frame of reference rotation for Positions0.
+	* @param RelativeRotation1 The frame of reference rotation for Positions1.
+	*/
 	UFUNCTION(BlueprintCallable, Category = "LearningAgents")
 	void SetPositionArraySimilarityReward(
-		int32 AgentId, 
+		const int32 AgentId,
 		const TArray<FVector>& Positions0, 
 		const TArray<FVector>& Positions1, 
-		FVector RelativePosition0 = FVector::ZeroVector,
-		FVector RelativePosition1 = FVector::ZeroVector,
-		FRotator RelativeRotation0 = FRotator::ZeroRotator,
-		FRotator RelativeRotation1 = FRotator::ZeroRotator);
+		const FVector RelativePosition0 = FVector::ZeroVector,
+		const FVector RelativePosition1 = FVector::ZeroVector,
+		const FRotator RelativeRotation0 = FRotator::ZeroRotator,
+		const FRotator RelativeRotation1 = FRotator::ZeroRotator);
 
 #if ENABLE_VISUAL_LOG
+	/** Describes this reward to the visual logger for debugging purposes. */
 	virtual void VisualLog(const UE::Learning::FIndexSet Instances) const override;
 #endif
 
 	TSharedPtr<UE::Learning::FPositionArraySimilarityReward> RewardObject;
 };
-
-//------------------------------------------------------------------
diff --git a/Engine/Plugins/Experimental/LearningAgents/Source/LearningAgentsTraining/Public/LearningAgentsTrainer.h b/Engine/Plugins/Experimental/LearningAgents/Source/LearningAgentsTraining/Public/LearningAgentsTrainer.h
--- a/Engine/Plugins/Experimental/LearningAgents/Source/LearningAgentsTraining/Public/LearningAgentsTrainer.h
+++ b/Engine/Plugins/Experimental/LearningAgents/Source/LearningAgentsTraining/Public/LearningAgentsTrainer.h
@@ -30,8 +30,7 @@ class ULearningAgentsReward;
 class ULearningAgentsPolicy;
 class ULearningAgentsCritic;
 
-//------------------------------------------------------------------
-
+/** Completion modes for episodes. */
 UENUM(BlueprintType, Category = "LearningAgents")
 enum class ELearningAgentsCompletion : uint8
 {
@@ -42,8 +41,7 @@ enum class ELearningAgentsCompletion : uint8
 	Termination	UMETA(DisplayName = "Termination"),
 };
 
-//------------------------------------------------------------------
-
+/** The configurable settings for a ULearningAgentsTrainer. */
 USTRUCT(BlueprintType, Category = "LearningAgents")
 struct FLearningAgentsTrainerSettings
 {
@@ -55,25 +53,27 @@ public:
 	UPROPERTY(EditAnywhere, Category = "LearningAgents")
 	ELearningAgentsCompletion MaxStepsCompletion = ELearningAgentsCompletion::Truncation;
 
-	/** Max number of steps to take while training before episode automatically terminates */
+	/** Max number of steps to take while training before episode automatically completes. */
 	UPROPERTY(EditAnywhere, Category = "LearningAgents", meta = (ClampMin = "1", UIMin = "1"))
 	int32 MaxStepNum = 300;
 
-	/** Maximum number of episodes to record before running a training iteration */
+	/** Maximum number of episodes to record before running a training iteration. */
 	UPROPERTY(EditAnywhere, Category = "LearningAgents", meta = (ClampMin = "1", UIMin = "1"))
 	int32 MaximumRecordedEpisodesPerIteration = 1000;
 
-	/** Maximum number of steps to record before running a training iteration */
+	/** Maximum number of steps to record before running a training iteration. */
 	UPROPERTY(EditAnywhere, Category = "LearningAgents", meta = (ClampMin = "1", UIMin = "1"))
 	int32 MaximumRecordedStepsPerIteration = 10000;
 
-	/** How long to wait for the training subprocess before timing out */
+	/** Time in seconds to wait for the training subprocess before timing out. */
 	UPROPERTY(EditAnywhere, Category = "LearningAgents", meta = (ClampMin = "0.0", UIMin = "0.0"))
 	float TrainerCommunicationTimeout = 20.0f;
 };
 
-//------------------------------------------------------------------
-
+/**
+* The configurable game settings for a ULearningAgentsTrainer. These allow the timestep and physics tick to be fixed
+* during training, which can enable ticking faster than real-time.
+*/
 USTRUCT(BlueprintType, Category="LearningAgents")
 struct FLearningAgentsTrainerGameSettings
 {
@@ -81,24 +81,36 @@ struct FLearningAgentsTrainerGameSettings
 
 public:
 
+	/**
+	* If true, the game will run in fixed time step mode (i.e the frame's delta times will always be the same
+	* regardless of how much wall time has passed). This can enable faster than real-time training if your game runs
+	* quickly. If false, the time steps will match real wall time.
+	*/
 	UPROPERTY(EditAnywhere, Category = "LearningAgents")
 	bool bUseFixedTimeStep = true;
 
+	/**
+	* Determines the amount of time for each frame when bUseFixedTimeStep is true; Ignored if false. You want this
+	* time step to match as closely as possible to the expected inference time steps, otherwise your training results
+	* may not generalize to your game.
+	*/
 	UPROPERTY(EditAnywhere, Category = "LearningAgents", meta = (DisplayName = "Fixed Time Step Frequency (Hz)"), meta = (ClampMin = "0.0", UIMin = "0.0"))
 	float FixedTimeStepFrequency = 60.0f;
 
+	/** If true, set the physics delta time to match the fixed time step. */
 	UPROPERTY(EditAnywhere, Category = "LearningAgents")
 	bool bSetMaxPhysicsStepToFixedTimeStep = true;
 
+	/** If true, VSync will be disabled; Otherwise, it will not. Disabling VSync can speed up the game simulation. */
 	UPROPERTY(EditAnywhere, Category = "LearningAgents")
 	bool bDisableVSync = true;
 
+	/** If true, the viewport rendering will be unlit; Otherwise, it will not. Disabling lighting can speed up the game simulation. */
 	UPROPERTY(EditAnywhere, Category = "LearningAgents")
 	bool bUseUnlitViewportRendering = false;
 };
 
-//------------------------------------------------------------------
-
+/** Enumeration of the training devices. */
 UENUM(BlueprintType, Category = "LearningAgents")
 enum class ELearningAgentsTrainerDevice : uint8
 {
@@ -106,8 +118,7 @@ enum class ELearningAgentsTrainerDevice : uint8
 	GPU,
 };
 
-//------------------------------------------------------------------
-
+/** The configurable settings for the training process. */
 USTRUCT(BlueprintType, Category = "LearningAgents")
 struct FLearningAgentsTrainerTrainingSettings
 {
@@ -115,33 +126,65 @@ struct FLearningAgentsTrainerTrainingSettings
 
 public:
 
+	/** The number of iterations to run before training is complete. */
 	UPROPERTY(EditAnywhere, Category = "LearningAgents", meta = (ClampMin = "0", UIMin = "0"))
 	int32 NumberOfIterations = 1000000;
 
+	/** If true, TensorBoard logs will be emitted to Intermediate\TensorBoard. Otherwise, no logs will be emitted. */
 	UPROPERTY(EditAnywhere, Category = "LearningAgents")
 	bool bUseTensorboard = false;
 
+	/**
+	* The initial scaling for the weights of the output layer of the neural network. Typically, you would use this to
+	* scale down the initial weights as it can stabilize the initial training and speed up convergence.
+	*/
 	UPROPERTY(EditAnywhere, Category = "LearningAgents", meta = (ClampMin = "0.0", UIMin = "0.0"))
 	float InitialActionScale = 0.1f;
 
+	/**
+	* The discount factor to use during training. This affects how much the agent cares about future rewards vs
+	* near-term rewards. Should typically be a value less than but near 1. 
+	*/
 	UPROPERTY(EditAnywhere, Category = "LearningAgents", meta = (ClampMin = "0.0", UIMin = "0.0"))
 	float DiscountFactor = 0.99f;
 
+	/** The seed used for any random sampling the trainer will perform, e.g. for weight initialization. */
 	UPROPERTY(EditAnywhere, Category = "LearningAgents", meta = (ClampMin = "0", UIMin = "0"))
 	int32 RandomSeed = 1234;
 
+	/** 
+	* The number of steps to trim from the start of the episode, e.g. can be useful if some things are still getting
+	* setup at the start of the episode.
+	*/
 	UPROPERTY(EditAnywhere, Category = "LearningAgents", meta = (ClampMin = "0", UIMin = "0"))
 	int32 NumberOfStepsToTrimAtStartOfEpisode = 0;
 
+	/**
+	* The number of steps to trim from the end of the episode. Can be useful if the end of the episode contains
+	* irrelevant data.
+	*/
 	UPROPERTY(EditAnywhere, Category = "LearningAgents", meta = (ClampMin = "0", UIMin = "0"))
 	int32 NumberOfStepsToTrimAtEndOfEpisode = 0;
 
+	/** The device to train on. */
 	UPROPERTY(EditAnywhere, Category = "LearningAgents")
 	ELearningAgentsTrainerDevice Device = ELearningAgentsTrainerDevice::GPU;
 };
 
-//------------------------------------------------------------------
-
+/**
+* The ULearningAgentsTrainer is the core class for reinforcement learning training. It has a few responsibilities:
+*   1) It keeps track of which agents are gathering training data.
+*   2) It defines how those agents' rewards, completions, and resets are implemented.
+*   3) It provides methods for orchestrating the training process.
+*
+* To use this class, you need to implement the SetupRewards and SetupCompletions functions (as well as their
+* corresponding SetRewards and SetCompletions functions), which will define the rewards and penalties the agent
+* receives and what conditions cause an episode to end. Before you can begin training, you need to call SetupTrainer,
+* which will initialize the underlying data structures, and you need to call AddAgent for each agent you want to gather
+* training data from.
+*
+* @see ULearningAgentsType to understand how observations and actions work.
+*/
 UCLASS(Abstract, BlueprintType, Blueprintable)
 class LEARNINGAGENTSTRAINING_API ULearningAgentsTrainer : public UActorComponent
 {
@@ -155,8 +198,16 @@ public:
 	ULearningAgentsTrainer(FVTableHelper& Helper);
 	virtual ~ULearningAgentsTrainer();
 
+	/** Will automatically call EndTraining if training is still in-progress when play is ending. */
 	virtual void EndPlay(const EEndPlayReason::Type EndPlayReason) override;
 
+	/**
+	* Initializes this object and runs the setup functions for rewards and completions.
+	* @param InAgentType The agent type we are training with.
+	* @param InPolicy The policy to be trained.
+	* @param InCritic Optional - only needs to be provided if we want the critic to be accessible at runtime.
+	* @param Settings The trainer settings to use.
+	*/
 	UFUNCTION(BlueprintCallable, Category = "LearningAgents")
 	void SetupTrainer(
 		ULearningAgentsType* InAgentType,
@@ -164,73 +215,113 @@ public:
 		ULearningAgentsCritic* InCritic = nullptr,
 		const FLearningAgentsTrainerSettings& Settings = FLearningAgentsTrainerSettings());
 
+	/** Returns true if SetupTrainer has been run successfully; Otherwise, false. */
 	UFUNCTION(BlueprintPure, Category = "LearningAgents")
 	bool IsTrainerSetupPerformed() const;
 
 // ----- Agent Management -----
 public:
 
+	/**
+	 * Adds an agent to this trainer.
+	 * @param AgentId The id of the agent to be added.
+	 * @warning The agent id must exist for the agent type.
+	 */
 	UFUNCTION(BlueprintCallable, Category = "LearningAgents")
-	void AddAgent(int32 AgentId);
+	void AddAgent(const int32 AgentId);
 
+	/**
+	* Removes an agent from this trainer.
+	* @param AgentId The id of the agent to be removed.
+	* @warning The agent id must exist for this trainer already.
+	*/
 	UFUNCTION(BlueprintCallable, Category = "LearningAgents")
-	void RemoveAgent(int32 AgentId);
+	void RemoveAgent(const int32 AgentId);
 
+	/** Returns true if the given id has been previously added to this trainer; Otherwise, false. */
 	UFUNCTION(BlueprintPure, Category = "LearningAgents")
-	bool HasAgent(int32 AgentId) const;
+	bool HasAgent(const int32 AgentId) const;
 
+	/**
+	* Gets the agent type this trainer is associated with.
+	* @param AgentClass The class to cast the agent type to (in blueprint).
+	*/
 	UFUNCTION(BlueprintPure, Category = "LearningAgents", meta = (DeterminesOutputType = "AgentClass"))
-	ULearningAgentsType* GetAgentType(TSubclassOf<ULearningAgentsType> AgentClass);
+	ULearningAgentsType* GetAgentType(const TSubclassOf<ULearningAgentsType> AgentClass);
 
 public:
 
-	const UObject* GetAgent(int32 AgentId) const;
-	UObject* GetAgent(int32 AgentId);
+	/** Gets the agent corresponding to the given id. */
+	const UObject* GetAgent(const int32 AgentId) const;
+
+	/** Gets the agent corresponding to the given id. */
+	UObject* GetAgent(const int32 AgentId);
 
+	/** Gets the associated agent type. */
 	const ULearningAgentsType* GetAgentType() const;
+
+	/** Gets the associated agent type. */
 	ULearningAgentsType* GetAgentType();
 
 // ----- Rewards -----
 public:
 	
 	/**
-	 * Event where all rewards should be set up for this agent trainer.
-	 * 
-	 * @param AgentTrainer This AgentTrainer. This is a convenience for blueprints so you don't have to find the Self pin.
+	 * During this event, all rewards/penalties should be added to this trainer.
+	 * @param AgentTrainer This trainer. This is a convenience for blueprints so you don't have to find the Self pin.
+	 * @see LearningAgentsRewards.h for the list of available rewards/penalties.
 	 */
 	UFUNCTION(BlueprintNativeEvent, Category = "LearningAgents")
 	void SetupRewards(ULearningAgentsTrainer* AgentTrainer);
 
+	/**
+	* During this event, all rewards/penalties should be set for each agent.
+	* @param AgentIds The list of agent ids to set rewards/penalties for.
+	* @see LearningAgentsRewards.h for the list of available rewards/penalties.
+	* @see GetAgent to get the agent corresponding to each id.
+	*/
 	UFUNCTION(BlueprintNativeEvent, Category = "LearningAgents")
 	void SetRewards(const TArray<int32>& AgentIds);
 
+	/**
+	* Used by objects derived from ULearningAgentsReward to add themselves to this trainer during their creation.
+	* You shouldn't need to call this directly.
+	*/
 	void AddReward(TObjectPtr<ULearningAgentsReward> Object, const TSharedRef<UE::Learning::FRewardObject>& Reward);
 
 // ----- Completions ----- 
 public:
 
 	/**
-	 * Event where all completions should be set up for this agent trainer.
-	 * 
-	 * @param AgentTrainer This AgentTrainer. This is a convenience for blueprints so you don't have to find the Self pin.
+	 * During this event, all completions should be added to this trainer.
+	 * @param AgentTrainer This trainer. This is a convenience for blueprints so you don't have to find the Self pin.
+	 * @see LearningAgentsCompletions.h for the list of available completions.
 	 */
 	UFUNCTION(BlueprintNativeEvent, Category = "LearningAgents")
 	void SetupCompletions(ULearningAgentsTrainer* AgentTrainer);
 
+	/**
+	* During this event, all completions should be set for each agent.
+	* @param AgentIds The list of agent ids to set completions for.
+	* @see LearningAgentsCompletions.h for the list of available completions.
+	* @see GetAgent to get the agent corresponding to each id.
+	*/
 	UFUNCTION(BlueprintNativeEvent, Category = "LearningAgents")
 	void SetCompletions(const TArray<int32>& AgentIds);
 
+	/**
+	* Used by objects derived from ULearningAgentsCompletion to add themselves to this trainer during their creation.
+	* You shouldn't need to call this directly.
+	*/
 	void AddCompletion(TObjectPtr<ULearningAgentsCompletion> Object, const TSharedRef<UE::Learning::FCompletionObject>& Completion);
 
 // ----- Resets ----- 
 public:
 
 	/**
-	 * Event where both the agent and its training environment should be reset so that a new training episode could
-	 * commence.
-	 *
-	 * @param AgentId The id of the current agent to be reset
-	 * @param Agent A pointer to the current agent to be reset
+	 * During this event, you will receive the ids of each agent that needs to be reset. Both the agent's actor and its
+	 * training environment should be reset for a new episode to commence.
+	 * @param AgentIds The ids of the agents that need resetting.
 	 */
 	UFUNCTION(BlueprintNativeEvent, Category = "LearningAgents")
 	void ResetInstance(const TArray<int32>& AgentIds);
@@ -238,9 +329,18 @@ public:
 // ----- Training Process -----
 public:
 
+	/** Returns true if the trainer is currently training; Otherwise, false. */
 	UFUNCTION(BlueprintPure, Category = "LearningAgents")
 	const bool IsTraining() const;
 
+	/**
+	* Begins the training process with the provided settings.
+	* @param TrainingSettings The settings for this training run.
+	* @param GameSettings The settings that will affect the game's simulation.
+	* @param CriticSettings The settings for the critic (if we are using one).
+	* @param bReinitializePolicyNetwork If true, reinitialize the policy. Set this to false if your policy is pre-trained, e.g. with imitation learning.
+	* @param bReinitializeCriticNetwork If true, reinitialize the critic. Set this to false if your critic is pre-trained.
+	*/
 	UFUNCTION(BlueprintCallable, Category = "LearningAgents")
 	void BeginTraining(
 		const FLearningAgentsTrainerTrainingSettings& TrainingSettings = FLearningAgentsTrainerTrainingSettings(),
@@ -249,39 +349,59 @@ public:
 		const bool bReinitializePolicyNetwork = true,
 		const bool bReinitializeCriticNetwork = true);
 
+	/** Stops the training process. */
 	UFUNCTION(BlueprintCallable, Category = "LearningAgents")
 	void EndTraining();
 
+	/**
+	* Call this function when it is time to evaluate the rewards for your agents. This should be done at the beginning
+	* of each iteration of your training loop after the initial step, i.e. after taking an action, you want to get into
+	* the next state before evaluating the rewards.
+	*/
 	UFUNCTION(BlueprintCallable, Category = "LearningAgents")
 	void EvaluateRewards();
 
+	/**
+	* Call this function when it is time to evaluate the completions for your agents. This should be done at the beginning
+	* of each iteration of your training loop after the initial step, i.e. after taking an action, you want to get into
+	* the next state before evaluating the completions.
+	*/
 	UFUNCTION(BlueprintCallable, Category = "LearningAgents")
 	void EvaluateCompletions();
 
+	/**
+	* Call this function at the end of each step of your training loop. This takes the current observations/actions/
+	* rewards and moves them into the current episode's experience buffer. Finished episodes will have their agents
+	* reset and their data will be sent to the external training process. Finally, the latest iteration of the
+	* trained policy will be synced back to UE so further experience can be acquired on-policy.
+	*/
 	UFUNCTION(BlueprintCallable, Category = "LearningAgents")
 	void IterateTraining();
 
-	/**
-	* Manually reset all agents. Does not record the experience gathered up to this point by each agent.
-	*/
+	/** Manually reset all agents. Does not record the experience gathered up to this point by each agent. */
 	UFUNCTION(BlueprintCallable, Category = "LearningAgents")
 	void ResetAllInstances();
 
 // ----- Private Data ----- 
 private:
 
+	/** The agent type this trainer is associated with. */
 	UPROPERTY(VisibleAnywhere, Category = "LearningAgents")
 	TObjectPtr<ULearningAgentsType> AgentType;
 
+	/** The agent ids this trainer is managing. */
 	UPROPERTY(VisibleAnywhere, Category = "LearningAgents")
 	TArray<int32> SelectedAgentIds;
 
+	/** The current policy for experience gathering. */
 	UPROPERTY(VisibleAnywhere, Category = "LearningAgents")
 	TObjectPtr<ULearningAgentsPolicy> Policy;
 
+	/** The current critic. */
 	UPROPERTY(VisibleAnywhere, Category = "LearningAgents")
 	TObjectPtr<ULearningAgentsCritic> Critic;
 
+	/** True if this trainer's SetupTrainer has been run; Otherwise, false. */
 	UPROPERTY(VisibleAnywhere, Category = "LearningAgents")
 	bool bTrainerSetupPerformed = false;
 
@@ -289,9 +409,11 @@ private:
 	UPROPERTY(VisibleAnywhere, Category = "LearningAgents")
 	bool bIsTraining = false;
 
+	/** The list of current reward objects. */
 	UPROPERTY(VisibleAnywhere, Category = "LearningAgents")
 	TArray<TObjectPtr<ULearningAgentsReward>> RewardObjects;
 
+	/** The list of current completion objects. */
 	UPROPERTY(VisibleAnywhere, Category = "LearningAgents")
 	TArray<TObjectPtr<ULearningAgentsCompletion>> CompletionObjects;
 
