LearningAgents: Added the ability to get the estimated expected discounted return from the critic and the ability to set action noise scale in policy. Various other small misc cleanups across Learning codebase.

diff --git a/Engine/Plugins/Experimental/LearningAgents/Source/Learning/Private/LearningRewardObject.cpp b/Engine/Plugins/Experimental/LearningAgents/Source/Learning/Private/LearningRewardObject.cpp
--- a/Engine/Plugins/Experimental/LearningAgents/Source/Learning/Private/LearningRewardObject.cpp
+++ b/Engine/Plugins/Experimental/LearningAgents/Source/Learning/Private/LearningRewardObject.cpp
@@ -308,10 +308,7 @@ namespace UE::Learning
 			Rewards[RewardIdx]->Evaluate(Instances);
 		}
 
-		for (const int32 InstanceIdx : Instances)
-		{
-			Reward[InstanceIdx] = 0.0f;
-		}
+		Array::Zero(Reward, Instances);
 
 		for (int32 RewardIdx = 0; RewardIdx < RewardNum; RewardIdx++)
 		{
diff --git a/Engine/Plugins/Experimental/LearningAgents/Source/Learning/Public/LearningArray.h b/Engine/Plugins/Experimental/LearningAgents/Source/Learning/Public/LearningArray.h
--- a/Engine/Plugins/Experimental/LearningAgents/Source/Learning/Public/LearningArray.h
+++ b/Engine/Plugins/Experimental/LearningAgents/Source/Learning/Public/LearningArray.h
@@ -79,7 +79,7 @@
 #define UE_LEARNING_CHECKF(...)
 #endif
 
-#define UE_LEARNING_NOT_IMPLEMENTED() UE_LEARNING_CHECKF(false, TEXT("Not Implemented!"))
+#define UE_LEARNING_NOT_IMPLEMENTED() checkf(false, TEXT("Not Implemented!"))
 
 #if UE_LEARNING_ARRAY_CHECK_SHAPE == 1
 #define UE_LEARNING_ARRAY_SHAPE_CHECK(Expr) check(Expr)
diff --git a/Engine/Plugins/Experimental/LearningAgents/Source/LearningAgents/Private/LearningAgentsCritic.cpp b/Engine/Plugins/Experimental/LearningAgents/Source/LearningAgents/Private/LearningAgentsCritic.cpp
--- a/Engine/Plugins/Experimental/LearningAgents/Source/LearningAgents/Private/LearningAgentsCritic.cpp
+++ b/Engine/Plugins/Experimental/LearningAgents/Source/LearningAgents/Private/LearningAgentsCritic.cpp
@@ -19,7 +19,7 @@
 
 namespace UE::Learning::Agents::Critic::Private
 {
-	static inline FString ArrayToString(const TLearningArrayView<1, const float> Array)
+	static inline FString FloatArrayToString(const TLearningArrayView<1, const float> Array)
 	{
 		const int32 ItemNum = Array.Num();
 		const int32 MaxItemNum = 32;
@@ -35,21 +35,18 @@ namespace UE::Learning::Agents::Critic::Private
 			{
 				Output += TEXT(" ");
 			}
-			else if (Idx == ItemNum - 1)
-			{
-				Output += TEXT("]");
-			}
 			else if (Idx == MaxItemNum - 1)
 			{
-				Output += TEXT("...]");
+				Output += TEXT("...");
 			}
 		}
 
+		Output += TEXT("]");
 
 		return Output;
 	}
 
-	static inline FString ArrayToStatsString(const TLearningArrayView<1, const float> Array)
+	static inline FString FloatArrayToStatsString(const TLearningArrayView<1, const float> Array)
 	{
 		const int32 ItemNum = Array.Num();
 
@@ -326,13 +323,32 @@ void ULearningAgentsCritic::EvaluateCritic()
 		return;
 	}
 
-	CriticObject->Evaluate(SelectedAgentIds);
+	CriticObject->Evaluate(SelectedAgentsSet);
 
 #if ENABLE_VISUAL_LOG
-	VisualLog(SelectedAgentIds);
+	VisualLog(SelectedAgentsSet);
 #endif
 }
 
+float ULearningAgentsCritic::GetEstimatedDiscountedReturn(const int32 AgentId) const
+{
+	if (!IsCriticSetupPerformed())
+	{
+		UE_LOG(LogLearning, Error, TEXT("Setup must be run before the critic can get the estimated discounted return."));
+		return 0.0f;
+	}
+
+	if (!SelectedAgentsSet.Contains(AgentId))
+	{
+		UE_LOG(LogLearning, Error, TEXT("Unable to get estimate for agent - AgentId %d not found in the added agents set."), AgentId);
+		return 0.0f;
+	}
+
+	const TLearningArrayView<1, const float> CriticOutputView = CriticObject->InstanceData->ConstView(CriticObject->OutputHandle);
+
+	return CriticOutputView[AgentId];
+}
+
 #if ENABLE_VISUAL_LOG
 void ULearningAgentsCritic::VisualLog(const UE::Learning::FIndexSet Instances) const
 {
@@ -345,8 +361,8 @@ void ULearningAgentsCritic::VisualLog(const UE::Learning::FIndexSet Instances) c
 	{
 		if (const AActor* Actor = Cast<AActor>(AgentType->GetAgent(Instance)))
 		{
-			const FString InputArrayString = UE::Learning::Agents::Critic::Private::ArrayToString(InputView[Instance]);
-			const FString InputStatsString = UE::Learning::Agents::Critic::Private::ArrayToStatsString(InputView[Instance]);
+			const FString InputArrayString = UE::Learning::Agents::Critic::Private::FloatArrayToString(InputView[Instance]);
+			const FString InputStatsString = UE::Learning::Agents::Critic::Private::FloatArrayToStatsString(InputView[Instance]);
 
 			UE_LEARNING_AGENTS_VLOG_STRING(this, LogLearning, Display,
 				Actor->GetActorLocation(),
diff --git a/Engine/Plugins/Experimental/LearningAgents/Source/LearningAgents/Private/LearningAgentsPolicy.cpp b/Engine/Plugins/Experimental/LearningAgents/Source/LearningAgents/Private/LearningAgentsPolicy.cpp
--- a/Engine/Plugins/Experimental/LearningAgents/Source/LearningAgents/Private/LearningAgentsPolicy.cpp
+++ b/Engine/Plugins/Experimental/LearningAgents/Source/LearningAgents/Private/LearningAgentsPolicy.cpp
@@ -19,7 +19,7 @@
 
 namespace UE::Learning::Agents::Policy::Private
 {
-	static inline FString ArrayToString(const TLearningArrayView<1, const float> Array)
+	static inline FString FloatArrayToString(const TLearningArrayView<1, const float> Array)
 	{
 		const int32 ItemNum = Array.Num();
 		const int32 MaxItemNum = 32;
@@ -35,21 +35,18 @@ namespace UE::Learning::Agents::Policy::Private
 			{
 				Output += TEXT(" ");
 			}
-			else if (Idx == ItemNum - 1)
-			{
-				Output += TEXT("]");
-			}
 			else if (Idx == MaxItemNum - 1)
 			{
-				Output += TEXT("...]");
+				Output += TEXT("...");
 			}
 		}
 
+		Output += TEXT("]");
 
 		return Output;
 	}
 
-	static inline FString ArrayToStatsString(const TLearningArrayView<1, const float> Array)
+	static inline FString FloatArrayToStatsString(const TLearningArrayView<1, const float> Array)
 	{
 		const int32 ItemNum = Array.Num();
 
@@ -116,6 +113,7 @@ void ULearningAgentsPolicy::SetupPolicy(ULearningAgentsType* InAgentType, const
 	UE::Learning::FNeuralNetworkPolicyFunctionSettings PolicyFunctionSettings;
 	PolicyFunctionSettings.ActionNoiseMin = PolicySettings.ActionNoiseMin;
 	PolicyFunctionSettings.ActionNoiseMax = PolicySettings.ActionNoiseMax;
+	PolicyFunctionSettings.ActionNoiseScale = PolicySettings.InitialActionNoiseScale;
 
 	PolicyObject = MakeShared<UE::Learning::FNeuralNetworkPolicyFunction>(
 		TEXT("PolicyObject"),
@@ -125,15 +123,7 @@ void ULearningAgentsPolicy::SetupPolicy(ULearningAgentsType* InAgentType, const
 		PolicySettings.ActionNoiseSeed,
 		PolicyFunctionSettings);
 
-	if (AgentType->GetInstanceData()->HasLink(AgentType->GetActionFeature().FeatureHandle))
-	{
-		UE_LOG(LogLearning, Warning, TEXT("AgentType actions are already linked to a difference policy."));
-	}
-	else
-	{
-		AgentType->GetInstanceData()->Link(AgentType->GetObservationFeature().FeatureHandle, PolicyObject->InputHandle);
-		AgentType->GetInstanceData()->Link(PolicyObject->OutputHandle, AgentType->GetActionFeature().FeatureHandle);
-	}
+	AgentType->GetInstanceData()->Link(AgentType->GetObservationFeature().FeatureHandle, PolicyObject->InputHandle);
 
 	// Done!
 	bPolicySetupPerformed = true;
@@ -342,11 +332,58 @@ void ULearningAgentsPolicy::EvaluatePolicy()
 
 	PolicyObject->Evaluate(SelectedAgentIds);
 
+	// Copy the actions computed by the policy into the agent-type's feature buffer.
+	// 
+	// Normally we would just link these two handles, but in this case we want to allow
+	// for multiple difference policies to be used for different agents, so that means
+	// there may be multiple writers to the action feature vector handle and so therefore
+	// the handles cannot be linked and we need to do the copy manually
+	UE::Learning::Array::Copy(
+		PolicyObject->InstanceData->View(AgentType->GetActionFeature().FeatureHandle),
+		PolicyObject->InstanceData->ConstView(PolicyObject->OutputHandle),
+		SelectedAgentIds);
+
 #if ENABLE_VISUAL_LOG
 	VisualLog(SelectedAgentIds);
 #endif
 }
 
+float ULearningAgentsPolicy::GetAgentActionNoiseScale(const int32 AgentId) const
+{
+	if (!IsPolicySetupPerformed())
+	{
+		UE_LOG(LogLearning, Error, TEXT("Setup must be run before getting the action noise."));
+		return 0.0f;
+	}
+
+	if (!SelectedAgentsSet.Contains(AgentId))
+	{
+		UE_LOG(LogLearning, Error, TEXT("Unable to get action noise for agent - AgentId %d not found in the added agents set."), AgentId);
+		return 0.0f;
+	}
+
+	const TLearningArrayView<1, const float> ActionNoiseScaleView = PolicyObject->InstanceData->ConstView(PolicyObject->ActionNoiseScaleHandle);
+	return ActionNoiseScaleView[AgentId];
+}
+
+void ULearningAgentsPolicy::SetAgentActionNoiseScale(const int32 AgentId, const float ActionNoiseScale)
+{
+	if (!IsPolicySetupPerformed())
+	{
+		UE_LOG(LogLearning, Error, TEXT("Setup must be run before setting the action noise."));
+		return;
+	}
+
+	if (!SelectedAgentsSet.Contains(AgentId))
+	{
+		UE_LOG(LogLearning, Error, TEXT("Unable to set action noise for agent - AgentId %d not found in the added agents set."), AgentId);
+		return;
+	}
+
+	const TLearningArrayView<1, float> ActionNoiseScaleView = PolicyObject->InstanceData->View(PolicyObject->ActionNoiseScaleHandle);
+	ActionNoiseScaleView[AgentId] = ActionNoiseScale;
+}
+
 #if ENABLE_VISUAL_LOG
 void ULearningAgentsPolicy::VisualLog(const UE::Learning::FIndexSet Instances) const
 {
@@ -362,12 +399,12 @@ void ULearningAgentsPolicy::VisualLog(const UE::Learning::FIndexSet Instances) c
 	{
 		if (const AActor* Actor = Cast<AActor>(AgentType->GetAgent(Instance)))
 		{
-			const FString InputArrayString = UE::Learning::Agents::Policy::Private::ArrayToString(InputView[Instance]);
-			const FString OutputArrayString = UE::Learning::Agents::Policy::Private::ArrayToString(OutputView[Instance]);
-			const FString OutputMeanArrayString = UE::Learning::Agents::Policy::Private::ArrayToString(OutputMeanView[Instance]);
-			const FString OutputStdArrayString = UE::Learning::Agents::Policy::Private::ArrayToString(OutputStdView[Instance]);
-			const FString InputStatsString = UE::Learning::Agents::Policy::Private::ArrayToStatsString(InputView[Instance]);
-			const FString OutputStatsString = UE::Learning::Agents::Policy::Private::ArrayToStatsString(OutputView[Instance]);
+			const FString InputArrayString = UE::Learning::Agents::Policy::Private::FloatArrayToString(InputView[Instance]);
+			const FString OutputArrayString = UE::Learning::Agents::Policy::Private::FloatArrayToString(OutputView[Instance]);
+			const FString OutputMeanArrayString = UE::Learning::Agents::Policy::Private::FloatArrayToString(OutputMeanView[Instance]);
+			const FString OutputStdArrayString = UE::Learning::Agents::Policy::Private::FloatArrayToString(OutputStdView[Instance]);
+			const FString InputStatsString = UE::Learning::Agents::Policy::Private::FloatArrayToStatsString(InputView[Instance]);
+			const FString OutputStatsString = UE::Learning::Agents::Policy::Private::FloatArrayToStatsString(OutputView[Instance]);
 
 			UE_LEARNING_AGENTS_VLOG_STRING(this, LogLearning, Display,
 				Actor->GetActorLocation(),
diff --git a/Engine/Plugins/Experimental/LearningAgents/Source/LearningAgents/Public/LearningAgentsCritic.h b/Engine/Plugins/Experimental/LearningAgents/Source/LearningAgents/Public/LearningAgentsCritic.h
--- a/Engine/Plugins/Experimental/LearningAgents/Source/LearningAgents/Public/LearningAgentsCritic.h
+++ b/Engine/Plugins/Experimental/LearningAgents/Source/LearningAgents/Public/LearningAgentsCritic.h
@@ -136,6 +136,18 @@ public:
 	UFUNCTION(BlueprintCallable, Category = "LearningAgents")
 	void EvaluateCritic();
 
+	/**
+	* Gets an estimate of the average discounted return expected by an agent according to the critic. I.E. the total
+	* sum of future rewards, scaled by the discount factor that was used during training. This value can be useful if 
+	* you want to make some decision based on how well the agent thinks they are doing at achieving their task. This 
+	* should be called after EvaluateCritic.
+	* 
+	* @param AgentId	The AgentId to look-up the estimated discounted return for
+	* @returns			The estimated average discounted return according to the critic
+	*/
+	UFUNCTION(BlueprintPure, Category = "LearningAgents")
+	float GetEstimatedDiscountedReturn(const int32 AgentId) const;
+
 // ----- Non-blueprint public interface -----
 public:
 
diff --git a/Engine/Plugins/Experimental/LearningAgents/Source/LearningAgents/Public/LearningAgentsPolicy.h b/Engine/Plugins/Experimental/LearningAgents/Source/LearningAgents/Public/LearningAgentsPolicy.h
--- a/Engine/Plugins/Experimental/LearningAgents/Source/LearningAgents/Public/LearningAgentsPolicy.h
+++ b/Engine/Plugins/Experimental/LearningAgents/Source/LearningAgents/Public/LearningAgentsPolicy.h
@@ -40,6 +40,10 @@ public:
 	UPROPERTY(EditAnywhere, Category = "LearningAgents", meta = (ClampMin = "0.0", UIMin = "0.0"))
 	float ActionNoiseMax = 0.25f;
 
+	/** Initial scale of the action noise used by the policy. Should be 1.0 for agents participating in training. */
+	UPROPERTY(EditAnywhere, Category = "LearningAgents", meta = (ClampMin = "0.0", UIMin = "0.0"))
+	float InitialActionNoiseScale = 1.0f;
+
 	/** Total layers for policy network including input, hidden, and output layers */
 	UPROPERTY(EditAnywhere, Category = "LearningAgents", meta = (ClampMin = "2", UIMin = "2"))
 	int32 LayerNum = 3;
@@ -149,6 +153,27 @@ public:
 	UFUNCTION(BlueprintCallable, Category = "LearningAgents")
 	void EvaluatePolicy();
 
+	/**
+	* Get the action noise scale used by an agent.
+	*
+	* @param AgentId			The AgentId to get the action noise scale for.
+	* @return					Action noise scale for that agent.
+	*/
+	UFUNCTION(BlueprintPure, Category = "LearningAgents")
+	float GetAgentActionNoiseScale(const int32 AgentId) const;
+
+	/**
+	* Set the action noise scale used by an agent. This can be useful if you have certain
+	* agents that are participating in training (and so should have an action noise scale of 1.0)
+	* and certain agents which you are testing the inference for (and so will want action noise
+	* scale of 0.0)
+	*
+	* @param AgentId			The AgentId to set the action noise scale for.
+	* @param ActionNoiseScale	Action noise scale for that agent.
+	*/
+	UFUNCTION(BlueprintCallable, Category = "LearningAgents")
+	void SetAgentActionNoiseScale(const int32 AgentId, const float ActionNoiseScale);
+
 // ----- Non-blueprint public interface -----
 public:
 
diff --git a/Engine/Plugins/Experimental/LearningAgents/Source/LearningTraining/Private/LearningSharedMemoryTraining.cpp b/Engine/Plugins/Experimental/LearningAgents/Source/LearningTraining/Private/LearningSharedMemoryTraining.cpp
--- a/Engine/Plugins/Experimental/LearningAgents/Source/LearningTraining/Private/LearningSharedMemoryTraining.cpp
+++ b/Engine/Plugins/Experimental/LearningAgents/Source/LearningTraining/Private/LearningSharedMemoryTraining.cpp
@@ -53,7 +53,7 @@ namespace UE::Learning::SharedMemoryTraining
 		const float SleepTime = 0.001f;
 		float WaitTime = 0.0f;
 
-		// Wait until the Policy is done being written by the sub-process
+		// Wait until the policy is done being written by the sub-process
 		while (!Controls[(uint8)EControls::PolicySignal])
 		{
 			// Check if Completed Signal has been raised
@@ -101,7 +101,7 @@ namespace UE::Learning::SharedMemoryTraining
 		const float SleepTime = 0.001f;
 		float WaitTime = 0.0f;
 
-		// Wait until the Critic is done being written by the sub-process
+		// Wait until the critic is done being written by the sub-process
 		while (!Controls[(uint8)EControls::CriticSignal])
 		{
 			FPlatformProcess::Sleep(SleepTime);
@@ -118,7 +118,7 @@ namespace UE::Learning::SharedMemoryTraining
 			UE_LOG(LogLearning, Display, TEXT("Pulling Critic..."));
 		}
 
-		// Read the policy
+		// Read the critic
 		{
 			FScopeNullableWriteLock ScopeLock(NetworkLock);
 			OutNetwork.DeserializeFromBytes(Critic);
@@ -141,7 +141,7 @@ namespace UE::Learning::SharedMemoryTraining
 		const float SleepTime = 0.001f;
 		float WaitTime = 0.0f;
 
-		// Wait until the Policy is requested by the sub-process
+		// Wait until the policy is requested by the sub-process
 		while (!Controls[(uint8)EControls::PolicySignal])
 		{
 			FPlatformProcess::Sleep(SleepTime);
@@ -181,7 +181,7 @@ namespace UE::Learning::SharedMemoryTraining
 		const float SleepTime = 0.001f;
 		float WaitTime = 0.0f;
 
-		// Wait until the Critic is requested by the sub-process
+		// Wait until the critic is requested by the sub-process
 		while (!Controls[(uint8)EControls::CriticSignal])
 		{
 			FPlatformProcess::Sleep(SleepTime);
@@ -247,17 +247,13 @@ namespace UE::Learning::SharedMemoryTraining
 		const int32 StepNum = ReplayBuffer.GetStepNum();
 
 		// Write experience to the shared memory
-		{
-			TRACE_CPUPROFILER_EVENT_SCOPE(Learning::SharedMemoryTraining::SendExperience::Copy);
-
-			Array::Copy(EpisodeStarts.Slice(0, EpisodeNum), ReplayBuffer.GetEpisodeStarts());
-			Array::Copy(EpisodeLengths.Slice(0, EpisodeNum), ReplayBuffer.GetEpisodeLengths());
-			Array::Copy(EpisodeCompletionModes.Slice(0, EpisodeNum), ReplayBuffer.GetEpisodeCompletionModes());
-			Array::Copy(EpisodeFinalObservations.Slice(0, EpisodeNum), ReplayBuffer.GetEpisodeFinalObservations());
-			Array::Copy(Observations.Slice(0, StepNum), ReplayBuffer.GetObservations());
-			Array::Copy(Actions.Slice(0, StepNum), ReplayBuffer.GetActions());
-			Array::Copy(Rewards.Slice(0, StepNum), ReplayBuffer.GetRewards());
-		}
+		Array::Copy(EpisodeStarts.Slice(0, EpisodeNum), ReplayBuffer.GetEpisodeStarts());
+		Array::Copy(EpisodeLengths.Slice(0, EpisodeNum), ReplayBuffer.GetEpisodeLengths());
+		Array::Copy(EpisodeCompletionModes.Slice(0, EpisodeNum), ReplayBuffer.GetEpisodeCompletionModes());
+		Array::Copy(EpisodeFinalObservations.Slice(0, EpisodeNum), ReplayBuffer.GetEpisodeFinalObservations());
+		Array::Copy(Observations.Slice(0, StepNum), ReplayBuffer.GetObservations());
+		Array::Copy(Actions.Slice(0, StepNum), ReplayBuffer.GetActions());
+		Array::Copy(Rewards.Slice(0, StepNum), ReplayBuffer.GetRewards());
 
 		// Indicate that experience is written
 		Controls[(uint8)EControls::ExperienceEpisodeNum] = EpisodeNum;
@@ -299,12 +295,8 @@ namespace UE::Learning::SharedMemoryTraining
 		const int32 StepNum = ObservationExperience.Num(0);
 
 		// Write experience to the shared memory
-		{
-			TRACE_CPUPROFILER_EVENT_SCOPE(Learning::SharedMemoryTraining::SendExperience::Copy);
-
-			Array::Copy(Observations.Slice(0, StepNum), ObservationExperience);
-			Array::Copy(Actions.Slice(0, StepNum), ActionExperience);
-		}
+		Array::Copy(Observations.Slice(0, StepNum), ObservationExperience);
+		Array::Copy(Actions.Slice(0, StepNum), ActionExperience);
 
 		// Confirm that experience is written
 		Controls[(uint8)EControls::ExperienceStepNum] = StepNum;
