Learning Agents: Split manager out of agent type, UProperty Transient tags

diff --git a/Engine/Plugins/Experimental/LearningAgents/Source/LearningAgents/Private/LearningAgentsActions.cpp b/Engine/Plugins/Experimental/LearningAgents/Source/LearningAgents/Private/LearningAgentsActions.cpp
--- a/Engine/Plugins/Experimental/LearningAgents/Source/LearningAgents/Private/LearningAgentsActions.cpp
+++ b/Engine/Plugins/Experimental/LearningAgents/Source/LearningAgents/Private/LearningAgentsActions.cpp
@@ -2,6 +2,7 @@
 
 #include "LearningAgentsActions.h"
 
+#include "LearningAgentsManager.h"
 #include "LearningAgentsType.h"
 #include "LearningArray.h"
 #include "LearningArrayMap.h"
@@ -39,8 +40,8 @@ namespace UE::Learning::Agents::Private
 
 		Action->FeatureObject = MakeShared<ActionFObject>(
 			Action->GetFName(),
-			AgentType->GetInstanceData().ToSharedRef(),
-			AgentType->GetMaxInstanceNum(),
+			AgentType->GetAgentManager()->GetInstanceData().ToSharedRef(),
+			AgentType->GetAgentManager()->GetMaxInstanceNum(),
 			Forward<InArgTypes>(Args)...);
 
 		AgentType->AddAction(Action, Action->FeatureObject.ToSharedRef());
diff --git a/Engine/Plugins/Experimental/LearningAgents/Source/LearningAgents/Private/LearningAgentsController.cpp b/Engine/Plugins/Experimental/LearningAgents/Source/LearningAgents/Private/LearningAgentsController.cpp
--- a/Engine/Plugins/Experimental/LearningAgents/Source/LearningAgents/Private/LearningAgentsController.cpp
+++ b/Engine/Plugins/Experimental/LearningAgents/Source/LearningAgents/Private/LearningAgentsController.cpp
@@ -2,6 +2,7 @@
 
 #include "LearningAgentsController.h"
 
+#include "LearningAgentsManager.h"
 #include "LearningAgentsType.h"
 #include "LearningAgentsActions.h"
 #include "LearningFeatureObject.h"
@@ -9,7 +10,7 @@
 
 #include "EngineDefines.h"
 
-ULearningAgentsController::ULearningAgentsController() : UActorComponent() {}
+ULearningAgentsController::ULearningAgentsController() : ULearningAgentsManagerComponent() {}
 ULearningAgentsController::ULearningAgentsController(FVTableHelper& Helper) : ULearningAgentsController() {}
 ULearningAgentsController::~ULearningAgentsController() {}
 
@@ -18,114 +19,78 @@ void ULearningAgentsController::SetActions_Implementation(const TArray<int32>& A
 	// Can be overridden to get actions without blueprints
 }
 
-void ULearningAgentsController::SetupController(ULearningAgentsType* InAgentType)
+void ULearningAgentsController::SetupController(ALearningAgentsManager* InAgentManager, ULearningAgentsType* InAgentType)
 {
-	if (IsControllerSetupPerformed())
+	if (IsSetup())
 	{
-		UE_LOG(LogLearning, Error, TEXT("Setup already performed!"));
+		UE_LOG(LogLearning, Error, TEXT("%s: Setup already performed!"), *GetName());
 		return;
 	}
 
-	if (!InAgentType)
+	if (!InAgentManager)
 	{
-		UE_LOG(LogLearning, Error, TEXT("SetupController called but AgentType is nullptr."));
+		UE_LOG(LogLearning, Error, TEXT("%s: InAgentManager is nullptr."), *GetName());
 		return;
 	}
 
-	if (!InAgentType->IsSetupPerformed())
+	if (!InAgentManager->IsManagerSetup())
 	{
-		UE_LOG(LogLearning, Error, TEXT("AgentType Setup must be run before controller can be setup."));
+		UE_LOG(LogLearning, Error, TEXT("%s's SetupManager() must be run before %s can be setup."), *InAgentManager->GetName(), *GetName());
 		return;
 	}
 
-	AgentType = InAgentType;
-}
-
-bool ULearningAgentsController::IsControllerSetupPerformed() const
-{
-	return AgentType ? true : false;
-}
-
-void ULearningAgentsController::AddAgent(const int32 AgentId)
-{
-	if (!IsControllerSetupPerformed())
-	{
-		UE_LOG(LogLearning, Error, TEXT("Controller setup must be run before agents can be added!"));
-		return;
-	}
-
-	if (!AgentType->GetOccupiedAgentSet().Contains(AgentId))
-	{
-		UE_LOG(LogLearning, Error, TEXT("Unable to add: AgentId %d not found on AgentType. Make sure to add agents to the agent type before adding."), AgentId);
-		return;
-	}
-
-	if (SelectedAgentIds.Contains(AgentId))
-	{
-		UE_LOG(LogLearning, Error, TEXT("AgentId %i is already included in agents set"), AgentId);
-		return;
-	}
-
-	SelectedAgentIds.Add(AgentId);
-	SelectedAgentsSet = SelectedAgentIds;
-	SelectedAgentsSet.TryMakeSlice();
-}
+	// This manager is not referenced in this class but we need it in blueprints to call GetAgent()
+	AgentManager = InAgentManager;
 
-void ULearningAgentsController::RemoveAgent(const int32 AgentId)
-{
-	if (!IsControllerSetupPerformed())
+	if (!InAgentType)
 	{
-		UE_LOG(LogLearning, Error, TEXT("Controller setup must be run before agents can be removed!"));
+		UE_LOG(LogLearning, Error, TEXT("%s: InAgentType is nullptr."), *GetName());
 		return;
 	}
 
-	if (SelectedAgentIds.RemoveSingleSwap(AgentId, false) == 0)
+	if (!InAgentType->IsSetup())
 	{
-		UE_LOG(LogLearning, Error, TEXT("Unable to remove: AgentId %d not found in the added agents set."), AgentId);
+		UE_LOG(LogLearning, Error, TEXT("%s: %s's Setup must be run before it can be used."), *GetName(), *InAgentType->GetName());
 		return;
 	}
 
-	SelectedAgentsSet = SelectedAgentIds;
-	SelectedAgentsSet.TryMakeSlice();
-}
-
-bool ULearningAgentsController::HasAgent(const int32 AgentId) const
-{
-	return SelectedAgentsSet.Contains(AgentId);
-}
-
-ULearningAgentsType* ULearningAgentsController::GetAgentType(TSubclassOf<ULearningAgentsType> AgentClass)
-{
-	if (!IsControllerSetupPerformed())
-	{
-		UE_LOG(LogLearning, Error, TEXT("Controller setup must be run before getting the agent type!"));
-		return nullptr;
-	}
+	AgentType = InAgentType;
 
-	return AgentType;
+	bIsSetup = true;
 }
 
 void ULearningAgentsController::EncodeActions()
 {
 	UE_LEARNING_TRACE_CPUPROFILER_EVENT_SCOPE(ULearningAgentsController::EncodeActions);
 
-	if (!IsControllerSetupPerformed())
+	if (!IsSetup())
 	{
 		UE_LOG(LogLearning, Error, TEXT("Setup must be run before actions can be encoded."));
 		return;
 	}
 
-	SetActions(SelectedAgentIds);
+	SetActions(AddedAgentIds);
 
-	AgentType->GetActionFeature().Encode(SelectedAgentsSet);
+	AgentType->GetActionFeature().Encode(AddedAgentSet);
 
 #if ENABLE_VISUAL_LOG
 	for (const ULearningAgentsAction* ActionObject : AgentType->GetActionObjects())
 	{
 		if (ActionObject)
 		{
-			ActionObject->VisualLog(SelectedAgentsSet);
+			ActionObject->VisualLog(AddedAgentSet);
 		}
 	}
 #endif
 }
+
+ULearningAgentsType* ULearningAgentsController::GetAgentType(const TSubclassOf<ULearningAgentsType> AgentTypeClass) const
+{
+	if (!AgentType)
+	{
+		UE_LOG(LogLearning, Error, TEXT("%s: AgentType is nullptr. Did we forget to call Setup on this component?"), *GetName());
+		return nullptr;
+	}
+
+	return AgentType;
+}
diff --git a/Engine/Plugins/Experimental/LearningAgents/Source/LearningAgents/Private/LearningAgentsCritic.cpp b/Engine/Plugins/Experimental/LearningAgents/Source/LearningAgents/Private/LearningAgentsCritic.cpp
--- a/Engine/Plugins/Experimental/LearningAgents/Source/LearningAgents/Private/LearningAgentsCritic.cpp
+++ b/Engine/Plugins/Experimental/LearningAgents/Source/LearningAgents/Private/LearningAgentsCritic.cpp
@@ -2,6 +2,7 @@
 
 #include "LearningAgentsCritic.h"
 
+#include "LearningAgentsManager.h"
 #include "LearningAgentsType.h"
 #include "LearningFeatureObject.h"
 #include "LearningNeuralNetwork.h"
@@ -69,36 +70,47 @@ namespace UE::Learning::Agents::Critic::Private
 	}
 }
 
-ULearningAgentsCritic::ULearningAgentsCritic() : UActorComponent() {}
+ULearningAgentsCritic::ULearningAgentsCritic() : ULearningAgentsManagerComponent() {}
 ULearningAgentsCritic::ULearningAgentsCritic(FVTableHelper& Helper) : ULearningAgentsCritic() {}
 ULearningAgentsCritic::~ULearningAgentsCritic() {}
 
-void ULearningAgentsCritic::SetupCritic(ULearningAgentsType* InAgentType, const FLearningAgentsCriticSettings& CriticSettings)
+void ULearningAgentsCritic::SetupCritic(ALearningAgentsManager* InAgentManager, ULearningAgentsType* InAgentType, const FLearningAgentsCriticSettings& CriticSettings)
 {
-	if (IsCriticSetupPerformed())
+	if (IsSetup())
 	{
-		UE_LOG(LogLearning, Error, TEXT("Setup already performed!"));
+		UE_LOG(LogLearning, Error, TEXT("%s: Setup already performed!"), *GetName());
 		return;
 	}
 
-	// Setup Agent Type
+	if (!InAgentManager)
+	{
+		UE_LOG(LogLearning, Error, TEXT("%s: InAgentManager is nullptr."), *GetName());
+		return;
+	}
+
+	if (!InAgentManager->IsManagerSetup())
+	{
+		UE_LOG(LogLearning, Error, TEXT("%s's SetupManager() must be run before %s can be setup."), *InAgentManager->GetName(), *GetName());
+		return;
+	}
+
+	AgentManager = InAgentManager;
 
 	if (!InAgentType)
 	{
-		UE_LOG(LogLearning, Error, TEXT("SetupCritic called but AgentType is nullptr."));
+		UE_LOG(LogLearning, Error, TEXT("%s: InAgentType is nullptr."), *GetName());
 		return;
 	}
 
-	if (!InAgentType->IsSetupPerformed())
+	if (!InAgentType->IsSetup())
 	{
-		UE_LOG(LogLearning, Error, TEXT("AgentType Setup must be run before critic can be setup."));
+		UE_LOG(LogLearning, Error, TEXT("%s: %s's Setup must be run before it can be used."), *GetName(), *InAgentType->GetName());
 		return;
 	}
 
 	AgentType = InAgentType;
 
 	// Setup Neural Network
-
 	Network = NewObject<ULearningAgentsNeuralNetwork>(this, TEXT("CriticNetwork"));
 	Network->NeuralNetwork = MakeShared<UE::Learning::FNeuralNetwork>();
 	Network->NeuralNetwork->Resize(
@@ -109,81 +121,15 @@ void ULearningAgentsCritic::SetupCritic(ULearningAgentsType* InAgentType, const
 	Network->NeuralNetwork->ActivationFunction = UE::Learning::Agents::GetActivationFunction(CriticSettings.ActivationFunction);
 
 	// Create Critic Object
-
 	CriticObject = MakeShared<UE::Learning::FNeuralNetworkCriticFunction>(
 		TEXT("CriticObject"),
-		AgentType->GetInstanceData().ToSharedRef(),
-		AgentType->GetMaxInstanceNum(),
+		AgentManager->GetInstanceData().ToSharedRef(),
+		AgentManager->GetMaxInstanceNum(),
 		Network->NeuralNetwork.ToSharedRef());
 
-	AgentType->GetInstanceData()->Link(AgentType->GetObservationFeature().FeatureHandle, CriticObject->InputHandle);
-
-	// Done!
-	bCriticSetupPerformed = true;
-}
-
-bool ULearningAgentsCritic::IsCriticSetupPerformed() const
-{
-	return bCriticSetupPerformed;
-}
-
-void ULearningAgentsCritic::AddAgent(const int32 AgentId)
-{
-	if (!IsCriticSetupPerformed())
-	{
-		UE_LOG(LogLearning, Error, TEXT("Critic setup must be run before agents can be added!"));
-		return;
-	}
-
-	if (!AgentType->GetOccupiedAgentSet().Contains(AgentId))
-	{
-		UE_LOG(LogLearning, Error, TEXT("Unable to add: AgentId %d not found on AgentType. Make sure to add agents to the agent type before adding."), AgentId);
-		return;
-	}
-
-	if (SelectedAgentIds.Contains(AgentId))
-	{
-		UE_LOG(LogLearning, Error, TEXT("AgentId %i is already included in agents set"), AgentId);
-		return;
-	}
-
-	SelectedAgentIds.Add(AgentId);
-	SelectedAgentsSet = SelectedAgentIds;
-	SelectedAgentsSet.TryMakeSlice();
-}
-
-void ULearningAgentsCritic::RemoveAgent(const int32 AgentId)
-{
-	if (!IsCriticSetupPerformed())
-	{
-		UE_LOG(LogLearning, Error, TEXT("Critic setup must be run before agents can be removed!"));
-		return;
-	}
-
-	if (SelectedAgentIds.RemoveSingleSwap(AgentId, false) == 0)
-	{
-		UE_LOG(LogLearning, Error, TEXT("Unable to remove: AgentId %d not found in the added agents set."), AgentId);
-		return;
-	}
-
-	SelectedAgentsSet = SelectedAgentIds;
-	SelectedAgentsSet.TryMakeSlice();
-}
-
-bool ULearningAgentsCritic::HasAgent(const int32 AgentId) const
-{
-	return SelectedAgentsSet.Contains(AgentId);
-}
-
-ULearningAgentsType* ULearningAgentsCritic::GetAgentType(const TSubclassOf<ULearningAgentsType> AgentClass)
-{
-	if (!IsCriticSetupPerformed())
-	{
-		UE_LOG(LogLearning, Error, TEXT("Critic setup must be run before getting the agent type!"));
-		return nullptr;
-	}
+	AgentManager->GetInstanceData()->Link(AgentType->GetObservationFeature().FeatureHandle, CriticObject->InputHandle);
 
-	return AgentType;
+	bIsSetup = true;
 }
 
 UE::Learning::FNeuralNetwork& ULearningAgentsCritic::GetCriticNetwork()
@@ -198,7 +144,7 @@ UE::Learning::FNeuralNetworkCriticFunction& ULearningAgentsCritic::GetCriticObje
 
 void ULearningAgentsCritic::LoadCriticFromSnapshot(const FDirectoryPath& Directory, const FString Filename)
 {
-	if (!IsCriticSetupPerformed())
+	if (!IsSetup())
 	{
 		UE_LOG(LogLearning, Error, TEXT("Critic setup must be run before network can be loaded."));
 		return;
@@ -230,7 +176,7 @@ void ULearningAgentsCritic::LoadCriticFromSnapshot(const FDirectoryPath& Directo
 
 void ULearningAgentsCritic::SaveCriticToSnapshot(const FDirectoryPath& Directory, const FString Filename) const
 {
-	if (!IsCriticSetupPerformed())
+	if (!IsSetup())
 	{
 		UE_LOG(LogLearning, Error, TEXT("Critic setup must be run before network can be saved."));
 		return;
@@ -254,7 +200,7 @@ void ULearningAgentsCritic::SaveCriticToSnapshot(const FDirectoryPath& Directory
 
 void ULearningAgentsCritic::LoadCriticFromAsset(const ULearningAgentsNeuralNetwork* NeuralNetworkAsset)
 {
-	if (!IsCriticSetupPerformed())
+	if (!IsSetup())
 	{
 		UE_LOG(LogLearning, Error, TEXT("Critic setup must be run before network can be loaded."));
 		return;
@@ -278,7 +224,7 @@ void ULearningAgentsCritic::LoadCriticFromAsset(const ULearningAgentsNeuralNetwo
 
 void ULearningAgentsCritic::SaveCriticToAsset(ULearningAgentsNeuralNetwork* NeuralNetworkAsset) const
 {
-	if (!IsCriticSetupPerformed())
+	if (!IsSetup())
 	{
 		UE_LOG(LogLearning, Error, TEXT("Critic setup must be run before network can be saved."));
 		return;
@@ -317,28 +263,28 @@ void ULearningAgentsCritic::EvaluateCritic()
 {
 	UE_LEARNING_TRACE_CPUPROFILER_EVENT_SCOPE(ULearningAgentsCritic::EvaluateCritic);
 
-	if (!IsCriticSetupPerformed())
+	if (!IsSetup())
 	{
 		UE_LOG(LogLearning, Error, TEXT("Setup must be run before the critic can be evaluated."));
 		return;
 	}
 
-	CriticObject->Evaluate(SelectedAgentsSet);
+	CriticObject->Evaluate(AddedAgentSet);
 
 #if ENABLE_VISUAL_LOG
-	VisualLog(SelectedAgentsSet);
+	VisualLog(AddedAgentSet);
 #endif
 }
 
 float ULearningAgentsCritic::GetEstimatedDiscountedReturn(const int32 AgentId) const
 {
-	if (!IsCriticSetupPerformed())
+	if (!IsSetup())
 	{
 		UE_LOG(LogLearning, Error, TEXT("Setup must be run before the critic can get the estimated discounted return."));
 		return 0.0f;
 	}
 
-	if (!SelectedAgentsSet.Contains(AgentId))
+	if (!HasAgent(AgentId))
 	{
 		UE_LOG(LogLearning, Error, TEXT("Unable to get estimate for agent - AgentId %d not found in the added agents set."), AgentId);
 		return 0.0f;
diff --git a/Engine/Plugins/Experimental/LearningAgents/Source/LearningAgents/Private/LearningAgentsManager.cpp b/Engine/Plugins/Experimental/LearningAgents/Source/LearningAgents/Private/LearningAgentsManager.cpp
new file mode 100644
--- /dev/null
+++ b/Engine/Plugins/Experimental/LearningAgents/Source/LearningAgents/Private/LearningAgentsManager.cpp
@@ -0,0 +1,220 @@
+// Copyright Epic Games, Inc. All Rights Reserved.
+
+#include "LearningAgentsManager.h"
+#include "LearningAgentsManagerComponent.h"
+
+#include "LearningArray.h"
+#include "LearningArrayMap.h"
+#include "LearningLog.h"
+
+#include "UObject/ScriptInterface.h"
+#include "EngineDefines.h"
+
+ALearningAgentsManager::ALearningAgentsManager() : AActor() {}
+ALearningAgentsManager::ALearningAgentsManager(FVTableHelper& Helper) : ALearningAgentsManager() {}
+ALearningAgentsManager::~ALearningAgentsManager() {}
+
+void ALearningAgentsManager::PostInitProperties()
+{
+	Super::PostInitProperties();
+
+	if (HasAnyFlags(RF_ClassDefaultObject))
+	{
+		return;
+	}
+
+	// Pre-populate the vacant ids
+	OccupiedAgentIds.Reserve(MaxInstanceNum);
+	VacantAgentIds.Reserve(MaxInstanceNum);
+	for (int32 i = MaxInstanceNum - 1; i >= 0; i--)
+	{
+		VacantAgentIds.Push(i);
+		Agents.Add(nullptr);
+	}
+
+	UpdateAgentSets();
+}
+
+void ALearningAgentsManager::SetupManager()
+{
+	if (IsManagerSetup())
+	{
+		UE_LOG(LogLearning, Error, TEXT("%s: Setup already performed!"), *GetName());
+		return;
+	}
+
+	InstanceData = MakeShared<UE::Learning::FArrayMap>();
+
+	GetComponents<ULearningAgentsManagerComponent>(CachedManagerComponents);
+	if (CachedManagerComponents.IsEmpty())
+	{
+		UE_LOG(LogLearning, Warning, TEXT("%s: Found zero attached manager components during setup."), *GetName());
+	}
+	else
+	{
+		for (const ULearningAgentsManagerComponent* ManagerComponent : CachedManagerComponents)
+		{
+			UE_LOG(LogLearning, Display, TEXT("Added %s manager component to %s's cache."), *ManagerComponent->GetName(), *GetName());
+		}
+	}
+
+	bIsSetup = true;
+}
+
+bool ALearningAgentsManager::IsManagerSetup() const
+{
+	return bIsSetup;
+}
+
+int32 ALearningAgentsManager::GetMaxInstanceNum() const
+{
+	return MaxInstanceNum;
+}
+
+const TSharedPtr<UE::Learning::FArrayMap>& ALearningAgentsManager::GetInstanceData() const
+{
+	return InstanceData;
+}
+
+TConstArrayView<TObjectPtr<UObject>> ALearningAgentsManager::GetAgents() const
+{
+	return Agents;
+}
+
+int32 ALearningAgentsManager::AddAgent(UObject* Agent)
+{
+	if (Agent == nullptr)
+	{
+		UE_LOG(LogLearning, Error, TEXT("%s: Attempted to add an agent but agent is nullptr."), *GetName());
+		return INDEX_NONE;
+	}
+
+	if (VacantAgentIds.IsEmpty())
+	{
+		UE_LOG(LogLearning, Error,
+			TEXT("%s: Attempting to add an agent but we have no more vacant ids. Increase MaxInstanceNum (%d) or remove unused agents."),
+			*GetName(),
+			MaxInstanceNum);
+		return INDEX_NONE;
+	}
+
+	// Add Agent
+	const int32 NewAgentId = VacantAgentIds.Pop();
+	Agents[NewAgentId] = Agent;
+
+	OccupiedAgentIds.Add(NewAgentId);
+
+	UpdateAgentSets();
+
+	OnAgentAdded(NewAgentId);
+
+	return NewAgentId;
+}
+
+void ALearningAgentsManager::OnAgentAdded_Implementation(const int32 AgentId)
+{
+	for (ULearningAgentsManagerComponent* ManagerComponent : CachedManagerComponents)
+	{
+		ManagerComponent->AddAgent(AgentId);
+	}
+}
+
+void ALearningAgentsManager::RemoveAgentById(const int32 AgentId)
+{
+	if (AgentId == INDEX_NONE)
+	{
+		UE_LOG(LogLearning, Warning, TEXT("%s: Attempting to remove an agent with id of INDEX_NONE."), *GetName());
+		return;
+	}
+
+	const int32 RemovedCount = OccupiedAgentIds.RemoveSingleSwap(AgentId, false);
+
+	if (RemovedCount == 0)
+	{
+		UE_LOG(LogLearning, Warning, TEXT("%s: Trying to remove an agent but its Id (%d) is not in the occupied agents."), *GetName(), AgentId);
+		return;
+	}
+
+	// Remove Agent
+	VacantAgentIds.Push(AgentId);
+	Agents[AgentId] = nullptr;
+
+	UpdateAgentSets();
+
+	OnAgentRemoved(AgentId);
+}
+
+void ALearningAgentsManager::RemoveAgent(UObject* Agent)
+{
+	if (Agent == nullptr)
+	{
+		UE_LOG(LogLearning, Error, TEXT("%s: Attempted to remove an agent but agent is nullptr."), *GetName());
+		return;
+	}
+
+	int32 AgentId = INDEX_NONE;
+	const bool bIsFound = Agents.Find(Agent, AgentId);
+
+	if (!bIsFound)
+	{
+		UE_LOG(LogLearning, Warning, TEXT("%s: Trying to remove an agent but it was not found."), *GetName());
+		return;
+	}
+
+	RemoveAgentById(AgentId);
+}
+
+void ALearningAgentsManager::OnAgentRemoved_Implementation(const int32 AgentId)
+{
+	for (ULearningAgentsManagerComponent* ManagerComponent : CachedManagerComponents)
+	{
+		if (ManagerComponent->HasAgent(AgentId))
+		{
+			ManagerComponent->RemoveAgent(AgentId);
+		}
+	}
+}
+
+bool ALearningAgentsManager::HasAgent(UObject* Agent) const
+{
+	return Agents.Find(Agent) ? true : false;
+}
+
+bool ALearningAgentsManager::HasAgentById(const int32 AgentId) const
+{
+	return OccupiedAgentSet.Contains(AgentId);
+}
+
+UObject* ALearningAgentsManager::GetAgent(const int32 AgentId, const TSubclassOf<UObject> AgentClass)
+{
+	if (!OccupiedAgentSet.Contains(AgentId))
+	{
+		UE_LOG(LogLearning, Error,
+			TEXT("%s: AgentId %d not found. Be sure to only use AgentIds returned by AddAgent() and check that the agent has not be removed."),
+			*GetName(),
+			AgentId);
+		return nullptr;
+	}
+
+	return Agents[AgentId];
+}
+
+const UObject* ALearningAgentsManager::GetAgent(const int32 AgentId) const
+{
+	UE_LEARNING_CHECK(OccupiedAgentSet.Contains(AgentId));
+	return Agents[AgentId];
+}
+
+UObject* ALearningAgentsManager::GetAgent(const int32 AgentId)
+{
+	UE_LEARNING_CHECK(OccupiedAgentSet.Contains(AgentId));
+	return Agents[AgentId];
+}
+
+void ALearningAgentsManager::UpdateAgentSets()
+{
+	OccupiedAgentSet = OccupiedAgentIds;
+	OccupiedAgentSet.TryMakeSlice();
+	VacantAgentSet = VacantAgentIds;
+	VacantAgentSet.TryMakeSlice();
+}
diff --git a/Engine/Plugins/Experimental/LearningAgents/Source/LearningAgents/Private/LearningAgentsManagerComponent.cpp b/Engine/Plugins/Experimental/LearningAgents/Source/LearningAgents/Private/LearningAgentsManagerComponent.cpp
new file mode 100644
--- /dev/null
+++ b/Engine/Plugins/Experimental/LearningAgents/Source/LearningAgents/Private/LearningAgentsManagerComponent.cpp
@@ -0,0 +1,102 @@
+// Copyright Epic Games, Inc. All Rights Reserved.
+
+#include "LearningAgentsManagerComponent.h"
+#include "LearningAgentsManager.h"
+#include "LearningLog.h"
+
+ULearningAgentsManagerComponent::ULearningAgentsManagerComponent() : UActorComponent() {}
+ULearningAgentsManagerComponent::ULearningAgentsManagerComponent(FVTableHelper& Helper) : ULearningAgentsManagerComponent() {}
+ULearningAgentsManagerComponent::~ULearningAgentsManagerComponent() {}
+
+bool ULearningAgentsManagerComponent::AddAgent(const int32 AgentId)
+{
+	if (HasAgent(AgentId))
+	{
+		UE_LOG(LogLearning, Warning, TEXT("AgentId %i is already added to %s."), AgentId, *GetName());
+		return false;
+	}
+
+	AddedAgentIds.Add(AgentId);
+	AddedAgentSet = AddedAgentIds;
+	AddedAgentSet.TryMakeSlice();
+
+	return true;
+}
+
+bool ULearningAgentsManagerComponent::RemoveAgent(const int32 AgentId)
+{
+	if (AddedAgentIds.RemoveSingleSwap(AgentId, false) == 0)
+	{
+		UE_LOG(LogLearning, Warning, TEXT("Unable to remove: AgentId %d not found in %s's agent set."), AgentId, *GetName());
+		return false;
+	}
+
+	AddedAgentSet = AddedAgentIds;
+	AddedAgentSet.TryMakeSlice();
+
+	return true;
+}
+
+bool ULearningAgentsManagerComponent::HasAgent(const int32 AgentId) const
+{
+	return AddedAgentSet.Contains(AgentId);
+}
+
+UObject* ULearningAgentsManagerComponent::GetAgent(const int32 AgentId, const TSubclassOf<UObject> AgentClass) const
+{
+	if (!AgentManager)
+	{
+		UE_LOG(LogLearning, Error, TEXT("Agent manager is nullptr. Call setup on this component prior to getting agents."));
+		return nullptr;
+	}
+
+	if (!AddedAgentSet.Contains(AgentId))
+	{
+		UE_LOG(LogLearning, Error,
+			TEXT("%s: AgentId %d not found. Be sure to only use AgentIds returned by AddAgent() and check that the agent has not be removed."),
+			*GetName(),
+			AgentId);
+		return nullptr;
+	}
+
+	return AgentManager->GetAgent(AgentId, AgentClass); // Calling this overload since it will log about missing manager ids
+}
+
+const UObject* ULearningAgentsManagerComponent::GetAgent(const int32 AgentId) const
+{
+	UE_LEARNING_CHECKF(AgentManager, TEXT("AgentManager is nullptr. Did we forget to call Setup on this component and set the manager?"));
+	UE_LEARNING_CHECKF(AddedAgentSet.Contains(AgentId), TEXT("AgentId not found. Make sure it was added via AddAgent()."));
+
+	return AgentManager->GetAgent(AgentId);
+}
+
+UObject* ULearningAgentsManagerComponent::GetAgent(const int32 AgentId)
+{
+	UE_LEARNING_CHECKF(AgentManager, TEXT("AgentManager is nullptr. Did we forget to call Setup on this component and set the manager?"));
+	UE_LEARNING_CHECKF(AddedAgentSet.Contains(AgentId), TEXT("AgentId not found. Make sure it was added via AddAgent()."));
+
+	return AgentManager->GetAgent(AgentId);
+}
+
+ALearningAgentsManager* ULearningAgentsManagerComponent::GetAgentManager() const
+{
+	UE_LEARNING_CHECKF(AgentManager, TEXT("AgentManager is nullptr. Did we forget to call Setup on this component and set the manager?"));
+
+	return AgentManager;
+}
+
+ALearningAgentsManager* ULearningAgentsManagerComponent::GetAgentManager(const TSubclassOf<ALearningAgentsManager> AgentManagerClass) const
+{
+	if (!AgentManager)
+	{
+		UE_LOG(LogLearning, Error, TEXT("AgentManager is nullptr. Did we forget to call Setup on this component and set the manager?"));
+		return nullptr;
+	}
+
+	return AgentManager;
+}
+
+bool ULearningAgentsManagerComponent::IsSetup() const
+{
+	return bIsSetup;
+}
diff --git a/Engine/Plugins/Experimental/LearningAgents/Source/LearningAgents/Private/LearningAgentsObservations.cpp b/Engine/Plugins/Experimental/LearningAgents/Source/LearningAgents/Private/LearningAgentsObservations.cpp
--- a/Engine/Plugins/Experimental/LearningAgents/Source/LearningAgents/Private/LearningAgentsObservations.cpp
+++ b/Engine/Plugins/Experimental/LearningAgents/Source/LearningAgents/Private/LearningAgentsObservations.cpp
@@ -2,6 +2,7 @@
 
 #include "LearningAgentsObservations.h"
 
+#include "LearningAgentsManager.h"
 #include "LearningAgentsType.h"
 #include "LearningArray.h"
 #include "LearningArrayMap.h"
@@ -39,8 +40,8 @@ namespace UE::Learning::Agents::Private
 
 		Observation->FeatureObject = MakeShared<ObservationFObject>(
 			Observation->GetFName(),
-			AgentType->GetInstanceData().ToSharedRef(),
-			AgentType->GetMaxInstanceNum(),
+			AgentType->GetAgentManager()->GetInstanceData().ToSharedRef(),
+			AgentType->GetAgentManager()->GetMaxInstanceNum(),
 			Forward<InArgTypes>(Args)...);
 
 		AgentType->AddObservation(Observation, Observation->FeatureObject.ToSharedRef());
diff --git a/Engine/Plugins/Experimental/LearningAgents/Source/LearningAgents/Private/LearningAgentsPolicy.cpp b/Engine/Plugins/Experimental/LearningAgents/Source/LearningAgents/Private/LearningAgentsPolicy.cpp
--- a/Engine/Plugins/Experimental/LearningAgents/Source/LearningAgents/Private/LearningAgentsPolicy.cpp
+++ b/Engine/Plugins/Experimental/LearningAgents/Source/LearningAgents/Private/LearningAgentsPolicy.cpp
@@ -2,6 +2,7 @@
 
 #include "LearningAgentsPolicy.h"
 
+#include "LearningAgentsManager.h"
 #include "LearningAgentsType.h"
 #include "LearningFeatureObject.h"
 #include "LearningNeuralNetwork.h"
@@ -69,36 +70,47 @@ namespace UE::Learning::Agents::Policy::Private
 	}
 }
 
-ULearningAgentsPolicy::ULearningAgentsPolicy() : UActorComponent() {}
+ULearningAgentsPolicy::ULearningAgentsPolicy() : ULearningAgentsManagerComponent() {}
 ULearningAgentsPolicy::ULearningAgentsPolicy(FVTableHelper& Helper) : ULearningAgentsPolicy() {}
 ULearningAgentsPolicy::~ULearningAgentsPolicy() {}
 
-void ULearningAgentsPolicy::SetupPolicy(ULearningAgentsType* InAgentType, const FLearningAgentsPolicySettings& PolicySettings)
+void ULearningAgentsPolicy::SetupPolicy(ALearningAgentsManager* InAgentManager, ULearningAgentsType* InAgentType, const FLearningAgentsPolicySettings& PolicySettings)
 {
-	if (IsPolicySetupPerformed())
+	if (IsSetup())
 	{
-		UE_LOG(LogLearning, Error, TEXT("Setup already performed!"));
+		UE_LOG(LogLearning, Error, TEXT("%s: Setup already performed!"), *GetName());
 		return;
 	}
 
-	// Setup Agent Type
+	if (!InAgentManager)
+	{
+		UE_LOG(LogLearning, Error, TEXT("%s: InAgentManager is nullptr."), *GetName());
+		return;
+	}
+
+	if (!InAgentManager->IsManagerSetup())
+	{
+		UE_LOG(LogLearning, Error, TEXT("%s's SetupManager() must be run before %s can be setup."), *InAgentManager->GetName(), *GetName());
+		return;
+	}
+
+	AgentManager = InAgentManager;
 
 	if (!InAgentType)
 	{
-		UE_LOG(LogLearning, Error, TEXT("SetupPolicy called but AgentType is nullptr."));
+		UE_LOG(LogLearning, Error, TEXT("%s: InAgentType is nullptr."), *GetName());
 		return;
 	}
 
-	if (!InAgentType->IsSetupPerformed())
+	if (!InAgentType->IsSetup())
 	{
-		UE_LOG(LogLearning, Error, TEXT("AgentType Setup must be run before policy can be setup."));
+		UE_LOG(LogLearning, Error, TEXT("%s: %s's Setup must be run before it can be used."), *GetName(), *InAgentType->GetName());
 		return;
 	}
 
 	AgentType = InAgentType;
 
 	// Setup Neural Network
-
 	Network = NewObject<ULearningAgentsNeuralNetwork>(this, TEXT("PolicyNetwork"));
 	Network->NeuralNetwork = MakeShared<UE::Learning::FNeuralNetwork>();
 	Network->NeuralNetwork->Resize(
@@ -109,7 +121,6 @@ void ULearningAgentsPolicy::SetupPolicy(ULearningAgentsType* InAgentType, const
 	Network->NeuralNetwork->ActivationFunction = UE::Learning::Agents::GetActivationFunction(PolicySettings.ActivationFunction);
 
 	// Create Policy Object
-
 	UE::Learning::FNeuralNetworkPolicyFunctionSettings PolicyFunctionSettings;
 	PolicyFunctionSettings.ActionNoiseMin = PolicySettings.ActionNoiseMin;
 	PolicyFunctionSettings.ActionNoiseMax = PolicySettings.ActionNoiseMax;
@@ -117,80 +128,15 @@ void ULearningAgentsPolicy::SetupPolicy(ULearningAgentsType* InAgentType, const
 
 	PolicyObject = MakeShared<UE::Learning::FNeuralNetworkPolicyFunction>(
 		TEXT("PolicyObject"),
-		AgentType->GetInstanceData().ToSharedRef(),
-		AgentType->GetMaxInstanceNum(),
+		AgentManager->GetInstanceData().ToSharedRef(),
+		AgentManager->GetMaxInstanceNum(),
 		Network->NeuralNetwork.ToSharedRef(),
 		PolicySettings.ActionNoiseSeed,
 		PolicyFunctionSettings);
 
-	AgentType->GetInstanceData()->Link(AgentType->GetObservationFeature().FeatureHandle, PolicyObject->InputHandle);
-
-	// Done!
-	bPolicySetupPerformed = true;
-}
-
-bool ULearningAgentsPolicy::IsPolicySetupPerformed() const
-{
-	return bPolicySetupPerformed;
-}
-
-void ULearningAgentsPolicy::AddAgent(const int32 AgentId)
-{
-	if (!IsPolicySetupPerformed())
-	{
-		UE_LOG(LogLearning, Error, TEXT("Policy setup must be run before agents can be added!"));
-		return;
-	}
-
-	if (!AgentType->GetOccupiedAgentSet().Contains(AgentId))
-	{
-		UE_LOG(LogLearning, Error, TEXT("Unable to add: AgentId %d not found on AgentType. Make sure to add agents to the agent type before adding."), AgentId);
-		return;
-	}
-
-	if (SelectedAgentIds.Contains(AgentId))
-	{
-		UE_LOG(LogLearning, Error, TEXT("AgentId %i is already included in agents set"), AgentId);
-		return;
-	}
-
-	SelectedAgentIds.Add(AgentId);
-	SelectedAgentsSet = SelectedAgentIds;
-	SelectedAgentsSet.TryMakeSlice();
-}
-
-void ULearningAgentsPolicy::RemoveAgent(const int32 AgentId)
-{
-	if (!IsPolicySetupPerformed())
-	{
-		UE_LOG(LogLearning, Error, TEXT("Policy setup must be run before agents can be removed!"));
-		return;
-	}
-
-	if (SelectedAgentIds.RemoveSingleSwap(AgentId, false) == 0)
-	{
-		UE_LOG(LogLearning, Error, TEXT("Unable to remove: AgentId %d not found in the added agents set."), AgentId);
-		return;
-	}
-
-	SelectedAgentsSet = SelectedAgentIds;
-	SelectedAgentsSet.TryMakeSlice();
-}
-
-bool ULearningAgentsPolicy::HasAgent(const int32 AgentId) const
-{
-	return SelectedAgentsSet.Contains(AgentId);
-}
-
-ULearningAgentsType* ULearningAgentsPolicy::GetAgentType(const TSubclassOf<ULearningAgentsType> AgentClass)
-{
-	if (!IsPolicySetupPerformed())
-	{
-		UE_LOG(LogLearning, Error, TEXT("Policy setup must be run before getting the agent type!"));
-		return nullptr;
-	}
+	AgentManager->GetInstanceData()->Link(AgentType->GetObservationFeature().FeatureHandle, PolicyObject->InputHandle);
 
-	return AgentType;
+	bIsSetup = true;
 }
 
 UE::Learning::FNeuralNetwork& ULearningAgentsPolicy::GetPolicyNetwork()
@@ -205,7 +151,7 @@ UE::Learning::FNeuralNetworkPolicyFunction& ULearningAgentsPolicy::GetPolicyObje
 
 void ULearningAgentsPolicy::LoadPolicyFromSnapshot(const FDirectoryPath& Directory, const FString Filename)
 {
-	if (!IsPolicySetupPerformed())
+	if (!IsSetup())
 	{
 		UE_LOG(LogLearning, Error, TEXT("Policy setup must be run before network can be loaded."));
 		return;
@@ -237,7 +183,7 @@ void ULearningAgentsPolicy::LoadPolicyFromSnapshot(const FDirectoryPath& Directo
 
 void ULearningAgentsPolicy::SavePolicyToSnapshot(const FDirectoryPath& Directory, const FString Filename) const
 {
-	if (!IsPolicySetupPerformed())
+	if (!IsSetup())
 	{
 		UE_LOG(LogLearning, Error, TEXT("Policy setup must be run before network can be saved."));
 		return;
@@ -261,7 +207,7 @@ void ULearningAgentsPolicy::SavePolicyToSnapshot(const FDirectoryPath& Directory
 
 void ULearningAgentsPolicy::LoadPolicyFromAsset(const ULearningAgentsNeuralNetwork* NeuralNetworkAsset)
 {
-	if (!IsPolicySetupPerformed())
+	if (!IsSetup())
 	{
 		UE_LOG(LogLearning, Error, TEXT("Policy setup must be run before network can be loaded."));
 		return;
@@ -285,7 +231,7 @@ void ULearningAgentsPolicy::LoadPolicyFromAsset(const ULearningAgentsNeuralNetwo
 
 void ULearningAgentsPolicy::SavePolicyToAsset(ULearningAgentsNeuralNetwork* NeuralNetworkAsset) const
 {
-	if (!IsPolicySetupPerformed())
+	if (!IsSetup())
 	{
 		UE_LOG(LogLearning, Error, TEXT("Policy setup must be run before network can be saved."));
 		return;
@@ -324,13 +270,13 @@ void ULearningAgentsPolicy::EvaluatePolicy()
 {
 	UE_LEARNING_TRACE_CPUPROFILER_EVENT_SCOPE(ULearningAgentsPolicy::EvaluatePolicy);
 
-	if (!IsPolicySetupPerformed())
+	if (!IsSetup())
 	{
 		UE_LOG(LogLearning, Error, TEXT("Setup must be run before the policy can be evaluated."));
 		return;
 	}
 
-	PolicyObject->Evaluate(SelectedAgentIds);
+	PolicyObject->Evaluate(AddedAgentIds);
 
 	// Copy the actions computed by the policy into the agent-type's feature buffer.
 	// 
@@ -341,22 +287,22 @@ void ULearningAgentsPolicy::EvaluatePolicy()
 	UE::Learning::Array::Copy(
 		PolicyObject->InstanceData->View(AgentType->GetActionFeature().FeatureHandle),
 		PolicyObject->InstanceData->ConstView(PolicyObject->OutputHandle),
-		SelectedAgentIds);
+		AddedAgentIds);
 
 #if ENABLE_VISUAL_LOG
-	VisualLog(SelectedAgentIds);
+	VisualLog(AddedAgentIds);
 #endif
 }
 
 float ULearningAgentsPolicy::GetAgentActionNoiseScale(const int32 AgentId) const
 {
-	if (!IsPolicySetupPerformed())
+	if (!IsSetup())
 	{
 		UE_LOG(LogLearning, Error, TEXT("Setup must be run before getting the action noise."));
 		return 0.0f;
 	}
 
-	if (!SelectedAgentsSet.Contains(AgentId))
+	if (!HasAgent(AgentId))
 	{
 		UE_LOG(LogLearning, Error, TEXT("Unable to get action noise for agent - AgentId %d not found in the added agents set."), AgentId);
 		return 0.0f;
@@ -368,13 +314,13 @@ float ULearningAgentsPolicy::GetAgentActionNoiseScale(const int32 AgentId) const
 
 void ULearningAgentsPolicy::SetAgentActionNoiseScale(const int32 AgentId, const float ActionNoiseScale)
 {
-	if (!IsPolicySetupPerformed())
+	if (!IsSetup())
 	{
 		UE_LOG(LogLearning, Error, TEXT("Setup must be run before setting the action noise."));
 		return;
 	}
 
-	if (!SelectedAgentsSet.Contains(AgentId))
+	if (!HasAgent(AgentId))
 	{
 		UE_LOG(LogLearning, Error, TEXT("Unable to set action noise for agent - AgentId %d not found in the added agents set."), AgentId);
 		return;
@@ -423,4 +369,4 @@ void ULearningAgentsPolicy::VisualLog(const UE::Learning::FIndexSet Instances) c
 }
 #endif
 
-#undef UE_LEARNING_AGENTS_VLOG_STRING
\ No newline at end of file
+#undef UE_LEARNING_AGENTS_VLOG_STRING
diff --git a/Engine/Plugins/Experimental/LearningAgents/Source/LearningAgents/Private/LearningAgentsType.cpp b/Engine/Plugins/Experimental/LearningAgents/Source/LearningAgents/Private/LearningAgentsType.cpp
--- a/Engine/Plugins/Experimental/LearningAgents/Source/LearningAgents/Private/LearningAgentsType.cpp
+++ b/Engine/Plugins/Experimental/LearningAgents/Source/LearningAgents/Private/LearningAgentsType.cpp
@@ -1,6 +1,8 @@
 // Copyright Epic Games, Inc. All Rights Reserved.
 
 #include "LearningAgentsType.h"
+
+#include "LearningAgentsManager.h"
 #include "LearningAgentsObservations.h"
 #include "LearningAgentsActions.h"
 
@@ -10,41 +12,31 @@
 #include "LearningLog.h"
 #include "EngineDefines.h"
 
-ULearningAgentsType::ULearningAgentsType() : UActorComponent() {}
+ULearningAgentsType::ULearningAgentsType() : ULearningAgentsManagerComponent() {}
 ULearningAgentsType::ULearningAgentsType(FVTableHelper& Helper) : ULearningAgentsType() {}
 ULearningAgentsType::~ULearningAgentsType() {}
 
-void ULearningAgentsType::PostInitProperties()
+void ULearningAgentsType::SetupAgentType(ALearningAgentsManager* InAgentManager)
 {
-	Super::PostInitProperties();
-
-	if (HasAnyFlags(RF_ClassDefaultObject))
+	if (IsSetup())
 	{
+		UE_LOG(LogLearning, Error, TEXT("%s: Setup already performed!"), *GetName());
 		return;
 	}
 
-	// Pre-populate the vacant ids
-	OccupiedAgentIds.Reserve(MaxInstanceNum);
-	VacantAgentIds.Reserve(MaxInstanceNum);
-	for (int32 i = MaxInstanceNum - 1; i >= 0; i--)
+	if (!InAgentManager)
 	{
-		VacantAgentIds.Push(i);
-		Agents.Add(nullptr);
+		UE_LOG(LogLearning, Error, TEXT("%s: InAgentManager is nullptr."), *GetName());
+		return;
 	}
 
-	UpdateAgentSets();
-}
-
-void ULearningAgentsType::SetupAgentType()
-{
-	if (IsSetupPerformed())
+	if (!InAgentManager->IsManagerSetup())
 	{
-		UE_LOG(LogLearning, Error, TEXT("Setup already performed!"));
+		UE_LOG(LogLearning, Error, TEXT("%s's SetupManager() must be run before %s can be setup."), *InAgentManager->GetName(), *GetName());
 		return;
 	}
 
-	// Allocate Instance Data
-	InstanceData = MakeShared<UE::Learning::FArrayMap>();
+	AgentManager = InAgentManager;
 
 	// Setup Observations
 	ObservationObjects.Empty();
@@ -52,8 +44,8 @@ void ULearningAgentsType::SetupAgentType()
 	SetupObservations(this);
 	Observations = MakeShared<UE::Learning::FConcatenateFeature>(TEXT("Observations"),
 		TLearningArrayView<1, const TSharedRef<UE::Learning::FFeatureObject>>(ObservationFeatures),
-		InstanceData.ToSharedRef(),
-		MaxInstanceNum);
+		InAgentManager->GetInstanceData().ToSharedRef(),
+		InAgentManager->GetMaxInstanceNum());
 
 	// Setup Actions
 	ActionObjects.Empty();
@@ -61,26 +53,10 @@ void ULearningAgentsType::SetupAgentType()
 	SetupActions(this);
 	Actions = MakeShared<UE::Learning::FConcatenateFeature>(TEXT("Actions"),
 		TLearningArrayView<1, const TSharedRef<UE::Learning::FFeatureObject>>(ActionFeatures),
-		InstanceData.ToSharedRef(),
-		MaxInstanceNum);
+		InAgentManager->GetInstanceData().ToSharedRef(),
+		InAgentManager->GetMaxInstanceNum());
 
-	// Done!
-	bSetupPerformed = true;
-}
-
-bool ULearningAgentsType::IsSetupPerformed() const
-{
-	return bSetupPerformed;
-}
-
-int32 ULearningAgentsType::GetMaxInstanceNum() const
-{
-	return MaxInstanceNum;
-}
-
-const TSharedPtr<UE::Learning::FArrayMap>& ULearningAgentsType::GetInstanceData() const
-{
-	return InstanceData;
+	bIsSetup = true;
 }
 
 UE::Learning::FFeatureObject& ULearningAgentsType::GetObservationFeature() const
@@ -103,120 +79,6 @@ TConstArrayView<ULearningAgentsAction*> ULearningAgentsType::GetActionObjects()
 	return ActionObjects;
 }
 
-TConstArrayView<TObjectPtr<UObject>> ULearningAgentsType::GetAgents() const
-{
-	return Agents;
-}
-
-UE::Learning::FIndexSet ULearningAgentsType::GetOccupiedAgentSet() const
-{
-	return OccupiedAgentSet;
-}
-
-UE::Learning::FIndexSet ULearningAgentsType::GetVacantAgentSet() const
-{
-	return VacantAgentSet;
-}
-
-int32 ULearningAgentsType::AddAgent(UObject* Agent)
-{
-	if (Agent == nullptr)
-	{
-		UE_LOG(LogLearning, Error, TEXT("Attempted to add an agent but agent is nullptr."));
-		return INDEX_NONE;
-	}
-
-	if (VacantAgentIds.IsEmpty())
-	{
-		UE_LOG(LogLearning, Error, TEXT("Attempting to add an agent but we have no more vacant ids. Increase MaxInstanceNum (%d) or remove unused agents."), Agents.Num());
-		return INDEX_NONE;
-	}
-
-	// Add Agent
-	const int32 NewAgentId = VacantAgentIds.Pop();
-	Agents[NewAgentId] = Agent;
-
-	OccupiedAgentIds.Add(NewAgentId);
-
-	UpdateAgentSets();
-
-	return NewAgentId;
-}
-
-void ULearningAgentsType::RemoveAgentById(const int32 AgentId)
-{
-	if (AgentId == INDEX_NONE)
-	{
-		UE_LOG(LogLearning, Warning, TEXT("Attempting to remove an agent with id of INDEX_NONE."));
-		return;
-	}
-
-	const int32 RemovedCount = OccupiedAgentIds.RemoveSingleSwap(AgentId, false);
-
-	if (RemovedCount == 0)
-	{
-		UE_LOG(LogLearning, Warning, TEXT("Trying to remove an agent but its Id (%d) is not in the occupied agents."), AgentId);
-		return;
-	}
-
-	// Remove Agent
-	VacantAgentIds.Push(AgentId);
-	Agents[AgentId] = nullptr;
-
-	UpdateAgentSets();
-}
-
-void ULearningAgentsType::RemoveAgent(UObject* Agent)
-{
-	if (Agent == nullptr)
-	{
-		UE_LOG(LogLearning, Error, TEXT("Attempted to remove an agent but agent is nullptr."));
-		return;
-	}
-
-	int32 AgentId = INDEX_NONE;
-	const bool bIsFound = Agents.Find(Agent, AgentId);
-
-	if (!bIsFound)
-	{
-		UE_LOG(LogLearning, Warning, TEXT("Trying to remove an agent but it was not found."), AgentId);
-		return;
-	}
-
-	RemoveAgentById(AgentId);
-}
-
-bool ULearningAgentsType::HasAgent(UObject* Agent) const
-{
-	return Agents.Find(Agent) ? true : false;
-}
-
-bool ULearningAgentsType::HasAgentById(const int32 AgentId) const
-{
-	return OccupiedAgentSet.Contains(AgentId);
-}
-
-UObject* ULearningAgentsType::GetAgent(const int32 AgentId, const TSubclassOf<UObject> AgentClass)
-{
-	if (AgentId < 0 || AgentId >= Agents.Num())
-	{
-		UE_LOG(LogLearning, Warning, TEXT("AgentId %d outside valid range [0, %d]"), AgentId, Agents.Num() - 1);
-		return nullptr;
-	}
-
-	return Agents[AgentId];
-}
-
-const UObject* ULearningAgentsType::GetAgent(const int32 AgentId) const
-{
-	return Agents[AgentId];
-}
-
-UObject* ULearningAgentsType::GetAgent(const int32 AgentId)
-{
-	return Agents[AgentId];
-}
-
 void ULearningAgentsType::SetupObservations_Implementation(ULearningAgentsType* AgentType)
 {
 	// Can be overridden to setup observations without blueprints
@@ -229,7 +91,7 @@ void ULearningAgentsType::SetObservations_Implementation(const TArray<int32>& Ag
 
 void ULearningAgentsType::AddObservation(TObjectPtr<ULearningAgentsObservation> Object, const TSharedRef<UE::Learning::FFeatureObject>& Feature)
 {
-	UE_LEARNING_CHECK(!IsSetupPerformed());
+	UE_LEARNING_CHECK(!IsSetup());
 	ObservationObjects.Add(Object);
 	ObservationFeatures.Add(Feature);
 }
@@ -246,7 +108,7 @@ void ULearningAgentsType::GetActions_Implementation(const TArray<int32>& AgentId
 
 void ULearningAgentsType::AddAction(TObjectPtr<ULearningAgentsAction> Object, const TSharedRef<UE::Learning::FFeatureObject>& Feature)
 {
-	UE_LEARNING_CHECK(!IsSetupPerformed());
+	UE_LEARNING_CHECK(!IsSetup());
 	ActionObjects.Add(Object);
 	ActionFeatures.Add(Feature);
 }
@@ -255,22 +117,22 @@ void ULearningAgentsType::EncodeObservations()
 {
 	UE_LEARNING_TRACE_CPUPROFILER_EVENT_SCOPE(ULearningAgentsType::EncodeObservations);
 
-	if (!IsSetupPerformed())
+	if (!IsSetup())
 	{
 		UE_LOG(LogLearning, Error, TEXT("Setup must be run before observations can be encoded."));
 		return;
 	}
 
-	SetObservations(OccupiedAgentIds);
+	SetObservations(AddedAgentIds);
 
-	Observations->Encode(OccupiedAgentSet);
+	Observations->Encode(AddedAgentSet);
 
 #if ENABLE_VISUAL_LOG
 	for (const ULearningAgentsObservation* ObservationObject : ObservationObjects)
 	{
 		if (ObservationObject)
 		{
-			ObservationObject->VisualLog(OccupiedAgentSet);
+			ObservationObject->VisualLog(AddedAgentSet);
 		}
 	}
 #endif
@@ -280,31 +142,23 @@ void ULearningAgentsType::DecodeActions()
 {
 	UE_LEARNING_TRACE_CPUPROFILER_EVENT_SCOPE(ULearningAgentsType::DecodeActions);
 
-	if (!IsSetupPerformed())
+	if (!IsSetup())
 	{
 		UE_LOG(LogLearning, Error, TEXT("Setup must be run before actions can be decoded."));
 		return;
 	}
 
-	Actions->Decode(OccupiedAgentSet);
+	Actions->Decode(AddedAgentSet);
 
-	GetActions(OccupiedAgentIds);
+	GetActions(AddedAgentIds);
 
 #if ENABLE_VISUAL_LOG
 	for (const ULearningAgentsAction* ActionObject : ActionObjects)
 	{
 		if (ActionObject)
 		{
-			ActionObject->VisualLog(OccupiedAgentSet);
+			ActionObject->VisualLog(AddedAgentSet);
 		}
 	}
 #endif
 }
-
-void ULearningAgentsType::UpdateAgentSets()
-{
-	OccupiedAgentSet = OccupiedAgentIds;
-	OccupiedAgentSet.TryMakeSlice();
-	VacantAgentSet = VacantAgentIds;
-	VacantAgentSet.TryMakeSlice();
-}
\ No newline at end of file
diff --git a/Engine/Plugins/Experimental/LearningAgents/Source/LearningAgents/Public/LearningAgentsController.h b/Engine/Plugins/Experimental/LearningAgents/Source/LearningAgents/Public/LearningAgentsController.h
--- a/Engine/Plugins/Experimental/LearningAgents/Source/LearningAgents/Public/LearningAgentsController.h
+++ b/Engine/Plugins/Experimental/LearningAgents/Source/LearningAgents/Public/LearningAgentsController.h
@@ -2,10 +2,8 @@
 
 #pragma once
 
+#include "LearningAgentsManagerComponent.h"
 #include "LearningArray.h"
-
-#include "Components/ActorComponent.h"
-
 #include "LearningAgentsController.generated.h"
 
 /**
@@ -14,11 +12,10 @@
 * for imitation learning purposes.
 */
 UCLASS(Abstract, BlueprintType, Blueprintable)
-class LEARNINGAGENTS_API ULearningAgentsController : public UActorComponent
+class LEARNINGAGENTS_API ULearningAgentsController : public ULearningAgentsManagerComponent
 {
 	GENERATED_BODY()
 
-// ----- Setup -----
 public:
 
 	// These constructors/destructors are needed to make forward declarations happy
@@ -26,46 +23,9 @@ public:
 	ULearningAgentsController(FVTableHelper& Helper);
 	virtual ~ULearningAgentsController();
 
-	/** Initializes this object to be used with the given agent type. */
-	UFUNCTION(BlueprintCallable, Category = "LearningAgents")
-	void SetupController(ULearningAgentsType* InAgentType);
-
-	/** Returns true if SetupController has been run successfully; Otherwise, false. */
-	UFUNCTION(BlueprintPure, Category = "LearningAgents")
-	bool IsControllerSetupPerformed() const;
-
-// ----- Agent Management -----
-public:
-
-	/**
-	 * Adds an agent to this controller.
-	 * @param AgentId The id of the agent to be added.
-	 * @warning The agent id must exist for the agent type.
-	 */
+	/** Initializes this object to be used with the given manager and agent type. */
 	UFUNCTION(BlueprintCallable, Category = "LearningAgents")
-	void AddAgent(const int32 AgentId);
-
-	/**
-	* Removes an agent from this controller.
-	* @param AgentId The id of the agent to be removed.
-	* @warning The agent id must exist for this controller already.
-	*/
-	UFUNCTION(BlueprintCallable, Category = "LearningAgents")
-	void RemoveAgent(const int32 AgentId);
-
-	/** Returns true if the given id has been previously added to this controller; Otherwise, false. */
-	UFUNCTION(BlueprintPure, Category = "LearningAgents")
-	bool HasAgent(const int32 AgentId) const;
-
-	/**
-	* Gets the agent type this controller is associated with.
-	* @param AgentClass The class to cast the agent type to (in blueprint).
-	*/
-	UFUNCTION(BlueprintPure, Category = "LearningAgents", meta = (DeterminesOutputType = "AgentClass"))
-	ULearningAgentsType* GetAgentType(const TSubclassOf<ULearningAgentsType> AgentClass);
-
-// ----- Actions -----
-public:
+	void SetupController(ALearningAgentsManager* InAgentManager, ULearningAgentsType* InAgentType);
 
 	/**
 	* During this event, you should set the actions of your agents.
@@ -83,18 +43,18 @@ public:
 	UFUNCTION(BlueprintCallable, Category = "LearningAgents")
 	void EncodeActions();
 
-// ----- Private Data -----
-private:
-
-	/** The agent type this controller is associated with. */
-	UPROPERTY(VisibleAnywhere, Category = "LearningAgents")
-	TObjectPtr<ULearningAgentsType> AgentType;
+protected:
 
-	/** The agent ids this controller is managing. */
-	UPROPERTY(VisibleAnywhere, Category = "LearningAgents")
-	TArray<int32> SelectedAgentIds;
+	/**
+	* Gets the agent type associated with this component.
+	* @param AgentClass The class to cast the agent type to (in blueprint).
+	*/
+	UFUNCTION(BlueprintPure, Category = "LearningAgents", meta = (DeterminesOutputType = "AgentTypeClass"))
+	ULearningAgentsType* GetAgentType(const TSubclassOf<ULearningAgentsType> AgentTypeClass) const;
 
 private:
 
-	UE::Learning::FIndexSet SelectedAgentsSet;
+	/** The agent type this controller is associated with. */
+	UPROPERTY(VisibleAnywhere, Transient, Category = "LearningAgents")
+	TObjectPtr<ULearningAgentsType> AgentType;
 };
diff --git a/Engine/Plugins/Experimental/LearningAgents/Source/LearningAgents/Public/LearningAgentsCritic.h b/Engine/Plugins/Experimental/LearningAgents/Source/LearningAgents/Public/LearningAgentsCritic.h
--- a/Engine/Plugins/Experimental/LearningAgents/Source/LearningAgents/Public/LearningAgentsCritic.h
+++ b/Engine/Plugins/Experimental/LearningAgents/Source/LearningAgents/Public/LearningAgentsCritic.h
@@ -2,10 +2,10 @@
 
 #pragma once
 
+#include "LearningAgentsManagerComponent.h"
+
 #include "LearningAgentsNeuralNetwork.h" // Included for ELearningAgentsActivationFunction
 #include "LearningArray.h"
-
-#include "Components/ActorComponent.h"
 #include "UObject/ObjectPtr.h"
 #include "EngineDefines.h"
 
@@ -43,11 +43,10 @@ public:
 
 /** A critic used by some algorithms for training the managed agents. */
 UCLASS(BlueprintType, Blueprintable)
-class LEARNINGAGENTS_API ULearningAgentsCritic : public UActorComponent
+class LEARNINGAGENTS_API ULearningAgentsCritic : public ULearningAgentsManagerComponent
 {
 	GENERATED_BODY()
 
-// ----- Setup -----
 public:
 
 	// These constructors/destructors are needed to make forward declarations happy
@@ -57,41 +56,10 @@ public:
 
 	/** Initializes this object to be used with the given agent type and critic settings. */
 	UFUNCTION(BlueprintCallable, Category = "LearningAgents")
-	void SetupCritic(ULearningAgentsType* InAgentType, const FLearningAgentsCriticSettings& CriticSettings = FLearningAgentsCriticSettings());
-
-	/** Returns true if SetupCritic has been run successfully; Otherwise, false. */
-	UFUNCTION(BlueprintPure, Category = "LearningAgents")
-	bool IsCriticSetupPerformed() const;
-
-// ----- Agent Management -----
-public:
-
-	/**
-	 * Adds an agent to this critic.
-	 * @param AgentId The id of the agent to be added.
-	 * @warning The agent id must exist for the agent type.
-	 */
-	UFUNCTION(BlueprintCallable, Category = "LearningAgents")
-	void AddAgent(const int32 AgentId);
-
-	/**
-	* Removes an agent from this critic.
-	* @param AgentId The id of the agent to be removed.
-	* @warning The agent id must exist for this critic already.
-	*/
-	UFUNCTION(BlueprintCallable, Category = "LearningAgents")
-	void RemoveAgent(const int32 AgentId);
-
-	/** Returns true if the given id has been previously added to this critic; Otherwise, false. */
-	UFUNCTION(BlueprintPure, Category = "LearningAgents")
-	bool HasAgent(const int32 AgentId) const;
-
-	/**
-	* Gets the agent type this critic is associated with.
-	* @param AgentClass The class to cast the agent type to (in blueprint).
-	*/
-	UFUNCTION(BlueprintPure, Category = "LearningAgents", meta = (DeterminesOutputType = "AgentClass"))
-	ULearningAgentsType* GetAgentType(const TSubclassOf<ULearningAgentsType> AgentClass);
+	void SetupCritic(
+		ALearningAgentsManager* InAgentManager,
+		ULearningAgentsType* InAgentType,
+		const FLearningAgentsCriticSettings& CriticSettings = FLearningAgentsCriticSettings());
 
 // ----- Load / Save -----
 public:
@@ -161,29 +129,17 @@ public:
 private:
 
 	/** The agent type this critic is associated with. */
-	UPROPERTY(VisibleAnywhere, Category = "LearningAgents")
+	UPROPERTY(VisibleAnywhere, Transient, Category = "LearningAgents")
 	TObjectPtr<ULearningAgentsType> AgentType;
 
-	/** The agent ids this critic is managing. */
-	UPROPERTY(VisibleAnywhere, Category = "LearningAgents")
-	TArray<int32> SelectedAgentIds;
-
-	/** True if this critic's SetupCritic has been run; Otherwise, false. */
-	UPROPERTY(VisibleAnywhere, Category = "LearningAgents")
-	bool bCriticSetupPerformed = false;
-
 	/** The underlying neural network. */
-	UPROPERTY(VisibleAnywhere, Category = "LearningAgents")
+	UPROPERTY(VisibleAnywhere, Transient, Category = "LearningAgents")
 	TObjectPtr<ULearningAgentsNeuralNetwork> Network;
 
 private:
 
 	TSharedPtr<UE::Learning::FNeuralNetworkCriticFunction> CriticObject;
 
-	UE::Learning::FIndexSet SelectedAgentsSet;
-
-private:
-
 #if ENABLE_VISUAL_LOG
 	/** Color used to draw this action in the visual log */
 	FLinearColor VisualLogColor = FColor::Orange;
diff --git a/Engine/Plugins/Experimental/LearningAgents/Source/LearningAgents/Public/LearningAgentsManager.h b/Engine/Plugins/Experimental/LearningAgents/Source/LearningAgents/Public/LearningAgentsManager.h
new file mode 100644
--- /dev/null
+++ b/Engine/Plugins/Experimental/LearningAgents/Source/LearningAgents/Public/LearningAgentsManager.h
@@ -0,0 +1,167 @@
+// Copyright Epic Games, Inc. All Rights Reserved.
+
+#pragma once
+
+#include "LearningArray.h"
+
+#include "GameFramework/Actor.h"
+#include "Containers/Array.h"
+#include "Containers/ArrayView.h"
+#include "Containers/ContainerAllocationPolicies.h"
+#include "Templates/SharedPointer.h"
+#include "UObject/Object.h"
+#include "UObject/ObjectPtr.h"
+
+#include "LearningAgentsManager.generated.h"
+
+namespace UE::Learning
+{
+	struct FArrayMap;
+}
+
+class ULearningAgentsManagerComponent;
+
+/**
+* The agent manager is responsible for tracking which game objects are agents. It's the central class around which
+* most of Learning Agents is built.
+*
+* If you have multiple different types of objects you want controlled by Learning Agents, you should consider creating
+* one agent manager per object type, rather than trying to share an agent manager.
+*/
+UCLASS(Abstract, BlueprintType, Blueprintable)
+class LEARNINGAGENTS_API ALearningAgentsManager : public AActor
+{
+	GENERATED_BODY()
+
+public:
+
+	// These constructors/destructors are needed to make forward declarations happy
+	ALearningAgentsManager();
+	ALearningAgentsManager(FVTableHelper& Helper);
+	virtual ~ALearningAgentsManager();
+
+	/** Sets up the agent ids so that agents can be added prior to calling SetupManager. */
+	virtual void PostInitProperties() override;
+
+// ----- Setup -----
+public:
+
+	/** Initializes this object and runs the setup events for observations and actions. */
+	UFUNCTION(BlueprintCallable, Category = "LearningAgents")
+	void SetupManager();
+
+	/** Returns true if SetupManager has been run successfully; Otherwise, false. */
+	UFUNCTION(BlueprintPure, Category = "LearningAgents")
+	bool IsManagerSetup() const;
+
+	/** Returns the maximum number of agents that this manager is configured to handle. */
+	UFUNCTION(BlueprintPure, Category = "LearningAgents")
+	int32 GetMaxInstanceNum() const;
+
+// ----- Agent Management -----
+public:
+
+	/**
+	 * Adds the given object as an agent to this manager. This can be called before or after SetupManager.
+	 * @param Agent The object to be added.
+	 * @return The agent's newly assigned id.
+	 */
+	UFUNCTION(BlueprintCallable, Category = "LearningAgents")
+	int32 AddAgent(UObject* Agent);
+
+	/**
+	 * Called whenever a new agent is added to this manager. By default, this will add the agent to each of this
+	 * manager's agent components, but you can override this event if you have custom logic. For example, maybe some of
+	 * the agents are inference agents only and should not be added to the trainer component. 
+	 * @param AgentId The agent's newly assigned id.
+	 */
+	UFUNCTION(BlueprintNativeEvent, Category = "LearningAgents")
+	void OnAgentAdded(const int32 AgentId);
+	
+	/**
+	 * Removes the agent with the given id from this manager.
+	 * @param AgentId The id of the agent to remove.
+	 */
+	UFUNCTION(BlueprintCallable, Category = "LearningAgents", meta = (AgentId = "-1"))
+	void RemoveAgentById(const int32 AgentId);
+
+	/**
+	* Removes the given agent from this manager. Use RemoveAgentById if you have the id available as this function
+	* must do a linear search to find the agent.
+	* @param Agent The agent to be removed.
+	*/
+	UFUNCTION(BlueprintCallable, Category = "LearningAgents")
+	void RemoveAgent(UObject* Agent);
+
+	/**
+	 * Called whenever an agent is removed from this manager. By default, this will remove the agent from each of
+	 * this manager's agent components, but you can override this event if you have additional logic. 
+	 * @param AgentId The removed agent's id.
+	 */
+	UFUNCTION(BlueprintNativeEvent, Category = "LearningAgents")
+	void OnAgentRemoved(const int32 AgentId);
+
+	/**
+	 * Gets the agent with the given id. Calling this from blueprint with the appropriate AgentClass will automatically
+	 * cast the object to the given type. If not in a blueprint, you should use one of the other GetAgent overloads.
+	 * @param AgentId The id of the agent to get.
+	 * @param AgentClass The class to cast the agent object to (in blueprint).
+	 * @return The agent object.
+	 */
+	UFUNCTION(BlueprintPure, Category = "LearningAgents", meta = (AgentId = "-1", DeterminesOutputType = "AgentClass"))
+	UObject* GetAgent(const int32 AgentId, const TSubclassOf<UObject> AgentClass);
+
+	/**
+	* Returns true if the given object is an agent in this manager; Otherwise, false.
+	* Use HasAgentById if you have the id available as this function must do a linear search to find the agent.
+	*/
+	UFUNCTION(BlueprintPure, Category = "LearningAgents")
+	bool HasAgent(UObject* Agent) const;
+
+	/** Returns true if the given id is used by an agent in this manager; Otherwise, false. */
+	UFUNCTION(BlueprintPure, Category = "LearningAgents", meta = (AgentId = "-1"))
+	bool HasAgentById(const int32 AgentId) const;
+
+// ----- Non-blueprint public interface -----
+public:
+
+	/** Gets the agent corresponding to the given id. */
+	const UObject* GetAgent(const int32 AgentId) const;
+
+	/** Gets the agent corresponding to the given id. */
+	UObject* GetAgent(const int32 AgentId);
+
+	/** Get a const reference to this manager's underlying instance data. */
+	const TSharedPtr<UE::Learning::FArrayMap>& GetInstanceData() const;
+
+	/** Get a const array view of this manager's agent objects. */
+	TConstArrayView<TObjectPtr<UObject>> GetAgents() const;
+
+private:
+
+	/** Maximum number of agents. Used to preallocate internal buffers. Setting this higher will allow more agents but use up more memory. */
+	UPROPERTY(EditDefaultsOnly, Category = "LearningAgents", meta = (ClampMin = "1", UIMin = "1"))
+	int32 MaxInstanceNum = 1;
+
+	/** True if SetupManager has been performed; Otherwise, false. */
+	UPROPERTY(VisibleAnywhere, Transient, Category = "LearningAgents")
+	bool bIsSetup = false;
+
+	/** The list of current agents. */
+	UPROPERTY(VisibleAnywhere, Transient, Category = "LearningAgents")
+	TArray<TObjectPtr<UObject>> Agents;
+
+	/** Update the agent sets to keep them in sync with the id lists. */
+	void UpdateAgentSets();
+
+	TArray<int32> OccupiedAgentIds;
+	TArray<int32> VacantAgentIds;
+	UE::Learning::FIndexSet OccupiedAgentSet;
+	UE::Learning::FIndexSet VacantAgentSet;
+
+	TSharedPtr<UE::Learning::FArrayMap> InstanceData;
+
+	/** The manager components that were found during setup. */
+	UPROPERTY(VisibleAnywhere, Transient, Category = "LearningAgents")
+	TArray<TObjectPtr<ULearningAgentsManagerComponent>> CachedManagerComponents;
+};
diff --git a/Engine/Plugins/Experimental/LearningAgents/Source/LearningAgents/Public/LearningAgentsManagerComponent.h b/Engine/Plugins/Experimental/LearningAgents/Source/LearningAgents/Public/LearningAgentsManagerComponent.h
new file mode 100644
--- /dev/null
+++ b/Engine/Plugins/Experimental/LearningAgents/Source/LearningAgents/Public/LearningAgentsManagerComponent.h
@@ -0,0 +1,102 @@
+// Copyright Epic Games, Inc. All Rights Reserved.
+
+#pragma once
+
+#include "Components/ActorComponent.h"
+
+#include "LearningArray.h"
+#include "Containers/Array.h"
+#include "UObject/Object.h"
+#include "UObject/ObjectPtr.h"
+
+#include "LearningAgentsManagerComponent.generated.h"
+
+class ALearningAgentsManager;
+
+/**
+ * Base class for components which can be attached to an ALearningAgentsManager.
+ */
+UCLASS(Abstract, BlueprintType, Blueprintable)
+class LEARNINGAGENTS_API ULearningAgentsManagerComponent : public UActorComponent
+{
+	GENERATED_BODY()
+
+public:
+
+	// These constructors/destructors are needed to make forward declarations happy
+	ULearningAgentsManagerComponent();
+	ULearningAgentsManagerComponent(FVTableHelper& Helper);
+	virtual ~ULearningAgentsManagerComponent();
+
+	/**
+	 * Adds an agent to this component.
+	 * @param AgentId The id of the agent to be added.
+	 * @return True if the agent was added successfully. Otherwise, false.
+	 */
+	UFUNCTION(BlueprintCallable, Category = "LearningAgents", meta = (AgentId = "-1"))
+	virtual bool AddAgent(const int32 AgentId);
+
+	/**
+	* Removes an agent from this component.
+	* @param AgentId The id of the agent to be removed.
+	* @return True if the agent was removed successfully. Otherwise, false.
+	*/
+	UFUNCTION(BlueprintCallable, Category = "LearningAgents", meta = (AgentId = "-1"))
+	virtual bool RemoveAgent(const int32 AgentId);
+
+	/** Returns true if the given id has been previously added to this component; Otherwise, false. */
+	UFUNCTION(BlueprintPure, Category = "LearningAgents", meta = (AgentId = "-1"))
+	bool HasAgent(const int32 AgentId) const;
+
+	/** Returns true if this component has been setup. Otherwise, false. */
+	UFUNCTION(BlueprintPure, Category = "LearningAgents")
+	bool IsSetup() const;
+
+// ----- Blueprint Convenience Functions -----
+protected:
+
+	/**
+	* Gets the agent with the given id from the manager. Calling this from blueprint with the appropriate AgentClass
+	* will automatically cast the object to the given type. If not in a blueprint, you should call the manager's
+	* GetAgent methods directly.
+	* @param AgentId The id of the agent to get.
+	* @param AgentClass The class to cast the agent object to (in blueprint).
+	* @return The agent object.
+	*/
+	UFUNCTION(BlueprintPure, Category = "LearningAgents", meta = (AgentId = "-1", DeterminesOutputType = "AgentClass"))
+	UObject* GetAgent(const int32 AgentId, const TSubclassOf<UObject> AgentClass) const;
+
+	/**
+	* Gets the agent manager associated with this component.
+	* @param AgentClass The class to cast the agent manager to (in blueprint).
+	*/
+	UFUNCTION(BlueprintPure, Category = "LearningAgents", meta = (DeterminesOutputType = "AgentManagerClass"))
+	ALearningAgentsManager* GetAgentManager(const TSubclassOf<ALearningAgentsManager> AgentManagerClass) const;
+
+// ----- Non-blueprint public interface -----
+public:
+
+	/** Gets the agent corresponding to the given id. */
+	const UObject* GetAgent(const int32 AgentId) const;
+
+	/** Gets the agent corresponding to the given id. */
+	UObject* GetAgent(const int32 AgentId);
+
+	ALearningAgentsManager* GetAgentManager() const;
+
+protected:
+
+	/** True if this component has been setup. Otherwise, false. */
+	UPROPERTY(VisibleAnywhere, Transient, Category = "LearningAgents")
+	bool bIsSetup = false;
+
+	/** The agent manager associated with this component. */
+	UPROPERTY(VisibleAnywhere, Transient, Category = "LearningAgents")
+	TObjectPtr<ALearningAgentsManager> AgentManager;
+
+	/** The agent ids added to this component. */
+	UPROPERTY(VisibleAnywhere, Transient, Category = "LearningAgents")
+	TArray<int32> AddedAgentIds;
+
+	UE::Learning::FIndexSet AddedAgentSet;
+};
diff --git a/Engine/Plugins/Experimental/LearningAgents/Source/LearningAgents/Public/LearningAgentsPolicy.h b/Engine/Plugins/Experimental/LearningAgents/Source/LearningAgents/Public/LearningAgentsPolicy.h
--- a/Engine/Plugins/Experimental/LearningAgents/Source/LearningAgents/Public/LearningAgentsPolicy.h
+++ b/Engine/Plugins/Experimental/LearningAgents/Source/LearningAgents/Public/LearningAgentsPolicy.h
@@ -2,10 +2,10 @@
 
 #pragma once
 
+#include "LearningAgentsManagerComponent.h"
+
 #include "LearningAgentsNeuralNetwork.h" // Included for ELearningAgentsActivationFunction
 #include "LearningArray.h"
-
-#include "Components/ActorComponent.h"
 #include "UObject/ObjectPtr.h"
 #include "EngineDefines.h"
 
@@ -59,11 +59,10 @@ public:
 
 /** A policy that maps from observations to actions for the managed agents. */
 UCLASS(BlueprintType, Blueprintable)
-class LEARNINGAGENTS_API ULearningAgentsPolicy : public UActorComponent
+class LEARNINGAGENTS_API ULearningAgentsPolicy : public ULearningAgentsManagerComponent
 {
 	GENERATED_BODY()
 
-// ----- Setup -----
 public:
 
 	// These constructors/destructors are needed to make forward declarations happy
@@ -73,41 +72,7 @@ public:
 
 	/** Initializes this object to be used with the given agent type and policy settings. */
 	UFUNCTION(BlueprintCallable, Category = "LearningAgents")
-	void SetupPolicy(ULearningAgentsType* InAgentType, const FLearningAgentsPolicySettings& PolicySettings = FLearningAgentsPolicySettings());
-
-	/** Returns true if SetupPolicy has been run successfully; Otherwise, false. */
-	UFUNCTION(BlueprintPure, Category = "LearningAgents")
-	bool IsPolicySetupPerformed() const;
-
-// ----- Agent Management -----
-public:
-
-	/**
-	 * Adds an agent to this policy.
-	 * @param AgentId The id of the agent to be added.
-	 * @warning The agent id must exist for the agent type.
-	 */
-	UFUNCTION(BlueprintCallable, Category = "LearningAgents")
-	void AddAgent(const int32 AgentId);
-
-	/**
-	* Removes an agent from this policy.
-	* @param AgentId The id of the agent to be removed.
-	* @warning The agent id must exist for this policy already.
-	*/
-	UFUNCTION(BlueprintCallable, Category = "LearningAgents")
-	void RemoveAgent(const int32 AgentId);
-
-	/** Returns true if the given id has been previously added to this policy; Otherwise, false. */
-	UFUNCTION(BlueprintPure, Category = "LearningAgents")
-	bool HasAgent(const int32 AgentId) const;
-
-	/**
-	* Gets the agent type this policy is associated with.
-	* @param AgentClass The class to cast the agent type to (in blueprint).
-	*/
-	UFUNCTION(BlueprintPure, Category = "LearningAgents", meta = (DeterminesOutputType = "AgentClass"))
-	ULearningAgentsType* GetAgentType(const TSubclassOf<ULearningAgentsType> AgentClass);
+	void SetupPolicy(ALearningAgentsManager* InAgentManager, ULearningAgentsType* InAgentType, const FLearningAgentsPolicySettings& PolicySettings = FLearningAgentsPolicySettings());
 
 // ----- Load / Save -----
 public:
@@ -187,29 +152,17 @@ public:
 private:
 
 	/** The agent type this policy is associated with. */
-	UPROPERTY(VisibleAnywhere, Category = "LearningAgents")
+	UPROPERTY(VisibleAnywhere, Transient, Category = "LearningAgents")
 	TObjectPtr<ULearningAgentsType> AgentType;
 
-	/** The agent ids this policy is managing. */
-	UPROPERTY(VisibleAnywhere, Category = "LearningAgents")
-	TArray<int32> SelectedAgentIds;
-
-	/** True if this policy's SetupPolicy has been run; Otherwise, false. */
-	UPROPERTY(VisibleAnywhere, Category = "LearningAgents")
-	bool bPolicySetupPerformed = false;
-
 	/** The underlying neural network. */
-	UPROPERTY(VisibleAnywhere, Category = "LearningAgents")
+	UPROPERTY(VisibleAnywhere, Transient, Category = "LearningAgents")
 	TObjectPtr<ULearningAgentsNeuralNetwork> Network;
 
 private:
 
 	TSharedPtr<UE::Learning::FNeuralNetworkPolicyFunction> PolicyObject;
 
-	UE::Learning::FIndexSet SelectedAgentsSet;
-
-private:
-
 #if ENABLE_VISUAL_LOG
 	/** Color used to draw this action in the visual log */
 	FLinearColor VisualLogColor = FColor::Purple;
diff --git a/Engine/Plugins/Experimental/LearningAgents/Source/LearningAgents/Public/LearningAgentsType.h b/Engine/Plugins/Experimental/LearningAgents/Source/LearningAgents/Public/LearningAgentsType.h
--- a/Engine/Plugins/Experimental/LearningAgents/Source/LearningAgents/Public/LearningAgentsType.h
+++ b/Engine/Plugins/Experimental/LearningAgents/Source/LearningAgents/Public/LearningAgentsType.h
@@ -2,9 +2,9 @@
 
 #pragma once
 
-#include "LearningArray.h"
+#include "LearningAgentsManagerComponent.h"
 
-#include "Components/ActorComponent.h"
+#include "LearningArray.h"
 #include "Containers/Array.h"
 #include "Containers/ArrayView.h"
 #include "Containers/ContainerAllocationPolicies.h"
@@ -21,27 +21,21 @@ namespace UE::Learning
 	struct FFeatureObject;
 }
 
+class ALearningAgentsManager;
 class ULearningAgentsAction;
 class ULearningAgentsObservation;
 
 /**
-* The agent type is the core class around which the rest of Learning Agents is built. It has a few responsibilities:
-*   1) It keeps track of which objects are agents.
-*   2) It defines how those agents' observations and actions are implemented.
-*   3) It provides methods that need to be called during the inference process of those agents.
+* The agent type defines how agents' observations and actions are implemented and provides methods to be called during
+* the inference process of those agents.
 *
 * To use this class, you need to implement the SetupObservations and SetupActions functions (as well as their
 * corresponding SetObservations and SetActions functions), which will define the size of inputs and outputs to your
 * policy. Before you can do inference, you need to call SetupAgentType, which will initialize the underlying data
 * structure, and you need to call AddAgent for each object you want controlled by this agent type.
-*
-* If you have multiple different types of objects you want controlled by Learning Agents, you should create
-* one agent type per object type, rather than trying to share an agent type.
-*
-* @see ULearningAgentsTrainer to understand how training works.
 */
 UCLASS(Abstract, BlueprintType, Blueprintable)
-class LEARNINGAGENTS_API ULearningAgentsType : public UActorComponent
+class LEARNINGAGENTS_API ULearningAgentsType : public ULearningAgentsManagerComponent
 {
 	GENERATED_BODY()
 
@@ -52,75 +46,9 @@ public:
 	ULearningAgentsType(FVTableHelper& Helper);
 	virtual ~ULearningAgentsType();
 
-	/** Sets up the agent ids so that agents can be added prior to calling SetupAgentType. */
-	virtual void PostInitProperties() override;
-
-// ----- Setup -----
-public:
-
 	/** Initializes this object and runs the setup events for observations and actions. */
 	UFUNCTION(BlueprintCallable, Category = "LearningAgents")
-	void SetupAgentType();
-
-	/** Returns true if SetupAgentType has been run successfully; Otherwise, false. */
-	UFUNCTION(BlueprintPure, Category = "LearningAgents")
-	bool IsSetupPerformed() const;
-
-	/** Returns the maximum number of agents that this agent type is configured to handle. */
-	UFUNCTION(BlueprintPure, Category = "LearningAgents")
-	int32 GetMaxInstanceNum() const;
-
-// ----- Agent Management -----
-public:
-
-	/**
-	 * Adds the given object as an agent to this agent type. This can be called before or after SetupAgentType.
-	 * @param Agent The object to be added.
-	 * @return The agent's newly assigned id.
-	 */
-	UFUNCTION(BlueprintCallable, Category = "LearningAgents")
-	int32 AddAgent(UObject* Agent);
-	
-	/**
-	 * Removes the agent with the given id from this agent type.
-	 * @param AgentId The id of the agent to remove.
-	 */
-	UFUNCTION(BlueprintCallable, Category = "LearningAgents")
-	void RemoveAgentById(const int32 AgentId);
-
-	/**
-	* Removes the given agent from this agent type. Use RemoveAgentById if you have the id available as this function
-	* is slower.
-	* @param Agent The agent to be removed.
-	*/
-	UFUNCTION(BlueprintCallable, Category = "LearningAgents")
-	void RemoveAgent(UObject* Agent);
-
-	/**
-	 * Gets the agent with the given id. Calling this from blueprint with the appropriate AgentClass will automatically
-	 * cast the object to the given type. If not in a blueprint, you should use one of the other GetAgent overloads.
-	 * @param AgentId The id of the agent to get.
-	 * @param AgentClass The class to cast the agent object to (in blueprint).
-	 * @return The agent object.
-	 */
-	UFUNCTION(BlueprintPure, Category = "LearningAgents", meta = (DeterminesOutputType = "AgentClass"))
-	UObject* GetAgent(const int32 AgentId, const TSubclassOf<UObject> AgentClass);
-
-	/** Returns true if the given object is an agent in this agent type; Otherwise, false. */
-	UFUNCTION(BlueprintPure, Category = "LearningAgents")
-	bool HasAgent(UObject* Agent) const;
-
-	/** Returns true if the given id is used by an agent in this agent type; Otherwise, false. */
-	UFUNCTION(BlueprintPure, Category = "LearningAgents")
-	bool HasAgentById(const int32 AgentId) const;
-
-public:
-
-	/** Gets the agent corresponding to the given id. */
-	const UObject* GetAgent(const int32 AgentId) const;
-
-	/** Gets the agent corresponding to the given id. */
-	UObject* GetAgent(const int32 AgentId);
+	void SetupAgentType(ALearningAgentsManager* InAgentManager);
 
 // ----- Observations -----
 public:
@@ -196,9 +124,6 @@ public:
 // ----- Non-blueprint public interface -----
 public:
 
-	/** Get a const reference to this agent type's underlying instance data. */
-	const TSharedPtr<UE::Learning::FArrayMap>& GetInstanceData() const;
-
 	/** Get a reference to this agent type's observation feature. */
 	UE::Learning::FFeatureObject& GetObservationFeature() const;
 
@@ -211,53 +136,22 @@ public:
 	/** Get a const array view of this agent type's action objects. */
 	TConstArrayView<ULearningAgentsAction*> GetActionObjects() const;
 
-	/** Get a const array view of this agent type's agent objects. */
-	TConstArrayView<TObjectPtr<UObject>> GetAgents() const;
+private:
 
-	/** Get an FIndexSet with this agent type's occupied agent ids. */
-	UE::Learning::FIndexSet GetOccupiedAgentSet() const;
+	/** The list of current observation objects. */
+	UPROPERTY(VisibleAnywhere, Transient, Category = "LearningAgents")
+	TArray<TObjectPtr<ULearningAgentsObservation>> ObservationObjects;
 
-	/** Get an FIndexSet with this agent type's vacant agent ids. */
-	UE::Learning::FIndexSet GetVacantAgentSet() const;
+	/** The list of current action objects. */
+	UPROPERTY(VisibleAnywhere, Transient, Category = "LearningAgents")
+	TArray<TObjectPtr<ULearningAgentsAction>> ActionObjects;
 
 // ----- Private Data -----
 private:
 
-	/** Update the agent sets to keep them in sync with the id lists. */
-	void UpdateAgentSets();
-
-	TArray<int32> OccupiedAgentIds;
-	TArray<int32> VacantAgentIds;
-	UE::Learning::FIndexSet OccupiedAgentSet;
-	UE::Learning::FIndexSet VacantAgentSet;
-
-	TSharedPtr<UE::Learning::FArrayMap> InstanceData;
-
 	TArray<TSharedRef<UE::Learning::FFeatureObject>, TInlineAllocator<16>> ObservationFeatures;
 	TArray<TSharedRef<UE::Learning::FFeatureObject>, TInlineAllocator<16>> ActionFeatures;
 
-	TSharedPtr<UE::Learning::FConcatenateFeature> Observations;	
+	TSharedPtr<UE::Learning::FConcatenateFeature> Observations;
 	TSharedPtr<UE::Learning::FConcatenateFeature> Actions;
-	
-private:
-
-	/** Maximum number of agent instances. Used to preallocate internal buffers. */
-	UPROPERTY(EditDefaultsOnly, Category = "LearningAgents", meta = (ClampMin = "1", UIMin = "1"))
-	int32 MaxInstanceNum = 1;
-
-	/** True if SetupAgentType has been performed; Otherwise, false. */
-	UPROPERTY(VisibleAnywhere, Category = "LearningAgents")
-	bool bSetupPerformed = false;
-
-	/** The list of current agents. */
-	UPROPERTY(VisibleAnywhere, Transient, Category = "LearningAgents")
-	TArray<TObjectPtr<UObject>> Agents;
-
-	/** The list of current observation objects. */
-	UPROPERTY(VisibleAnywhere, Category = "LearningAgents")
-	TArray<TObjectPtr<ULearningAgentsObservation>> ObservationObjects;
-
-	/** The list of current action objects. */
-	UPROPERTY(VisibleAnywhere, Category = "LearningAgents")
-	TArray<TObjectPtr<ULearningAgentsAction>> ActionObjects;
 };
diff --git a/Engine/Plugins/Experimental/LearningAgents/Source/LearningAgentsTraining/Private/LearningAgentsCompletions.cpp b/Engine/Plugins/Experimental/LearningAgents/Source/LearningAgentsTraining/Private/LearningAgentsCompletions.cpp
--- a/Engine/Plugins/Experimental/LearningAgents/Source/LearningAgentsTraining/Private/LearningAgentsCompletions.cpp
+++ b/Engine/Plugins/Experimental/LearningAgents/Source/LearningAgentsTraining/Private/LearningAgentsCompletions.cpp
@@ -3,7 +3,7 @@
 #include "LearningAgentsCompletions.h"
 
 #include "LearningAgentsTrainer.h"
-#include "LearningAgentsType.h"
+#include "LearningAgentsManager.h"
 #include "LearningArray.h"
 #include "LearningArrayMap.h"
 #include "LearningCompletion.h"
@@ -40,8 +40,8 @@ namespace UE::Learning::Agents::Private
 
 		Completion->CompletionObject = MakeShared<CompletionFObject>(
 			Completion->GetFName(),
-			AgentTrainer->GetAgentType()->GetInstanceData().ToSharedRef(),
-			AgentTrainer->GetAgentType()->GetMaxInstanceNum(),
+			AgentTrainer->GetAgentManager()->GetInstanceData().ToSharedRef(),
+			AgentTrainer->GetAgentManager()->GetMaxInstanceNum(),
 			Forward<InArgTypes>(Args)...);
 
 		AgentTrainer->AddCompletion(Completion, Completion->CompletionObject.ToSharedRef());
diff --git a/Engine/Plugins/Experimental/LearningAgents/Source/LearningAgentsTraining/Private/LearningAgentsDataStorage.cpp b/Engine/Plugins/Experimental/LearningAgents/Source/LearningAgentsTraining/Private/LearningAgentsDataStorage.cpp
--- a/Engine/Plugins/Experimental/LearningAgents/Source/LearningAgentsTraining/Private/LearningAgentsDataStorage.cpp
+++ b/Engine/Plugins/Experimental/LearningAgents/Source/LearningAgentsTraining/Private/LearningAgentsDataStorage.cpp
@@ -78,7 +78,7 @@ ULearningAgentsRecord* ULearningAgentsDataStorage::CreateRecord(const FName Reco
 		return nullptr;
 	}
 
-	if (!AgentType->IsSetupPerformed())
+	if (!AgentType->IsSetup())
 	{
 		UE_LOG(LogLearning, Warning, TEXT("CreateRecord: AgentType Setup must be performed before record can be created."));
 		return nullptr;
@@ -134,7 +134,7 @@ ULearningAgentsRecord* ULearningAgentsDataStorage::LoadRecord(ULearningAgentsTyp
 		return nullptr;
 	}
 
-	if (!AgentType->IsSetupPerformed())
+	if (!AgentType->IsSetup())
 	{
 		UE_LOG(LogLearning, Warning, TEXT("LoadRecord: AgentType Setup must be performed before record can be created."));
 		return nullptr;
diff --git a/Engine/Plugins/Experimental/LearningAgents/Source/LearningAgentsTraining/Private/LearningAgentsImitationTrainer.cpp b/Engine/Plugins/Experimental/LearningAgents/Source/LearningAgentsTraining/Private/LearningAgentsImitationTrainer.cpp
--- a/Engine/Plugins/Experimental/LearningAgents/Source/LearningAgentsTraining/Private/LearningAgentsImitationTrainer.cpp
+++ b/Engine/Plugins/Experimental/LearningAgents/Source/LearningAgentsTraining/Private/LearningAgentsImitationTrainer.cpp
@@ -48,7 +48,7 @@ void ULearningAgentsImitationTrainer::BeginTraining(
 		return;
 	}
 
-	if (!InPolicy->IsPolicySetupPerformed())
+	if (!InPolicy->IsSetup())
 	{
 		UE_LOG(LogLearning, Error, TEXT("Policy Setup not Performed"));
 		return;
diff --git a/Engine/Plugins/Experimental/LearningAgents/Source/LearningAgentsTraining/Private/LearningAgentsRecorder.cpp b/Engine/Plugins/Experimental/LearningAgents/Source/LearningAgentsTraining/Private/LearningAgentsRecorder.cpp
--- a/Engine/Plugins/Experimental/LearningAgents/Source/LearningAgentsTraining/Private/LearningAgentsRecorder.cpp
+++ b/Engine/Plugins/Experimental/LearningAgents/Source/LearningAgentsTraining/Private/LearningAgentsRecorder.cpp
@@ -2,14 +2,16 @@
 
 #include "LearningAgentsRecorder.h"
 
-#include "Engine/World.h"
-#include "LearningAgentsDataStorage.h"
+#include "LearningAgentsManager.h"
 #include "LearningAgentsType.h"
+#include "LearningAgentsDataStorage.h"
 #include "LearningFeatureObject.h"
 #include "LearningLog.h"
 #include "LearningTrainer.h"
 
-ULearningAgentsRecorder::ULearningAgentsRecorder() : UActorComponent() {}
+#include "Engine/World.h"
+
+ULearningAgentsRecorder::ULearningAgentsRecorder() : ULearningAgentsManagerComponent() {}
 ULearningAgentsRecorder::ULearningAgentsRecorder(FVTableHelper& Helper) : ULearningAgentsRecorder() {}
 ULearningAgentsRecorder::~ULearningAgentsRecorder() {}
 
@@ -23,29 +25,41 @@ void ULearningAgentsRecorder::EndPlay(const EEndPlayReason::Type EndPlayReason)
 	Super::EndPlay(EndPlayReason);
 }
 
-void ULearningAgentsRecorder::SetupRecorder(ULearningAgentsType* InAgentType)
+void ULearningAgentsRecorder::SetupRecorder(ALearningAgentsManager* InAgentManager, ULearningAgentsType* InAgentType)
 {
-	if (IsRecorderSetupPerformed())
+	if (IsSetup())
 	{
-		UE_LOG(LogLearning, Error, TEXT("Setup already performed!"));
+		UE_LOG(LogLearning, Error, TEXT("%s: Setup already performed!"), *GetName());
 		return;
 	}
 
-	if (!InAgentType)
+	if (!InAgentManager)
 	{
-		UE_LOG(LogLearning, Error, TEXT("SetupRecorder called but AgentType is nullptr."));
+		UE_LOG(LogLearning, Error, TEXT("%s: InAgentManager is nullptr."), *GetName());
 		return;
 	}
 
-	if (!InAgentType->IsSetupPerformed())
+	if (!InAgentManager->IsManagerSetup())
 	{
-		UE_LOG(LogLearning, Error, TEXT("AgentType Setup must be run before recorder can be setup."));
+		UE_LOG(LogLearning, Error, TEXT("%s's SetupManager() must be run before %s can be setup."), *InAgentManager->GetName(), *GetName());
 		return;
 	}
 
-	AgentType = InAgentType;
+	AgentManager = InAgentManager;
+
+	if (!InAgentType)
+	{
+		UE_LOG(LogLearning, Error, TEXT("%s: InAgentType is nullptr."), *GetName());
+		return;
+}
+
+	if (!InAgentType->IsSetup())
+	{
+		UE_LOG(LogLearning, Error, TEXT("%s: %s's Setup must be run before it can be used."), *GetName(), *InAgentType->GetName());
+		return;
+	}
 
-	// Create Data Storage
+	AgentType = InAgentType;
 
 	DataStorage = NewObject<ULearningAgentsDataStorage>(this, TEXT("DataStorage"));
 #if WITH_EDITOR
@@ -53,81 +67,33 @@ void ULearningAgentsRecorder::SetupRecorder(ULearningAgentsType* InAgentType)
 #else
 	UE_LOG(LogLearning, Error, TEXT("ULearningAgentsRecorder: DataDirectory was not set. This is non-editor build so need a directory setting."));
 #endif
-}
 
-bool ULearningAgentsRecorder::IsRecorderSetupPerformed() const
-{
-	return AgentType ? true : false;
+	bIsSetup = true;
 }
 
-void ULearningAgentsRecorder::AddAgent(const int32 AgentId)
+bool ULearningAgentsRecorder::AddAgent(const int32 AgentId)
 {
-	if (!IsRecorderSetupPerformed())
-	{
-		UE_LOG(LogLearning, Error, TEXT("Recorder setup must be run before agents can be added!"));
-		return;
-	}
+	bool bSuccess = Super::AddAgent(AgentId);
 
-	if (!AgentType->GetOccupiedAgentSet().Contains(AgentId))
-	{
-		UE_LOG(LogLearning, Error, TEXT("Unable to add: AgentId %d not found on AgentType. Make sure to add agents to the agent type before adding."), AgentId);
-		return;
-	}
-
-	if (SelectedAgentIds.Contains(AgentId))
-	{
-		UE_LOG(LogLearning, Error, TEXT("AgentId %i is already included in agents set"), AgentId);
-		return;
-	}
-
-	SelectedAgentIds.Add(AgentId);
-	SelectedAgentsSet = SelectedAgentIds;
-	SelectedAgentsSet.TryMakeSlice();
-
-	if (IsRecording() && !CurrentRecords.Contains(AgentId))
+	if (bSuccess && IsRecording() && !CurrentRecords.Contains(AgentId))
 	{
 		CurrentRecords.Add(AgentId, DataStorage->CreateRecord(FName(AgentType->GetName() + "_id" + FString::FromInt(AgentId)), AgentType));
 	}
+
+	return bSuccess;
 }
 
-void ULearningAgentsRecorder::RemoveAgent(const int32 AgentId)
+bool ULearningAgentsRecorder::RemoveAgent(const int32 AgentId)
 {
-	if (!IsRecorderSetupPerformed())
-	{
-		UE_LOG(LogLearning, Error, TEXT("Recorder setup must be run before agents can be removed!"));
-		return;
-	}
+	bool bSuccess = Super::RemoveAgent(AgentId);
 
-	if (SelectedAgentIds.RemoveSingleSwap(AgentId, false) == 0)
-	{
-		UE_LOG(LogLearning, Error, TEXT("Unable to remove: AgentId %d not found in the added agents set."), AgentId);
-		return;
-	}
-
-	SelectedAgentsSet = SelectedAgentIds;
-	SelectedAgentsSet.TryMakeSlice();
-
-	if (IsRecording() && CurrentRecords.Contains(AgentId))
+	if (bSuccess && IsRecording() && CurrentRecords.Contains(AgentId))
 	{
 		CurrentRecords[AgentId]->Trim();
 		CurrentRecords.Remove(AgentId);
 	}
-}
 
-bool ULearningAgentsRecorder::HasAgent(const int32 AgentId) const
-{
-	return SelectedAgentsSet.Contains(AgentId);
-}
-
-ULearningAgentsType* ULearningAgentsRecorder::GetAgentType(TSubclassOf<ULearningAgentsType> AgentClass)
-{
-	if (!IsRecorderSetupPerformed())
-	{
-		UE_LOG(LogLearning, Error, TEXT("Recorder setup must be run before getting the agent type!"));
-		return nullptr;
-	}
-
-	return AgentType;
+	return bSuccess;
 }
 
 void ULearningAgentsRecorder::AddExperience()
@@ -138,9 +104,7 @@ void ULearningAgentsRecorder::AddExperience()
 		return;
 	}
 
-	// Add Experience to Records
-
-	for (const int32 AgentId : SelectedAgentsSet)
+	for (const int32 AgentId : AddedAgentSet)
 	{
 		CurrentRecords[AgentId]->AddExperience(
 			AgentType->GetObservationFeature().FeatureBuffer()[AgentId], 
@@ -161,32 +125,24 @@ void ULearningAgentsRecorder::EndRecording()
 		return;
 	}
 
-	// Save Records
-
 	if (bSaveDataOnEndPlay)
 	{
 		DataStorage->SaveAllRecords(DataDirectory);
 	}
 
-	// Trim Records
-
-	for (const int32 AgentId : SelectedAgentsSet)
+	for (const int32 AgentId : AddedAgentSet)
 	{
 		CurrentRecords[AgentId]->Trim();
 	}
 
-	// Reset Current Records
-
 	CurrentRecords.Empty();
 
-	// Done
-
 	bIsRecording = false;
 }
 
 void ULearningAgentsRecorder::BeginRecording()
 {
-	if (IsRecorderSetupPerformed())
+	if (!IsSetup())
 	{
 		UE_LOG(LogLearning, Error, TEXT("Setup must be run before recording can begin."));
 		return;
@@ -198,14 +154,10 @@ void ULearningAgentsRecorder::BeginRecording()
 		return;
 	}
 
-	// Create Records
-
-	for (const int32 AgentId : SelectedAgentsSet)
+	for (const int32 AgentId : AddedAgentSet)
 	{
 		CurrentRecords.Add(AgentId, DataStorage->CreateRecord(FName(AgentType->GetName() + "_id" + FString::FromInt(AgentId)), AgentType));
 	}
 
-	// Done
-
 	bIsRecording = true;
 }
diff --git a/Engine/Plugins/Experimental/LearningAgents/Source/LearningAgentsTraining/Private/LearningAgentsRewards.cpp b/Engine/Plugins/Experimental/LearningAgents/Source/LearningAgentsTraining/Private/LearningAgentsRewards.cpp
--- a/Engine/Plugins/Experimental/LearningAgents/Source/LearningAgentsTraining/Private/LearningAgentsRewards.cpp
+++ b/Engine/Plugins/Experimental/LearningAgents/Source/LearningAgentsTraining/Private/LearningAgentsRewards.cpp
@@ -3,7 +3,7 @@
 #include "LearningAgentsRewards.h"
 
 #include "LearningAgentsTrainer.h"
-#include "LearningAgentsType.h"
+#include "LearningAgentsManager.h"
 #include "LearningArray.h"
 #include "LearningArrayMap.h"
 #include "LearningRewardObject.h"
@@ -40,8 +40,8 @@ namespace UE::Learning::Agents::Private
 
 		Reward->RewardObject = MakeShared<RewardFObject>(
 			Reward->GetFName(),
-			AgentTrainer->GetAgentType()->GetInstanceData().ToSharedRef(),
-			AgentTrainer->GetAgentType()->GetMaxInstanceNum(),
+			AgentTrainer->GetAgentManager()->GetInstanceData().ToSharedRef(),
+			AgentTrainer->GetAgentManager()->GetMaxInstanceNum(),
 			Forward<InArgTypes>(Args)...);
 
 		AgentTrainer->AddReward(Reward, Reward->RewardObject.ToSharedRef());
diff --git a/Engine/Plugins/Experimental/LearningAgents/Source/LearningAgentsTraining/Private/LearningAgentsTrainer.cpp b/Engine/Plugins/Experimental/LearningAgents/Source/LearningAgentsTraining/Private/LearningAgentsTrainer.cpp
--- a/Engine/Plugins/Experimental/LearningAgents/Source/LearningAgentsTraining/Private/LearningAgentsTrainer.cpp
+++ b/Engine/Plugins/Experimental/LearningAgents/Source/LearningAgentsTraining/Private/LearningAgentsTrainer.cpp
@@ -5,6 +5,7 @@
 #include "GenericPlatform/GenericPlatformMisc.h"
 #include "HAL/FileManager.h"
 #include "LearningAgentsType.h"
+#include "LearningAgentsManager.h"
 #include "LearningAgentsRewards.h"
 #include "LearningAgentsCompletions.h"
 #include "LearningAgentsObservations.h"
@@ -28,33 +29,46 @@
 #include "Engine/GameViewportClient.h"
 #include "EngineDefines.h"
 
-ULearningAgentsTrainer::ULearningAgentsTrainer() : UActorComponent() {}
+ULearningAgentsTrainer::ULearningAgentsTrainer() : ULearningAgentsManagerComponent() {}
 ULearningAgentsTrainer::ULearningAgentsTrainer(FVTableHelper& Helper) : ULearningAgentsTrainer() {}
 ULearningAgentsTrainer::~ULearningAgentsTrainer() {}
 
 void ULearningAgentsTrainer::SetupTrainer(
+	ALearningAgentsManager* InAgentManager,
 	ULearningAgentsType* InAgentType,
 	ULearningAgentsPolicy* InPolicy,
 	ULearningAgentsCritic* InCritic,
 	const FLearningAgentsTrainerSettings& Settings)
 {
-	if (IsTrainerSetupPerformed())
+	if (IsSetup())
 	{
-		UE_LOG(LogLearning, Error, TEXT("Setup already performed!"));
+		UE_LOG(LogLearning, Error, TEXT("%s: Setup already performed!"), *GetName());
 		return;
 	}
 
-	// Setup Agent Type, Policy, and Critic
+	if (!InAgentManager)
+	{
+		UE_LOG(LogLearning, Error, TEXT("%s: InAgentManager is nullptr."), *GetName());
+		return;
+	}
+
+	if (!InAgentManager->IsManagerSetup())
+	{
+		UE_LOG(LogLearning, Error, TEXT("%s's SetupManager() must be run before %s can be setup."), *InAgentManager->GetName(), *GetName());
+		return;
+	}
+
+	AgentManager = InAgentManager;
 
 	if (!InAgentType)
 	{
-		UE_LOG(LogLearning, Error, TEXT("SetupTrainer called with nullptr for AgentType."));
+		UE_LOG(LogLearning, Error, TEXT("%s: InAgentType is nullptr."), *GetName());
 		return;
 	}
 
-	if (!InAgentType->IsSetupPerformed())
+	if (!InAgentType->IsSetup())
 	{
-		UE_LOG(LogLearning, Error, TEXT("AgentType Setup not performed."));
+		UE_LOG(LogLearning, Error, TEXT("%s: %s's Setup must be run before it can be used."), *GetName(), *InAgentType->GetName());
 		return;
 	}
 
@@ -62,23 +76,24 @@ void ULearningAgentsTrainer::SetupTrainer(
 
 	if (!InPolicy)
 	{
-		UE_LOG(LogLearning, Error, TEXT("SetupTrainer called with nullptr for Policy."));
+		UE_LOG(LogLearning, Error, TEXT("%s: InPolicy is nullptr."), *GetName());
 		return;
 	}
 
-	if (!InPolicy->IsPolicySetupPerformed())
+	if (!InPolicy->IsSetup())
 	{
-		UE_LOG(LogLearning, Error, TEXT("Policy Setup not Performed"));
+		UE_LOG(LogLearning, Error, TEXT("%s: %s's Setup must be run before it can be used."), *GetName(), *InPolicy->GetName());
 		return;
 	}
 
 	Policy = InPolicy;
 
+	// The critic is optional unlike the other components
 	if (InCritic)
 	{
-		if (!InCritic->IsCriticSetupPerformed())
+		if (!InCritic->IsSetup())
 		{
-			UE_LOG(LogLearning, Error, TEXT("Critic Setup not Performed"));
+			UE_LOG(LogLearning, Error, TEXT("%s: %s's Setup must be run before it can be used."), *GetName(), *InCritic->GetName());
 			return;
 		}
 
@@ -86,30 +101,27 @@ void ULearningAgentsTrainer::SetupTrainer(
 	}
 
 	// Setup Rewards
-
 	RewardObjects.Empty();
 	RewardFeatures.Empty();
 	SetupRewards(this);
 	Rewards = MakeShared<UE::Learning::FSumReward>(TEXT("Rewards"),
 		TLearningArrayView<1, const TSharedRef<UE::Learning::FRewardObject>>(RewardFeatures),
-		AgentType->GetInstanceData().ToSharedRef(),
-		AgentType->GetMaxInstanceNum());
+		AgentManager->GetInstanceData().ToSharedRef(),
+		AgentManager->GetMaxInstanceNum());
 
 	// Setup Completions
-
 	CompletionObjects.Empty();
 	CompletionFeatures.Empty();
 	SetupCompletions(this);
 	Completions = MakeShared<UE::Learning::FAnyCompletion>(TEXT("Completions"),
 		TLearningArrayView<1, const TSharedRef<UE::Learning::FCompletionObject>>(CompletionFeatures),
-		AgentType->GetInstanceData().ToSharedRef(),
-		AgentType->GetMaxInstanceNum());
+		AgentManager->GetInstanceData().ToSharedRef(),
+		AgentManager->GetMaxInstanceNum());
 
 	// Create Episode Buffer
-
 	EpisodeBuffer = MakeUnique<UE::Learning::FEpisodeBuffer>();
 	EpisodeBuffer->Resize(
-		AgentType->GetMaxInstanceNum(),
+		AgentManager->GetMaxInstanceNum(),
 		Settings.MaxStepNum,
 		AgentType->GetObservationFeature().DimNum(),
 		AgentType->GetActionFeature().DimNum());
@@ -117,7 +129,6 @@ void ULearningAgentsTrainer::SetupTrainer(
 	MaxStepsCompletion = Settings.MaxStepsCompletion;
 
 	// Create Replay Buffer
-
 	ReplayBuffer = MakeUnique<UE::Learning::FReplayBuffer>();
 	ReplayBuffer->Resize(
 		AgentType->GetObservationFeature().DimNum(),
@@ -126,113 +137,28 @@ void ULearningAgentsTrainer::SetupTrainer(
 		Settings.MaximumRecordedStepsPerIteration);
 
 	// Create Reset Buffer
-
 	ResetBuffer = MakeUnique<UE::Learning::FResetInstanceBuffer>();
-	ResetBuffer->Resize(AgentType->GetMaxInstanceNum());
+	ResetBuffer->Resize(AgentManager->GetMaxInstanceNum());
 
 	// Record Timeout Setting
-
 	TrainerTimeout = Settings.TrainerCommunicationTimeout;
 
-	// Done!
-
-	bTrainerSetupPerformed = true;
+	bIsSetup = true;
 }
 
-bool ULearningAgentsTrainer::IsTrainerSetupPerformed() const
+bool ULearningAgentsTrainer::AddAgent(const int32 AgentId)
 {
-	return bTrainerSetupPerformed;
-}
+	bool bSuccess = Super::AddAgent(AgentId);
 
-void ULearningAgentsTrainer::AddAgent(const int32 AgentId)
-{
-	if (!IsTrainerSetupPerformed())
-	{
-		UE_LOG(LogLearning, Error, TEXT("Trainer setup must be run before agents can be added!"));
-		return;
-	}
-
-	if (!AgentType->GetOccupiedAgentSet().Contains(AgentId))
-	{
-		UE_LOG(LogLearning, Error, TEXT("Unable to add: AgentId %d not found on AgentType. Make sure to add agents to the agent type before adding."), AgentId);
-		return;
-	}
-
-	if (SelectedAgentIds.Contains(AgentId))
-	{
-		UE_LOG(LogLearning, Error, TEXT("AgentId %i is already included in agents set"), AgentId);
-		return;
-	}
-
-	SelectedAgentIds.Add(AgentId);
-	SelectedAgentsSet = SelectedAgentIds;
-	SelectedAgentsSet.TryMakeSlice();
-
-	if (IsTraining())
+	if (bSuccess && IsTraining())
 	{
 		// Reset the instance and the buffer in case we have stale 
 		// data from another agent that was using this id previously
 		ResetInstance({ AgentId });
 		EpisodeBuffer->Reset(AgentId);
 	}
-}
-
-void ULearningAgentsTrainer::RemoveAgent(const int32 AgentId)
-{
-	if (!IsTrainerSetupPerformed())
-	{
-		UE_LOG(LogLearning, Error, TEXT("Trainer setup must be run before agents can be removed!"));
-		return;
-	}
 
-	if (SelectedAgentIds.RemoveSingleSwap(AgentId, false) == 0)
-	{
-		UE_LOG(LogLearning, Error, TEXT("Unable to remove: AgentId %d not found in the added agents set."), AgentId);
-		return;
-	}
-
-	SelectedAgentsSet = SelectedAgentIds;
-	SelectedAgentsSet.TryMakeSlice();
-}
-
-bool ULearningAgentsTrainer::HasAgent(const int32 AgentId) const
-{
-	return SelectedAgentsSet.Contains(AgentId);
-}
-
-ULearningAgentsType* ULearningAgentsTrainer::GetAgentType(TSubclassOf<ULearningAgentsType> AgentClass)
-{
-	if (!IsTrainerSetupPerformed())
-	{
-		UE_LOG(LogLearning, Error, TEXT("Trainer setup must be run before getting the agent type!"));
-		return nullptr;
-	}
-
-	return AgentType;
-}
-
-const UObject* ULearningAgentsTrainer::GetAgent(const int32 AgentId) const
-{
-	UE_LEARNING_CHECK(AgentType);
-	return AgentType->GetAgent(AgentId);
-}
-
-UObject* ULearningAgentsTrainer::GetAgent(const int32 AgentId)
-{
-	UE_LEARNING_CHECK(AgentType);
-	return AgentType->GetAgent(AgentId);
-}
-
-const ULearningAgentsType* ULearningAgentsTrainer::GetAgentType() const
-{
-	UE_LEARNING_CHECK(AgentType);
-	return AgentType;
-}
-
-ULearningAgentsType* ULearningAgentsTrainer::GetAgentType()
-{
-	UE_LEARNING_CHECK(AgentType);
-	return AgentType;
+	return bSuccess;
 }
 
 void ULearningAgentsTrainer::EndPlay(const EEndPlayReason::Type EndPlayReason)
@@ -257,7 +183,7 @@ void ULearningAgentsTrainer::SetRewards_Implementation(const TArray<int32>& Agen
 
 void ULearningAgentsTrainer::AddReward(TObjectPtr<ULearningAgentsReward> Object, const TSharedRef<UE::Learning::FRewardObject>& Reward)
 {
-	UE_LEARNING_CHECK(!IsTrainerSetupPerformed());
+	UE_LEARNING_CHECK(!IsSetup());
 	RewardObjects.Add(Object);
 	RewardFeatures.Add(Reward);
 }
@@ -274,7 +200,7 @@ void ULearningAgentsTrainer::SetCompletions_Implementation(const TArray<int32>&
 
 void ULearningAgentsTrainer::AddCompletion(TObjectPtr<ULearningAgentsCompletion> Object, const TSharedRef<UE::Learning::FCompletionObject>& Completion)
 {
-	UE_LEARNING_CHECK(!IsTrainerSetupPerformed());
+	UE_LEARNING_CHECK(!IsSetup());
 	CompletionObjects.Add(Object);
 	CompletionFeatures.Add(Completion);
 }
@@ -291,7 +217,7 @@ void ULearningAgentsTrainer::BeginTraining(
 	const bool bReinitializePolicyNetwork,
 	const bool bReinitializeCriticNetwork)
 {
-	if (!IsTrainerSetupPerformed())
+	if (!IsSetup())
 	{
 		UE_LOG(LogLearning, Error, TEXT("Setup must be run before training can start."));
 		return;
@@ -319,6 +245,7 @@ void ULearningAgentsTrainer::BeginTraining(
 	{
 		ViewModeIndex = ViewportClient->ViewModeIndex;
 	}
+
 	// Apply Training GameState Settings
 
 	FApp::SetUseFixedTimeStep(TrainerGameSettings.bUseFixedTimeStep);
@@ -479,12 +406,10 @@ void ULearningAgentsTrainer::BeginTraining(
 	}
 
 	// Reset Agents, Episode Buffer, and Replay Buffer
-	ResetInstance(SelectedAgentIds);
-
-	EpisodeBuffer->Reset(SelectedAgentsSet);
+	ResetInstance(AddedAgentIds);
+	EpisodeBuffer->Reset(AddedAgentSet);
 	ReplayBuffer->Reset();
 
-	// Done!
 	bIsTraining = true;
 }
 
@@ -520,7 +445,6 @@ void ULearningAgentsTrainer::DoneTraining()
 			ViewportClient->ViewModeIndex = ViewModeIndex;
 		}
 
-		// Done
 		bIsTraining = false;
 	}
 }
@@ -539,22 +463,22 @@ void ULearningAgentsTrainer::EvaluateRewards()
 {
 	UE_LEARNING_TRACE_CPUPROFILER_EVENT_SCOPE(ULearningAgentsTrainer::EvaluateRewards);
 
-	if (!IsTrainerSetupPerformed())
+	if (!IsSetup())
 	{
 		UE_LOG(LogLearning, Error, TEXT("Setup must be run before rewards can be evaluated."));
 		return;
 	}
 
-	SetRewards(SelectedAgentIds);
+	SetRewards(AddedAgentIds);
 
-	Rewards->Evaluate(SelectedAgentsSet);
+	Rewards->Evaluate(AddedAgentSet);
 
 #if ENABLE_VISUAL_LOG
 	for (const ULearningAgentsReward* RewardObject : RewardObjects)
 	{
 		if (RewardObject)
 		{
-			RewardObject->VisualLog(SelectedAgentsSet);
+			RewardObject->VisualLog(AddedAgentSet);
 		}
 	}
 #endif
@@ -564,22 +488,22 @@ void ULearningAgentsTrainer::EvaluateCompletions()
 {
 	UE_LEARNING_TRACE_CPUPROFILER_EVENT_SCOPE(ULearningAgentsTrainer::EvaluateCompletions);
 
-	if (!IsTrainerSetupPerformed())
+	if (!IsSetup())
 	{
 		UE_LOG(LogLearning, Error, TEXT("Setup must be run before completions can be evaluated."));
 		return;
 	}
 
-	SetCompletions(SelectedAgentIds);
+	SetCompletions(AddedAgentIds);
 
-	Completions->Evaluate(SelectedAgentsSet);
+	Completions->Evaluate(AddedAgentSet);
 	
 #if ENABLE_VISUAL_LOG
 	for (const ULearningAgentsCompletion* CompletionObject : CompletionObjects)
 	{
 		if (CompletionObject)
 		{
-			CompletionObject->VisualLog(SelectedAgentsSet);
+			CompletionObject->VisualLog(AddedAgentSet);
 		}
 	}
 #endif
@@ -587,7 +511,7 @@ void ULearningAgentsTrainer::EvaluateCompletions()
 
 void ULearningAgentsTrainer::IterateTraining()
 {
-	if (!IsTrainerSetupPerformed())
+	if (!IsSetup())
 	{
 		UE_LOG(LogLearning, Error, TEXT("Setup must be run before training can be performed."));
 		return;
@@ -603,15 +527,13 @@ void ULearningAgentsTrainer::IterateTraining()
 	UE::Learning::FFeatureObject& Actions = AgentType->GetActionFeature();
 
 	// Add Experience to Episode Buffer
-
 	EpisodeBuffer->Push(
 		Observations.FeatureBuffer(),
 		Actions.FeatureBuffer(),
 		Rewards->RewardBuffer(),
-		SelectedAgentsSet);
+		AddedAgentSet);
 
 	// Check for completion based on reaching the maximum episode length
-
 	UE::Learning::Completion::EvaluateEndOfEpisodeCompletions(
 		Completions->CompletionBuffer(),
 		EpisodeBuffer->GetEpisodeStepNums(),
@@ -619,18 +541,15 @@ void ULearningAgentsTrainer::IterateTraining()
 		MaxStepsCompletion == ELearningAgentsCompletion::Truncation 
 			? UE::Learning::ECompletionMode::Truncated 
 			: UE::Learning::ECompletionMode::Terminated,
-		SelectedAgentsSet);
+		AddedAgentSet);
 
 	// Find the set of Instances that need to be reset
-
-	ResetBuffer->SetResetInstancesFromCompletions(Completions->CompletionBuffer(), SelectedAgentsSet);
+	ResetBuffer->SetResetInstancesFromCompletions(Completions->CompletionBuffer(), AddedAgentSet);
 
 	if (ResetBuffer->GetResetInstanceNum() > 0)
 	{
 		// Encode Observations for completed Instances
-
 		AgentType->SetObservations(ResetBuffer->GetResetInstances().ToArray());
-
 		Observations.Encode(ResetBuffer->GetResetInstances());
 		
 #if ENABLE_VISUAL_LOG
@@ -643,8 +562,6 @@ void ULearningAgentsTrainer::IterateTraining()
 		}
 #endif
 
-		// Push Experience to Replay Buffer
-
 		const bool bReplayBufferFull = ReplayBuffer->AddEpisodes(
 			Completions->CompletionBuffer(),
 			Observations.FeatureBuffer(),
@@ -653,8 +570,6 @@ void ULearningAgentsTrainer::IterateTraining()
 
 		if (bReplayBufferFull)
 		{
-			// Send Experience
-
 			UE::Learning::ETrainerResponse Response = Trainer->SendExperience(*ReplayBuffer, TrainerTimeout);
 
 			if (Response != UE::Learning::ETrainerResponse::Success)
@@ -664,12 +579,9 @@ void ULearningAgentsTrainer::IterateTraining()
 				return;
 			}
 
-			// Reset Replay Buffer
-
 			ReplayBuffer->Reset();
 
 			// Get Updated Policy
-
 			Response = Trainer->RecvPolicy(Policy->GetPolicyNetwork(), TrainerTimeout);
 
 			if (Response == UE::Learning::ETrainerResponse::Completed)
@@ -686,7 +598,6 @@ void ULearningAgentsTrainer::IterateTraining()
 			}
 
 			// Get Updated Critic
-
 			if (Critic)
 			{
 				Response = Trainer->RecvCritic(Critic->GetCriticNetwork(), TrainerTimeout);
@@ -700,10 +611,9 @@ void ULearningAgentsTrainer::IterateTraining()
 			}
 
 			// Mark all instances for reset since we have a new policy
-			ResetBuffer->SetResetInstances(SelectedAgentsSet);
+			ResetBuffer->SetResetInstances(AddedAgentSet);
 		}
 
-		// Reset Instances
 		ResetInstance(ResetBuffer->GetResetInstances().ToArray());
 
 		EpisodeBuffer->Reset(ResetBuffer->GetResetInstances());
@@ -714,15 +624,15 @@ void ULearningAgentsTrainer::ResetAllInstances()
 {
 	UE_LEARNING_TRACE_CPUPROFILER_EVENT_SCOPE(ULearningAgentsTrainer::ResetAllInstances);
 
-	if (!IsTrainerSetupPerformed())
+	if (!IsSetup())
 	{
 		UE_LOG(LogLearning, Error, TEXT("Setup must be run before agents can be reset."));
 		return;
 	}
 
-	ResetInstance(SelectedAgentIds);
+	ResetInstance(AddedAgentIds);
 
-	EpisodeBuffer->Reset(SelectedAgentsSet);
+	EpisodeBuffer->Reset(AddedAgentSet);
 }
 
 void ULearningAgentsTrainer::ResetInstance_Implementation(const TArray<int32>& AgentId)
diff --git a/Engine/Plugins/Experimental/LearningAgents/Source/LearningAgentsTraining/Public/LearningAgentsDataStorage.h b/Engine/Plugins/Experimental/LearningAgents/Source/LearningAgentsTraining/Public/LearningAgentsDataStorage.h
--- a/Engine/Plugins/Experimental/LearningAgents/Source/LearningAgentsTraining/Public/LearningAgentsDataStorage.h
+++ b/Engine/Plugins/Experimental/LearningAgents/Source/LearningAgentsTraining/Public/LearningAgentsDataStorage.h
@@ -131,7 +131,7 @@ public:
 private:
 	
 	/** All currently loaded records. */
-	UPROPERTY(VisibleAnywhere, Category = "LearningAgents")
+	UPROPERTY(VisibleAnywhere, Transient, Category = "LearningAgents")
 	TArray<TObjectPtr<ULearningAgentsRecord>> Records;
 
 	/** If true, prepends a timestamp to the saved file names. */
diff --git a/Engine/Plugins/Experimental/LearningAgents/Source/LearningAgentsTraining/Public/LearningAgentsImitationTrainer.h b/Engine/Plugins/Experimental/LearningAgents/Source/LearningAgentsTraining/Public/LearningAgentsImitationTrainer.h
--- a/Engine/Plugins/Experimental/LearningAgents/Source/LearningAgentsTraining/Public/LearningAgentsImitationTrainer.h
+++ b/Engine/Plugins/Experimental/LearningAgents/Source/LearningAgentsTraining/Public/LearningAgentsImitationTrainer.h
@@ -121,15 +121,15 @@ public:
 private:
 
 	/** The policy being trained. */
-	UPROPERTY(VisibleAnywhere, Category = "LearningAgents")
+	UPROPERTY(VisibleAnywhere, Transient, Category = "LearningAgents")
 	TObjectPtr<ULearningAgentsPolicy> Policy;
 
 	/** True if training is currently in-progress. Otherwise, false. */
-	UPROPERTY(VisibleAnywhere, Category = "LearningAgents")
+	UPROPERTY(VisibleAnywhere, Transient, Category = "LearningAgents")
 	bool bIsTraining = false;
 
 	/** True if training is completed. Otherwise, false. */
-	UPROPERTY(VisibleAnywhere, Category = "LearningAgents")
+	UPROPERTY(VisibleAnywhere, Transient, Category = "LearningAgents")
 	bool bIsTrainingComplete = false;
 
 	TLearningArray<2, float> RecordedObservations;
diff --git a/Engine/Plugins/Experimental/LearningAgents/Source/LearningAgentsTraining/Public/LearningAgentsRecorder.h b/Engine/Plugins/Experimental/LearningAgents/Source/LearningAgentsTraining/Public/LearningAgentsRecorder.h
--- a/Engine/Plugins/Experimental/LearningAgents/Source/LearningAgentsTraining/Public/LearningAgentsRecorder.h
+++ b/Engine/Plugins/Experimental/LearningAgents/Source/LearningAgentsTraining/Public/LearningAgentsRecorder.h
@@ -2,9 +2,9 @@
 
 #pragma once
 
-#include "LearningArray.h"
+#include "LearningAgentsManagerComponent.h"
 
-#include "Components/ActorComponent.h"
+#include "LearningArray.h"
 #include "Containers/Map.h"
 #include "UObject/ObjectPtr.h"
 
@@ -15,7 +15,7 @@ class ULearningAgentsRecord;
 
 /** A component that can be used to create recordings of training data for imitation learning. */
 UCLASS(BlueprintType, Blueprintable)
-class LEARNINGAGENTSTRAINING_API ULearningAgentsRecorder : public UActorComponent
+class LEARNINGAGENTSTRAINING_API ULearningAgentsRecorder : public ULearningAgentsManagerComponent
 {
 	GENERATED_BODY()
 
@@ -35,41 +35,15 @@ public:
 	* @param InAgentType The agent type we are recording with.
 	*/
 	UFUNCTION(BlueprintCallable, Category = "LearningAgents")
-	void SetupRecorder(ULearningAgentsType* InAgentType);
+	void SetupRecorder(ALearningAgentsManager* InAgentManager, ULearningAgentsType* InAgentType);
 
-	/** Returns true if SetupRecorder has been run successfully; Otherwise, false. */
-	UFUNCTION(BlueprintPure, Category = "LearningAgents")
-	bool IsRecorderSetupPerformed() const;
-
-// ----- Agent Management -----
 public:
 
-	/**
-	 * Adds an agent to this recorder.
-	 * @param AgentId The id of the agent to be added.
-	 * @warning The agent id must exist for the agent type.
-	 */
-	UFUNCTION(BlueprintCallable, Category = "LearningAgents")
-	void AddAgent(const int32 AgentId);
-
-	/**
-	* Removes an agent from this recorder.
-	* @param AgentId The id of the agent to be removed.
-	* @warning The agent id must exist for this recorder already.
-	*/
-	UFUNCTION(BlueprintCallable, Category = "LearningAgents")
-	void RemoveAgent(const int32 AgentId);
-
-	/** Returns true if the given id has been previously added to this recorder; Otherwise, false. */
-	UFUNCTION(BlueprintPure, Category = "LearningAgents")
-	bool HasAgent(const int32 AgentId) const;
+	//~ Begin ULearningAgentsManagerComponent Interface
+	virtual bool AddAgent(const int32 AgentId) override;
 
-	/**
-	* Gets the agent type this recorder is associated with.
-	* @param AgentClass The class to cast the agent type to (in blueprint).
-	*/
-	UFUNCTION(BlueprintPure, Category = "LearningAgents", meta = (DeterminesOutputType = "AgentClass"))
-	ULearningAgentsType* GetAgentType(const TSubclassOf<ULearningAgentsType> AgentClass);
+	virtual bool RemoveAgent(const int32 AgentId) override;
+	//~ End ULearningAgentsManagerComponent Interface
 
 // ----- Recording Process -----
 public:
@@ -98,13 +72,9 @@ public:
 private:
 
 	/** The agent type this recorder is associated with. */
-	UPROPERTY(VisibleAnywhere, Category = "LearningAgents")
+	UPROPERTY(VisibleAnywhere, Transient, Category = "LearningAgents")
 	TObjectPtr<ULearningAgentsType> AgentType;
 
-	/** The agent ids this recorder is managing. */
-	UPROPERTY(VisibleAnywhere, Category = "LearningAgents")
-	TArray<int32> SelectedAgentIds;
-
 // ----- Recorder Configuration -----
 private:
 
@@ -120,18 +90,14 @@ private:
 private:
 
 	/** True if recording is currently in-progress. Otherwise, false. */
-	UPROPERTY(VisibleAnywhere, Category = "LearningAgents")
+	UPROPERTY(VisibleAnywhere, Transient, Category = "LearningAgents")
 	bool bIsRecording = false;
 
 	/** The data storage manager. It can be used to save/load agent records. */
-	UPROPERTY(VisibleAnywhere, Category = "LearningAgents")
+	UPROPERTY(VisibleAnywhere, Transient, Category = "LearningAgents")
 	TObjectPtr<ULearningAgentsDataStorage> DataStorage;
 
 	/** All records which are currently being written to. */
-	UPROPERTY(VisibleAnywhere, Category = "LearningAgents")
+	UPROPERTY(VisibleAnywhere, Transient, Category = "LearningAgents")
 	TMap<int32, TObjectPtr<ULearningAgentsRecord>> CurrentRecords;
-
-private:
-
-	UE::Learning::FIndexSet SelectedAgentsSet;
 };
diff --git a/Engine/Plugins/Experimental/LearningAgents/Source/LearningAgentsTraining/Public/LearningAgentsTrainer.h b/Engine/Plugins/Experimental/LearningAgents/Source/LearningAgentsTraining/Public/LearningAgentsTrainer.h
--- a/Engine/Plugins/Experimental/LearningAgents/Source/LearningAgentsTraining/Public/LearningAgentsTrainer.h
+++ b/Engine/Plugins/Experimental/LearningAgents/Source/LearningAgentsTraining/Public/LearningAgentsTrainer.h
@@ -2,7 +2,7 @@
 
 #pragma once
 
-#include "Components/ActorComponent.h"
+#include "LearningAgentsManagerComponent.h"
 #include "LearningAgentsCritic.h" // Included for FLearningAgentsCriticSettings()
 
 #include "Templates/SharedPointer.h"
@@ -24,6 +24,7 @@ namespace UE::Learning
 	struct FCompletionObject;
 }
 
+class ALearningAgentsManager;
 class ULearningAgentsType;
 class ULearningAgentsCompletion;
 class ULearningAgentsReward;
@@ -186,7 +187,7 @@ public:
 * @see ULearningAgentsType to understand how observations and actions work.
 */
 UCLASS(Abstract, BlueprintType, Blueprintable)
-class LEARNINGAGENTSTRAINING_API ULearningAgentsTrainer : public UActorComponent
+class LEARNINGAGENTSTRAINING_API ULearningAgentsTrainer : public ULearningAgentsManagerComponent
 {
 	GENERATED_BODY()
 
@@ -203,6 +204,7 @@ public:
 
 	/**
 	* Initializes this object and runs the setup functions for rewards and completions.
+	* @param InAgentManager The agent manager we are associated with.
 	* @param InAgentType The agent type we are training with.
 	* @param InPolicy The policy to be trained.
 	* @param InCritic Optional - only needs to be provided if we want the critic to be accessible at runtime.
@@ -210,58 +212,17 @@ public:
 	*/
 	UFUNCTION(BlueprintCallable, Category = "LearningAgents")
 	void SetupTrainer(
+		ALearningAgentsManager* InAgentManager,
 		ULearningAgentsType* InAgentType,
 		ULearningAgentsPolicy* InPolicy,
 		ULearningAgentsCritic* InCritic = nullptr,
 		const FLearningAgentsTrainerSettings& Settings = FLearningAgentsTrainerSettings());
 
-	/** Returns true if SetupTrainer has been run successfully; Otherwise, false. */
-	UFUNCTION(BlueprintPure, Category = "LearningAgents")
-	bool IsTrainerSetupPerformed() const;
-
-// ----- Agent Management -----
-public:
-
-	/**
-	 * Adds an agent to this trainer.
-	 * @param AgentId The id of the agent to be added.
-	 * @warning The agent id must exist for the agent type.
-	 */
-	UFUNCTION(BlueprintCallable, Category = "LearningAgents")
-	void AddAgent(const int32 AgentId);
-
-	/**
-	* Removes an agent from this trainer.
-	* @param AgentId The id of the agent to be removed.
-	* @warning The agent id must exist for this trainer already.
-	*/
-	UFUNCTION(BlueprintCallable, Category = "LearningAgents")
-	void RemoveAgent(const int32 AgentId);
-
-	/** Returns true if the given id has been previously added to this trainer; Otherwise, false. */
-	UFUNCTION(BlueprintPure, Category = "LearningAgents")
-	bool HasAgent(const int32 AgentId) const;
-
-	/**
-	* Gets the agent type this trainer is associated with.
-	* @param AgentClass The class to cast the agent type to (in blueprint).
-	*/
-	UFUNCTION(BlueprintPure, Category = "LearningAgents", meta = (DeterminesOutputType = "AgentClass"))
-	ULearningAgentsType* GetAgentType(const TSubclassOf<ULearningAgentsType> AgentClass);
-
-public:
+public: 
 
-	/** Gets the agent corresponding to the given id. */
-	const UObject* GetAgent(const int32 AgentId) const;
-
-	/** Gets the agent corresponding to the given id. */
-	UObject* GetAgent(const int32 AgentId);
-
-	/** Gets the associated agent type. */
-	const ULearningAgentsType* GetAgentType() const;
-
-	/** Gets the associated agent type. */
-	ULearningAgentsType* GetAgentType();
+	//~ Begin ULearningAgentsManagerComponent Interface
+	virtual bool AddAgent(const int32 AgentId) override;
+	//~ End ULearningAgentsManagerComponent Interface
 
 // ----- Rewards -----
 public:
@@ -385,36 +346,28 @@ public:
 // ----- Private Data ----- 
 private:
 
-	/** The agent type this trainer is associated with. */
-	UPROPERTY(VisibleAnywhere, Category = "LearningAgents")
+	/** The agent type associated with this component. */
+	UPROPERTY(VisibleAnywhere, Transient, Category = "LearningAgents")
 	TObjectPtr<ULearningAgentsType> AgentType;
 
-	/** The agent ids this trainer is managing. */
-	UPROPERTY(VisibleAnywhere, Category = "LearningAgents")
-	TArray<int32> SelectedAgentIds;
-
 	/** The current policy for experience gathering. */
-	UPROPERTY(VisibleAnywhere, Category = "LearningAgents")
+	UPROPERTY(VisibleAnywhere, Transient, Category = "LearningAgents")
 	TObjectPtr<ULearningAgentsPolicy> Policy;
 
 	/** The current critic. */
-	UPROPERTY(VisibleAnywhere, Category = "LearningAgents")
+	UPROPERTY(VisibleAnywhere, Transient, Category = "LearningAgents")
 	TObjectPtr<ULearningAgentsCritic> Critic;
 
-	/** True if this trainer's SetupTrainer has been run; Otherwise, false. */
-	UPROPERTY(VisibleAnywhere, Category = "LearningAgents")
-	bool bTrainerSetupPerformed = false;
-
 	/** True if training is currently in-progress. Otherwise, false. */
-	UPROPERTY(VisibleAnywhere, Category = "LearningAgents")
+	UPROPERTY(VisibleAnywhere, Transient, Category = "LearningAgents")
 	bool bIsTraining = false;
 
 	/** The list of current reward objects. */
-	UPROPERTY(VisibleAnywhere, Category = "LearningAgents")
+	UPROPERTY(VisibleAnywhere, Transient, Category = "LearningAgents")
 	TArray<TObjectPtr<ULearningAgentsReward>> RewardObjects;
 
 	/** The list of current completion objects. */
-	UPROPERTY(VisibleAnywhere, Category = "LearningAgents")
+	UPROPERTY(VisibleAnywhere, Transient, Category = "LearningAgents")
 	TArray<TObjectPtr<ULearningAgentsCompletion>> CompletionObjects;
 
 	TArray<TSharedRef<UE::Learning::FRewardObject>, TInlineAllocator<16>> RewardFeatures;
@@ -434,8 +387,6 @@ private:
 
 	void DoneTraining();
 
-	UE::Learning::FIndexSet SelectedAgentsSet;
-
 // ----- Private Recording of GameSettings ----- 
 private:
 
