LearningAgents: Added some additional observations, actions, rewards, completions and helpers. Cleaned up some of the error messaging to make it more consistent.

diff --git a/Engine/Plugins/Experimental/LearningAgents/Source/Learning/Private/LearningCompletionObject.cpp b/Engine/Plugins/Experimental/LearningAgents/Source/Learning/Private/LearningCompletionObject.cpp
--- a/Engine/Plugins/Experimental/LearningAgents/Source/Learning/Private/LearningCompletionObject.cpp
+++ b/Engine/Plugins/Experimental/LearningAgents/Source/Learning/Private/LearningCompletionObject.cpp
@@ -84,6 +84,34 @@ namespace UE::Learning
 
 	//------------------------------------------------------------------
 
+	FTimeElapsedCompletion::FTimeElapsedCompletion(
+		const FName& InIdentifier,
+		const TSharedRef<FArrayMap>& InInstanceData,
+		const int32 InMaxInstanceNum,
+		const float InThreshold,
+		const ECompletionMode InCompletionMode)
+		: FCompletionObject(InIdentifier, InInstanceData, InMaxInstanceNum, InCompletionMode)
+	{
+		TimeHandle = InstanceData->Add<1, float>({ InIdentifier, TEXT("Time") }, { InMaxInstanceNum }, 0.0f);
+		ThresholdHandle = InstanceData->Add<1, float>({ InIdentifier, TEXT("Threshold") }, { InMaxInstanceNum }, InThreshold);
+	}
+
+	void FTimeElapsedCompletion::Evaluate(const FIndexSet Instances)
+	{
+		UE_LEARNING_TRACE_CPUPROFILER_EVENT_SCOPE(FTimeElapsedCompletion::Evaluate);
+
+		const TLearningArrayView<1, const float> Time = InstanceData->ConstView(TimeHandle);
+		const TLearningArrayView<1, const float> Threshold = InstanceData->ConstView(ThresholdHandle);
+		TLearningArrayView<1, ECompletionMode> Completion = InstanceData->View(CompletionHandle);
+
+		for (const int32 InstanceIdx : Instances)
+		{
+			Completion[InstanceIdx] = Time[InstanceIdx] > Threshold[InstanceIdx] ? CompletionMode : ECompletionMode::Running;
+		}
+	}
+
+	//------------------------------------------------------------------
+
 	FScalarPositionDifferenceCompletion::FScalarPositionDifferenceCompletion(
 		const FName& InIdentifier,
 		const TSharedRef<FArrayMap>& InInstanceData,
@@ -171,6 +199,53 @@ namespace UE::Learning
 		}
 	}
 
+	FPlanarPositionSimilarityCompletion::FPlanarPositionSimilarityCompletion(
+		const FName& InIdentifier,
+		const TSharedRef<FArrayMap>& InInstanceData,
+		const int32 InMaxInstanceNum,
+		const int32 InPositionNum,
+		const float InThreshold,
+		const ECompletionMode InCompletionMode,
+		const FVector InAxis0,
+		const FVector InAxis1)
+		: FCompletionObject(InIdentifier, InInstanceData, InMaxInstanceNum, InCompletionMode)
+		, Axis0(InAxis0)
+		, Axis1(InAxis1)
+	{
+		Position0Handle = InstanceData->Add<2, FVector>({ InIdentifier, TEXT("Position0") }, { InMaxInstanceNum, InPositionNum }, FVector::ZeroVector);
+		Position1Handle = InstanceData->Add<2, FVector>({ InIdentifier, TEXT("Position1") }, { InMaxInstanceNum, InPositionNum }, FVector::ZeroVector);
+		ThresholdHandle = InstanceData->Add<1, float>({ InIdentifier, TEXT("Threshold") }, { InMaxInstanceNum }, InThreshold);
+	}
+
+	void FPlanarPositionSimilarityCompletion::Evaluate(const FIndexSet Instances)
+	{
+		UE_LEARNING_TRACE_CPUPROFILER_EVENT_SCOPE(FPlanarPositionSimilarityCompletion::Evaluate);
+
+		const TLearningArrayView<2, const FVector> Position0 = InstanceData->ConstView(Position0Handle);
+		const TLearningArrayView<2, const FVector> Position1 = InstanceData->ConstView(Position1Handle);
+		const TLearningArrayView<1, const float> Threshold = InstanceData->ConstView(ThresholdHandle);
+		TLearningArrayView<1, ECompletionMode> Completion = InstanceData->View(CompletionHandle);
+
+		const int32 PositionNum = Position0.Num<1>();
+
+		for (const int32 InstanceIdx : Instances)
+		{
+			Completion[InstanceIdx] = ECompletionMode::Running;
+
+			for (int32 PositionIdx = 0; PositionIdx < PositionNum; PositionIdx++)
+			{
+				const FVector ProjectedPosition0 = FVector(Axis0.Dot(Position0[InstanceIdx][PositionIdx]), Axis1.Dot(Position0[InstanceIdx][PositionIdx]), 0.0f);
+				const FVector ProjectedPosition1 = FVector(Axis0.Dot(Position1[InstanceIdx][PositionIdx]), Axis1.Dot(Position1[InstanceIdx][PositionIdx]), 0.0f);
+
+				if (FVector::Distance(ProjectedPosition0, ProjectedPosition1) < Threshold[InstanceIdx])
+				{
+					Completion[InstanceIdx] = CompletionMode;
+					break;
+				}
+			}
+		}
+	}
+
 	//------------------------------------------------------------------
 
 	FScalarVelocityDifferenceCompletion::FScalarVelocityDifferenceCompletion(
diff --git a/Engine/Plugins/Experimental/LearningAgents/Source/Learning/Private/LearningFeatureObject.cpp b/Engine/Plugins/Experimental/LearningAgents/Source/Learning/Private/LearningFeatureObject.cpp
--- a/Engine/Plugins/Experimental/LearningAgents/Source/Learning/Private/LearningFeatureObject.cpp
+++ b/Engine/Plugins/Experimental/LearningAgents/Source/Learning/Private/LearningFeatureObject.cpp
@@ -223,23 +223,6 @@ namespace UE::Learning
 			Output[1] = FMath::Cos(LocalAngle) / FMath::Max(Scale, Epsilon);
 		}
 
-		void PlanarRotation(
-			TLearningArrayView<1, float> Output,
-			const FQuat Rotation,
-			const FQuat RelativeRotation,
-			const float Scale,
-			const FVector Axis,
-			const float Epsilon)
-		{
-			UE_LEARNING_ARRAY_SHAPE_CHECK(Output.Num() == 2);
-			UE_LEARNING_ARRAY_VALUE_CHECK(Scale > 0.0f);
-
-			const FQuat LocalRotation = RelativeRotation.Inverse() * Rotation;
-			const float LocalAngle = LocalRotation.GetTwistAngle(Axis);
-			Output[0] = FMath::Sin(LocalAngle) / FMath::Max(Scale, Epsilon);
-			Output[1] = FMath::Cos(LocalAngle) / FMath::Max(Scale, Epsilon);
-		}
-
 		void Rotation(
 			TLearningArrayView<1, float> Output,
 			const FQuat Rotation,
@@ -288,19 +271,6 @@ namespace UE::Learning
 			Output = AngularVelocity / FMath::Max(Scale, Epsilon);
 		}
 
-		void PlanarAngularVelocity(
-			float& Output,
-			const FVector AngularVelocity,
-			const FQuat RelativeRotation,
-			const float Scale,
-			const FVector Axis,
-			const float Epsilon)
-		{
-			UE_LEARNING_ARRAY_VALUE_CHECK(Scale > 0.0f);
-
-			Output = RelativeRotation.UnrotateVector(AngularVelocity).Dot(Axis) / FMath::Max(Scale, Epsilon);
-		}
-
 		void AngularVelocity(
 			TLearningArrayView<1, float> Output,
 			const FVector AngularVelocity,
@@ -550,19 +520,6 @@ namespace UE::Learning
 			OutAngle = FMath::Atan2(Input[0] * Scale, Input[1] * Scale) + RelativeAngle;
 		}
 
-		void PlanarRotation(
-			FQuat& OutPlanarRotation,
-			const TLearningArrayView<1, const float> Input,
-			const FQuat RelativeRotation,
-			const float Scale,
-			const FVector Axis)
-		{
-			UE_LEARNING_ARRAY_SHAPE_CHECK(Input.Num() == 2);
-
-			const float LocalAngle = FMath::Atan2(Input[0] * Scale, Input[1] * Scale);
-			OutPlanarRotation = RelativeRotation * FQuat::MakeFromRotationVector(LocalAngle * Axis);
-		}
-
 		void Rotation(
 			FQuat& OutRotation,
 			const TLearningArrayView<1, const float> Input,
@@ -605,16 +562,6 @@ namespace UE::Learning
 			OutScalarVelocity = Input * Scale;
 		}
 
-		void PlanarAngularVelocity(
-			FVector& OutPlanarAngularVelocity,
-			const float Input,
-			const FQuat RelativeRotation,
-			const float Scale,
-			const FVector Axis)
-		{
-			OutPlanarAngularVelocity = RelativeRotation.RotateVector(Axis * Input * Scale);
-		}
-
 		void AngularVelocity(
 			FVector& OutAngularVelocity,
 			const TLearningArrayView<1, const float> Input,
@@ -873,7 +820,7 @@ namespace UE::Learning
 
 	void FFloatFeature::Encode(const FIndexSet Instances)
 	{
-		UE_LEARNING_TRACE_CPUPROFILER_EVENT_SCOPE(FScalarFeature::Encode);
+		UE_LEARNING_TRACE_CPUPROFILER_EVENT_SCOPE(FFloatFeature::Encode);
 
 		const TLearningArrayView<2, const float> Value = InstanceData->ConstView(ValueHandle);
 		TLearningArrayView<2, float> Feature = InstanceData->View(FeatureHandle);
@@ -894,7 +841,7 @@ namespace UE::Learning
 
 	void FFloatFeature::Decode(const FIndexSet Instances)
 	{
-		UE_LEARNING_TRACE_CPUPROFILER_EVENT_SCOPE(FScalarFeature::Decode);
+		UE_LEARNING_TRACE_CPUPROFILER_EVENT_SCOPE(FFloatFeature::Decode);
 
 		TLearningArrayView<2, float> Value = InstanceData->View(ValueHandle);
 		const TLearningArrayView<2, const float> Feature = InstanceData->ConstView(FeatureHandle);
@@ -1711,68 +1658,6 @@ namespace UE::Learning
 		}
 	}
 
-	FPlanarRotationFeature::FPlanarRotationFeature(
-		const FName& InIdentifier,
-		const TSharedRef<FArrayMap>& InInstanceData,
-		const int32 InMaxInstanceNum,
-		const int32 InRotationNum,
-		const float InScale,
-		const FVector InAxis)
-		: FFeatureObject(InIdentifier, InInstanceData, InMaxInstanceNum, InRotationNum * 2, InScale)
-		, Axis(InAxis)
-	{
-		RotationHandle = InstanceData->Add<2, FQuat>({ InIdentifier, TEXT("Rotation") }, { InMaxInstanceNum, InRotationNum }, FQuat::Identity);
-		RelativeRotationHandle = InstanceData->Add<1, FQuat>({ InIdentifier, TEXT("RelativeRotation") }, { InMaxInstanceNum }, FQuat::Identity);
-	}
-
-	void FPlanarRotationFeature::Encode(const FIndexSet Instances)
-	{
-		UE_LEARNING_TRACE_CPUPROFILER_EVENT_SCOPE(FPlanarRotationFeature::Encode);
-
-		const TLearningArrayView<2, const FQuat> Rotation = InstanceData->ConstView(RotationHandle);
-		const TLearningArrayView<1, const FQuat> RelativeRotation = InstanceData->ConstView(RelativeRotationHandle);
-		TLearningArrayView<2, float> Feature = InstanceData->View(FeatureHandle);
-
-		const int32 RotationNum = Rotation.Num<1>();
-
-		for (const int32 InstanceIdx : Instances)
-		{
-			for (int32 RotationIdx = 0; RotationIdx < RotationNum; RotationIdx++)
-			{
-				Feature::Encode::PlanarRotation(
-					Feature[InstanceIdx].Slice(RotationIdx * 2, 2),
-					Rotation[InstanceIdx][RotationIdx],
-					RelativeRotation[InstanceIdx],
-					Scale,
-					Axis);
-			}
-		}
-	}
-
-	void FPlanarRotationFeature::Decode(const FIndexSet Instances)
-	{
-		UE_LEARNING_TRACE_CPUPROFILER_EVENT_SCOPE(FPlanarRotationFeature::Decode);
-
-		TLearningArrayView<2, FQuat> Rotation = InstanceData->View(RotationHandle);
-		const TLearningArrayView<1, const FQuat> RelativeRotation = InstanceData->ConstView(RelativeRotationHandle);
-		const TLearningArrayView<2, const float> Feature = InstanceData->ConstView(FeatureHandle);
-
-		const int32 RotationNum = Rotation.Num<1>();
-
-		for (const int32 InstanceIdx : Instances)
-		{
-			for (int32 RotationIdx = 0; RotationIdx < RotationNum; RotationIdx++)
-			{
-				Feature::Decode::PlanarRotation(
-					Rotation[InstanceIdx][RotationIdx],
-					Feature[InstanceIdx].Slice(RotationIdx * 2, 2),
-					RelativeRotation[InstanceIdx],
-					Scale,
-					Axis);
-			}
-		}
-	}
-
 	FRotationFeature::FRotationFeature(
 		const FName& InIdentifier,
 		const TSharedRef<FArrayMap>& InInstanceData,
@@ -1941,68 +1826,6 @@ namespace UE::Learning
 		}
 	}
 
-	FPlanarAngularVelocityFeature::FPlanarAngularVelocityFeature(
-		const FName& InIdentifier,
-		const TSharedRef<FArrayMap>& InInstanceData,
-		const int32 InMaxInstanceNum,
-		const int32 InAngularVelocityNum,
-		const float InScale,
-		const FVector InAxis)
-		: FFeatureObject(InIdentifier, InInstanceData, InMaxInstanceNum, InAngularVelocityNum, InScale)
-		, Axis(InAxis)
-	{
-		AngularVelocityHandle = InstanceData->Add<2, FVector>({ InIdentifier, TEXT("AngularVelocity") }, { InMaxInstanceNum }, FVector::ZeroVector);
-		RelativeRotationHandle = InstanceData->Add<1, FQuat>({ InIdentifier, TEXT("RelativeRotation") }, { InMaxInstanceNum }, FQuat::Identity);
-	}
-
-	void FPlanarAngularVelocityFeature::Encode(const FIndexSet Instances)
-	{
-		UE_LEARNING_TRACE_CPUPROFILER_EVENT_SCOPE(FPlanarAngularVelocityFeature::Encode);
-
-		const TLearningArrayView<2, const FVector> AngularVelocity = InstanceData->ConstView(AngularVelocityHandle);
-		const TLearningArrayView<1, const FQuat> RelativeRotation = InstanceData->ConstView(RelativeRotationHandle);
-		TLearningArrayView<2, float> Feature = InstanceData->View(FeatureHandle);
-
-		const int32 AngularVelocityNum = AngularVelocity.Num<1>();
-
-		for (const int32 InstanceIdx : Instances)
-		{
-			for (int32 AngularVelocityIdx = 0; AngularVelocityIdx < AngularVelocityNum; AngularVelocityIdx++)
-			{
-				Feature::Encode::PlanarAngularVelocity(
-					Feature[InstanceIdx][AngularVelocityIdx],
-					AngularVelocity[InstanceIdx][AngularVelocityIdx],
-					RelativeRotation[InstanceIdx],
-					Scale,
-					Axis);
-			}
-		}
-	}
-
-	void FPlanarAngularVelocityFeature::Decode(const FIndexSet Instances)
-	{
-		UE_LEARNING_TRACE_CPUPROFILER_EVENT_SCOPE(FPlanarAngularVelocityFeature::Decode);
-
-		TLearningArrayView<2, FVector> AngularVelocity = InstanceData->View(AngularVelocityHandle);
-		const TLearningArrayView<1, const FQuat> RelativeRotation = InstanceData->ConstView(RelativeRotationHandle);
-		const TLearningArrayView<2, const float> Feature = InstanceData->ConstView(FeatureHandle);
-
-		const int32 AngularVelocityNum = AngularVelocity.Num<1>();
-
-		for (const int32 InstanceIdx : Instances)
-		{
-			for (int32 AngularVelocityIdx = 0; AngularVelocityIdx < AngularVelocityNum; AngularVelocityIdx++)
-			{
-				Feature::Decode::PlanarAngularVelocity(
-					AngularVelocity[InstanceIdx][AngularVelocityIdx],
-					Feature[InstanceIdx][AngularVelocityIdx],
-					RelativeRotation[InstanceIdx],
-					Scale,
-					Axis);
-			}
-		}
-	}
-
 	FAngularVelocityFeature::FAngularVelocityFeature(
 		const FName& InIdentifier,
 		const TSharedRef<FArrayMap>& InInstanceData,
diff --git a/Engine/Plugins/Experimental/LearningAgents/Source/Learning/Public/LearningArray.h b/Engine/Plugins/Experimental/LearningAgents/Source/Learning/Public/LearningArray.h
--- a/Engine/Plugins/Experimental/LearningAgents/Source/Learning/Public/LearningArray.h
+++ b/Engine/Plugins/Experimental/LearningAgents/Source/Learning/Public/LearningArray.h
@@ -3,6 +3,8 @@
 #pragma once
 
 #include "CoreTypes.h"
+#include "Templates/Function.h"
+#include "Containers/UnrealString.h"
 
 #include "MultiArrayView.h"
 #include "MultiArray.h"
@@ -104,9 +106,6 @@
 #endif
 
 
-template <typename FuncType> class TFunctionRef;
-
-
 // Learning Multi-Dimensional Array Types
 
 template<uint32 DimNum>
@@ -972,6 +971,47 @@ namespace UE::Learning
 			return INDEX_NONE;
 		}
 
+		template<typename InElementType>
+		FString Format(const TLearningArrayView<1, const InElementType> Array, const TFunctionRef<FString(const InElementType&)> Formatter, const int32 MaxItemNum = 16)
+		{
+			const int32 ItemNum = Array.Num();
+
+			FString Output = TEXT("[");
+
+			if (ItemNum <= MaxItemNum)
+			{
+				for (int32 Idx = 0; Idx < ItemNum; Idx++)
+				{
+					Output.Append(Formatter(Array[Idx]));
+					if (Idx != ItemNum - 1) { Output.Append(TEXT(" ")); }
+				}
+			}
+			else
+			{
+				for (int32 Idx = 0; Idx < MaxItemNum / 2; Idx++)
+				{
+					Output.Append(Formatter(Array[Idx]));
+					Output.Append(TEXT(" "));
+				}
+
+				Output.Append(TEXT("... "));
+
+				for (int32 Idx = 0; Idx < MaxItemNum / 2; Idx++)
+				{
+					Output.Append(Formatter(Array[ItemNum - Idx - 1]));
+					if (Idx != MaxItemNum / 2 - 1) { Output.Append(TEXT(" ")); }
+				}
+			}
+
+			Output.Append(TEXT("]"));
+
+			return Output;
+		}
+
+		static inline FString FormatFloat(const TLearningArrayView<1, const float> Array, const int32 MaxItemNum = 16)
+		{
+			return Format<float>(Array, [](const float& Value) { return FString::Printf(TEXT("%6.3f"), Value); }, MaxItemNum);
+		}
 	}
 }
 
diff --git a/Engine/Plugins/Experimental/LearningAgents/Source/Learning/Public/LearningCompletionObject.h b/Engine/Plugins/Experimental/LearningAgents/Source/Learning/Public/LearningCompletionObject.h
--- a/Engine/Plugins/Experimental/LearningAgents/Source/Learning/Public/LearningCompletionObject.h
+++ b/Engine/Plugins/Experimental/LearningAgents/Source/Learning/Public/LearningCompletionObject.h
@@ -17,6 +17,7 @@ namespace UE::Learning
 		static constexpr float DefaultThresholdAngle = UE_PI / 2.0f; // rad
 		static constexpr float DefaultThresholdAngularVelocity = UE_PI; // rad/s
 		static constexpr float DefaultThresholdHeight = 0.0f; // cm
+		static constexpr float DefaultThresholdTime = 10.0f; // s
 	}
 
 	//------------------------------------------------------------------
@@ -84,6 +85,26 @@ namespace UE::Learning
 
 	//------------------------------------------------------------------
 
+	/**
+	* Completion due to time having elapsed
+	*/
+	struct LEARNING_API FTimeElapsedCompletion : public FCompletionObject
+	{
+		FTimeElapsedCompletion(
+			const FName& InIdentifier,
+			const TSharedRef<FArrayMap>& InInstanceData,
+			const int32 InMaxInstanceNum,
+			const float InThreshold = Completion::DefaultThresholdTime,
+			const ECompletionMode InCompletionMode = ECompletionMode::Terminated);
+
+		virtual void Evaluate(const FIndexSet Instances) override final;
+
+		TArrayMapHandle<1, float> TimeHandle;
+		TArrayMapHandle<1, float> ThresholdHandle;
+	};
+
+	//------------------------------------------------------------------
+
 	/**
 	* Completion due to the difference between any two scalar positions exceeding a given threshold
 	*/
@@ -129,6 +150,31 @@ namespace UE::Learning
 		TArrayMapHandle<1, float> ThresholdHandle;
 	};
 
+	/**
+	* Completion due to the similarity between any two planar positions being below a given threshold
+	*/
+	struct LEARNING_API FPlanarPositionSimilarityCompletion : public FCompletionObject
+	{
+		FPlanarPositionSimilarityCompletion(
+			const FName& InIdentifier,
+			const TSharedRef<FArrayMap>& InInstanceData,
+			const int32 InMaxInstanceNum,
+			const int32 InPositionNum,
+			const float InThreshold = Completion::DefaultThresholdPosition,
+			const ECompletionMode InCompletionMode = ECompletionMode::Terminated,
+			const FVector InAxis0 = FVector::ForwardVector,
+			const FVector InAxis1 = FVector::RightVector);
+
+		virtual void Evaluate(const FIndexSet Instances) override final;
+
+		FVector Axis0 = FVector::ForwardVector;
+		FVector Axis1 = FVector::RightVector;
+
+		TArrayMapHandle<2, FVector> Position0Handle;
+		TArrayMapHandle<2, FVector> Position1Handle;
+		TArrayMapHandle<1, float> ThresholdHandle;
+	};
+
 	//------------------------------------------------------------------
 
 	/**
diff --git a/Engine/Plugins/Experimental/LearningAgents/Source/Learning/Public/LearningFeatureObject.h b/Engine/Plugins/Experimental/LearningAgents/Source/Learning/Public/LearningFeatureObject.h
--- a/Engine/Plugins/Experimental/LearningAgents/Source/Learning/Public/LearningFeatureObject.h
+++ b/Engine/Plugins/Experimental/LearningAgents/Source/Learning/Public/LearningFeatureObject.h
@@ -139,14 +139,6 @@ namespace UE::Learning
 				const float Scale = DefaultScaleDirection, // angles are encoded as directions
 				const float Epsilon = UE_SMALL_NUMBER);
 
-			LEARNING_API void PlanarRotation(
-				TLearningArrayView<1, float> Output,
-				const FQuat Rotation,
-				const FQuat RelativeRotation = FQuat::Identity,
-				const float Scale = DefaultScaleDirection,  // rotations are encoded as directions
-				const FVector Axis = FVector::UpVector,
-				const float Epsilon = UE_SMALL_NUMBER);
-
 			LEARNING_API void Rotation(
 				TLearningArrayView<1, float> Output,
 				const FQuat Rotation,
@@ -167,14 +159,6 @@ namespace UE::Learning
 				const float Scale = DefaultScaleAngularVelocity,
 				const float Epsilon = UE_SMALL_NUMBER);
 
-			LEARNING_API void PlanarAngularVelocity(
-				float& Output,
-				const FVector AngularVelocity,
-				const FQuat RelativeRotation = FQuat::Identity,
-				const float Scale = DefaultScaleAngularVelocity,
-				const FVector Axis = FVector::UpVector,
-				const float Epsilon = UE_SMALL_NUMBER);
-
 			LEARNING_API void AngularVelocity(
 				TLearningArrayView<1, float> Output,
 				const FVector AngularVelocity,
@@ -314,13 +298,6 @@ namespace UE::Learning
 				const float RelativeRotation = 0.0f,
 				const float Scale = DefaultScaleDirection); // angles are encoded as directions
 
-			LEARNING_API void PlanarRotation(
-				FQuat& OutPlanarRotation,
-				const TLearningArrayView<1, const float> Input,
-				const FQuat RelativeRotation = FQuat::Identity,
-				const float Scale = DefaultScaleDirection, // rotations are encoded as directions
-				const FVector Axis = FVector::UpVector);
-
 			LEARNING_API void Rotation(
 				FQuat& OutRotation,
 				const TLearningArrayView<1, const float> Input,
@@ -339,13 +316,6 @@ namespace UE::Learning
 				const float Input,
 				const float Scale = DefaultScaleAngularVelocity);
 
-			LEARNING_API void PlanarAngularVelocity(
-				FVector& OutPlanarVelocity,
-				const float Input,
-				const FQuat RelativeRotation = FQuat::Identity,
-				const float Scale = DefaultScaleAngularVelocity,
-				const FVector Axis = FVector::UpVector);
-
 			LEARNING_API void AngularVelocity(
 				FVector& OutAngularVelocity,
 				const TLearningArrayView<1, const float> Input,
@@ -763,30 +733,6 @@ namespace UE::Learning
 		TArrayMapHandle<1, float> RelativeAngleHandle;
 	};
 
-	/**
-	* Feature for an array of rotations
-	*/
-	struct LEARNING_API FPlanarRotationFeature : public FFeatureObject
-	{
-		FPlanarRotationFeature(
-			const FName& InIdentifier,
-			const TSharedRef<FArrayMap>& InInstanceData,
-			const int32 InMaxInstanceNum,
-			const int32 InRotationNum,
-			const float InScale = Feature::DefaultScaleDirection, // rotations are encoded as directions
-			const FVector Axis = FVector::UpVector);
-
-		virtual bool IsEncodable() const override final { return true; }
-		virtual bool IsDecodable() const override final { return true; }
-		virtual void Encode(const FIndexSet Instances) override final;
-		virtual void Decode(const FIndexSet Instances) override final;
-
-		FVector Axis = FVector::UpVector;
-
-		TArrayMapHandle<2, FQuat> RotationHandle;
-		TArrayMapHandle<1, FQuat> RelativeRotationHandle;
-	};
-
 	/**
 	* Feature for an array of rotations
 	*/
@@ -850,30 +796,6 @@ namespace UE::Learning
 		TArrayMapHandle<2, float> AngularVelocityHandle;
 	};
 
-	/**
-	* Feature for an array of planar angular velocities
-	*/
-	struct LEARNING_API FPlanarAngularVelocityFeature : public FFeatureObject
-	{
-		FPlanarAngularVelocityFeature(
-			const FName& InIdentifier,
-			const TSharedRef<FArrayMap>& InInstanceData,
-			const int32 InMaxInstanceNum,
-			const int32 InAngularAccelerationNum,
-			const float InScale = Feature::DefaultScaleAngularVelocity,
-			const FVector InAxis = FVector::UpVector);
-
-		virtual bool IsEncodable() const override final { return true; }
-		virtual bool IsDecodable() const override final { return true; }
-		virtual void Encode(const FIndexSet Instances) override final;
-		virtual void Decode(const FIndexSet Instances) override final;
-
-		FVector Axis = FVector::UpVector;
-
-		TArrayMapHandle<2, FVector> AngularVelocityHandle;
-		TArrayMapHandle<1, FQuat> RelativeRotationHandle;
-	};
-
 	/**
 	* Feature for an array of angular velocities
 	*/
diff --git a/Engine/Plugins/Experimental/LearningAgents/Source/LearningAgents/Private/LearningAgentsActions.cpp b/Engine/Plugins/Experimental/LearningAgents/Source/LearningAgents/Private/LearningAgentsActions.cpp
--- a/Engine/Plugins/Experimental/LearningAgents/Source/LearningAgents/Private/LearningAgentsActions.cpp
+++ b/Engine/Plugins/Experimental/LearningAgents/Source/LearningAgents/Private/LearningAgentsActions.cpp
@@ -10,41 +10,28 @@
 #include "LearningLog.h"
 
 #include "GameFramework/Actor.h"
-#include "VisualLogger/VisualLogger.h"
 
-#define UE_LEARNING_AGENTS_VLOG_STRING(Owner, Category, Verbosity, Location, Color, Format, ...) \
-	UE_VLOG_LOCATION(Owner, Category, Verbosity, Location, 0.0f, Color, Format, ##__VA_ARGS__)
-
-#define UE_LEARNING_AGENTS_VLOG_TRANSFORM(Owner, Category, Verbosity, Location, Rotation, Color, Format, ...) \
-	UE_VLOG_SEGMENT(Owner, Category, Verbosity, Location, Location + 15.0f * Rotation.RotateVector(FVector::ForwardVector), FColor::Red, TEXT("")); \
-	UE_VLOG_SEGMENT(Owner, Category, Verbosity, Location, Location + 15.0f * Rotation.RotateVector(FVector::RightVector), FColor::Green, TEXT("")); \
-	UE_VLOG_SEGMENT(Owner, Category, Verbosity, Location, Location + 15.0f * Rotation.RotateVector(FVector::UpVector), FColor::Blue, TEXT("")); \
-	UE_VLOG_OBOX(Owner, Category, Verbosity, FBox(10.0f * FVector(-1, -1, -1), 10.0f * FVector(1, 1, 1)), FTransform(Rotation, Location, FVector::OneVector).ToMatrixNoScale(), Color, TEXT("")); \
-	UE_LEARNING_AGENTS_VLOG_STRING(Owner, Category, Verbosity, Location + FVector(0.0f, 0.0f, 20.0f), Color, Format, ##__VA_ARGS__)
-
-#define UE_LEARNING_AGENTS_VLOG_PLANE(Owner, Category, Verbosity, Location, Rotation, Axis0, Axis1, Color, Format, ...) \
-	UE_VLOG_OBOX(Owner, Category, Verbosity, FBox(-25.0f * (Axis0 + Axis1), 25.0f * (Axis0 + Axis1)), FTransform(Rotation, Location, FVector::OneVector).ToMatrixNoScale(), Color, Format, ##__VA_ARGS__)
-
-namespace UE::Learning::Agents::Private
+namespace UE::Learning::Agents::Actions::Private
 {
 	template<typename ActionUObject, typename ActionFObject, typename... InArgTypes>
-	ActionUObject* AddAction(ULearningAgentsType* AgentType, const FName Name, InArgTypes&& ...Args)
+	ActionUObject* AddAction(ULearningAgentsType* InAgentType, const FName Name, InArgTypes&& ...Args)
 	{
-		if (!AgentType)
+		if (!InAgentType)
 		{
-			UE_LOG(LogLearning, Error, TEXT("AgentType is nullptr"));
+			UE_LOG(LogLearning, Error, TEXT("InAgentType is nullptr."));
 			return nullptr;
 		}
 
-		ActionUObject* Action = NewObject<ActionUObject>(AgentType, Name);
+		ActionUObject* Action = NewObject<ActionUObject>(InAgentType, Name);
 
+		Action->AgentType = InAgentType;
 		Action->FeatureObject = MakeShared<ActionFObject>(
 			Action->GetFName(),
-			AgentType->GetAgentManager()->GetInstanceData().ToSharedRef(),
-			AgentType->GetAgentManager()->GetMaxInstanceNum(),
+			InAgentType->GetAgentManager()->GetInstanceData().ToSharedRef(),
+			InAgentType->GetAgentManager()->GetMaxInstanceNum(),
 			Forward<InArgTypes>(Args)...);
 
-		AgentType->AddAction(Action, Action->FeatureObject.ToSharedRef());
+		InAgentType->AddAction(Action, Action->FeatureObject.ToSharedRef());
 
 		return Action;
 	}
@@ -52,72 +39,123 @@ namespace UE::Learning::Agents::Private
 
 //------------------------------------------------------------------
 
-UFloatAction* UFloatAction::AddFloatAction(ULearningAgentsType* AgentType, const FName Name, const float Scale)
+UFloatAction* UFloatAction::AddFloatAction(ULearningAgentsType* InAgentType, const FName Name, const float Scale)
 {
-	return UE::Learning::Agents::Private::AddAction<UFloatAction, UE::Learning::FFloatFeature>(AgentType, Name, 1, Scale);
+	return UE::Learning::Agents::Actions::Private::AddAction<UFloatAction, UE::Learning::FFloatFeature>(InAgentType, Name, 1, Scale);
 }
 
-float UFloatAction::GetFloatAction(const int32 AgentId)
+float UFloatAction::GetFloatAction(const int32 AgentId) const
 {
-	const TLearningArrayView<2, const float> View = FeatureObject->InstanceData->ConstView(FeatureObject->ValueHandle);
-
-	if (AgentId == INDEX_NONE)
+	if (!AgentType->HasAgent(AgentId))
 	{
-		UE_LOG(LogLearning, Error, TEXT("AgentId is invalid (INDEX_NONE)"));
+		UE_LOG(LogLearning, Error, TEXT("%s: AgentId %d not found in the agents set."), *GetName(), AgentId);
 		return 0.0f;
 	}
 
-	if (AgentId < 0 || AgentId >= View.Num<0>())
+	return FeatureObject->InstanceData->ConstView(FeatureObject->ValueHandle)[AgentId][0];
+}
+
+void UFloatAction::SetFloatAction(const int32 AgentId, const float Value)
+{
+	if (!AgentType->HasAgent(AgentId))
 	{
-		UE_LOG(LogLearning, Error, TEXT("AgentId %d is out of index. Valid range [0, %d]."), AgentId, View.Num<0>() - 1);
-		return 0.0f;
+		UE_LOG(LogLearning, Error, TEXT("%s: AgentId %d not found in the agents set."), *GetName(), AgentId);
+		return;
 	}
 
-	return View[AgentId][0];
+	FeatureObject->InstanceData->View(FeatureObject->ValueHandle)[AgentId][0] = Value;
 }
 
-void UFloatAction::SetFloatAction(const int32 AgentId, const float Value)
+#if UE_LEARNING_AGENTS_ENABLE_VISUAL_LOG
+void UFloatAction::VisualLog(const UE::Learning::FIndexSet Instances) const
 {
-	const TLearningArrayView<2, float> View = FeatureObject->InstanceData->View(FeatureObject->ValueHandle);
+	UE_LEARNING_TRACE_CPUPROFILER_EVENT_SCOPE(UFloatAction::VisualLog);
+
+	const TLearningArrayView<2, const float> ValueView = FeatureObject->InstanceData->ConstView(FeatureObject->ValueHandle);
+	const TLearningArrayView<2, const float> FeatureView = FeatureObject->InstanceData->ConstView(FeatureObject->FeatureHandle);
+
+	for (const int32 Instance : Instances)
+	{
+		if (const AActor* Actor = Cast<AActor>(AgentType->GetAgent(Instance)))
+		{
+			UE_LEARNING_AGENTS_VLOG_STRING(this, LogLearning, Display,
+				Actor->GetActorLocation(),
+				VisualLogColor.ToFColor(true),
+				TEXT("Agent %i\nScale: [% 6.2f]\nValue: [% 6.2f]\nEncoded: [% 6.3f]"),
+				Instance,
+				FeatureObject->Scale,
+				ValueView[Instance][0],
+				FeatureView[Instance][0]);
+		}
+	}
+}
+#endif
 
-	if (AgentId == INDEX_NONE)
+UFloatArrayAction* UFloatArrayAction::AddFloatArrayAction(ULearningAgentsType* InAgentType, const FName Name, const int32 Num, const float Scale)
+{
+	if (Num < 1)
+	{
+		UE_LOG(LogLearning, Error, TEXT("Number of elements in array must be at least 1, got %i."), Num);
+		return nullptr;
+	}
+
+	return UE::Learning::Agents::Actions::Private::AddAction<UFloatArrayAction, UE::Learning::FFloatFeature>(InAgentType, Name, Num, Scale);
+}
+
+void UFloatArrayAction::GetFloatArrayAction(const int32 AgentId, TArray<float>& OutValues) const
+{
+	if (!AgentType->HasAgent(AgentId))
 	{
-		UE_LOG(LogLearning, Error, TEXT("AgentId is invalid (INDEX_NONE)"));
+		UE_LOG(LogLearning, Error, TEXT("%s: AgentId %d not found in the agents set."), *GetName(), AgentId);
+		OutValues.Empty();
 		return;
 	}
 
-	if (AgentId < 0 || AgentId >= View.Num<0>())
+	const TLearningArrayView<2, const float> View = FeatureObject->InstanceData->ConstView(FeatureObject->ValueHandle);
+
+	OutValues.SetNumUninitialized(View.Num<1>());
+	UE::Learning::Array::Copy<1, float>(OutValues, View[AgentId]);
+}
+
+void UFloatArrayAction::SetFloatArrayAction(const int32 AgentId, const TArray<float>& Values)
+{
+	if (!AgentType->HasAgent(AgentId))
+	{
+		UE_LOG(LogLearning, Error, TEXT("%s: AgentId %d not found in the agents set."), *GetName(), AgentId);
+		return;
+	}
+
+	const TLearningArrayView<2, float> View = FeatureObject->InstanceData->View(FeatureObject->ValueHandle);
+
+	if (Values.Num() != View.Num<1>())
 	{
-		UE_LOG(LogLearning, Error, TEXT("AgentId %d is out of index. Valid range [0, %d]."), AgentId, View.Num<0>() - 1);
+		UE_LOG(LogLearning, Error, TEXT("%s: Got wrong number of elements in array. Expected %i, got %i."), *GetName(), View.Num<1>(), Values.Num());
 		return;
 	}
 
-	View[AgentId][0] = Value;
+	UE::Learning::Array::Copy<1, float>(View[AgentId], Values);
 }
 
-#if ENABLE_VISUAL_LOG
-void UFloatAction::VisualLog(const UE::Learning::FIndexSet Instances) const
+#if UE_LEARNING_AGENTS_ENABLE_VISUAL_LOG
+void UFloatArrayAction::VisualLog(const UE::Learning::FIndexSet Instances) const
 {
-	UE_LEARNING_TRACE_CPUPROFILER_EVENT_SCOPE(UFloatAction::VisualLog);
+	UE_LEARNING_TRACE_CPUPROFILER_EVENT_SCOPE(UFloatArrayAction::VisualLog);
 
 	const TLearningArrayView<2, const float> ValueView = FeatureObject->InstanceData->ConstView(FeatureObject->ValueHandle);
 	const TLearningArrayView<2, const float> FeatureView = FeatureObject->InstanceData->ConstView(FeatureObject->FeatureHandle);
 
-	if (const ULearningAgentsType* AgentType = Cast<ULearningAgentsType>(GetOuter()))
+	for (const int32 Instance : Instances)
 	{
-		for (const int32 Instance : Instances)
+		if (const AActor* Actor = Cast<AActor>(AgentType->GetAgent(Instance)))
 		{
-			if (const AActor* Actor = Cast<AActor>(AgentType->GetAgent(Instance)))
-			{
-				UE_LEARNING_AGENTS_VLOG_STRING(this, LogLearning, Display,
-					Actor->GetActorLocation(),
-					VisualLogColor.ToFColor(true),
-					TEXT("Agent %i\nScale: [% 6.2f]\nValue: [% 6.2f]\nEncoded: [% 6.3f]"),
-					Instance,
-					FeatureObject->Scale,
-					ValueView[Instance][0],
-					FeatureView[Instance][0]);
-			}
+			UE_LEARNING_AGENTS_VLOG_STRING(this, LogLearning, Display,
+				Actor->GetActorLocation(),
+				VisualLogColor.ToFColor(true),
+				TEXT("Agent %i\nScale: [% 6.2f]\nValue: %s\nEncoded: %s"),
+				Instance,
+				FeatureObject->Scale,
+				*UE::Learning::Array::FormatFloat(ValueView[Instance]),
+				*UE::Learning::Array::FormatFloat(FeatureView[Instance]));
 		}
 	}
 }
@@ -125,52 +163,40 @@ void UFloatAction::VisualLog(const UE::Learning::FIndexSet Instances) const
 
 //------------------------------------------------------------------
 
-UVectorAction* UVectorAction::AddVectorAction(ULearningAgentsType* AgentType, const FName Name, const float Scale)
+UVectorAction* UVectorAction::AddVectorAction(ULearningAgentsType* InAgentType, const FName Name, const float Scale)
 {
-	return UE::Learning::Agents::Private::AddAction<UVectorAction, UE::Learning::FFloatFeature>(AgentType, Name, 3, Scale);
+	return UE::Learning::Agents::Actions::Private::AddAction<UVectorAction, UE::Learning::FFloatFeature>(InAgentType, Name, 3, Scale);
 }
 
-FVector UVectorAction::GetVectorAction(const int32 AgentId)
+FVector UVectorAction::GetVectorAction(const int32 AgentId) const
 {
-	const TLearningArrayView<2, const float> View = FeatureObject->InstanceData->ConstView(FeatureObject->ValueHandle);
-
-	if (AgentId == INDEX_NONE)
+	if (!AgentType->HasAgent(AgentId))
 	{
-		UE_LOG(LogLearning, Error, TEXT("AgentId is invalid (INDEX_NONE)"));
+		UE_LOG(LogLearning, Error, TEXT("%s: AgentId %d not found in the agents set."), *GetName(), AgentId);
 		return FVector::ZeroVector;
 	}
 
-	if (AgentId < 0 || AgentId >= View.Num<0>())
-	{
-		UE_LOG(LogLearning, Error, TEXT("AgentId %d is out of index. Valid range [0, %d]."), AgentId, View.Num<0>() - 1);
-		return FVector::ZeroVector;
-	}
+	const TLearningArrayView<2, const float> View = FeatureObject->InstanceData->ConstView(FeatureObject->ValueHandle);
 
 	return FVector(View[AgentId][0], View[AgentId][1], View[AgentId][2]);
 }
 
 void UVectorAction::SetVectorAction(const int32 AgentId, const FVector InAction)
 {
-	const TLearningArrayView<2, float> View = FeatureObject->InstanceData->View(FeatureObject->ValueHandle);
-
-	if (AgentId == INDEX_NONE)
+	if (!AgentType->HasAgent(AgentId))
 	{
-		UE_LOG(LogLearning, Error, TEXT("AgentId is invalid (INDEX_NONE)"));
+		UE_LOG(LogLearning, Error, TEXT("%s: AgentId %d not found in the agents set."), *GetName(), AgentId);
 		return;
 	}
 
-	if (AgentId < 0 || AgentId >= View.Num<0>())
-	{
-		UE_LOG(LogLearning, Error, TEXT("AgentId %d is out of index. Valid range [0, %d]."), AgentId, View.Num<0>() - 1);
-		return;
-	}
+	const TLearningArrayView<2, float> View = FeatureObject->InstanceData->View(FeatureObject->ValueHandle);
 
 	View[AgentId][0] = InAction.X;
 	View[AgentId][1] = InAction.Y;
 	View[AgentId][2] = InAction.Z;
 }
 
-#if ENABLE_VISUAL_LOG
+#if UE_LEARNING_AGENTS_ENABLE_VISUAL_LOG
 void UVectorAction::VisualLog(const UE::Learning::FIndexSet Instances) const
 {
 	UE_LEARNING_TRACE_CPUPROFILER_EVENT_SCOPE(UVectorAction::VisualLog);
@@ -178,34 +204,133 @@ void UVectorAction::VisualLog(const UE::Learning::FIndexSet Instances) const
 	const TLearningArrayView<2, const float> ValueView = FeatureObject->InstanceData->ConstView(FeatureObject->ValueHandle);
 	const TLearningArrayView<2, const float> FeatureView = FeatureObject->InstanceData->ConstView(FeatureObject->FeatureHandle);
 
-	if (const ULearningAgentsType* AgentType = Cast<ULearningAgentsType>(GetOuter()))
+	for (const int32 Instance : Instances)
+	{
+		if (const AActor* Actor = Cast<AActor>(AgentType->GetAgent(Instance)))
+		{
+			const FVector Vector(ValueView[Instance][0], ValueView[Instance][1], ValueView[Instance][2]);
+
+			UE_LEARNING_AGENTS_VLOG_ARROW(this, LogLearning, Display,
+				Actor->GetActorLocation(),
+				Actor->GetActorLocation() + Vector,
+				VisualLogColor.ToFColor(true),
+				TEXT(""));
+
+			UE_LEARNING_AGENTS_VLOG_STRING(this, LogLearning, Display,
+				Actor->GetActorLocation() + Vector,
+				VisualLogColor.ToFColor(true),
+				TEXT("Vector: [% 6.4f % 6.4f % 6.4f]"),
+				Vector.X, Vector.Y, Vector.Z);
+
+			UE_LEARNING_AGENTS_VLOG_STRING(this, LogLearning, Display,
+				Actor->GetActorLocation(),
+				VisualLogColor.ToFColor(true),
+				TEXT("Agent %i\nScale: [% 6.2f]\nEncoded: [% 6.3f % 6.3f % 6.3f]"),
+				Instance,
+				FeatureObject->Scale,
+				FeatureView[Instance][0], FeatureView[Instance][1], FeatureView[Instance][2]);
+		}
+	}
+}
+#endif
+
+UVectorArrayAction* UVectorArrayAction::AddVectorArrayAction(ULearningAgentsType* InAgentType, const FName Name, const int32 Num, const float Scale)
+{
+	if (Num < 1)
+	{
+		UE_LOG(LogLearning, Error, TEXT("Number of elements in array must be at least 1, got %i."), Num);
+		return nullptr;
+	}
+
+	return UE::Learning::Agents::Actions::Private::AddAction<UVectorArrayAction, UE::Learning::FFloatFeature>(InAgentType, Name, Num * 3, Scale);
+}
+
+void UVectorArrayAction::GetVectorArrayAction(const int32 AgentId, TArray<FVector>& OutVectors) const
+{
+	if (!AgentType->HasAgent(AgentId))
+	{
+		UE_LOG(LogLearning, Error, TEXT("%s: AgentId %d not found in the agents set."), *GetName(), AgentId);
+		OutVectors.Empty();
+		return;
+	}
+
+	const TLearningArrayView<2, const float> View = FeatureObject->InstanceData->ConstView(FeatureObject->ValueHandle);
+
+	OutVectors.SetNumUninitialized(View.Num<1>() / 3);
+
+	for (int32 VectorIdx = 0; VectorIdx < View.Num<1>() / 3; VectorIdx++)
+	{
+		OutVectors[VectorIdx] = FVector(
+			View[AgentId][VectorIdx * 3 + 0], 
+			View[AgentId][VectorIdx * 3 + 1], 
+			View[AgentId][VectorIdx * 3 + 2]);
+	}
+}
+
+void UVectorArrayAction::SetVectorArrayAction(const int32 AgentId, const TArray<FVector>& Vectors)
+{
+	if (!AgentType->HasAgent(AgentId))
+	{
+		UE_LOG(LogLearning, Error, TEXT("%s: AgentId %d not found in the agents set."), *GetName(), AgentId);
+		return;
+	}
+
+	const TLearningArrayView<2, float> View = FeatureObject->InstanceData->View(FeatureObject->ValueHandle);
+
+	if (Vectors.Num() != View.Num<1>() / 3)
+	{
+		UE_LOG(LogLearning, Error, TEXT("%s: Got wrong number of elements in array. Expected %i, got %i."), *GetName(), View.Num<1>() / 3, Vectors.Num());
+		return;
+	}
+
+	for (int32 VectorIdx = 0; VectorIdx < Vectors.Num(); VectorIdx++)
+	{
+		View[AgentId][VectorIdx * 3 + 0] = Vectors[VectorIdx].X;
+		View[AgentId][VectorIdx * 3 + 1] = Vectors[VectorIdx].Y;
+		View[AgentId][VectorIdx * 3 + 2] = Vectors[VectorIdx].Z;
+	}
+}
+
+#if UE_LEARNING_AGENTS_ENABLE_VISUAL_LOG
+void UVectorArrayAction::VisualLog(const UE::Learning::FIndexSet Instances) const
+{
+	UE_LEARNING_TRACE_CPUPROFILER_EVENT_SCOPE(UVectorArrayAction::VisualLog);
+
+	const TLearningArrayView<2, const float> ValueView = FeatureObject->InstanceData->ConstView(FeatureObject->ValueHandle);
+	const TLearningArrayView<2, const float> FeatureView = FeatureObject->InstanceData->ConstView(FeatureObject->FeatureHandle);
+
+	const int32 VectorNum = ValueView.Num<1>() / 3;
+
+	for (const int32 Instance : Instances)
 	{
-		for (const int32 Instance : Instances)
+		if (const AActor* Actor = Cast<AActor>(AgentType->GetAgent(Instance)))
 		{
-			if (const AActor* Actor = Cast<AActor>(AgentType->GetAgent(Instance)))
+			for (int32 VectorIdx = 0; VectorIdx < VectorNum; VectorIdx++)
 			{
-				const FVector Vector(ValueView[Instance][0], ValueView[Instance][1], ValueView[Instance][2]);
+				const FVector Vector(ValueView[Instance][VectorIdx * 3 + 0], ValueView[Instance][VectorIdx * 3 + 1], ValueView[Instance][VectorIdx * 3 + 2]);
+				const FVector Offset = UE::Learning::Agents::Debug::GridOffsetForIndex(VectorIdx, VectorNum);
 
-				UE_VLOG_ARROW(this, LogLearning, Display,
-					Actor->GetActorLocation(),
-					Actor->GetActorLocation() + Vector,
+				UE_LEARNING_AGENTS_VLOG_ARROW(this, LogLearning, Display,
+					Actor->GetActorLocation() + Offset,
+					Actor->GetActorLocation() + Offset + Vector,
 					VisualLogColor.ToFColor(true),
 					TEXT(""));
 
 				UE_LEARNING_AGENTS_VLOG_STRING(this, LogLearning, Display,
-					Actor->GetActorLocation() + Vector,
+					Actor->GetActorLocation() + Offset + Vector,
 					VisualLogColor.ToFColor(true),
-					TEXT("Vector: [% 6.4f % 6.4f % 6.4f]"),
+					TEXT("Vector %i: [% 6.4f % 6.4f % 6.4f]"),
+					VectorIdx,
 					Vector.X, Vector.Y, Vector.Z);
-
-				UE_LEARNING_AGENTS_VLOG_STRING(this, LogLearning, Display,
-					Actor->GetActorLocation(),
-					VisualLogColor.ToFColor(true),
-					TEXT("Agent %i\nScale: [% 6.2f]\nEncoded: [% 6.3f % 6.3f % 6.3f]"),
-					Instance,
-					FeatureObject->Scale,
-					FeatureView[Instance][0], FeatureView[Instance][1], FeatureView[Instance][2]);
 			}
+
+			UE_LEARNING_AGENTS_VLOG_STRING(this, LogLearning, Display,
+				Actor->GetActorLocation(),
+				VisualLogColor.ToFColor(true),
+				TEXT("Agent %i\nScale: [% 6.2f]\nEncoded: %s"),
+				Instance,
+				FeatureObject->Scale,
+				*UE::Learning::Array::FormatFloat(FeatureView[Instance]));
 		}
 	}
 }
@@ -213,148 +338,179 @@ void UVectorAction::VisualLog(const UE::Learning::FIndexSet Instances) const
 
 //------------------------------------------------------------------
 
-URotationVectorArrayAction* URotationVectorArrayAction::AddRotationVectorArrayAction(ULearningAgentsType* AgentType, const FName Name, const int32 RotationVectorNum, const float Scale)
+URotationAction* URotationAction::AddRotationAction(ULearningAgentsType* InAgentType, const FName Name, const float Scale)
 {
-	return UE::Learning::Agents::Private::AddAction<URotationVectorArrayAction, UE::Learning::FRotationVectorFeature>(AgentType, Name, RotationVectorNum, FMath::DegreesToRadians(Scale));
+	return UE::Learning::Agents::Actions::Private::AddAction<URotationAction, UE::Learning::FRotationVectorFeature>(InAgentType, Name, 1, Scale);
 }
 
-void URotationVectorArrayAction::GetRotationVectorArrayAction(const int32 AgentId, TArray<FVector>& OutRotationVectors)
+FRotator URotationAction::GetRotationAction(const int32 AgentId) const
 {
+	if (!AgentType->HasAgent(AgentId))
+	{
+		UE_LOG(LogLearning, Error, TEXT("%s: AgentId %d not found in the agents set."), *GetName(), AgentId);
+		return FRotator::ZeroRotator;
+	}
+
 	const TLearningArrayView<2, const FVector> View = FeatureObject->InstanceData->ConstView(FeatureObject->RotationVectorsHandle);
 
-	if (AgentId == INDEX_NONE)
+	return FQuat::MakeFromRotationVector((UE_TWO_PI / 180.0f) * View[AgentId][0]).Rotator();
+}
+
+FVector URotationAction::GetRotationActionAsRotationVector(const int32 AgentId) const
+{
+	if (!AgentType->HasAgent(AgentId))
 	{
-		UE_LOG(LogLearning, Error, TEXT("AgentId is invalid (INDEX_NONE)"));
-		return;
+		UE_LOG(LogLearning, Error, TEXT("%s: AgentId %d not found in the agents set."), *GetName(), AgentId);
+		return FVector::ZeroVector;
 	}
 
-	if (AgentId < 0 || AgentId >= View.Num<0>())
+	return FeatureObject->InstanceData->ConstView(FeatureObject->RotationVectorsHandle)[AgentId][0];
+}
+
+FQuat URotationAction::GetRotationActionAsQuat(const int32 AgentId) const
+{
+	if (!AgentType->HasAgent(AgentId))
 	{
-		UE_LOG(LogLearning, Error, TEXT("AgentId %d is out of index. Valid range [0, %d]."), AgentId, View.Num<0>() - 1);
-		return;
+		UE_LOG(LogLearning, Error, TEXT("%s: AgentId %d not found in the agents set."), *GetName(), AgentId);
+		return FQuat::Identity;
 	}
 
-	OutRotationVectors.SetNumUninitialized(View.Num<1>());
-	for (int32 RotationVectorIdx = 0; RotationVectorIdx < View.Num<1>(); RotationVectorIdx++)
+	const TLearningArrayView<2, const FVector> View = FeatureObject->InstanceData->ConstView(FeatureObject->RotationVectorsHandle);
+
+	return FQuat::MakeFromRotationVector((UE_TWO_PI / 180.0f) * View[AgentId][0]);
+}
+
+#if UE_LEARNING_AGENTS_ENABLE_VISUAL_LOG
+void URotationAction::VisualLog(const UE::Learning::FIndexSet Instances) const
+{
+	UE_LEARNING_TRACE_CPUPROFILER_EVENT_SCOPE(URotationAction::VisualLog);
+
+	const TLearningArrayView<2, const FVector> ValueView = FeatureObject->InstanceData->ConstView(FeatureObject->RotationVectorsHandle);
+	const TLearningArrayView<2, const float> FeatureView = FeatureObject->InstanceData->ConstView(FeatureObject->FeatureHandle);
+
+	const int32 RotationVectorNum = ValueView.Num<1>();
+
+	for (const int32 Instance : Instances)
 	{
-		OutRotationVectors[RotationVectorIdx] = View[AgentId][RotationVectorIdx];
+		if (const AActor* Actor = Cast<AActor>(AgentType->GetAgent(Instance)))
+		{
+			const FQuat Rotation = FQuat::MakeFromRotationVector((UE_TWO_PI / 180.0f) * ValueView[Instance][0]);
+			const FRotator Rotator = Rotation.Rotator();
+
+			UE_LEARNING_AGENTS_VLOG_TRANSFORM(this, LogLearning, Display,
+				Actor->GetActorLocation(),
+				Rotation,
+				VisualLogColor.ToFColor(true),
+				TEXT("Agent %i\nScale: [% 6.2f]\nRotation: [% 6.1f % 6.1f % 6.1f]\nEncoded: %s"),
+				Instance,
+				FeatureObject->Scale,
+				Rotator.Pitch, Rotator.Roll, Rotator.Yaw,
+				*UE::Learning::Array::FormatFloat(FeatureView[Instance]));
+		}
 	}
 }
+#endif
 
-void URotationVectorArrayAction::GetRotationVectorArrayActionAsQuats(const int32 AgentId, TArray<FQuat>& OutRotations)
-{
-	const TLearningArrayView<2, const FVector> View = FeatureObject->InstanceData->ConstView(FeatureObject->RotationVectorsHandle);
 
-	if (AgentId == INDEX_NONE)
+URotationArrayAction* URotationArrayAction::AddRotationArrayAction(ULearningAgentsType* InAgentType, const FName Name, const int32 RotationNum, const float Scale)
+{
+	if (RotationNum < 1)
 	{
-		UE_LOG(LogLearning, Error, TEXT("AgentId is invalid (INDEX_NONE)"));
-		return;
+		UE_LOG(LogLearning, Error, TEXT("Number of elements in array must be at least 1, got %i."), RotationNum);
+		return nullptr;
 	}
 
-	if (AgentId < 0 || AgentId >= View.Num<0>())
+	return UE::Learning::Agents::Actions::Private::AddAction<URotationArrayAction, UE::Learning::FRotationVectorFeature>(InAgentType, Name, RotationNum, Scale);
+}
+
+void URotationArrayAction::GetRotationArrayAction(const int32 AgentId, TArray<FRotator>& OutRotations) const
+{
+	if (!AgentType->HasAgent(AgentId))
 	{
-		UE_LOG(LogLearning, Error, TEXT("AgentId %d is out of index. Valid range [0, %d]."), AgentId, View.Num<0>() - 1);
+		UE_LOG(LogLearning, Error, TEXT("%s: AgentId %d not found in the agents set."), *GetName(), AgentId);
 		return;
 	}
 
+	const TLearningArrayView<2, const FVector> View = FeatureObject->InstanceData->ConstView(FeatureObject->RotationVectorsHandle);
+
 	OutRotations.SetNumUninitialized(View.Num<1>());
 	for (int32 RotationVectorIdx = 0; RotationVectorIdx < View.Num<1>(); RotationVectorIdx++)
 	{
-		OutRotations[RotationVectorIdx] = FQuat::MakeFromRotationVector(View[AgentId][RotationVectorIdx]);
+		OutRotations[RotationVectorIdx] = FQuat::MakeFromRotationVector((UE_TWO_PI / 180.0f) * View[AgentId][RotationVectorIdx]).Rotator();
 	}
 }
 
-void URotationVectorArrayAction::GetRotationVectorArrayActionAsRotators(const int32 AgentId, TArray<FRotator>& OutRotations)
+void URotationArrayAction::GetRotationArrayActionAsRotationVectors(const int32 AgentId, TArray<FVector>& OutRotationVectors) const
 {
+	if (!AgentType->HasAgent(AgentId))
+	{
+		UE_LOG(LogLearning, Error, TEXT("%s: AgentId %d not found in the agents set."), *GetName(), AgentId);
+		return;
+	}
+
 	const TLearningArrayView<2, const FVector> View = FeatureObject->InstanceData->ConstView(FeatureObject->RotationVectorsHandle);
 
-	if (AgentId == INDEX_NONE)
+	OutRotationVectors.SetNumUninitialized(View.Num<1>());
+	for (int32 RotationVectorIdx = 0; RotationVectorIdx < View.Num<1>(); RotationVectorIdx++)
 	{
-		UE_LOG(LogLearning, Error, TEXT("AgentId is invalid (INDEX_NONE)"));
-		return;
+		OutRotationVectors[RotationVectorIdx] = View[AgentId][RotationVectorIdx];
 	}
+}
 
-	if (AgentId < 0 || AgentId >= View.Num<0>())
+void URotationArrayAction::GetRotationArrayActionAsQuats(const int32 AgentId, TArray<FQuat>& OutRotations) const
+{
+	if (!AgentType->HasAgent(AgentId))
 	{
-		UE_LOG(LogLearning, Error, TEXT("AgentId %d is out of index. Valid range [0, %d]."), AgentId, View.Num<0>() - 1);
+		UE_LOG(LogLearning, Error, TEXT("%s: AgentId %d not found in the agents set."), *GetName(), AgentId);
 		return;
 	}
 
+	const TLearningArrayView<2, const FVector> View = FeatureObject->InstanceData->ConstView(FeatureObject->RotationVectorsHandle);
+
 	OutRotations.SetNumUninitialized(View.Num<1>());
 	for (int32 RotationVectorIdx = 0; RotationVectorIdx < View.Num<1>(); RotationVectorIdx++)
 	{
-		OutRotations[RotationVectorIdx] = FQuat::MakeFromRotationVector(View[AgentId][RotationVectorIdx]).Rotator();
+		OutRotations[RotationVectorIdx] = FQuat::MakeFromRotationVector((UE_TWO_PI / 180.0f) * View[AgentId][RotationVectorIdx]);
 	}
 }
 
-#if ENABLE_VISUAL_LOG
-void URotationVectorArrayAction::VisualLog(const UE::Learning::FIndexSet Instances) const
+#if UE_LEARNING_AGENTS_ENABLE_VISUAL_LOG
+void URotationArrayAction::VisualLog(const UE::Learning::FIndexSet Instances) const
 {
-	UE_LEARNING_TRACE_CPUPROFILER_EVENT_SCOPE(URotationVectorArrayAction::VisualLog);
+	UE_LEARNING_TRACE_CPUPROFILER_EVENT_SCOPE(URotationArrayAction::VisualLog);
 
 	const TLearningArrayView<2, const FVector> ValueView = FeatureObject->InstanceData->ConstView(FeatureObject->RotationVectorsHandle);
 	const TLearningArrayView<2, const float> FeatureView = FeatureObject->InstanceData->ConstView(FeatureObject->FeatureHandle);
 
-	const int32 RotationVectorNum = ValueView.Num<1>();
+	const int32 RotationNum = ValueView.Num<1>();
 
-	if (const ULearningAgentsType* AgentType = Cast<ULearningAgentsType>(GetOuter()))
+	for (const int32 Instance : Instances)
 	{
-		for (const int32 Instance : Instances)
+		if (const AActor* Actor = Cast<AActor>(AgentType->GetAgent(Instance)))
 		{
-			if (const AActor* Actor = Cast<AActor>(AgentType->GetAgent(Instance)))
+			for (int32 RotationIdx = 0; RotationIdx < RotationNum; RotationIdx++)
 			{
-				for (int32 RotationVectorIdx = 0; RotationVectorIdx < RotationVectorNum; RotationVectorIdx++)
-				{
-					const FVector Offset = FVector(0.0, 10.0f * RotationVectorIdx, 0.0f);
-					const FVector RotationVector = ValueView[Instance][RotationVectorIdx];
-
-					UE_VLOG_LOCATION(this, LogLearning, Display,
-						Actor->GetActorLocation() + Offset,
-						2.5f,
-						VisualLogColor.ToFColor(true),
-						TEXT(""));
-
-					UE_VLOG_ARROW(this, LogLearning, Display,
-						Actor->GetActorLocation() + Offset,
-						Actor->GetActorLocation() + Offset + RotationVector,
-						VisualLogColor.ToFColor(true),
-						TEXT(""));
-
-					UE_LEARNING_AGENTS_VLOG_STRING(this, LogLearning, Display,
-						Actor->GetActorLocation() + Offset + RotationVector,
-						VisualLogColor.ToFColor(true),
-						TEXT("Rotation Vector: [% 6.4f % 6.4f % 6.4f]"),
-						RotationVector.X, RotationVector.Y, RotationVector.Z);
-
-				}
-
-				if (RotationVectorNum > 0)
-				{
-					UE_LEARNING_AGENTS_VLOG_STRING(this, LogLearning, Display,
-						Actor->GetActorLocation(),
-						VisualLogColor.ToFColor(true),
-						TEXT("Agent %i\nScale: [% 6.2f]\nEncoded: [% 6.3f % 6.3f % 6.3f ...]"),
-						Instance,
-						FeatureObject->Scale,
-						FeatureView[Instance][0], FeatureView[Instance][1], FeatureView[Instance][2]);
-				}
-				else
-				{
-					UE_LEARNING_AGENTS_VLOG_STRING(this, LogLearning, Display,
-						Actor->GetActorLocation(),
-						VisualLogColor.ToFColor(true),
-						TEXT("Agent %i\nScale: [% 6.2f]\nEncoded: []"),
-						Instance,
-						FeatureObject->Scale);
-				}
+				const FVector Offset = UE::Learning::Agents::Debug::GridOffsetForIndex(RotationIdx, RotationNum);
+				const FQuat Rotation = FQuat::MakeFromRotationVector((UE_TWO_PI / 180.0f) * ValueView[Instance][RotationIdx]);
+				const FRotator Rotator = Rotation.Rotator();
+
+				UE_LEARNING_AGENTS_VLOG_TRANSFORM(this, LogLearning, Display,
+					Actor->GetActorLocation() + Offset,
+					Rotation,
+					VisualLogColor.ToFColor(true),
+					TEXT("Rotation %i: [% 6.3f % 6.3f % 6.3f]"),
+					RotationIdx,
+					Rotator.Pitch, Rotator.Roll, Rotator.Yaw);
 			}
+
+			UE_LEARNING_AGENTS_VLOG_STRING(this, LogLearning, Display,
+				Actor->GetActorLocation(),
+				VisualLogColor.ToFColor(true),
+				TEXT("Agent %i\nScale: [% 6.2f]\nEncoded: %s"),
+				Instance,
+				FeatureObject->Scale,
+				*UE::Learning::Array::FormatFloat(FeatureView[Instance]));
 		}
 	}
 }
 #endif
-
-//------------------------------------------------------------------
-
-#undef UE_LEARNING_AGENTS_VLOG_STRING
-#undef UE_LEARNING_AGENTS_VLOG_TRANSFORM
-#undef UE_LEARNING_AGENTS_VLOG_PLANE
diff --git a/Engine/Plugins/Experimental/LearningAgents/Source/LearningAgents/Private/LearningAgentsController.cpp b/Engine/Plugins/Experimental/LearningAgents/Source/LearningAgents/Private/LearningAgentsController.cpp
--- a/Engine/Plugins/Experimental/LearningAgents/Source/LearningAgents/Private/LearningAgentsController.cpp
+++ b/Engine/Plugins/Experimental/LearningAgents/Source/LearningAgents/Private/LearningAgentsController.cpp
@@ -23,7 +23,7 @@ void ULearningAgentsController::SetupController(ALearningAgentsManager* InAgentM
 {
 	if (IsSetup())
 	{
-		UE_LOG(LogLearning, Error, TEXT("%s: Setup already performed!"), *GetName());
+		UE_LOG(LogLearning, Error, TEXT("%s: Setup already run!"), *GetName());
 		return;
 	}
 
@@ -35,7 +35,7 @@ void ULearningAgentsController::SetupController(ALearningAgentsManager* InAgentM
 
 	if (!InAgentManager->IsManagerSetup())
 	{
-		UE_LOG(LogLearning, Error, TEXT("%s's SetupManager() must be run before %s can be setup."), *InAgentManager->GetName(), *GetName());
+		UE_LOG(LogLearning, Error, TEXT("%s: %s's SetupManager must be run before it can be used."), *GetName(), *InAgentManager->GetName());
 		return;
 	}
 
@@ -65,7 +65,7 @@ void ULearningAgentsController::EncodeActions()
 
 	if (!IsSetup())
 	{
-		UE_LOG(LogLearning, Error, TEXT("Setup must be run before actions can be encoded."));
+		UE_LOG(LogLearning, Error, TEXT("%s: Setup not run."), *GetName());
 		return;
 	}
 
@@ -94,3 +94,18 @@ ULearningAgentsType* ULearningAgentsController::GetAgentType(const TSubclassOf<U
 
 	return AgentType;
 }
+
+void ULearningAgentsController::RunController()
+{
+	UE_LEARNING_TRACE_CPUPROFILER_EVENT_SCOPE(ULearningAgentsController::RunController);
+
+	if (!IsSetup())
+	{
+		UE_LOG(LogLearning, Error, TEXT("%s: Setup not run."), *GetName());
+		return;
+	}
+
+	AgentType->EncodeObservations();
+	EncodeActions();
+	AgentType->DecodeActions();
+}
diff --git a/Engine/Plugins/Experimental/LearningAgents/Source/LearningAgents/Private/LearningAgentsCritic.cpp b/Engine/Plugins/Experimental/LearningAgents/Source/LearningAgents/Private/LearningAgentsCritic.cpp
--- a/Engine/Plugins/Experimental/LearningAgents/Source/LearningAgents/Private/LearningAgentsCritic.cpp
+++ b/Engine/Plugins/Experimental/LearningAgents/Source/LearningAgents/Private/LearningAgentsCritic.cpp
@@ -13,62 +13,6 @@
 #include "Misc/FileHelper.h"
 #include "Misc/Paths.h"
 #include "GameFramework/Actor.h"
-#include "VisualLogger/VisualLogger.h"
-
-#define UE_LEARNING_AGENTS_VLOG_STRING(Owner, Category, Verbosity, Location, Color, Format, ...) \
-	UE_VLOG_LOCATION(Owner, Category, Verbosity, Location, 0.0f, Color, Format, ##__VA_ARGS__)
-
-namespace UE::Learning::Agents::Critic::Private
-{
-	static inline FString FloatArrayToString(const TLearningArrayView<1, const float> Array)
-	{
-		const int32 ItemNum = Array.Num();
-		const int32 MaxItemNum = 32;
-		const int32 OutputItemNum = FMath::Min(ItemNum, MaxItemNum);
-
-		FString Output = TEXT("[");
-
-		for (int32 Idx = 0; Idx < OutputItemNum; Idx++)
-		{
-			Output.Appendf(TEXT("% 6.3f"), Array[Idx]);
-
-			if (Idx < OutputItemNum - 1)
-			{
-				Output += TEXT(" ");
-			}
-			else if (Idx == MaxItemNum - 1)
-			{
-				Output += TEXT("...");
-			}
-		}
-
-		Output += TEXT("]");
-
-		return Output;
-	}
-
-	static inline FString FloatArrayToStatsString(const TLearningArrayView<1, const float> Array)
-	{
-		const int32 ItemNum = Array.Num();
-
-		float Min = +FLT_MAX, Max = -FLT_MAX, Mean = 0.0f;
-		for (int32 Idx = 0; Idx < ItemNum; Idx++)
-		{
-			Min = FMath::Min(Min, Array[Idx]);
-			Max = FMath::Max(Max, Array[Idx]);
-			Mean += Array[Idx] / ItemNum;
-		}
-
-		float Var = 0.0f;
-		for (int32 Idx = 0; Idx < ItemNum; Idx++)
-		{
-			Var += FMath::Square(Array[Idx] - Mean) / ItemNum;
-		}
-
-		return FString::Printf(TEXT("[% 6.3f/% 6.3f/% 6.3f/% 6.3f]"),
-			Min, Max, Mean, FMath::Sqrt(Var));
-	}
-}
 
 ULearningAgentsCritic::ULearningAgentsCritic() : ULearningAgentsManagerComponent() {}
 ULearningAgentsCritic::ULearningAgentsCritic(FVTableHelper& Helper) : ULearningAgentsCritic() {}
@@ -78,7 +22,7 @@ void ULearningAgentsCritic::SetupCritic(ALearningAgentsManager* InAgentManager,
 {
 	if (IsSetup())
 	{
-		UE_LOG(LogLearning, Error, TEXT("%s: Setup already performed!"), *GetName());
+		UE_LOG(LogLearning, Error, TEXT("%s: Setup already run!"), *GetName());
 		return;
 	}
 
@@ -90,7 +34,7 @@ void ULearningAgentsCritic::SetupCritic(ALearningAgentsManager* InAgentManager,
 
 	if (!InAgentManager->IsManagerSetup())
 	{
-		UE_LOG(LogLearning, Error, TEXT("%s's SetupManager() must be run before %s can be setup."), *InAgentManager->GetName(), *GetName());
+		UE_LOG(LogLearning, Error, TEXT("%s: %s's SetupManager must be run before it can be used."), *GetName(), *InAgentManager->GetName());
 		return;
 	}
 
@@ -146,7 +90,7 @@ void ULearningAgentsCritic::LoadCriticFromSnapshot(const FDirectoryPath& Directo
 {
 	if (!IsSetup())
 	{
-		UE_LOG(LogLearning, Error, TEXT("Critic setup must be run before network can be loaded."));
+		UE_LOG(LogLearning, Error, TEXT("%s: Setup not run."), *GetName());
 		return;
 	}
 
@@ -162,7 +106,7 @@ void ULearningAgentsCritic::LoadCriticFromSnapshot(const FDirectoryPath& Directo
 
 		if (NetworkData.Num() != TotalByteNum)
 		{
-			UE_LOG(LogLearning, Error, TEXT("Failed to load network from file %s. File size incorrect."), *FilePath);
+			UE_LOG(LogLearning, Error, TEXT("%s: Failed to load network from file %s. File size incorrect."), *GetName(), *FilePath);
 			return;
 		}
 
@@ -170,7 +114,7 @@ void ULearningAgentsCritic::LoadCriticFromSnapshot(const FDirectoryPath& Directo
 	}
 	else
 	{
-		UE_LOG(LogLearning, Error, TEXT("Failed to load network. File not found: %s"), *FilePath);
+		UE_LOG(LogLearning, Error, TEXT("%s: Failed to load network. File not found: %s"), *GetName(), *FilePath);
 	}
 }
 
@@ -178,7 +122,7 @@ void ULearningAgentsCritic::SaveCriticToSnapshot(const FDirectoryPath& Directory
 {
 	if (!IsSetup())
 	{
-		UE_LOG(LogLearning, Error, TEXT("Critic setup must be run before network can be saved."));
+		UE_LOG(LogLearning, Error, TEXT("%s: Setup not run."), *GetName());
 		return;
 	}
 
@@ -194,7 +138,7 @@ void ULearningAgentsCritic::SaveCriticToSnapshot(const FDirectoryPath& Directory
 	const FString FilePath = Directory.Path + FGenericPlatformMisc::GetDefaultPathSeparator() + Filename;
 	if (!FFileHelper::SaveArrayToFile(NetworkData, *FilePath))
 	{
-		UE_LOG(LogLearning, Error, TEXT("Failed to save network to file: %s"), *FilePath);
+		UE_LOG(LogLearning, Error, TEXT("%s: Failed to save network to file: %s"), *GetName(), *FilePath);
 	}
 }
 
@@ -202,20 +146,20 @@ void ULearningAgentsCritic::LoadCriticFromAsset(const ULearningAgentsNeuralNetwo
 {
 	if (!IsSetup())
 	{
-		UE_LOG(LogLearning, Error, TEXT("Critic setup must be run before network can be loaded."));
+		UE_LOG(LogLearning, Error, TEXT("%s: Setup not run."), *GetName());
 		return;
 	}
 
 	if (!NeuralNetworkAsset || !NeuralNetworkAsset->NeuralNetwork)
 	{
-		UE_LOG(LogLearning, Error, TEXT("Cannot load critic from invalid asset."));
+		UE_LOG(LogLearning, Error, TEXT("%s: Asset is invalid."), *GetName());
 		return;
 	}
 
 	if (NeuralNetworkAsset->NeuralNetwork->GetInputNum() != Network->NeuralNetwork->GetInputNum() ||
 		NeuralNetworkAsset->NeuralNetwork->GetOutputNum() != Network->NeuralNetwork->GetOutputNum())
 	{
-		UE_LOG(LogLearning, Error, TEXT("Failed to load critic from asset. Network Asset inputs and outputs don't match."));
+		UE_LOG(LogLearning, Error, TEXT("%s: Failed to load network from asset. Inputs and outputs don't match."), *GetName());
 		return;
 	}
 
@@ -226,13 +170,13 @@ void ULearningAgentsCritic::SaveCriticToAsset(ULearningAgentsNeuralNetwork* Neur
 {
 	if (!IsSetup())
 	{
-		UE_LOG(LogLearning, Error, TEXT("Critic setup must be run before network can be saved."));
+		UE_LOG(LogLearning, Error, TEXT("%s: Setup not run."), *GetName());
 		return;
 	}
 
 	if (!NeuralNetworkAsset)
 	{
-		UE_LOG(LogLearning, Error, TEXT("Cannot save critic to invalid asset."));
+		UE_LOG(LogLearning, Error, TEXT("%s: Asset is invalid."), *GetName());
 		return;
 	}
 
@@ -265,13 +209,13 @@ void ULearningAgentsCritic::EvaluateCritic()
 
 	if (!IsSetup())
 	{
-		UE_LOG(LogLearning, Error, TEXT("Setup must be run before the critic can be evaluated."));
+		UE_LOG(LogLearning, Error, TEXT("%s: Setup not run."), *GetName());
 		return;
 	}
 
 	CriticObject->Evaluate(AddedAgentSet);
 
-#if ENABLE_VISUAL_LOG
+#if UE_LEARNING_AGENTS_ENABLE_VISUAL_LOG
 	VisualLog(AddedAgentSet);
 #endif
 }
@@ -280,13 +224,13 @@ float ULearningAgentsCritic::GetEstimatedDiscountedReturn(const int32 AgentId) c
 {
 	if (!IsSetup())
 	{
-		UE_LOG(LogLearning, Error, TEXT("Setup must be run before the critic can get the estimated discounted return."));
+		UE_LOG(LogLearning, Error, TEXT("%s: Setup not run."), *GetName());
 		return 0.0f;
 	}
 
 	if (!HasAgent(AgentId))
 	{
-		UE_LOG(LogLearning, Error, TEXT("Unable to get estimate for agent - AgentId %d not found in the added agents set."), AgentId);
+		UE_LOG(LogLearning, Error, TEXT("%s: AgentId %d not found in the agents set."), *GetName(), AgentId);
 		return 0.0f;
 	}
 
@@ -295,7 +239,7 @@ float ULearningAgentsCritic::GetEstimatedDiscountedReturn(const int32 AgentId) c
 	return CriticOutputView[AgentId];
 }
 
-#if ENABLE_VISUAL_LOG
+#if UE_LEARNING_AGENTS_ENABLE_VISUAL_LOG
 void ULearningAgentsCritic::VisualLog(const UE::Learning::FIndexSet Instances) const
 {
 	UE_LEARNING_TRACE_CPUPROFILER_EVENT_SCOPE(ULearningAgentsCritic::VisualLog);
@@ -307,20 +251,15 @@ void ULearningAgentsCritic::VisualLog(const UE::Learning::FIndexSet Instances) c
 	{
 		if (const AActor* Actor = Cast<AActor>(AgentType->GetAgent(Instance)))
 		{
-			const FString InputArrayString = UE::Learning::Agents::Critic::Private::FloatArrayToString(InputView[Instance]);
-			const FString InputStatsString = UE::Learning::Agents::Critic::Private::FloatArrayToStatsString(InputView[Instance]);
-
 			UE_LEARNING_AGENTS_VLOG_STRING(this, LogLearning, Display,
 				Actor->GetActorLocation(),
 				VisualLogColor.ToFColor(true),
 				TEXT("Agent %i\nInput: %s\nInput Stats (Min/Max/Mean/Std): %s\nOutput: [% 6.3f]"),
 				Instance,
-				*InputArrayString,
-				*InputStatsString,
+				*UE::Learning::Array::FormatFloat(InputView[Instance]),
+				*UE::Learning::Agents::Debug::FloatArrayToStatsString(InputView[Instance]),
 				OutputView[Instance]);
 		}
 	}
 }
 #endif
-
-#undef UE_LEARNING_AGENTS_VLOG_STRING
diff --git a/Engine/Plugins/Experimental/LearningAgents/Source/LearningAgents/Private/LearningAgentsDebug.cpp b/Engine/Plugins/Experimental/LearningAgents/Source/LearningAgents/Private/LearningAgentsDebug.cpp
new file mode 100644
--- /dev/null
+++ b/Engine/Plugins/Experimental/LearningAgents/Source/LearningAgents/Private/LearningAgentsDebug.cpp
@@ -0,0 +1,33 @@
+#include "LearningAgentsDebug.h"
+
+namespace UE::Learning::Agents::Debug
+{
+	FVector GridOffsetForIndex(const int32 Idx, const int32 Num, const float Spacing, const float Height)
+	{
+		const int32 RowNum = FMath::Max(FMath::CeilToInt(FMath::Sqrt((float)Num)), 1);
+
+		return Spacing * (FVector((Idx / RowNum), (Idx % RowNum), 0.0f) - FVector(RowNum - 1, RowNum - 1, 0.0f) / 2.0f) + FVector(0.0f, 0.0f, Height);
+	}
+
+	FString FloatArrayToStatsString(const TLearningArrayView<1, const float> Array)
+	{
+		const int32 ItemNum = Array.Num();
+
+		float Min = +FLT_MAX, Max = -FLT_MAX, Mean = 0.0f;
+		for (int32 Idx = 0; Idx < ItemNum; Idx++)
+		{
+			Min = FMath::Min(Min, Array[Idx]);
+			Max = FMath::Max(Max, Array[Idx]);
+			Mean += Array[Idx] / ItemNum;
+		}
+
+		float Var = 0.0f;
+		for (int32 Idx = 0; Idx < ItemNum; Idx++)
+		{
+			Var += FMath::Square(Array[Idx] - Mean) / ItemNum;
+		}
+
+		return FString::Printf(TEXT("[% 6.3f/% 6.3f/% 6.3f/% 6.3f]"),
+			Min, Max, Mean, FMath::Sqrt(Var));
+	}
+}
diff --git a/Engine/Plugins/Experimental/LearningAgents/Source/LearningAgents/Private/LearningAgentsHelpers.cpp b/Engine/Plugins/Experimental/LearningAgents/Source/LearningAgents/Private/LearningAgentsHelpers.cpp
--- a/Engine/Plugins/Experimental/LearningAgents/Source/LearningAgents/Private/LearningAgentsHelpers.cpp
+++ b/Engine/Plugins/Experimental/LearningAgents/Source/LearningAgents/Private/LearningAgentsHelpers.cpp
@@ -2,90 +2,129 @@
 
 #include "LearningAgentsHelpers.h"
 
+#include "LearningAgentsManager.h"
+#include "LearningAgentsManagerComponent.h"
 #include "LearningLog.h"
 
 #include "Components/MeshComponent.h"
 #include "Components/SplineComponent.h"
+#include "Engine/World.h"
 
 //------------------------------------------------------------------
 
-FTransform ULearningAgentsHelpers::ProjectTransformOntoGroundPlane(const FTransform Transform, const FVector LocalForwardVector)
+USplineComponentHelper* USplineComponentHelper::AddSplineComponentHelper(ULearningAgentsManagerComponent* InManagerComponent, const FName Name)
 {
-	FVector Position = Transform.GetLocation();
-	Position.Z = 0.0f;
-
-	const FVector Direction = (FVector(1.0f, 1.0f, 0.0f) * Transform.TransformVectorNoScale(LocalForwardVector)).GetSafeNormal(UE_SMALL_NUMBER, FVector::ForwardVector);
-
-	return FTransform(FQuat::FindBetweenNormals(FVector::ForwardVector, Direction), Position);
-}
-
-void ULearningAgentsHelpers::ProjectPositionRotationOntoGroundPlane(FVector& OutPosition, FRotator& OutRotation, const FVector InPosition, const FRotator InRotation, const FVector LocalForwardVector)
-{
-	OutPosition = InPosition;
-	OutPosition.Z = 0.0f;
+	if (!InManagerComponent)
+	{
+		UE_LOG(LogLearning, Error, TEXT("InManagerComponent is nullptr."));
+		return nullptr;
+	}
 
-	const FVector Direction = (FVector(1.0f, 1.0f, 0.0f) * InRotation.RotateVector(LocalForwardVector)).GetSafeNormal(UE_SMALL_NUMBER, FVector::ForwardVector);
-	OutRotation = FQuat::FindBetweenNormals(FVector::ForwardVector, Direction).Rotator();
+	USplineComponentHelper* Helper = NewObject<USplineComponentHelper>(InManagerComponent, Name);
+	Helper->ManagerComponent = InManagerComponent;
+	InManagerComponent->AddHelper(Helper);
+	return Helper;
 }
 
-//------------------------------------------------------------------
-
-void ULearningAgentsHelpers::GetMeshBonePositions(const UMeshComponent* MeshComponent, const TArray<FName>& BoneNames, TArray<FVector>& OutBonePositions)
+float USplineComponentHelper::GetProportionAlongSpline(const int32 AgentId, const USplineComponent* SplineComponent, const float DistanceAlongSpline) const
 {
-	if (!MeshComponent)
+	if (!ManagerComponent->HasAgent(AgentId))
 	{
-		UE_LOG(LogLearning, Warning, TEXT("Mesh Component was nullptr."));
-		return;
+		UE_LOG(LogLearning, Error, TEXT("%s: AgentId %d not found in the agents set."), *GetName(), AgentId);
+		return 0.0f;
 	}
 
-	const int32 BoneNum = BoneNames.Num();
+	if (!SplineComponent)
+	{
+		UE_LOG(LogLearning, Error, TEXT("%s: SplineComponent is nullptr."), *GetName());
+		return 0.0f;
+	}
 
-	OutBonePositions.SetNumUninitialized(BoneNum);
+	const float Proportion = FMath::Clamp(DistanceAlongSpline / FMath::Max(SplineComponent->GetSplineLength(), UE_SMALL_NUMBER), 0.0f, 1.0f);
 
-	for (int32 BoneIdx = 0; BoneIdx < BoneNum; BoneIdx++)
+#if UE_LEARNING_AGENTS_ENABLE_VISUAL_LOG
 	{
-		OutBonePositions[BoneIdx] = MeshComponent->GetSocketLocation(BoneNames[BoneIdx]);
+		const FVector Position = SplineComponent->GetLocationAtDistanceAlongSpline(DistanceAlongSpline, ESplineCoordinateSpace::World);
+
+		UE_LEARNING_AGENTS_VLOG_LOCATION(this, LogLearning, Display,
+			Position,
+			10.0f,
+			VisualLogColor.ToFColor(true),
+			TEXT("GetProportionAlongSpline\nAgent %i\nDistance Along Spline: [% 6.1f]\nProportion: [% 6.3f]"),
+			AgentId,
+			DistanceAlongSpline,
+			Proportion);
 	}
-}
+#endif
 
-//------------------------------------------------------------------
+	return Proportion;
+}
 
-float ULearningAgentsHelpers::GetProportionAlongSpline(const USplineComponent* SplineComponent, const float DistanceAlongSpline)
+float USplineComponentHelper::GetProportionAlongSplineAsAngle(const int32 AgentId, const USplineComponent* SplineComponent, const float DistanceAlongSpline) const
 {
-	if (!SplineComponent)
+	if (!ManagerComponent->HasAgent(AgentId))
 	{
-		UE_LOG(LogLearning, Warning, TEXT("Spline Component was nullptr."));
+		UE_LOG(LogLearning, Error, TEXT("%s: AgentId %d not found in the agents set."), *GetName(), AgentId);
 		return 0.0f;
 	}
 
-	return FMath::Clamp(DistanceAlongSpline / FMath::Max(SplineComponent->GetSplineLength(), UE_SMALL_NUMBER), 0.0f, 1.0f);
-}
-
-float ULearningAgentsHelpers::GetProportionAlongSplineAsAngle(const USplineComponent* SplineComponent, const float DistanceAlongSpline)
-{
 	if (!SplineComponent)
 	{
-		UE_LOG(LogLearning, Warning, TEXT("Spline Component was nullptr."));
+		UE_LOG(LogLearning, Error, TEXT("%s: SplineComponent is nullptr."), *GetName());
 		return 0.0f;
 	}
 
 	if (!SplineComponent->IsClosedLoop())
 	{
-		UE_LOG(LogLearning, Warning, TEXT("Getting proportion along spline as angle, but spline is not closed loop. Consider using GetPropotionAlongSpline instead."));
+		UE_LOG(LogLearning, Warning, TEXT("%s: Getting proportion along spline as angle, but spline is not closed loop. Consider using GetPropotionAlongSpline instead."), *GetName());
 	}
 
 	const float TotalDistance = SplineComponent->GetSplineLength();
+	const float WrapDistance = FMath::Wrap(DistanceAlongSpline, 0.0f, TotalDistance);
+	const float Proportion = WrapDistance / FMath::Max(TotalDistance, UE_SMALL_NUMBER);
+	const float Angle = FMath::Wrap(UE_TWO_PI * Proportion, -UE_PI, UE_PI);
 
-	const float Ratio = UE_TWO_PI * (FMath::Wrap(DistanceAlongSpline, 0.0f, TotalDistance) / FMath::Max(TotalDistance, UE_SMALL_NUMBER)) - UE_PI;
+#if UE_LEARNING_AGENTS_ENABLE_VISUAL_LOG
+	{
+		const FVector Position = SplineComponent->GetLocationAtDistanceAlongSpline(WrapDistance, ESplineCoordinateSpace::World);
+
+		UE_LEARNING_AGENTS_VLOG_ANGLE(this, LogLearning, Display,
+			Angle,
+			0.0f,
+			Position,
+			50.0f,
+			VisualLogColor.ToFColor(true),
+			TEXT("GetProportionAlongSplineAsAngle\nAgent %i\nDistance Along Spline: [% 6.1f]\nProportion: [% 6.3f]\nAngle: [% 6.1f]"),
+			AgentId,
+			WrapDistance,
+			Proportion,
+			FMath::RadiansToDegrees(Angle));
+	}
+#endif
 
-	return FMath::RadiansToDegrees(Ratio);
+	return FMath::RadiansToDegrees(Angle);
 }
 
-void ULearningAgentsHelpers::GetPositionsAlongSpline(TArray<FVector>& OutPositions, const USplineComponent* SplineComponent, const int32 PositionNum, const float StartDistanceAlongSpline, const float StopDistanceAlongSpline, const ESplineCoordinateSpace::Type CoordinateSpace)
+void USplineComponentHelper::GetPositionsAlongSpline(TArray<FVector>& OutPositions, const int32 AgentId, const USplineComponent* SplineComponent, const int32 PositionNum, const float StartDistanceAlongSpline, const float StopDistanceAlongSpline, const ESplineCoordinateSpace::Type CoordinateSpace) const
 {
+	if (!ManagerComponent->HasAgent(AgentId))
+	{
+		UE_LOG(LogLearning, Error, TEXT("%s: AgentId %d not found in the agents set."), *GetName(), AgentId);
+		OutPositions.Empty();
+		return;
+	}
+
 	if (!SplineComponent)
 	{
-		UE_LOG(LogLearning, Warning, TEXT("Spline Component was nullptr."));
+		UE_LOG(LogLearning, Error, TEXT("%s: SplineComponent is nullptr."), *GetName());
+		OutPositions.Empty();
+		return;
+	}
+
+	if (PositionNum < 1)
+	{
+		UE_LOG(LogLearning, Warning, TEXT("%s: Expected at least one position along the spline got %i."), *GetName(), PositionNum);
+		OutPositions.Empty();
 		return;
 	}
 
@@ -106,18 +145,49 @@ void ULearningAgentsHelpers::GetPositionsAlongSpline(TArray<FVector>& OutPositio
 
 		OutPositions[PositionIdx] = SplineComponent->GetLocationAtDistanceAlongSpline(PositionDistance, CoordinateSpace);
 	}
+
+#if UE_LEARNING_AGENTS_ENABLE_VISUAL_LOG
+	{
+		for (int32 PositionIdx = 0; PositionIdx < PositionNum; PositionIdx++)
+		{
+			UE_LEARNING_AGENTS_VLOG_LOCATION(this, LogLearning, Display,
+				OutPositions[PositionIdx],
+				2.5f,
+				VisualLogColor.ToFColor(true),
+				TEXT("GetPositionsAlongSpline\nAgent: %i\nPosition: [% 6.1f % 6.1f % 6.1f]"),
+				AgentId,
+				OutPositions[PositionIdx].X, OutPositions[PositionIdx].Y, OutPositions[PositionIdx].Z);
+		}
+	}
+#endif
 }
 
-float ULearningAgentsHelpers::GetVelocityAlongSpline(const USplineComponent* SplineComponent, const FVector Position, const FVector Velocity, const float FiniteDifferenceDelta, const ESplineCoordinateSpace::Type CoordinateSpace)
+float USplineComponentHelper::GetVelocityAlongSpline(const int32 AgentId, const USplineComponent* SplineComponent, const FVector Position, const FVector Velocity, const float FiniteDifferenceDelta, const ESplineCoordinateSpace::Type CoordinateSpace) const
 {
+	if (!ManagerComponent->HasAgent(AgentId))
+	{
+		UE_LOG(LogLearning, Error, TEXT("%s: AgentId %d not found in the agents set."), *GetName(), AgentId);
+		return 0.0f;
+	}
+
 	if (!SplineComponent)
 	{
-		UE_LOG(LogLearning, Warning, TEXT("Spline Component was nullptr."));
+		UE_LOG(LogLearning, Error, TEXT("%s: SplineComponent is nullptr."), *GetName());
 		return 0.0f;
 	}
 
-	float Distance0 = SplineComponent->GetDistanceAlongSplineAtLocation(Position, CoordinateSpace);
-	float Distance1 = SplineComponent->GetDistanceAlongSplineAtLocation(Position + FiniteDifferenceDelta * Velocity, CoordinateSpace);
+	float FiniteDiff = FiniteDifferenceDelta;
+
+	if (FiniteDiff < UE_KINDA_SMALL_NUMBER)
+	{
+		UE_LOG(LogLearning, Warning, TEXT("%s: FiniteDifferenceDelta is too small (%6.4f). Clamping to %6.4f."), *GetName(), FiniteDifferenceDelta, UE_KINDA_SMALL_NUMBER);
+		FiniteDiff = UE_KINDA_SMALL_NUMBER;
+	}
+
+	const float RawDistance0 = SplineComponent->GetDistanceAlongSplineAtLocation(Position, CoordinateSpace);
+	const float RawDistance1 = SplineComponent->GetDistanceAlongSplineAtLocation(Position + FiniteDiff * Velocity.GetSafeNormal(), CoordinateSpace);
+
+	float Distance0 = RawDistance0, Distance1 = RawDistance1;
 
 	if (SplineComponent->IsClosedLoop())
 	{
@@ -133,6 +203,368 @@ float ULearningAgentsHelpers::GetVelocityAlongSpline(const USplineComponent* Spl
 		}
 	}
 
-	return (Distance1 - Distance0) / FMath::Max(FiniteDifferenceDelta, UE_SMALL_NUMBER);
+	const float SplineVelocity = ((Distance1 - Distance0) / FiniteDiff) * Velocity.Length();
+
+#if UE_LEARNING_AGENTS_ENABLE_VISUAL_LOG
+	{
+		const FVector VelocityEnd = Position + Velocity;
+		const FVector VelocityDiff = Position + FiniteDiff * Velocity.GetSafeNormal();
+
+		UE_LEARNING_AGENTS_VLOG_ARROW(this, LogLearning, Display,
+			Position,
+			VelocityEnd,
+			VisualLogColor.ToFColor(true),
+			TEXT("GetVelocityAlongSpline\nAgent %i\nPosition [% 6.1f % 6.1f % 6.1f]\nVelocity: [% 6.1f % 6.1f % 6.1f]\nMagnitude: [% 6.1f]"),
+			AgentId,
+			Position.X, Position.Y, Position.Z,
+			Velocity.X, Velocity.Y, Velocity.Z,
+			Velocity.Length());
+
+		UE_LEARNING_AGENTS_VLOG_SEGMENT(this, LogLearning, Display,
+			Position,
+			VelocityDiff,
+			VisualLogColor.ToFColor(true),
+			TEXT(""));
+
+		const FVector Position0 = SplineComponent->GetLocationAtDistanceAlongSpline(RawDistance0, CoordinateSpace);
+		const FVector Position1 = SplineComponent->GetLocationAtDistanceAlongSpline(RawDistance1, CoordinateSpace);
+
+		UE_LEARNING_AGENTS_VLOG_SEGMENT(this, LogLearning, Display,
+			Position0,
+			Position1,
+			VisualLogColor.ToFColor(true),
+			TEXT(""));
+
+		UE_LEARNING_AGENTS_VLOG_SEGMENT(this, LogLearning, Display,
+			Position,
+			Position0,
+			VisualLogColor.ToFColor(true),
+			TEXT(""));
+
+		UE_LEARNING_AGENTS_VLOG_SEGMENT(this, LogLearning, Display,
+			VelocityDiff,
+			Position1,
+			VisualLogColor.ToFColor(true),
+			TEXT(""));
+
+		UE_LEARNING_AGENTS_VLOG_ARROW(this, LogLearning, Display,
+			Position0,
+			Position0 + (Position1 - Position0).GetSafeNormal() * FMath::Abs(SplineVelocity),
+			VisualLogColor.ToFColor(true),
+			TEXT("GetVelocityAlongSpline\nAgent %i\nSpline Velocity: [% 6.1f]"),
+			AgentId,
+			SplineVelocity);
+	}
+#endif
+
+	return SplineVelocity;
+}
+
+//------------------------------------------------------------------
+
+UProjectionHelper* UProjectionHelper::AddProjectionHelper(ULearningAgentsManagerComponent* InManagerComponent, const FName Name)
+{
+	if (!InManagerComponent)
+	{
+		UE_LOG(LogLearning, Error, TEXT("InManagerComponent is nullptr."));
+		return nullptr;
+	}
+
+	UProjectionHelper* Helper = NewObject<UProjectionHelper>(InManagerComponent, Name);
+	Helper->ManagerComponent = InManagerComponent;
+	InManagerComponent->AddHelper(Helper);
+	return Helper;
+}
+
+FTransform UProjectionHelper::ProjectTransformOntoGroundPlane(const int32 AgentId, const FTransform Transform, const FVector LocalForwardVector) const
+{
+	if (!ManagerComponent->HasAgent(AgentId))
+	{
+		UE_LOG(LogLearning, Error, TEXT("%s: AgentId %d not found in the agents set."), *GetName(), AgentId);
+		return FTransform::Identity;
+	}
+
+	FVector Position = Transform.GetLocation();
+	Position.Z = 0.0f;
+
+	const FVector Direction = (FVector(1.0f, 1.0f, 0.0f) * Transform.TransformVectorNoScale(LocalForwardVector)).GetSafeNormal(UE_SMALL_NUMBER, FVector::ForwardVector);
+
+	const FTransform Projected = FTransform(FQuat::FindBetweenNormals(FVector::ForwardVector, Direction), Position);
+
+#if UE_LEARNING_AGENTS_ENABLE_VISUAL_LOG
+	{
+		UE_LEARNING_AGENTS_VLOG_TRANSFORM(this, LogLearning, Display,
+			Transform.GetLocation(),
+			Transform.GetRotation(),
+			VisualLogColor.ToFColor(true),
+			TEXT("ProjectTransformOntoGroundPlane\nAgent %i\nTransform\nLocal Forward: [% 6.3f % 6.3f % 6.3f]"),
+			LocalForwardVector.X, LocalForwardVector.Y, LocalForwardVector.Z);
+
+		UE_LEARNING_AGENTS_VLOG_ARROW(this, LogLearning, Display,
+			Transform.GetLocation(),
+			Transform.GetLocation() + 25.0f * Transform.GetRotation().RotateVector(LocalForwardVector),
+			VisualLogColor.ToFColor(true),
+			TEXT(""));
+
+		UE_LEARNING_AGENTS_VLOG_SEGMENT(this, LogLearning, Display,
+			Transform.GetLocation(),
+			Projected.GetLocation(),
+			VisualLogColor.ToFColor(true),
+			TEXT(""));
+
+		UE_LEARNING_AGENTS_VLOG_TRANSFORM(this, LogLearning, Display,
+			Projected.GetLocation(),
+			Projected.GetRotation(),
+			VisualLogColor.ToFColor(true),
+			TEXT("ProjectTransformOntoGroundPlane\nAgent %i\nProjected"));
+	}
+#endif
+
+	return Projected;
+}
+
+void UProjectionHelper::ProjectPositionRotationOntoGroundPlane(FVector& OutPosition, FRotator& OutRotation, const int32 AgentId, const FVector InPosition, const FRotator InRotation, const FVector LocalForwardVector) const
+{
+	if (!ManagerComponent->HasAgent(AgentId))
+	{
+		UE_LOG(LogLearning, Error, TEXT("%s: AgentId %d not found in the agents set."), *GetName(), AgentId);
+		OutPosition = FVector::ZeroVector;
+		OutRotation = FRotator::ZeroRotator;
+		return;
+	}
+
+	OutPosition = InPosition;
+	OutPosition.Z = 0.0f;
+
+	const FVector Direction = (FVector(1.0f, 1.0f, 0.0f) * InRotation.RotateVector(LocalForwardVector)).GetSafeNormal(UE_SMALL_NUMBER, FVector::ForwardVector);
+	OutRotation = FQuat::FindBetweenNormals(FVector::ForwardVector, Direction).Rotator();
+
+#if UE_LEARNING_AGENTS_ENABLE_VISUAL_LOG
+	{
+		UE_LEARNING_AGENTS_VLOG_TRANSFORM(this, LogLearning, Display,
+			InPosition,
+			InRotation,
+			VisualLogColor.ToFColor(true),
+			TEXT("ProjectPositionRotationOntoGroundPlane\nAgent %i\nTransform\nLocal Forward: [% 6.3f % 6.3f % 6.3f]"),
+			LocalForwardVector.X, LocalForwardVector.Y, LocalForwardVector.Z);
+
+		UE_LEARNING_AGENTS_VLOG_ARROW(this, LogLearning, Display,
+			InPosition,
+			InPosition + 25.0f * InRotation.RotateVector(LocalForwardVector),
+			VisualLogColor.ToFColor(true),
+			TEXT(""));
+
+		UE_LEARNING_AGENTS_VLOG_SEGMENT(this, LogLearning, Display,
+			InPosition,
+			OutPosition,
+			VisualLogColor.ToFColor(true),
+			TEXT(""));
+
+		UE_LEARNING_AGENTS_VLOG_TRANSFORM(this, LogLearning, Display,
+			OutPosition,
+			OutRotation,
+			VisualLogColor.ToFColor(true),
+			TEXT("ProjectPositionRotationOntoGroundPlane\nAgent %i\nProjected"));
+	}
+#endif
+}
+
+//------------------------------------------------------------------
+
+UMeshComponentHelper* UMeshComponentHelper::AddMeshComponentHelper(ULearningAgentsManagerComponent* InManagerComponent, const FName Name)
+{
+	if (!InManagerComponent)
+	{
+		UE_LOG(LogLearning, Error, TEXT("InManagerComponent is nullptr."));
+		return nullptr;
+	}
+
+	UMeshComponentHelper* Helper = NewObject<UMeshComponentHelper>(InManagerComponent, Name);
+	Helper->ManagerComponent = InManagerComponent;
+	InManagerComponent->AddHelper(Helper);
+	return Helper;
+}
+
+void UMeshComponentHelper::GetMeshBonePositions(TArray<FVector>& OutBonePositions, const int32 AgentId, const UMeshComponent* MeshComponent, const TArray<FName>& BoneNames) const
+{
+	if (!ManagerComponent->HasAgent(AgentId))
+	{
+		UE_LOG(LogLearning, Error, TEXT("%s: AgentId %d not found in the agents set."), *GetName(), AgentId);
+		OutBonePositions.Empty();
+		return;
+	}
+
+	if (!MeshComponent)
+	{
+		UE_LOG(LogLearning, Error, TEXT("%s: Mesh Component is nullptr."), *GetName());
+		OutBonePositions.Empty();
+		return;
+	}
+
+	if (BoneNames.IsEmpty())
+	{
+		UE_LOG(LogLearning, Warning, TEXT("%s: No bone names provided!"), *GetName());
+		OutBonePositions.Empty();
+		return;
+	}
+
+	const int32 BoneNum = BoneNames.Num();
+
+	OutBonePositions.SetNumUninitialized(BoneNum);
+
+	for (int32 BoneIdx = 0; BoneIdx < BoneNum; BoneIdx++)
+	{
+		if (MeshComponent->DoesSocketExist(BoneNames[BoneIdx]))
+		{
+			OutBonePositions[BoneIdx] = MeshComponent->GetSocketLocation(BoneNames[BoneIdx]);
+		}
+		else
+		{
+			OutBonePositions[BoneIdx] = FVector::ZeroVector;
+			UE_LOG(LogLearning, Warning, TEXT("%s: Bone \"%s\" does not exist!"), *GetName(), *BoneNames[BoneIdx].ToString());
+		}
+	}
+
+#if ENABLE_VISUAL_LOG
+	{
+		for (int32 BoneIdx = 0; BoneIdx < BoneNum; BoneIdx++)
+		{
+			UE_LEARNING_AGENTS_VLOG_LOCATION(this, LogLearning, Display,
+				OutBonePositions[BoneIdx],
+				2.5f,
+				VisualLogColor.ToFColor(true),
+				TEXT("Agent %i\nBone: \"%s\"\nPosition: [% 6.1f % 6.1f % 6.1f]"),
+				AgentId,
+				*BoneNames[BoneIdx].ToString(),
+				OutBonePositions[BoneIdx].X, OutBonePositions[BoneIdx].Y, OutBonePositions[BoneIdx].Z);
+		}
+	}
+#endif
 }
 
+//------------------------------------------------------------------
+
+URayCastHelper* URayCastHelper::AddRayCastHelper(ULearningAgentsManagerComponent* InManagerComponent, const FName Name)
+{
+	if (!InManagerComponent)
+	{
+		UE_LOG(LogLearning, Error, TEXT("InAgentType is nullptr."));
+		return nullptr;
+	}
+
+	URayCastHelper* Helper = NewObject<URayCastHelper>(InManagerComponent, Name);
+	Helper->ManagerComponent = InManagerComponent;
+	InManagerComponent->AddHelper(Helper);
+	return Helper;
+}
+
+void URayCastHelper::RayCastGridHeights(
+	TArray<float>& OutHeights, 
+	const int32 AgentId,
+	const FVector Position, 
+	const FRotator Rotation, 
+	const int32 RowNum, 
+	const int32 ColNum, 
+	const float RowWidth, 
+	const float ColWidth, 
+	const float MaxHeight, 
+	const float MinHeight,
+	const ECollisionChannel CollisionChannel) const
+{
+	if (!ManagerComponent->HasAgent(AgentId))
+	{
+		UE_LOG(LogLearning, Error, TEXT("%s: AgentId %d not found in the agents set."), *GetName(), AgentId);
+		OutHeights.Empty();
+		return;
+	}
+
+	if (RowNum < 1 || ColNum < 1)
+	{
+		UE_LOG(LogLearning, Error, TEXT("%s: Need at least 1 row and 1 col in grid."), *GetName());
+		OutHeights.Empty();
+		return;
+	}
+
+	OutHeights.SetNumUninitialized(ColNum * RowNum);
+
+	TArray<FHitResult, TInlineAllocator<25>> TraceHits;
+	TraceHits.SetNum(ColNum * RowNum);
+
+	const FRotator YawRotator(0.0f, Rotation.Yaw, 0.0f);
+
+	for (int32 RowIdx = 0; RowIdx < RowNum; RowIdx++)
+	{
+		for (int32 ColIdx = 0; ColIdx < ColNum; ColIdx++)
+		{
+			const FVector Start = Position + YawRotator.RotateVector(FVector(
+				RowWidth * (((float)RowIdx) / (RowNum - 1) - 0.5f),
+				ColWidth * (((float)ColIdx) / (ColNum - 1) - 0.5f),
+				MaxHeight));
+
+			const FVector End = Position + YawRotator.RotateVector(FVector(
+				RowWidth * (((float)RowIdx) / (RowNum - 1) - 0.5f),
+				ColWidth * (((float)ColIdx) / (ColNum - 1) - 0.5f),
+				MinHeight));
+
+			FCollisionObjectQueryParams ObjectQueryParams;
+			ObjectQueryParams.AddObjectTypesToQuery(CollisionChannel);
+
+			const bool bHit = GetWorld()->LineTraceSingleByObjectType(TraceHits[RowIdx * ColNum + ColIdx], Start, End, ObjectQueryParams);
+
+			if (bHit)
+			{
+				OutHeights[RowIdx * ColNum + ColIdx] = TraceHits[RowIdx * ColNum + ColIdx].ImpactPoint.Z;
+			}
+			else
+			{
+				OutHeights[RowIdx * ColNum + ColIdx] = MinHeight;
+			}
+		}
+	}
+
+#if ENABLE_VISUAL_LOG
+	{
+		UE_LEARNING_AGENTS_VLOG_TRANSFORM(this, LogLearning, Display,
+			Position,
+			Rotation,
+			VisualLogColor.ToFColor(true),
+			TEXT(""));
+
+		for (int32 RowIdx = 0; RowIdx < RowNum; RowIdx++)
+		{
+			for (int32 ColIdx = 0; ColIdx < ColNum; ColIdx++)
+			{
+				const FVector Start = Position + YawRotator.RotateVector(FVector(
+					RowWidth * (((float)RowIdx) / (RowNum - 1) - 0.5f),
+					ColWidth * (((float)ColIdx) / (ColNum - 1) - 0.5f),
+					MaxHeight));
+
+				const FVector End = Position + YawRotator.RotateVector(FVector(
+					RowWidth * (((float)RowIdx) / (RowNum - 1) - 0.5f),
+					ColWidth * (((float)ColIdx) / (ColNum - 1) - 0.5f),
+					MinHeight));
+
+				const FVector ImpactPoint = 
+					TraceHits[RowIdx * ColNum + ColIdx].bBlockingHit ?
+					TraceHits[RowIdx * ColNum + ColIdx].ImpactPoint : End;
+
+				UE_LEARNING_AGENTS_VLOG_ARROW(this, LogLearning, Display,
+					Start,
+					End,
+					VisualLogColor.ToFColor(true),
+					TEXT(""));
+
+				UE_LEARNING_AGENTS_VLOG_LOCATION(this, LogLearning, Display,
+					ImpactPoint,
+					2.5f,
+					VisualLogColor.ToFColor(true),
+					TEXT("Agent %i\nRow: %i\nCol: %i\nImpact Point: [% 6.1f % 6.1f % 6.1f]\nHeight: [% 6.1f]"),
+					AgentId,
+					RowIdx, ColIdx,
+					ImpactPoint.X, ImpactPoint.Y, ImpactPoint.Z,
+					OutHeights[RowIdx * ColNum + ColIdx]);
+			}
+		}
+	}
+#endif
+}
diff --git a/Engine/Plugins/Experimental/LearningAgents/Source/LearningAgents/Private/LearningAgentsManager.cpp b/Engine/Plugins/Experimental/LearningAgents/Source/LearningAgents/Private/LearningAgentsManager.cpp
--- a/Engine/Plugins/Experimental/LearningAgents/Source/LearningAgents/Private/LearningAgentsManager.cpp
+++ b/Engine/Plugins/Experimental/LearningAgents/Source/LearningAgents/Private/LearningAgentsManager.cpp
@@ -39,7 +39,7 @@ void ALearningAgentsManager::SetupManager()
 {
 	if (IsManagerSetup())
 	{
-		UE_LOG(LogLearning, Error, TEXT("%s: Setup already performed!"), *GetName());
+		UE_LOG(LogLearning, Error, TEXT("%s: Setup already run!"), *GetName());
 		return;
 	}
 
@@ -131,7 +131,7 @@ void ALearningAgentsManager::RemoveAgentById(const int32 AgentId)
 
 	if (RemovedCount == 0)
 	{
-		UE_LOG(LogLearning, Warning, TEXT("%s: Trying to remove an agent but its Id (%d) is not in the occupied agents."), *GetName(), AgentId);
+		UE_LOG(LogLearning, Warning, TEXT("%s: Trying to remove an agent with id of %i but it was not found."), *GetName(), AgentId);
 		return;
 	}
 
@@ -148,7 +148,7 @@ void ALearningAgentsManager::RemoveAgent(UObject* Agent)
 {
 	if (Agent == nullptr)
 	{
-		UE_LOG(LogLearning, Error, TEXT("%s: Attempted to remove an agent but agent is nullptr."), *GetName());
+		UE_LOG(LogLearning, Error, TEXT("%s: Attempted to remove agent but agent is nullptr."), *GetName());
 		return;
 	}
 
@@ -157,7 +157,7 @@ void ALearningAgentsManager::RemoveAgent(UObject* Agent)
 
 	if (!bIsFound)
 	{
-		UE_LOG(LogLearning, Warning, TEXT("%s: Trying to remove an agent but it was not found."), *GetName());
+		UE_LOG(LogLearning, Warning, TEXT("%s: Trying to remove agent %s but it was not found."), *Agent->GetName(), *GetName());
 		return;
 	}
 
@@ -189,10 +189,7 @@ UObject* ALearningAgentsManager::GetAgent(const int32 AgentId, const TSubclassOf
 {
 	if (!OccupiedAgentSet.Contains(AgentId))
 	{
-		UE_LOG(LogLearning, Error,
-			TEXT("%s: AgentId %d not found. Be sure to only use AgentIds returned by AddAgent() and check that the agent has not be removed."),
-			*GetName(),
-			AgentId);
+		UE_LOG(LogLearning, Error, TEXT("%s: AgentId %d not found. Be sure to only use AgentIds returned by AddAgent() and check that the agent has not be removed."), *GetName(), AgentId);
 		return nullptr;
 	}
 
@@ -211,6 +208,11 @@ UObject* ALearningAgentsManager::GetAgent(const int32 AgentId)
 	return Agents[AgentId];
 }
 
+void ALearningAgentsManager::GetAgentIds(TArray<int32>& OutAgentIds) const
+{
+	OutAgentIds = OccupiedAgentIds;
+}
+
 void ALearningAgentsManager::UpdateAgentSets()
 {
 	OccupiedAgentSet = OccupiedAgentIds;
diff --git a/Engine/Plugins/Experimental/LearningAgents/Source/LearningAgents/Private/LearningAgentsManagerComponent.cpp b/Engine/Plugins/Experimental/LearningAgents/Source/LearningAgents/Private/LearningAgentsManagerComponent.cpp
--- a/Engine/Plugins/Experimental/LearningAgents/Source/LearningAgents/Private/LearningAgentsManagerComponent.cpp
+++ b/Engine/Plugins/Experimental/LearningAgents/Source/LearningAgents/Private/LearningAgentsManagerComponent.cpp
@@ -12,7 +12,7 @@ bool ULearningAgentsManagerComponent::AddAgent(const int32 AgentId)
 {
 	if (HasAgent(AgentId))
 	{
-		UE_LOG(LogLearning, Warning, TEXT("AgentId %i is already added to %s."), AgentId, *GetName());
+		UE_LOG(LogLearning, Warning, TEXT("%s: AgentId %i is already added."), *GetName(), AgentId);
 		return false;
 	}
 
@@ -27,7 +27,7 @@ bool ULearningAgentsManagerComponent::RemoveAgent(const int32 AgentId)
 {
 	if (AddedAgentIds.RemoveSingleSwap(AgentId, false) == 0)
 	{
-		UE_LOG(LogLearning, Warning, TEXT("Unable to remove: AgentId %d not found in %s's agent set."), AgentId, *GetName());
+		UE_LOG(LogLearning, Warning, TEXT("%s: Trying to remove an agent with id of %i but it was not found."), *GetName(), AgentId);
 		return false;
 	}
 
@@ -46,20 +46,18 @@ UObject* ULearningAgentsManagerComponent::GetAgent(const int32 AgentId, const TS
 {
 	if (!AgentManager)
 	{
-		UE_LOG(LogLearning, Error, TEXT("Agent manager is nullptr. Call setup on this component prior to getting agents."));
+		UE_LOG(LogLearning, Error, TEXT("%s: Agent manager is nullptr. Call setup on this component prior to getting agents."), *GetName());
 		return nullptr;
 	}
 
 	if (!AddedAgentSet.Contains(AgentId))
 	{
-		UE_LOG(LogLearning, Error,
-			TEXT("%s: AgentId %d not found. Be sure to only use AgentIds returned by AddAgent() and check that the agent has not be removed."),
-			*GetName(),
-			AgentId);
+		UE_LOG(LogLearning, Error, TEXT("%s: AgentId %d not found. Be sure to only use AgentIds returned by AddAgent() and check that the agent has not be removed."), *GetName(), AgentId);
 		return nullptr;
 	}
 
-	return AgentManager->GetAgent(AgentId, AgentClass); // Calling this overload since it will log about missing manager ids
+	// Calling this overload since it will log about missing manager ids
+	return AgentManager->GetAgent(AgentId, AgentClass); 
 }
 
 const UObject* ULearningAgentsManagerComponent::GetAgent(const int32 AgentId) const
@@ -89,7 +87,7 @@ ALearningAgentsManager* ULearningAgentsManagerComponent::GetAgentManager(const T
 {
 	if (!AgentManager)
 	{
-		UE_LOG(LogLearning, Error, TEXT("AgentManager is nullptr. Did we forget to call Setup on this component and set the manager?"));
+		UE_LOG(LogLearning, Error, TEXT("%s: AgentManager is nullptr. Did we forget to call Setup on this component and set the manager?"), *GetName());
 		return nullptr;
 	}
 
@@ -100,3 +98,8 @@ bool ULearningAgentsManagerComponent::IsSetup() const
 {
 	return bIsSetup;
 }
+
+void ULearningAgentsManagerComponent::AddHelper(TObjectPtr<ULearningAgentsHelper> Object)
+{
+	HelperObjects.Add(Object);
+}
\ No newline at end of file
