LearningAgents: Added some additional observations, actions, rewards, completions and helpers. Cleaned up some of the error messaging to make it more consistent.

diff --git a/Engine/Plugins/Experimental/LearningAgents/Source/LearningAgents/Private/LearningAgentsObservations.cpp b/Engine/Plugins/Experimental/LearningAgents/Source/LearningAgents/Private/LearningAgentsObservations.cpp
--- a/Engine/Plugins/Experimental/LearningAgents/Source/LearningAgents/Private/LearningAgentsObservations.cpp
+++ b/Engine/Plugins/Experimental/LearningAgents/Source/LearningAgents/Private/LearningAgentsObservations.cpp
@@ -10,41 +10,28 @@
 #include "LearningLog.h"
 
 #include "GameFramework/Actor.h"
-#include "VisualLogger/VisualLogger.h"
 
-#define UE_LEARNING_AGENTS_VLOG_STRING(Owner, Category, Verbosity, Location, Color, Format, ...) \
-	UE_VLOG_LOCATION(Owner, Category, Verbosity, Location, 0.0f, Color, Format, ##__VA_ARGS__)
-
-#define UE_LEARNING_AGENTS_VLOG_TRANSFORM(Owner, Category, Verbosity, Location, Rotation, Color, Format, ...) \
-	UE_VLOG_SEGMENT(Owner, Category, Verbosity, Location, Location + 15.0f * Rotation.RotateVector(FVector::ForwardVector), FColor::Red, TEXT("")); \
-	UE_VLOG_SEGMENT(Owner, Category, Verbosity, Location, Location + 15.0f * Rotation.RotateVector(FVector::RightVector), FColor::Green, TEXT("")); \
-	UE_VLOG_SEGMENT(Owner, Category, Verbosity, Location, Location + 15.0f * Rotation.RotateVector(FVector::UpVector), FColor::Blue, TEXT("")); \
-	UE_VLOG_OBOX(Owner, Category, Verbosity, FBox(10.0f * FVector(-1, -1, -1), 10.0f * FVector(1, 1, 1)), FTransform(Rotation, Location, FVector::OneVector).ToMatrixNoScale(), Color, TEXT("")); \
-	UE_LEARNING_AGENTS_VLOG_STRING(Owner, Category, Verbosity, Location + FVector(0.0f, 0.0f, 20.0f), Color, Format, ##__VA_ARGS__)
-
-#define UE_LEARNING_AGENTS_VLOG_PLANE(Owner, Category, Verbosity, Location, Rotation, Axis0, Axis1, Color, Format, ...) \
-	UE_VLOG_OBOX(Owner, Category, Verbosity, FBox(-25.0f * (Axis0 + Axis1), 25.0f * (Axis0 + Axis1)), FTransform(Rotation, Location, FVector::OneVector).ToMatrixNoScale(), Color, Format, ##__VA_ARGS__)
-
-namespace UE::Learning::Agents::Private
+namespace UE::Learning::Agents::Observations::Private
 {
 	template<typename ObservationUObject, typename ObservationFObject, typename... InArgTypes>
-	ObservationUObject* AddObservation(ULearningAgentsType* AgentType, const FName Name, InArgTypes&& ...Args)
+	ObservationUObject* AddObservation(ULearningAgentsType* InAgentType, const FName Name, InArgTypes&& ...Args)
 	{
-		if (!AgentType)
+		if (!InAgentType)
 		{
-			UE_LOG(LogLearning, Error, TEXT("AgentType is nullptr"));
+			UE_LOG(LogLearning, Error, TEXT("InAgentType is nullptr."));
 			return nullptr;
 		}
 
-		ObservationUObject* Observation = NewObject<ObservationUObject>(AgentType, Name);
+		ObservationUObject* Observation = NewObject<ObservationUObject>(InAgentType, Name);
 
+		Observation->AgentType = InAgentType;
 		Observation->FeatureObject = MakeShared<ObservationFObject>(
 			Observation->GetFName(),
-			AgentType->GetAgentManager()->GetInstanceData().ToSharedRef(),
-			AgentType->GetAgentManager()->GetMaxInstanceNum(),
+			InAgentType->GetAgentManager()->GetInstanceData().ToSharedRef(),
+			InAgentType->GetAgentManager()->GetMaxInstanceNum(),
 			Forward<InArgTypes>(Args)...);
 
-		AgentType->AddObservation(Observation, Observation->FeatureObject.ToSharedRef());
+		InAgentType->AddObservation(Observation, Observation->FeatureObject.ToSharedRef());
 
 		return Observation;
 	}
@@ -52,31 +39,23 @@ namespace UE::Learning::Agents::Private
 
 //------------------------------------------------------------------
 
-UFloatObservation* UFloatObservation::AddFloatObservation(ULearningAgentsType* AgentType, const FName Name, const float Scale)
+UFloatObservation* UFloatObservation::AddFloatObservation(ULearningAgentsType* InAgentType, const FName Name, const float Scale)
 {
-	return UE::Learning::Agents::Private::AddObservation<UFloatObservation, UE::Learning::FFloatFeature>(AgentType, Name, 1, Scale);
+	return UE::Learning::Agents::Observations::Private::AddObservation<UFloatObservation, UE::Learning::FFloatFeature>(InAgentType, Name, 1, Scale);
 }
 
 void UFloatObservation::SetFloatObservation(const int32 AgentId, const float Value)
 {
-	const TLearningArrayView<2, float> View = FeatureObject->InstanceData->View(FeatureObject->ValueHandle);
-
-	if (AgentId == INDEX_NONE)
-	{
-		UE_LOG(LogLearning, Error, TEXT("AgentId is invalid (INDEX_NONE)"));
-		return;
-	}
-
-	if (AgentId < 0 || AgentId >= View.Num<0>())
+	if (!AgentType->HasAgent(AgentId))
 	{
-		UE_LOG(LogLearning, Error, TEXT("AgentId %d is out of index. Valid range [0, %d]."), AgentId, View.Num<0>() - 1);
+		UE_LOG(LogLearning, Error, TEXT("%s: AgentId %d not found in the agents set."), *GetName(), AgentId);
 		return;
 	}
 
-	View[AgentId][0] = Value;
+	FeatureObject->InstanceData->View(FeatureObject->ValueHandle)[AgentId][0] = Value;
 }
 
-#if ENABLE_VISUAL_LOG
+#if UE_LEARNING_AGENTS_ENABLE_VISUAL_LOG
 void UFloatObservation::VisualLog(const UE::Learning::FIndexSet Instances) const
 {
 	UE_LEARNING_TRACE_CPUPROFILER_EVENT_SCOPE(UFloatObservation::VisualLog);
@@ -84,21 +63,18 @@ void UFloatObservation::VisualLog(const UE::Learning::FIndexSet Instances) const
 	const TLearningArrayView<2, const float> ValueView = FeatureObject->InstanceData->ConstView(FeatureObject->ValueHandle);
 	const TLearningArrayView<2, const float> FeatureView = FeatureObject->InstanceData->ConstView(FeatureObject->FeatureHandle);
 
-	if (const ULearningAgentsType* AgentType = Cast<ULearningAgentsType>(GetOuter()))
+	for (const int32 Instance : Instances)
 	{
-		for (const int32 Instance : Instances)
+		if (const AActor* Actor = Cast<AActor>(AgentType->GetAgent(Instance)))
 		{
-			if (const AActor* Actor = Cast<AActor>(AgentType->GetAgent(Instance)))
-			{
-				UE_LEARNING_AGENTS_VLOG_STRING(this, LogLearning, Display,
-					Actor->GetActorLocation(),
-					VisualLogColor.ToFColor(true),
-					TEXT("Agent %i\nScale: [% 6.2f]\nValue: [% 6.2f]\nEncoded: [% 6.3f]"),
-					Instance,
-					ValueView[Instance][0],
-					FeatureObject->Scale,
-					FeatureView[Instance][0]);
-			}
+			UE_LEARNING_AGENTS_VLOG_STRING(this, LogLearning, Display,
+				Actor->GetActorLocation(),
+				VisualLogColor.ToFColor(true),
+				TEXT("Agent %i\nScale: [% 6.2f]\nValue: [% 6.2f]\nEncoded: [% 6.3f]"),
+				Instance,
+				FeatureObject->Scale,
+				ValueView[Instance][0],
+				*UE::Learning::Array::FormatFloat(FeatureView[Instance]));
 		}
 	}
 }
@@ -106,68 +82,56 @@ void UFloatObservation::VisualLog(const UE::Learning::FIndexSet Instances) const
 
 //------------------------------------------------------------------
 
-UVectorObservation* UVectorObservation::AddVectorObservation(ULearningAgentsType* AgentType, const FName Name, const float Scale)
+UFloatArrayObservation* UFloatArrayObservation::AddFloatArrayObservation(ULearningAgentsType* InAgentType, const FName Name, const int32 Num, const float Scale)
 {
-	return UE::Learning::Agents::Private::AddObservation<UVectorObservation, UE::Learning::FFloatFeature>(AgentType, Name, 3, Scale);
+	if (Num < 1)
+	{
+		UE_LOG(LogLearning, Error, TEXT("Number of elements in array must be at least 1, got %i."), Num);
+		return nullptr;
+	}
+
+	return UE::Learning::Agents::Observations::Private::AddObservation<UFloatArrayObservation, UE::Learning::FFloatFeature>(InAgentType, Name, Num, Scale);
 }
 
-void UVectorObservation::SetVectorObservation(const int32 AgentId, const FVector Vector)
+void UFloatArrayObservation::SetFloatArrayObservation(const int32 AgentId, const TArray<float>& Values)
 {
-	const TLearningArrayView<2, float> View = FeatureObject->InstanceData->View(FeatureObject->ValueHandle);
-
-	if (AgentId == INDEX_NONE)
+	if (!AgentType->HasAgent(AgentId))
 	{
-		UE_LOG(LogLearning, Error, TEXT("AgentId is invalid (INDEX_NONE)"));
+		UE_LOG(LogLearning, Error, TEXT("%s: AgentId %d not found in the agents set."), *GetName(), AgentId);
 		return;
 	}
 
-	if (AgentId < 0 || AgentId >= View.Num<0>())
+	const TLearningArrayView<2, float> View = FeatureObject->InstanceData->View(FeatureObject->ValueHandle);
+
+	if (Values.Num() != View.Num<1>())
 	{
-		UE_LOG(LogLearning, Error, TEXT("AgentId %d is out of index. Valid range [0, %d]."), AgentId, View.Num<0>() - 1);
+		UE_LOG(LogLearning, Error, TEXT("%s: Got wrong number of elements in array. Expected %i, got %i."), *GetName(), View.Num<1>(), Values.Num());
 		return;
 	}
 
-	View[AgentId][0] = Vector.X;
-	View[AgentId][1] = Vector.Y;
-	View[AgentId][2] = Vector.Z;
+	UE::Learning::Array::Copy<1, float>(View[AgentId], Values);
 }
 
-#if ENABLE_VISUAL_LOG
-void UVectorObservation::VisualLog(const UE::Learning::FIndexSet Instances) const
+#if UE_LEARNING_AGENTS_ENABLE_VISUAL_LOG
+void UFloatArrayObservation::VisualLog(const UE::Learning::FIndexSet Instances) const
 {
-	UE_LEARNING_TRACE_CPUPROFILER_EVENT_SCOPE(UVectorObservation::VisualLog);
+	UE_LEARNING_TRACE_CPUPROFILER_EVENT_SCOPE(UFloatArrayObservation::VisualLog);
 
 	const TLearningArrayView<2, const float> ValueView = FeatureObject->InstanceData->ConstView(FeatureObject->ValueHandle);
 	const TLearningArrayView<2, const float> FeatureView = FeatureObject->InstanceData->ConstView(FeatureObject->FeatureHandle);
 
-	if (const ULearningAgentsType* AgentType = Cast<ULearningAgentsType>(GetOuter()))
+	for (const int32 Instance : Instances)
 	{
-		for (const int32 Instance : Instances)
+		if (const AActor* Actor = Cast<AActor>(AgentType->GetAgent(Instance)))
 		{
-			if (const AActor* Actor = Cast<AActor>(AgentType->GetAgent(Instance)))
-			{
-				const FVector Vector(ValueView[Instance][0], ValueView[Instance][1], ValueView[Instance][2]);
-
-				UE_VLOG_ARROW(this, LogLearning, Display,
-					Actor->GetActorLocation(),
-					Actor->GetActorLocation() + Vector,
-					VisualLogColor.ToFColor(true),
-					TEXT(""));
-
-				UE_LEARNING_AGENTS_VLOG_STRING(this, LogLearning, Display,
-					Actor->GetActorLocation() + Vector,
-					VisualLogColor.ToFColor(true),
-					TEXT("Vector: [% 6.4f % 6.4f % 6.4f]"),
-					Vector.X, Vector.Y, Vector.Z);
-
-				UE_LEARNING_AGENTS_VLOG_STRING(this, LogLearning, Display,
-					Actor->GetActorLocation(),
-					VisualLogColor.ToFColor(true),
-					TEXT("Agent %i\nScale: [% 6.2f]\nEncoded: [% 6.3f % 6.3f % 6.3f]"),
-					Instance,
-					FeatureObject->Scale,
-					FeatureView[Instance][0], FeatureView[Instance][1], FeatureView[Instance][2]);
-			}
+			UE_LEARNING_AGENTS_VLOG_STRING(this, LogLearning, Display,
+				Actor->GetActorLocation(),
+				VisualLogColor.ToFColor(true),
+				TEXT("Agent %i\nScale: [% 6.2f]\nValue: %s\nEncoded: %s"),
+				Instance,
+				FeatureObject->Scale,
+				*UE::Learning::Array::FormatFloat(ValueView[Instance]),
+				*UE::Learning::Array::FormatFloat(FeatureView[Instance]));
 		}
 	}
 }
@@ -175,681 +139,1869 @@ void UVectorObservation::VisualLog(const UE::Learning::FIndexSet Instances) cons
 
 //------------------------------------------------------------------
 
-UAngleObservation* UAngleObservation::AddAngleObservation(ULearningAgentsType* AgentType, const FName Name, const float Scale)
+UVectorObservation* UVectorObservation::AddVectorObservation(ULearningAgentsType* InAgentType, const FName Name, const float Scale)
 {
-	return UE::Learning::Agents::Private::AddObservation<UAngleObservation, UE::Learning::FAngleFeature>(AgentType, Name, 1, Scale);
+	return UE::Learning::Agents::Observations::Private::AddObservation<UVectorObservation, UE::Learning::FFloatFeature>(InAgentType, Name, 3, Scale);
 }
 
-void UAngleObservation::SetAngleObservation(const int32 AgentId, const float Angle, const float RelativeAngle)
+void UVectorObservation::SetVectorObservation(const int32 AgentId, const FVector Vector)
 {
-	const TLearningArrayView<2, float> AngleView = FeatureObject->InstanceData->View(FeatureObject->AngleHandle);
-	const TLearningArrayView<1, float> RelativeAngleView = FeatureObject->InstanceData->View(FeatureObject->RelativeAngleHandle);
-
-	if (AgentId == INDEX_NONE)
+	if (!AgentType->HasAgent(AgentId))
 	{
-		UE_LOG(LogLearning, Error, TEXT("AgentId is invalid (INDEX_NONE)"));
+		UE_LOG(LogLearning, Error, TEXT("%s: AgentId %d not found in the agents set."), *GetName(), AgentId);
 		return;
 	}
 
-	if (AgentId < 0 || AgentId >= AngleView.Num<0>())
-	{
-		UE_LOG(LogLearning, Error, TEXT("AgentId %d is out of index. Valid range [0, %d]."), AgentId, AngleView.Num<0>() - 1);
-		return;
-	}
+	const TLearningArrayView<2, float> View = FeatureObject->InstanceData->View(FeatureObject->ValueHandle);
 
-	AngleView[AgentId][0] = FMath::DegreesToRadians(Angle);
-	RelativeAngleView[AgentId] = FMath::DegreesToRadians(RelativeAngle);
+	View[AgentId][0] = Vector.X;
+	View[AgentId][1] = Vector.Y;
+	View[AgentId][2] = Vector.Z;
 }
 
-#if ENABLE_VISUAL_LOG
-void UAngleObservation::VisualLog(const UE::Learning::FIndexSet Instances) const
+#if UE_LEARNING_AGENTS_ENABLE_VISUAL_LOG
+void UVectorObservation::VisualLog(const UE::Learning::FIndexSet Instances) const
 {
-	UE_LEARNING_TRACE_CPUPROFILER_EVENT_SCOPE(UAngleObservation::VisualLog);
+	UE_LEARNING_TRACE_CPUPROFILER_EVENT_SCOPE(UVectorObservation::VisualLog);
 
-	const TLearningArrayView<2, const float> AngleView = FeatureObject->InstanceData->ConstView(FeatureObject->AngleHandle);
-	const TLearningArrayView<1, const float> RelativeAngleView = FeatureObject->InstanceData->ConstView(FeatureObject->RelativeAngleHandle);
+	const TLearningArrayView<2, const float> ValueView = FeatureObject->InstanceData->ConstView(FeatureObject->ValueHandle);
 	const TLearningArrayView<2, const float> FeatureView = FeatureObject->InstanceData->ConstView(FeatureObject->FeatureHandle);
 
-	if (const ULearningAgentsType* AgentType = Cast<ULearningAgentsType>(GetOuter()))
+	for (const int32 Instance : Instances)
 	{
-		for (const int32 Instance : Instances)
+		if (const AActor* Actor = Cast<AActor>(AgentType->GetAgent(Instance)))
 		{
-			if (const AActor* Actor = Cast<AActor>(AgentType->GetAgent(Instance)))
-			{
-				const float Angle = AngleView[Instance][0];
-				const float RelativeAngle = RelativeAngleView[Instance];
-
-				UE_VLOG_CIRCLE(this, LogLearning, Display,
-					Actor->GetActorLocation(),
-					FVector::UpVector,
-					50.0f,
-					VisualLogColor.ToFColor(true),
-					TEXT(""));
-
-				UE_VLOG_SEGMENT(this, LogLearning, Display,
-					Actor->GetActorLocation(),
-					Actor->GetActorLocation() + 50.0f * FVector(FMath::Sin(RelativeAngle), FMath::Cos(RelativeAngle), 0.0f),
-					VisualLogColor.ToFColor(true),
-					TEXT(""));
-
-				UE_VLOG_LOCATION(this, LogLearning, Display,
-					Actor->GetActorLocation() + 50.0f * FVector(FMath::Sin(RelativeAngle), FMath::Cos(RelativeAngle), 0.0f),
-					2.5f,
-					VisualLogColor.ToFColor(true),
-					TEXT("Relative Angle: [% 6.1f]"),
-					RelativeAngle);
-
-				UE_VLOG_SEGMENT(this, LogLearning, Display,
-					Actor->GetActorLocation(),
-					Actor->GetActorLocation() + 50.0f * FVector(FMath::Sin(Angle), FMath::Cos(Angle), 0.0f),
-					VisualLogColor.ToFColor(true),
-					TEXT(""));
-
-				UE_VLOG_LOCATION(this, LogLearning, Display,
-					Actor->GetActorLocation() + 50.0f * FVector(FMath::Sin(Angle), FMath::Cos(Angle), 0.0f),
-					2.5f,
-					VisualLogColor.ToFColor(true),
-					TEXT("Angle: [% 6.1f]"),
-					Angle);
-
-				UE_VLOG_LOCATION(this, LogLearning, Display,
-					Actor->GetActorLocation(),
-					5.0f,
-					VisualLogColor.ToFColor(true),
-					TEXT("Agent %i\nScale: [% 6.2f]\nEncoded: [% 6.3f % 6.3f]"),
-					Instance,
-					FeatureObject->Scale,
-					FeatureView[Instance][0], FeatureView[Instance][1]);
-			}
+			const FVector Vector(ValueView[Instance][0], ValueView[Instance][1], ValueView[Instance][2]);
+
+			UE_LEARNING_AGENTS_VLOG_ARROW(this, LogLearning, Display,
+				Actor->GetActorLocation(),
+				Actor->GetActorLocation() + Vector,
+				VisualLogColor.ToFColor(true),
+				TEXT(""));
+
+			UE_LEARNING_AGENTS_VLOG_STRING(this, LogLearning, Display,
+				Actor->GetActorLocation() + Vector,
+				VisualLogColor.ToFColor(true),
+				TEXT("Vector: [% 6.4f % 6.4f % 6.4f]"),
+				Vector.X, Vector.Y, Vector.Z);
+
+			UE_LEARNING_AGENTS_VLOG_STRING(this, LogLearning, Display,
+				Actor->GetActorLocation(),
+				VisualLogColor.ToFColor(true),
+				TEXT("Agent %i\nScale: [% 6.2f]\nEncoded: %s"),
+				Instance,
+				FeatureObject->Scale,
+				*UE::Learning::Array::FormatFloat(FeatureView[Instance]));
 		}
 	}
 }
 #endif
 
-//------------------------------------------------------------------
-
-UPlanarDirectionObservation* UPlanarDirectionObservation::AddPlanarDirectionObservation(ULearningAgentsType* AgentType, const FName Name, const float Scale, const FVector Axis0, const FVector Axis1)
+UVectorArrayObservation* UVectorArrayObservation::AddVectorArrayObservation(ULearningAgentsType* InAgentType, const FName Name, const int32 Num, const float Scale)
 {
-	return UE::Learning::Agents::Private::AddObservation<UPlanarDirectionObservation, UE::Learning::FPlanarDirectionFeature>(AgentType, Name, 1, Scale, Axis0, Axis1);
+	if (Num < 1)
+	{
+		UE_LOG(LogLearning, Error, TEXT("Number of elements in array must be at least 1, got %i."), Num);
+		return nullptr;
+	}
+
+	return UE::Learning::Agents::Observations::Private::AddObservation<UVectorArrayObservation, UE::Learning::FFloatFeature>(InAgentType, Name, Num * 3, Scale);
 }
 
-void UPlanarDirectionObservation::SetPlanarDirectionObservation(const int32 AgentId, const FVector Direction, const FRotator RelativeRotation)
+void UVectorArrayObservation::SetVectorArrayObservation(const int32 AgentId, const TArray<FVector>& Vectors)
 {
-	const TLearningArrayView<2, FVector> DirectionView = FeatureObject->InstanceData->View(FeatureObject->DirectionHandle);
-	const TLearningArrayView<1, FQuat> RelativeRotationView = FeatureObject->InstanceData->View(FeatureObject->RelativeRotationHandle);
-
-	if (AgentId == INDEX_NONE)
+	if (!AgentType->HasAgent(AgentId))
 	{
-		UE_LOG(LogLearning, Error, TEXT("AgentId is invalid (INDEX_NONE)"));
+		UE_LOG(LogLearning, Error, TEXT("%s: AgentId %d not found in the agents set."), *GetName(), AgentId);
 		return;
 	}
 
-	if (AgentId < 0 || AgentId >= DirectionView.Num<0>())
+	const TLearningArrayView<2, float> View = FeatureObject->InstanceData->View(FeatureObject->ValueHandle);
+
+	if (Vectors.Num() != View.Num<1>() / 3)
 	{
-		UE_LOG(LogLearning, Error, TEXT("AgentId %d is out of index. Valid range [0, %d]."), AgentId, DirectionView.Num<0>() - 1);
+		UE_LOG(LogLearning, Error, TEXT("%s: Got wrong number of elements in array. Expected %i, got %i."), *GetName(), View.Num<1>() / 3, Vectors.Num());
 		return;
 	}
 
-	RelativeRotationView[AgentId] = FQuat::MakeFromRotator(RelativeRotation);
-	DirectionView[AgentId][0] = Direction;
+	for (int32 VectorIdx = 0; VectorIdx < Vectors.Num(); VectorIdx++)
+	{
+		View[AgentId][VectorIdx * 3 + 0] = Vectors[VectorIdx].X;
+		View[AgentId][VectorIdx * 3 + 1] = Vectors[VectorIdx].Y;
+		View[AgentId][VectorIdx * 3 + 2] = Vectors[VectorIdx].Z;
+	}
 }
 
-#if ENABLE_VISUAL_LOG
-void UPlanarDirectionObservation::VisualLog(const UE::Learning::FIndexSet Instances) const
+#if UE_LEARNING_AGENTS_ENABLE_VISUAL_LOG
+void UVectorArrayObservation::VisualLog(const UE::Learning::FIndexSet Instances) const
 {
-	UE_LEARNING_TRACE_CPUPROFILER_EVENT_SCOPE(UPlanarDirectionObservation::VisualLog);
+	UE_LEARNING_TRACE_CPUPROFILER_EVENT_SCOPE(UVectorArrayObservation::VisualLog);
 
-	const TLearningArrayView<2, const FVector> DirectionView = FeatureObject->InstanceData->ConstView(FeatureObject->DirectionHandle);
-	const TLearningArrayView<1, const FQuat> RelativeRotationView = FeatureObject->InstanceData->ConstView(FeatureObject->RelativeRotationHandle);
+	const TLearningArrayView<2, const float> ValueView = FeatureObject->InstanceData->ConstView(FeatureObject->ValueHandle);
 	const TLearningArrayView<2, const float> FeatureView = FeatureObject->InstanceData->ConstView(FeatureObject->FeatureHandle);
 
-	if (const ULearningAgentsType* AgentType = Cast<ULearningAgentsType>(GetOuter()))
+	const int32 VectorNum = ValueView.Num<1>() / 3;
+
+	for (const int32 Instance : Instances)
 	{
-		for (const int32 Instance : Instances)
+		if (const AActor* Actor = Cast<AActor>(AgentType->GetAgent(Instance)))
 		{
-			if (const AActor* Actor = Cast<AActor>(AgentType->GetAgent(Instance)))
+			for (int32 VectorIdx = 0; VectorIdx < VectorNum; VectorIdx++)
 			{
-				const FVector Direction = DirectionView[Instance][0];
-				const FQuat RelativeRotation = RelativeRotationView[Instance];
-				const FVector LocalDirection = RelativeRotation.UnrotateVector(Direction);
+				const FVector Vector(ValueView[Instance][VectorIdx * 3 + 0], ValueView[Instance][VectorIdx * 3 + 1], ValueView[Instance][VectorIdx * 3 + 2]);
+				const FVector Offset = UE::Learning::Agents::Debug::GridOffsetForIndex(VectorIdx, VectorNum);
 
-				UE_VLOG_ARROW(this, LogLearning, Display,
-					Actor->GetActorLocation(),
-					Actor->GetActorLocation() + 100.0f * Direction,
+				UE_LEARNING_AGENTS_VLOG_ARROW(this, LogLearning, Display,
+					Actor->GetActorLocation() + Offset,
+					Actor->GetActorLocation() + Offset + Vector,
 					VisualLogColor.ToFColor(true),
 					TEXT(""));
 
 				UE_LEARNING_AGENTS_VLOG_STRING(this, LogLearning, Display,
-					Actor->GetActorLocation() + 100.0f * Direction,
-					VisualLogColor.ToFColor(true),
-					TEXT("Direction: [% 6.3f % 6.3f % 6.3f]\nLocal Direction: [% 6.3f % 6.3f % 6.3f]"),
-					Direction.X, Direction.Y, Direction.Z,
-					LocalDirection.X, LocalDirection.Y, LocalDirection.Z);
-
-				UE_LEARNING_AGENTS_VLOG_PLANE(this, LogLearning, Display,
-					Actor->GetActorLocation(),
-					RelativeRotation,
-					FeatureObject->Axis0,
-					FeatureObject->Axis1,
-					VisualLogColor.ToFColor(true),
-					TEXT(""));
-
-				UE_LEARNING_AGENTS_VLOG_TRANSFORM(this, LogLearning, Display,
-					Actor->GetActorLocation(),
-					RelativeRotation,
+					Actor->GetActorLocation() + Offset + Vector,
 					VisualLogColor.ToFColor(true),
-					TEXT("Agent %i\nScale: [% 6.2f]\nEncoded: [% 6.3f % 6.3f]"),
-					Instance,
-					FeatureObject->Scale,
-					FeatureView[Instance][0], FeatureView[Instance][1]);
+					TEXT("Vector %i: [% 6.4f % 6.4f % 6.4f]"),
+					VectorIdx,
+					Vector.X, Vector.Y, Vector.Z);
 			}
+
+			UE_LEARNING_AGENTS_VLOG_STRING(this, LogLearning, Display,
+				Actor->GetActorLocation(),
+				VisualLogColor.ToFColor(true),
+				TEXT("Agent %i\nScale: [% 6.2f]\nEncoded: %s"),
+				Instance,
+				FeatureObject->Scale,
+				*UE::Learning::Array::FormatFloat(FeatureView[Instance]));
 		}
 	}
 }
 #endif
 
-UDirectionObservation* UDirectionObservation::AddDirectionObservation(ULearningAgentsType* AgentType, const FName Name, const float Scale)
+//------------------------------------------------------------------
+
+UEnumObservation* UEnumObservation::AddEnumObservation(ULearningAgentsType* InAgentType, const UEnum* EnumType, const FName Name)
 {
-	return UE::Learning::Agents::Private::AddObservation<UDirectionObservation, UE::Learning::FDirectionFeature>(AgentType, Name, 1, Scale);
+	if (!EnumType)
+	{
+		UE_LOG(LogLearning, Error, TEXT("Invalid Enum."));
+		return nullptr;
+	}
+
+	if (EnumType->NumEnums() < 1)
+	{
+		UE_LOG(LogLearning, Error, TEXT("Enum requires at least one entry to be used as an observation."));
+		return nullptr;
+	}
+
+	UEnumObservation* Observation = UE::Learning::Agents::Observations::Private::AddObservation<UEnumObservation, UE::Learning::FFloatFeature>(InAgentType, Name, EnumType->NumEnums());
+	if (Observation) { Observation->Enum = EnumType; }
+	
+	return Observation;
 }
 
-void UDirectionObservation::SetDirectionObservation(const int32 AgentId, const FVector Direction, const FRotator RelativeRotation)
+void UEnumObservation::SetEnumObservation(const int32 AgentId, const uint8 EnumValue)
 {
-	const TLearningArrayView<2, FVector> DirectionView = FeatureObject->InstanceData->View(FeatureObject->DirectionHandle);
-	const TLearningArrayView<1, FQuat> RelativeRotationView = FeatureObject->InstanceData->View(FeatureObject->RelativeRotationHandle);
-
-	if (AgentId == INDEX_NONE)
+	if (!AgentType->HasAgent(AgentId))
 	{
-		UE_LOG(LogLearning, Error, TEXT("AgentId is invalid (INDEX_NONE)"));
+		UE_LOG(LogLearning, Error, TEXT("%s: AgentId %d not found in the agents set."), *GetName(), AgentId);
 		return;
 	}
 
-	if (AgentId < 0 || AgentId >= DirectionView.Num<0>())
+	const int32 Index = Enum->GetIndexByValue(EnumValue);
+
+	if (Index == INDEX_NONE)
 	{
-		UE_LOG(LogLearning, Error, TEXT("AgentId %d is out of index. Valid range [0, %d]."), AgentId, DirectionView.Num<0>() - 1);
+		UE_LOG(LogLearning, Error, TEXT("%s: Invalid enum value: %i."), *GetName(), EnumValue);
 		return;
 	}
 
-	RelativeRotationView[AgentId] = FQuat::MakeFromRotator(RelativeRotation);
-	DirectionView[AgentId][0] = Direction;
+	const TLearningArrayView<2, float> EnumView = FeatureObject->InstanceData->View(FeatureObject->ValueHandle);
+	UE::Learning::Array::Zero(EnumView[AgentId]);
+	EnumView[AgentId][Index] = 1.0f;
 }
 
-#if ENABLE_VISUAL_LOG
-void UDirectionObservation::VisualLog(const UE::Learning::FIndexSet Instances) const
+#if UE_LEARNING_AGENTS_ENABLE_VISUAL_LOG
+void UEnumObservation::VisualLog(const UE::Learning::FIndexSet Instances) const
 {
-	UE_LEARNING_TRACE_CPUPROFILER_EVENT_SCOPE(UDirectionObservation::VisualLog);
+	UE_LEARNING_TRACE_CPUPROFILER_EVENT_SCOPE(UEnumObservation::VisualLog);
 
-	const TLearningArrayView<2, const FVector> DirectionView = FeatureObject->InstanceData->ConstView(FeatureObject->DirectionHandle);
-	const TLearningArrayView<1, const FQuat> RelativeRotationView = FeatureObject->InstanceData->ConstView(FeatureObject->RelativeRotationHandle);
+	const TLearningArrayView<2, const float> ValueView = FeatureObject->InstanceData->ConstView(FeatureObject->ValueHandle);
 	const TLearningArrayView<2, const float> FeatureView = FeatureObject->InstanceData->ConstView(FeatureObject->FeatureHandle);
 
-	if (const ULearningAgentsType* AgentType = Cast<ULearningAgentsType>(GetOuter()))
+	const int32 EnumNum = ValueView.Num<1>();
+		
+	for (const int32 Instance : Instances)
 	{
-		for (const int32 Instance : Instances)
+		if (const AActor* Actor = Cast<AActor>(AgentType->GetAgent(Instance)))
 		{
-			if (const AActor* Actor = Cast<AActor>(AgentType->GetAgent(Instance)))
-			{
-				const FVector Direction = DirectionView[Instance][0];
-				const FQuat RelativeRotation = RelativeRotationView[Instance];
-				const FVector LocalDirection = RelativeRotation.UnrotateVector(Direction);
+			int32 EnumEntryIdx = INDEX_NONE;
 
-				UE_VLOG_ARROW(this, LogLearning, Display,
-					Actor->GetActorLocation(),
-					Actor->GetActorLocation() + 100.0f * Direction,
-					VisualLogColor.ToFColor(true),
-					TEXT(""));
+			for (int32 EnumIdx = 0; EnumIdx < EnumNum; EnumIdx++)
+			{
+				if (ValueView[Instance][EnumIdx] != 0.0f)
+				{
+					EnumEntryIdx = EnumIdx;
+					break;
+				}
+			}
 
+			if (EnumEntryIdx != INDEX_NONE)
+			{
 				UE_LEARNING_AGENTS_VLOG_STRING(this, LogLearning, Display,
-					Actor->GetActorLocation() + 100.0f * Direction,
-					VisualLogColor.ToFColor(true),
-					TEXT("Direction: [% 6.3f % 6.3f % 6.3f]\nLocal Direction: [% 6.3f % 6.3f % 6.3f]"),
-					Direction.X, Direction.Y, Direction.Z,
-					LocalDirection.X, LocalDirection.Y, LocalDirection.Z);
-
-				UE_LEARNING_AGENTS_VLOG_TRANSFORM(this, LogLearning, Display,
 					Actor->GetActorLocation(),
-					RelativeRotation,
 					VisualLogColor.ToFColor(true),
-					TEXT("Agent %i\nScale: [% 6.2f]\nEncoded: [% 6.3f % 6.3f % 6.3f]"),
+					TEXT("Agent %i\nValue: %i\nIndex: %i\nName: \"%s\"\nEncoded: %s"),
 					Instance,
-					FeatureObject->Scale,
-					FeatureView[Instance][0], FeatureView[Instance][1], FeatureView[Instance][2]);
+					Enum->GetValueByIndex(EnumEntryIdx),
+					EnumEntryIdx,
+					*Enum->GetDisplayNameTextByIndex(EnumEntryIdx).ToString(),
+					*UE::Learning::Array::FormatFloat(FeatureView[Instance]));
+			}
+			else
+			{
+				UE_LOG(LogLearning, Error, TEXT("Invalid Enum encoding."));
 			}
 		}
 	}
 }
 #endif
 
-//------------------------------------------------------------------
-
-UPlanarPositionObservation* UPlanarPositionObservation::AddPlanarPositionObservation(ULearningAgentsType* AgentType, const FName Name, const float Scale, const FVector Axis0, const FVector Axis1)
+UEnumArrayObservation* UEnumArrayObservation::AddEnumArrayObservation(ULearningAgentsType* InAgentType, const UEnum* EnumType, const FName Name, const int32 EnumNum)
 {
-	return UE::Learning::Agents::Private::AddObservation<UPlanarPositionObservation, UE::Learning::FPlanarPositionFeature>(AgentType, Name, 1, Scale, Axis0, Axis1);
+	if (!EnumType)
+	{
+		UE_LOG(LogLearning, Error, TEXT("Invalid Enum."));
+		return nullptr;
+	}
+
+	if (EnumType->NumEnums() < 1)
+	{
+		UE_LOG(LogLearning, Error, TEXT("Enum requires at least one entry to be used as an observation."));
+		return nullptr;
+	}
+
+	if (EnumNum < 1)
+	{
+		UE_LOG(LogLearning, Error, TEXT("Number of elements in array must be at least 1, got %i."), EnumNum);
+		return nullptr;
+	}
+
+	UEnumArrayObservation* Observation = UE::Learning::Agents::Observations::Private::AddObservation<UEnumArrayObservation, UE::Learning::FFloatFeature>(InAgentType, Name, EnumNum * EnumType->NumEnums());
+	if (Observation) { Observation->Enum = EnumType; }
+
+	return Observation;
 }
 
-void UPlanarPositionObservation::SetPlanarPositionObservation(const int32 AgentId, const FVector Position, const FVector RelativePosition, const FRotator RelativeRotation)
+void UEnumArrayObservation::SetEnumArrayObservation(const int32 AgentId, const TArray<uint8>& EnumValues)
 {
-	const TLearningArrayView<2, FVector> PositionView = FeatureObject->InstanceData->View(FeatureObject->PositionHandle);
-	const TLearningArrayView<1, FVector> RelativePositionView = FeatureObject->InstanceData->View(FeatureObject->RelativePositionHandle);
-	const TLearningArrayView<1, FQuat> RelativeRotationView = FeatureObject->InstanceData->View(FeatureObject->RelativeRotationHandle);
-
-	if (AgentId == INDEX_NONE)
+	if (!AgentType->HasAgent(AgentId))
 	{
-		UE_LOG(LogLearning, Error, TEXT("AgentId is invalid (INDEX_NONE)"));
+		UE_LOG(LogLearning, Error, TEXT("%s: AgentId %d not found in the agents set."), *GetName(), AgentId);
 		return;
 	}
 
-	if (AgentId < 0 || AgentId >= PositionView.Num<0>())
+	const TLearningArrayView<2, float> EnumView = FeatureObject->InstanceData->View(FeatureObject->ValueHandle);
+
+	const int32 EnumEntryNum = Enum->NumEnums();
+
+	if (EnumValues.Num() != EnumView.Num<1>() / EnumEntryNum)
 	{
-		UE_LOG(LogLearning, Error, TEXT("AgentId %d is out of index. Valid range [0, %d]."), AgentId, PositionView.Num<0>() - 1);
+		UE_LOG(LogLearning, Error, TEXT("%s: Got wrong number of elements in array. Expected %i, got %i."), *GetName(), EnumView.Num<1>() / EnumEntryNum, EnumValues.Num());
 		return;
 	}
 
-	RelativePositionView[AgentId] = RelativePosition;
-	RelativeRotationView[AgentId] = FQuat::MakeFromRotator(RelativeRotation);
-	PositionView[AgentId][0] = Position;
+	for (int32 EnumIdx = 0; EnumIdx < EnumValues.Num(); EnumIdx++)
+	{
+		const TLearningArrayView<1, float> EnumViewSlice = EnumView[AgentId].Slice(EnumIdx * EnumEntryNum, EnumEntryNum);
+		UE::Learning::Array::Zero(EnumViewSlice);
+
+		const int32 Index = Enum->GetIndexByValue(EnumValues[EnumIdx]);
+
+		if (Index == INDEX_NONE)
+		{
+			UE_LOG(LogLearning, Error, TEXT("%s: Invalid enum value: %i"), *GetName(), EnumValues[EnumIdx]);
+			continue;
+		}
+
+		EnumViewSlice[Index] = 1.0f;
+	}
 }
 
-#if ENABLE_VISUAL_LOG
-void UPlanarPositionObservation::VisualLog(const UE::Learning::FIndexSet Instances) const
+#if UE_LEARNING_AGENTS_ENABLE_VISUAL_LOG
+void UEnumArrayObservation::VisualLog(const UE::Learning::FIndexSet Instances) const
 {
-	UE_LEARNING_TRACE_CPUPROFILER_EVENT_SCOPE(UPlanarPositionObservation::VisualLog);
+	UE_LEARNING_TRACE_CPUPROFILER_EVENT_SCOPE(UEnumArrayObservation::VisualLog);
 
-	const TLearningArrayView<2, const FVector> PositionView = FeatureObject->InstanceData->ConstView(FeatureObject->PositionHandle);
-	const TLearningArrayView<1, const FVector> RelativePositionView = FeatureObject->InstanceData->ConstView(FeatureObject->RelativePositionHandle);
-	const TLearningArrayView<1, const FQuat> RelativeRotationView = FeatureObject->InstanceData->ConstView(FeatureObject->RelativeRotationHandle);
+	const TLearningArrayView<2, const float> ValueView = FeatureObject->InstanceData->ConstView(FeatureObject->ValueHandle);
 	const TLearningArrayView<2, const float> FeatureView = FeatureObject->InstanceData->ConstView(FeatureObject->FeatureHandle);
 
-	if (const ULearningAgentsType* AgentType = Cast<ULearningAgentsType>(GetOuter()))
+	for (const int32 Instance : Instances)
 	{
-		for (const int32 Instance : Instances)
+		if (const AActor* Actor = Cast<AActor>(AgentType->GetAgent(Instance)))
 		{
-			if (const AActor* Actor = Cast<AActor>(AgentType->GetAgent(Instance)))
-			{
-				const FVector Position = PositionView[Instance][0];
-				const FVector RelativePosition = RelativePositionView[Instance];
-				const FQuat RelativeRotation = RelativeRotationView[Instance];
-				const FVector LocalPosition = RelativeRotation.UnrotateVector(Position - RelativePosition);
-
-				UE_VLOG_LOCATION(this, LogLearning, Display,
-					Position,
-					10.0f,
-					VisualLogColor.ToFColor(true),
-					TEXT("Position: [% 6.1f % 6.1f % 6.1f]\nLocal Position: [% 6.1f % 6.1f % 6.1f]"),
-					Position.X, Position.Y, Position.Z,
-					LocalPosition.X, LocalPosition.Y, LocalPosition.Z);
-
-				UE_VLOG_SEGMENT(this, LogLearning, Display,
-					RelativePosition,
-					Position,
-					VisualLogColor.ToFColor(true),
-					TEXT(""));
+			UE_LEARNING_AGENTS_VLOG_STRING(this, LogLearning, Display,
+				Actor->GetActorLocation(),
+				VisualLogColor.ToFColor(true),
+				TEXT("Agent %i\nEncoded: %s"),
+				Instance,
+				*UE::Learning::Array::FormatFloat(FeatureView[Instance]));
 
-				UE_LEARNING_AGENTS_VLOG_PLANE(this, LogLearning, Display,
-					RelativePosition,
-					RelativeRotation,
-					FeatureObject->Axis0,
-					FeatureObject->Axis1,
-					VisualLogColor.ToFColor(true),
-					TEXT(""));
-
-				UE_LEARNING_AGENTS_VLOG_TRANSFORM(this, LogLearning, Display,
-					RelativePosition,
-					RelativeRotation,
-					VisualLogColor.ToFColor(true),
-					TEXT("Agent %i\nScale: [% 6.2f]\nEncoded: [% 6.3f % 6.3f]"),
-					Instance,
-					FeatureObject->Scale,
-					FeatureView[Instance][0], FeatureView[Instance][1]);
-			}
 		}
 	}
 }
 #endif
 
-UPositionObservation* UPositionObservation::AddPositionObservation(ULearningAgentsType* AgentType, const FName Name, const float Scale)
+//------------------------------------------------------------------
+
+UTimeObservation* UTimeObservation::AddTimeObservation(ULearningAgentsType* InAgentType, const FName Name, const float Scale)
 {
-	return UE::Learning::Agents::Private::AddObservation<UPositionObservation, UE::Learning::FPositionFeature>(AgentType, Name, 1, Scale);
+	return UE::Learning::Agents::Observations::Private::AddObservation<UTimeObservation, UE::Learning::FTimeFeature>(InAgentType, Name, 1, Scale);
 }
 
-void UPositionObservation::SetPositionObservation(const int32 AgentId, const FVector Position, const FVector RelativePosition, const FRotator RelativeRotation)
+void UTimeObservation::SetTimeObservation(const int32 AgentId, const float Time, const float RelativeTime)
 {
-	const TLearningArrayView<2, FVector> PositionView = FeatureObject->InstanceData->View(FeatureObject->PositionHandle);
-	const TLearningArrayView<1, FVector> RelativePositionView = FeatureObject->InstanceData->View(FeatureObject->RelativePositionHandle);
-	const TLearningArrayView<1, FQuat> RelativeRotationView = FeatureObject->InstanceData->View(FeatureObject->RelativeRotationHandle);
-
-	if (AgentId == INDEX_NONE)
-	{
-		UE_LOG(LogLearning, Error, TEXT("AgentId is invalid (INDEX_NONE)"));
-		return;
-	}
-
-	if (AgentId < 0 || AgentId >= PositionView.Num<0>())
+	if (!AgentType->HasAgent(AgentId))
 	{
-		UE_LOG(LogLearning, Error, TEXT("AgentId %d is out of index. Valid range [0, %d]."), AgentId, PositionView.Num<0>() - 1);
+		UE_LOG(LogLearning, Error, TEXT("%s: AgentId %d not found in the agents set."), *GetName(), AgentId);
 		return;
 	}
 
-	RelativePositionView[AgentId] = RelativePosition;
-	RelativeRotationView[AgentId] = FQuat::MakeFromRotator(RelativeRotation);
-	PositionView[AgentId][0] = Position;
+	FeatureObject->InstanceData->View(FeatureObject->TimeHandle)[AgentId][0] = Time;
+	FeatureObject->InstanceData->View(FeatureObject->RelativeTimeHandle)[AgentId] = RelativeTime;
 }
 
-#if ENABLE_VISUAL_LOG
-void UPositionObservation::VisualLog(const UE::Learning::FIndexSet Instances) const
+#if UE_LEARNING_AGENTS_ENABLE_VISUAL_LOG
+void UTimeObservation::VisualLog(const UE::Learning::FIndexSet Instances) const
 {
-	UE_LEARNING_TRACE_CPUPROFILER_EVENT_SCOPE(UPositionObservation::VisualLog);
+	UE_LEARNING_TRACE_CPUPROFILER_EVENT_SCOPE(UTimeObservation::VisualLog);
 
-	const TLearningArrayView<2, const FVector> PositionView = FeatureObject->InstanceData->ConstView(FeatureObject->PositionHandle);
-	const TLearningArrayView<1, const FVector> RelativePositionView = FeatureObject->InstanceData->ConstView(FeatureObject->RelativePositionHandle);
-	const TLearningArrayView<1, const FQuat> RelativeRotationView = FeatureObject->InstanceData->ConstView(FeatureObject->RelativeRotationHandle);
+	const TLearningArrayView<2, const float> TimeView = FeatureObject->InstanceData->ConstView(FeatureObject->TimeHandle);
+	const TLearningArrayView<1, const float> RelativeTimeView = FeatureObject->InstanceData->ConstView(FeatureObject->RelativeTimeHandle);
 	const TLearningArrayView<2, const float> FeatureView = FeatureObject->InstanceData->ConstView(FeatureObject->FeatureHandle);
 
-	if (const ULearningAgentsType* AgentType = Cast<ULearningAgentsType>(GetOuter()))
+	for (const int32 Instance : Instances)
 	{
-		for (const int32 Instance : Instances)
+		if (const AActor* Actor = Cast<AActor>(AgentType->GetAgent(Instance)))
 		{
-			if (const AActor* Actor = Cast<AActor>(AgentType->GetAgent(Instance)))
-			{
-				const FVector Position = PositionView[Instance][0];
-				const FVector RelativePosition = RelativePositionView[Instance];
-				const FQuat RelativeRotation = RelativeRotationView[Instance];
-				const FVector LocalPosition = RelativeRotation.UnrotateVector(Position - RelativePosition);
-
-				UE_VLOG_LOCATION(this, LogLearning, Display,
-					Position,
-					10.0f,
-					VisualLogColor.ToFColor(true),
-					TEXT("Position: [% 6.1f % 6.1f % 6.1f]\nLocal Position: [% 6.1f % 6.1f % 6.1f]"),
-					Position.X, Position.Y, Position.Z,
-					LocalPosition.X, LocalPosition.Y, LocalPosition.Z);
-
-				UE_VLOG_SEGMENT(this, LogLearning, Display,
-					RelativePosition,
-					Position,
-					VisualLogColor.ToFColor(true),
-					TEXT(""));
-
-				UE_LEARNING_AGENTS_VLOG_TRANSFORM(this, LogLearning, Display,
-					RelativePosition,
-					RelativeRotation,
-					VisualLogColor.ToFColor(true),
-					TEXT("Agent %i\nScale: [% 6.2f]\nEncoded: [% 6.3f % 6.3f % 6.3f]"),
-					Instance,
-					FeatureObject->Scale,
-					FeatureView[Instance][0], FeatureView[Instance][1], FeatureView[Instance][2]);
-			}
+			const float Time = TimeView[Instance][0];
+			const float RelativeTime = RelativeTimeView[Instance];
+
+			UE_LEARNING_AGENTS_VLOG_STRING(this, LogLearning, Display,
+				Actor->GetActorLocation(),
+				VisualLogColor.ToFColor(true),
+				TEXT("Agent %i\nScale: [% 6.2f]\nRelative Time: [% 6.3f]\nTime: [% 6.3f]\nEncoded: %s"),
+				Instance,
+				FeatureObject->Scale,
+				RelativeTime,
+				Time,
+				*UE::Learning::Array::FormatFloat(FeatureView[Instance]));
 		}
 	}
 }
 #endif
 
-UPlanarPositionArrayObservation* UPlanarPositionArrayObservation::AddPlanarPositionArrayObservation(ULearningAgentsType* AgentType, const FName Name, const int32 PositionNum, const float Scale, const FVector Axis0, const FVector Axis1)
-{
-	return UE::Learning::Agents::Private::AddObservation<UPlanarPositionArrayObservation, UE::Learning::FPlanarPositionFeature>(AgentType, Name, PositionNum, Scale, Axis0, Axis1);
-}
-
-void UPlanarPositionArrayObservation::SetPlanarPositionArrayObservation(const int32 AgentId, const TArray<FVector>& Positions, const FVector RelativePosition, const FRotator RelativeRotation)
+UTimeArrayObservation* UTimeArrayObservation::AddTimeArrayObservation(ULearningAgentsType* InAgentType, const FName Name, const int32 TimeNum, const float Scale)
 {
-	const TLearningArrayView<2, FVector> PositionView = FeatureObject->InstanceData->View(FeatureObject->PositionHandle);
-	const TLearningArrayView<1, FVector> RelativePositionView = FeatureObject->InstanceData->View(FeatureObject->RelativePositionHandle);
-	const TLearningArrayView<1, FQuat> RelativeRotationView = FeatureObject->InstanceData->View(FeatureObject->RelativeRotationHandle);
-
-	if (AgentId == INDEX_NONE)
+	if (TimeNum < 1)
 	{
-		UE_LOG(LogLearning, Error, TEXT("AgentId is invalid (INDEX_NONE)"));
-		return;
+		UE_LOG(LogLearning, Error, TEXT("Number of elements in array must be at least 1, got %i."), TimeNum);
+		return nullptr;
 	}
 
-	if (AgentId < 0 || AgentId >= PositionView.Num<0>())
+	return UE::Learning::Agents::Observations::Private::AddObservation<UTimeArrayObservation, UE::Learning::FTimeFeature>(InAgentType, Name, TimeNum, Scale);
+}
+
+void UTimeArrayObservation::SetTimeArrayObservation(const int32 AgentId, const TArray<float>& Times, const float RelativeTime)
+{
+	if (!AgentType->HasAgent(AgentId))
 	{
-		UE_LOG(LogLearning, Error, TEXT("AgentId %d is out of index. Valid range [0, %d]."), AgentId, PositionView.Num<0>() - 1);
+		UE_LOG(LogLearning, Error, TEXT("%s: AgentId %d not found in the agents set."), *GetName(), AgentId);
 		return;
 	}
 
-	const int32 PositionNum = PositionView.Num<1>();
+	const TLearningArrayView<2, float> TimeView = FeatureObject->InstanceData->View(FeatureObject->TimeHandle);
+	const TLearningArrayView<1, float> RelativeTimeView = FeatureObject->InstanceData->View(FeatureObject->RelativeTimeHandle);
 
-	if (Positions.Num() != PositionNum)
+	if (Times.Num() != TimeView.Num<1>())
 	{
-		UE_LOG(LogLearning, Error, TEXT("Incorrect number of positions in array. Got %i, expected %i."), Positions.Num(), PositionNum);
+		UE_LOG(LogLearning, Error, TEXT("%s: Got wrong number of elements in array. Expected %i, got %i."), *GetName(), TimeView.Num<1>(), Times.Num());
 		return;
 	}
 
-	RelativePositionView[AgentId] = RelativePosition;
-	RelativeRotationView[AgentId] = FQuat::MakeFromRotator(RelativeRotation);
-	UE::Learning::Array::Copy(PositionView[AgentId], TLearningArrayView<1, const FVector>(Positions));
+	UE::Learning::Array::Copy<1, float>(TimeView[AgentId], Times);
+	RelativeTimeView[AgentId] = RelativeTime;
 }
 
-#if ENABLE_VISUAL_LOG
-void UPlanarPositionArrayObservation::VisualLog(const UE::Learning::FIndexSet Instances) const
+#if UE_LEARNING_AGENTS_ENABLE_VISUAL_LOG
+void UTimeArrayObservation::VisualLog(const UE::Learning::FIndexSet Instances) const
 {
-	UE_LEARNING_TRACE_CPUPROFILER_EVENT_SCOPE(UPlanarPositionArrayObservation::VisualLog);
+	UE_LEARNING_TRACE_CPUPROFILER_EVENT_SCOPE(UTimeArrayObservation::VisualLog);
 
-	const TLearningArrayView<2, const FVector> PositionView = FeatureObject->InstanceData->ConstView(FeatureObject->PositionHandle);
-	const TLearningArrayView<1, const FVector> RelativePositionView = FeatureObject->InstanceData->ConstView(FeatureObject->RelativePositionHandle);
-	const TLearningArrayView<1, const FQuat> RelativeRotationView = FeatureObject->InstanceData->ConstView(FeatureObject->RelativeRotationHandle);
+	const TLearningArrayView<2, const float> TimeView = FeatureObject->InstanceData->ConstView(FeatureObject->TimeHandle);
+	const TLearningArrayView<1, const float> RelativeTimeView = FeatureObject->InstanceData->ConstView(FeatureObject->RelativeTimeHandle);
 	const TLearningArrayView<2, const float> FeatureView = FeatureObject->InstanceData->ConstView(FeatureObject->FeatureHandle);
 
-	const int32 PositionNum = PositionView.Num<1>();
-
-	if (const ULearningAgentsType* AgentType = Cast<ULearningAgentsType>(GetOuter()))
+	for (const int32 Instance : Instances)
 	{
-		for (const int32 Instance : Instances)
+		if (const AActor* Actor = Cast<AActor>(AgentType->GetAgent(Instance)))
 		{
-			if (const AActor* Actor = Cast<AActor>(AgentType->GetAgent(Instance)))
-			{
-				const FVector RelativePosition = RelativePositionView[Instance];
-				const FQuat RelativeRotation = RelativeRotationView[Instance];
+			UE_LEARNING_AGENTS_VLOG_STRING(this, LogLearning, Display,
+				Actor->GetActorLocation(),
+				VisualLogColor.ToFColor(true),
+				TEXT("Agent %i\nScale: [% 6.2f]\nRelative Time: [% 6.3f]\nTimes: %s\nEncoded: %s"),
+				Instance,
+				FeatureObject->Scale,
+				RelativeTimeView[Instance],
+				*UE::Learning::Array::FormatFloat(TimeView[Instance]),
+				*UE::Learning::Array::FormatFloat(FeatureView[Instance]));
+		}
+	}
+}
+#endif
 
-				for (int32 PositionIdx = 0; PositionIdx < PositionNum; PositionIdx++)
-				{
-					const FVector LocalPosition = RelativeRotation.UnrotateVector(PositionView[Instance][PositionIdx] - RelativePosition);
-
-					UE_VLOG_LOCATION(this, LogLearning, Display,
-						PositionView[Instance][PositionIdx],
-						10.0f,
-						VisualLogColor.ToFColor(true),
-						TEXT("Position: [% 6.1f % 6.1f % 6.1f]\nLocal Position: [% 6.1f % 6.1f % 6.1f]"),
-						PositionView[Instance][PositionIdx].X, 
-						PositionView[Instance][PositionIdx].Y, 
-						PositionView[Instance][PositionIdx].Z,
-						LocalPosition.X, LocalPosition.Y, LocalPosition.Z);
-
-					UE_VLOG_SEGMENT(this, LogLearning, Display,
-						RelativePosition,
-						PositionView[Instance][PositionIdx],
-						VisualLogColor.ToFColor(true),
-						TEXT(""));
-				}
+//------------------------------------------------------------------
+
+UAngleObservation* UAngleObservation::AddAngleObservation(ULearningAgentsType* InAgentType, const FName Name, const float Scale)
+{
+	return UE::Learning::Agents::Observations::Private::AddObservation<UAngleObservation, UE::Learning::FAngleFeature>(InAgentType, Name, 1, Scale);
+}
+
+void UAngleObservation::SetAngleObservation(const int32 AgentId, const float Angle, const float RelativeAngle)
+{
+	if (!AgentType->HasAgent(AgentId))
+	{
+		UE_LOG(LogLearning, Error, TEXT("%s: AgentId %d not found in the agents set."), *GetName(), AgentId);
+		return;
+	}
+
+	FeatureObject->InstanceData->View(FeatureObject->AngleHandle)[AgentId][0] = FMath::DegreesToRadians(Angle);
+	FeatureObject->InstanceData->View(FeatureObject->RelativeAngleHandle)[AgentId] = FMath::DegreesToRadians(RelativeAngle);
+}
+
+#if UE_LEARNING_AGENTS_ENABLE_VISUAL_LOG
+void UAngleObservation::VisualLog(const UE::Learning::FIndexSet Instances) const
+{
+	UE_LEARNING_TRACE_CPUPROFILER_EVENT_SCOPE(UAngleObservation::VisualLog);
+
+	const TLearningArrayView<2, const float> AngleView = FeatureObject->InstanceData->ConstView(FeatureObject->AngleHandle);
+	const TLearningArrayView<1, const float> RelativeAngleView = FeatureObject->InstanceData->ConstView(FeatureObject->RelativeAngleHandle);
+	const TLearningArrayView<2, const float> FeatureView = FeatureObject->InstanceData->ConstView(FeatureObject->FeatureHandle);
+
+	for (const int32 Instance : Instances)
+	{
+		if (const AActor* Actor = Cast<AActor>(AgentType->GetAgent(Instance)))
+		{
+			const float Angle = AngleView[Instance][0];
+			const float RelativeAngle = RelativeAngleView[Instance];
+
+			UE_LEARNING_AGENTS_VLOG_ANGLE(this, LogLearning, Display,
+				Angle,
+				0.0f,
+				Actor->GetActorLocation(),
+				50.0f,
+				VisualLogColor.ToFColor(true),
+				TEXT("Agent %i\nScale: [% 6.2f]\nRelative Angle: [% 6.1f]\nAngle: [% 6.1f]\nEncoded: %s"),
+				Instance,
+				FeatureObject->Scale,
+				FMath::RadiansToDegrees(RelativeAngle),
+				FMath::RadiansToDegrees(Angle),
+				*UE::Learning::Array::FormatFloat(FeatureView[Instance]));
+		}
+	}
+}
+#endif
+
+UAngleArrayObservation* UAngleArrayObservation::AddAngleArrayObservation(ULearningAgentsType* InAgentType, const FName Name, const int32 AngleNum, const float Scale)
+{
+	if (AngleNum < 1)
+	{
+		UE_LOG(LogLearning, Error, TEXT("Number of elements in array must be at least 1, got %i."), AngleNum);
+		return nullptr;
+	}
+
+	return UE::Learning::Agents::Observations::Private::AddObservation<UAngleArrayObservation, UE::Learning::FAngleFeature>(InAgentType, Name, AngleNum, Scale);
+}
+
+void UAngleArrayObservation::SetAngleArrayObservation(const int32 AgentId, const TArray<float>& Angles, const float RelativeAngle)
+{
+	if (!AgentType->HasAgent(AgentId))
+	{
+		UE_LOG(LogLearning, Error, TEXT("%s: AgentId %d not found in the agents set."), *GetName(), AgentId);
+		return;
+	}
+
+	const TLearningArrayView<2, float> AngleView = FeatureObject->InstanceData->View(FeatureObject->AngleHandle);
+	const TLearningArrayView<1, float> RelativeAngleView = FeatureObject->InstanceData->View(FeatureObject->RelativeAngleHandle);
+
+	if (Angles.Num() != AngleView.Num<1>())
+	{
+		UE_LOG(LogLearning, Error, TEXT("%s: Got wrong number of elements in array. Expected %i, got %i."), *GetName(), AngleView.Num<1>(), Angles.Num());
+		return;
+	}
+
+	for (int32 AngleIdx = 0; AngleIdx < Angles.Num(); AngleIdx++)
+	{
+		AngleView[AgentId][AngleIdx] = FMath::DegreesToRadians(Angles[AngleIdx]);
+	}
+	
+	RelativeAngleView[AgentId] = FMath::DegreesToRadians(RelativeAngle);
+}
+
+#if UE_LEARNING_AGENTS_ENABLE_VISUAL_LOG
+void UAngleArrayObservation::VisualLog(const UE::Learning::FIndexSet Instances) const
+{
+	UE_LEARNING_TRACE_CPUPROFILER_EVENT_SCOPE(UAngleArrayObservation::VisualLog);
+
+	const TLearningArrayView<2, const float> AngleView = FeatureObject->InstanceData->ConstView(FeatureObject->AngleHandle);
+	const TLearningArrayView<1, const float> RelativeAngleView = FeatureObject->InstanceData->ConstView(FeatureObject->RelativeAngleHandle);
+	const TLearningArrayView<2, const float> FeatureView = FeatureObject->InstanceData->ConstView(FeatureObject->FeatureHandle);
+
+	const int32 AngleNum = AngleView.Num<1>();
+
+	for (const int32 Instance : Instances)
+	{
+		if (const AActor* Actor = Cast<AActor>(AgentType->GetAgent(Instance)))
+		{
+			const float RelativeAngle = RelativeAngleView[Instance];
+
+			for (int32 AngleIdx = 0; AngleIdx < AngleNum; AngleIdx++)
+			{
+				const float Angle = AngleView[Instance][AngleIdx];
+				const FVector Offset = UE::Learning::Agents::Debug::GridOffsetForIndex(AngleIdx, AngleNum);
+
+				UE_LEARNING_AGENTS_VLOG_ANGLE(this, LogLearning, Display,
+					Angle,
+					RelativeAngle,
+					Actor->GetActorLocation() + Offset,
+					5.0f,
+					VisualLogColor.ToFColor(true),
+					TEXT("Angle %i: [% 6.4f]"),
+					AngleIdx,
+					FMath::RadiansToDegrees(Angle));
+			}
+
+			UE_LEARNING_AGENTS_VLOG_STRING(this, LogLearning, Display,
+				Actor->GetActorLocation() + FVector(0.0f, 0.0f, 20.0f),
+				VisualLogColor.ToFColor(true),
+				TEXT("Agent %i\nScale: [% 6.2f]\nRelative Angle: [% 6.1f]\nEncoded: %s"),
+				Instance,
+				FeatureObject->Scale,
+				FMath::RadiansToDegrees(RelativeAngle),
+				*UE::Learning::Array::FormatFloat(FeatureView[Instance]));
+		}
+	}
+}
+#endif
+
+//------------------------------------------------------------------
+
+URotationObservation* URotationObservation::AddRotationObservation(ULearningAgentsType* InAgentType, const FName Name, const float Scale)
+{
+	return UE::Learning::Agents::Observations::Private::AddObservation<URotationObservation, UE::Learning::FRotationFeature>(InAgentType, Name, 1, Scale);
+}
+
+void URotationObservation::SetRotationObservation(const int32 AgentId, const FRotator Rotation, const FRotator RelativeRotation)
+{
+	SetRotationObservationFromQuat(AgentId, Rotation.Quaternion(), RelativeRotation.Quaternion());
+}
+
+void URotationObservation::SetRotationObservationFromQuat(const int32 AgentId, const FQuat Rotation, const FQuat RelativeRotation)
+{
+	if (!AgentType->HasAgent(AgentId))
+	{
+		UE_LOG(LogLearning, Error, TEXT("%s: AgentId %d not found in the agents set."), *GetName(), AgentId);
+		return;
+	}
+
+	FeatureObject->InstanceData->View(FeatureObject->RotationHandle)[AgentId][0] = Rotation;
+	FeatureObject->InstanceData->View(FeatureObject->RelativeRotationHandle)[AgentId] = RelativeRotation;
+}
+
+#if UE_LEARNING_AGENTS_ENABLE_VISUAL_LOG
+void URotationObservation::VisualLog(const UE::Learning::FIndexSet Instances) const
+{
+	UE_LEARNING_TRACE_CPUPROFILER_EVENT_SCOPE(URotationObservation::VisualLog);
+
+	const TLearningArrayView<2, const FQuat> RotationView = FeatureObject->InstanceData->ConstView(FeatureObject->RotationHandle);
+	const TLearningArrayView<1, const FQuat> RelativeRotationView = FeatureObject->InstanceData->ConstView(FeatureObject->RelativeRotationHandle);
+	const TLearningArrayView<2, const float> FeatureView = FeatureObject->InstanceData->ConstView(FeatureObject->FeatureHandle);
+
+	for (const int32 Instance : Instances)
+	{
+		if (const AActor* Actor = Cast<AActor>(AgentType->GetAgent(Instance)))
+		{
+			const FQuat Rotation = RotationView[Instance][0];
+			const FRotator Rotator = Rotation.Rotator();
+			const FQuat RelativeRotation = RelativeRotationView[Instance];
+			const FRotator RelativeRotator = RelativeRotation.Rotator();
+
+			UE_LEARNING_AGENTS_VLOG_TRANSFORM(this, LogLearning, Display,
+				Actor->GetActorLocation() + FVector(20.0f, 0.0f, 0.0f),
+				Rotation,
+				VisualLogColor.ToFColor(true),
+				TEXT("Relative Rotation: [% 6.1f % 6.1f % 6.1f]"),
+				RelativeRotator.Pitch, RelativeRotator.Roll, RelativeRotator.Yaw);
+
+			UE_LEARNING_AGENTS_VLOG_TRANSFORM(this, LogLearning, Display,
+				Actor->GetActorLocation(),
+				Rotation,
+				VisualLogColor.ToFColor(true),
+				TEXT("Agent %i\nScale: [% 6.2f]\nRotation: [% 6.1f % 6.1f % 6.1f]\nEncoded: %s"),
+				Instance,
+				FeatureObject->Scale,
+				Rotator.Pitch, Rotator.Roll, Rotator.Yaw,
+				*UE::Learning::Array::FormatFloat(FeatureView[Instance]));
+		}
+	}
+}
+#endif
+
+URotationArrayObservation* URotationArrayObservation::AddRotationArrayObservation(ULearningAgentsType* InAgentType, const FName Name, const int32 RotationNum, const float Scale)
+{
+	if (RotationNum < 1)
+	{
+		UE_LOG(LogLearning, Error, TEXT("Number of elements in array must be at least 1, got %i."), RotationNum);
+		return nullptr;
+	}
+
+	return UE::Learning::Agents::Observations::Private::AddObservation<URotationArrayObservation, UE::Learning::FRotationFeature>(InAgentType, Name, RotationNum, Scale);
+}
+
+void URotationArrayObservation::SetRotationArrayObservation(const int32 AgentId, const TArray<FRotator>& Rotations, const FRotator RelativeRotation)
+{
+	if (!AgentType->HasAgent(AgentId))
+	{
+		UE_LOG(LogLearning, Error, TEXT("%s: AgentId %d not found in the agents set."), *GetName(), AgentId);
+		return;
+	}
+
+	const TLearningArrayView<2, FQuat> RotationView = FeatureObject->InstanceData->View(FeatureObject->RotationHandle);
+	const TLearningArrayView<1, FQuat> RelativeRotationView = FeatureObject->InstanceData->View(FeatureObject->RelativeRotationHandle);
+
+	if (Rotations.Num() != RotationView.Num<1>())
+	{
+		UE_LOG(LogLearning, Error, TEXT("%s: Got wrong number of elements in array. Expected %i, got %i."), *GetName(), RotationView.Num<1>(), Rotations.Num());
+		return;
+	}
+
+	for (int32 RotationIdx = 0; RotationIdx < Rotations.Num(); RotationIdx++)
+	{
+		RotationView[AgentId][RotationIdx] = Rotations[RotationIdx].Quaternion();
+	}
+
+	RelativeRotationView[AgentId] = RelativeRotation.Quaternion();
+}
+
+void URotationArrayObservation::SetRotationArrayObservationFromQuats(const int32 AgentId, const TArray<FQuat>& Rotations, const FQuat RelativeRotation)
+{
+	if (!AgentType->HasAgent(AgentId))
+	{
+		UE_LOG(LogLearning, Error, TEXT("%s: AgentId %d not found in the agents set."), *GetName(), AgentId);
+		return;
+	}
+
+	const TLearningArrayView<2, FQuat> RotationView = FeatureObject->InstanceData->View(FeatureObject->RotationHandle);
+	const TLearningArrayView<1, FQuat> RelativeRotationView = FeatureObject->InstanceData->View(FeatureObject->RelativeRotationHandle);
+
+	if (Rotations.Num() != RotationView.Num<1>())
+	{
+		UE_LOG(LogLearning, Error, TEXT("%s: Got wrong number of elements in array. Expected %i, got %i."), *GetName(), RotationView.Num<1>(), Rotations.Num());
+		return;
+	}
+
+	UE::Learning::Array::Copy<1, FQuat>(RotationView[AgentId], Rotations);
+	RelativeRotationView[AgentId] = RelativeRotation;
+}
+
+#if UE_LEARNING_AGENTS_ENABLE_VISUAL_LOG
+void URotationArrayObservation::VisualLog(const UE::Learning::FIndexSet Instances) const
+{
+	UE_LEARNING_TRACE_CPUPROFILER_EVENT_SCOPE(URotationArrayObservation::VisualLog);
+
+	const TLearningArrayView<2, const FQuat> RotationView = FeatureObject->InstanceData->ConstView(FeatureObject->RotationHandle);
+	const TLearningArrayView<1, const FQuat> RelativeRotationView = FeatureObject->InstanceData->ConstView(FeatureObject->RelativeRotationHandle);
+	const TLearningArrayView<2, const float> FeatureView = FeatureObject->InstanceData->ConstView(FeatureObject->FeatureHandle);
+
+	const int32 RotationNum = RotationView.Num<1>();
+
+	for (const int32 Instance : Instances)
+	{
+		if (const AActor* Actor = Cast<AActor>(AgentType->GetAgent(Instance)))
+		{
+			const FQuat RelativeRotation = RelativeRotationView[Instance];
+			const FRotator RelativeRotator = RelativeRotation.Rotator();
+
+			for (int32 RotationIdx = 0; RotationIdx < RotationNum; RotationIdx++)
+			{
+				const FQuat Rotation = RotationView[Instance][RotationIdx];
+				const FRotator Rotator = Rotation.Rotator();
+				const FVector Offset = UE::Learning::Agents::Debug::GridOffsetForIndex(RotationIdx, RotationNum);
+
+				UE_LEARNING_AGENTS_VLOG_TRANSFORM(this, LogLearning, Display,
+					Actor->GetActorLocation() + Offset,
+					Rotation,
+					VisualLogColor.ToFColor(true),
+					TEXT("Rotation %i: [% 6.3f % 6.3f % 6.3f]"),
+					RotationIdx,
+					Rotator.Pitch, Rotator.Roll, Rotator.Yaw);
+			}
+
+			UE_LEARNING_AGENTS_VLOG_TRANSFORM(this, LogLearning, Display,
+				Actor->GetActorLocation(),
+				RelativeRotation,
+				VisualLogColor.ToFColor(true),
+				TEXT("Agent %i\nScale: [% 6.2f]\nRelative Rotation: [% 6.3f % 6.3f % 6.3f]\nEncoded: %s"),
+				Instance,
+				FeatureObject->Scale,
+				RelativeRotator.Pitch, RelativeRotator.Roll, RelativeRotator.Yaw,
+				*UE::Learning::Array::FormatFloat(FeatureView[Instance]));
+		}
+	}
+}
+#endif
+
+
+//------------------------------------------------------------------
+
+UDirectionObservation* UDirectionObservation::AddDirectionObservation(ULearningAgentsType* InAgentType, const FName Name, const float Scale)
+{
+	return UE::Learning::Agents::Observations::Private::AddObservation<UDirectionObservation, UE::Learning::FDirectionFeature>(InAgentType, Name, 1, Scale);
+}
+
+void UDirectionObservation::SetDirectionObservation(const int32 AgentId, const FVector Direction, const FRotator RelativeRotation)
+{
+	if (!AgentType->HasAgent(AgentId))
+	{
+		UE_LOG(LogLearning, Error, TEXT("%s: AgentId %d not found in the agents set."), *GetName(), AgentId);
+		return;
+	}
+
+	FeatureObject->InstanceData->View(FeatureObject->DirectionHandle)[AgentId][0] = Direction;
+	FeatureObject->InstanceData->View(FeatureObject->RelativeRotationHandle)[AgentId] = RelativeRotation.Quaternion();
+}
+
+#if UE_LEARNING_AGENTS_ENABLE_VISUAL_LOG
+void UDirectionObservation::VisualLog(const UE::Learning::FIndexSet Instances) const
+{
+	UE_LEARNING_TRACE_CPUPROFILER_EVENT_SCOPE(UDirectionObservation::VisualLog);
+
+	const TLearningArrayView<2, const FVector> DirectionView = FeatureObject->InstanceData->ConstView(FeatureObject->DirectionHandle);
+	const TLearningArrayView<1, const FQuat> RelativeRotationView = FeatureObject->InstanceData->ConstView(FeatureObject->RelativeRotationHandle);
+	const TLearningArrayView<2, const float> FeatureView = FeatureObject->InstanceData->ConstView(FeatureObject->FeatureHandle);
+
+	for (const int32 Instance : Instances)
+	{
+		if (const AActor* Actor = Cast<AActor>(AgentType->GetAgent(Instance)))
+		{
+			const FVector Direction = DirectionView[Instance][0];
+			const FQuat RelativeRotation = RelativeRotationView[Instance];
+			const FVector LocalDirection = RelativeRotation.UnrotateVector(Direction);
+
+			UE_LEARNING_AGENTS_VLOG_ARROW(this, LogLearning, Display,
+				Actor->GetActorLocation(),
+				Actor->GetActorLocation() + 100.0f * Direction,
+				VisualLogColor.ToFColor(true),
+				TEXT(""));
+
+			UE_LEARNING_AGENTS_VLOG_STRING(this, LogLearning, Display,
+				Actor->GetActorLocation() + 100.0f * Direction,
+				VisualLogColor.ToFColor(true),
+				TEXT("Direction: [% 6.3f % 6.3f % 6.3f]\nLocal Direction: [% 6.3f % 6.3f % 6.3f]"),
+				Direction.X, Direction.Y, Direction.Z,
+				LocalDirection.X, LocalDirection.Y, LocalDirection.Z);
+
+			UE_LEARNING_AGENTS_VLOG_TRANSFORM(this, LogLearning, Display,
+				Actor->GetActorLocation(),
+				RelativeRotation,
+				VisualLogColor.ToFColor(true),
+				TEXT("Agent %i\nScale: [% 6.2f]\nEncoded: %s"),
+				Instance,
+				FeatureObject->Scale,
+				*UE::Learning::Array::FormatFloat(FeatureView[Instance]));
+		}
+	}
+}
+#endif
+
+UDirectionArrayObservation* UDirectionArrayObservation::AddDirectionArrayObservation(ULearningAgentsType* InAgentType, const FName Name, const int32 DirectionNum, const float Scale)
+{
+	if (DirectionNum < 1)
+	{
+		UE_LOG(LogLearning, Error, TEXT("Number of elements in array must be at least 1, got %i."), DirectionNum);
+		return nullptr;
+	}
+
+	return UE::Learning::Agents::Observations::Private::AddObservation<UDirectionArrayObservation, UE::Learning::FDirectionFeature>(InAgentType, Name, DirectionNum, Scale);
+}
+
+void UDirectionArrayObservation::SetDirectionArrayObservation(const int32 AgentId, const TArray<FVector>& Directions, const FRotator RelativeRotation)
+{
+	if (!AgentType->HasAgent(AgentId))
+	{
+		UE_LOG(LogLearning, Error, TEXT("%s: AgentId %d not found in the agents set."), *GetName(), AgentId);
+		return;
+	}
+
+	const TLearningArrayView<2, FVector> DirectionView = FeatureObject->InstanceData->View(FeatureObject->DirectionHandle);
+	const TLearningArrayView<1, FQuat> RelativeRotationView = FeatureObject->InstanceData->View(FeatureObject->RelativeRotationHandle);
+
+	if (Directions.Num() != DirectionView.Num<1>())
+	{
+		UE_LOG(LogLearning, Error, TEXT("%s: Got wrong number of elements in array. Expected %i, got %i."), *GetName(), DirectionView.Num<1>(), Directions.Num());
+		return;
+	}
+
+	UE::Learning::Array::Copy<1, FVector>(DirectionView[AgentId], Directions);
+	RelativeRotationView[AgentId] = RelativeRotation.Quaternion();
+}
+
+#if UE_LEARNING_AGENTS_ENABLE_VISUAL_LOG
+void UDirectionArrayObservation::VisualLog(const UE::Learning::FIndexSet Instances) const
+{
+	UE_LEARNING_TRACE_CPUPROFILER_EVENT_SCOPE(UDirectionArrayObservation::VisualLog);
+
+	const TLearningArrayView<2, const FVector> DirectionView = FeatureObject->InstanceData->ConstView(FeatureObject->DirectionHandle);
+	const TLearningArrayView<1, const FQuat> RelativeRotationView = FeatureObject->InstanceData->ConstView(FeatureObject->RelativeRotationHandle);
+	const TLearningArrayView<2, const float> FeatureView = FeatureObject->InstanceData->ConstView(FeatureObject->FeatureHandle);
+
+	const int32 DirectionNum = DirectionView.Num<1>();
+
+	for (const int32 Instance : Instances)
+	{
+		if (const AActor* Actor = Cast<AActor>(AgentType->GetAgent(Instance)))
+		{
+			const FQuat RelativeRotation = RelativeRotationView[Instance];
+
+			for (int32 DirectionIdx = 0; DirectionIdx < DirectionNum; DirectionIdx++)
+			{
+				const FVector Direction = DirectionView[Instance][0];
+				const FVector LocalDirection = RelativeRotation.UnrotateVector(Direction);
+				const FVector Offset = UE::Learning::Agents::Debug::GridOffsetForIndex(DirectionIdx, DirectionNum);
+
+				UE_LEARNING_AGENTS_VLOG_ARROW(this, LogLearning, Display,
+					Actor->GetActorLocation() + Offset,
+					Actor->GetActorLocation() + Offset + 10.0f * Direction,
+					VisualLogColor.ToFColor(true),
+					TEXT(""));
+
+				UE_LEARNING_AGENTS_VLOG_STRING(this, LogLearning, Display,
+					Actor->GetActorLocation() + Offset + 10.0f * Direction,
+					VisualLogColor.ToFColor(true),
+					TEXT("Direction %i: [% 6.3f % 6.3f % 6.3f]\nLocal Direction: [% 6.3f % 6.3f % 6.3f]"),
+					DirectionIdx,
+					Direction.X, Direction.Y, Direction.Z,
+					LocalDirection.X, LocalDirection.Y, LocalDirection.Z);
+			}
+
+			UE_LEARNING_AGENTS_VLOG_TRANSFORM(this, LogLearning, Display,
+				Actor->GetActorLocation(),
+				RelativeRotation,
+				VisualLogColor.ToFColor(true),
+				TEXT("Agent %i\nScale: [% 6.2f]\nEncoded: %s"),
+				Instance,
+				FeatureObject->Scale,
+				*UE::Learning::Array::FormatFloat(FeatureView[Instance]));
+		}
+	}
+}
+#endif
+
+UPlanarDirectionObservation* UPlanarDirectionObservation::AddPlanarDirectionObservation(ULearningAgentsType* InAgentType, const FName Name, const float Scale, const FVector Axis0, const FVector Axis1)
+{
+	return UE::Learning::Agents::Observations::Private::AddObservation<UPlanarDirectionObservation, UE::Learning::FPlanarDirectionFeature>(InAgentType, Name, 1, Scale, Axis0, Axis1);
+}
+
+void UPlanarDirectionObservation::SetPlanarDirectionObservation(const int32 AgentId, const FVector Direction, const FRotator RelativeRotation)
+{
+	if (!AgentType->HasAgent(AgentId))
+	{
+		UE_LOG(LogLearning, Error, TEXT("%s: AgentId %d not found in the agents set."), *GetName(), AgentId);
+		return;
+	}
+
+	FeatureObject->InstanceData->View(FeatureObject->DirectionHandle)[AgentId][0] = Direction;
+	FeatureObject->InstanceData->View(FeatureObject->RelativeRotationHandle)[AgentId] = RelativeRotation.Quaternion();
+}
+
+#if UE_LEARNING_AGENTS_ENABLE_VISUAL_LOG
+void UPlanarDirectionObservation::VisualLog(const UE::Learning::FIndexSet Instances) const
+{
+	UE_LEARNING_TRACE_CPUPROFILER_EVENT_SCOPE(UPlanarDirectionObservation::VisualLog);
+
+	const TLearningArrayView<2, const FVector> DirectionView = FeatureObject->InstanceData->ConstView(FeatureObject->DirectionHandle);
+	const TLearningArrayView<1, const FQuat> RelativeRotationView = FeatureObject->InstanceData->ConstView(FeatureObject->RelativeRotationHandle);
+	const TLearningArrayView<2, const float> FeatureView = FeatureObject->InstanceData->ConstView(FeatureObject->FeatureHandle);
+
+	for (const int32 Instance : Instances)
+	{
+		if (const AActor* Actor = Cast<AActor>(AgentType->GetAgent(Instance)))
+		{
+			const FVector Direction = DirectionView[Instance][0];
+			const FQuat RelativeRotation = RelativeRotationView[Instance];
+			const FVector LocalDirection = RelativeRotation.UnrotateVector(Direction);
+
+			UE_LEARNING_AGENTS_VLOG_ARROW(this, LogLearning, Display,
+				Actor->GetActorLocation(),
+				Actor->GetActorLocation() + 100.0f * Direction,
+				VisualLogColor.ToFColor(true),
+				TEXT(""));
+
+			UE_LEARNING_AGENTS_VLOG_STRING(this, LogLearning, Display,
+				Actor->GetActorLocation() + 100.0f * Direction,
+				VisualLogColor.ToFColor(true),
+				TEXT("Direction: [% 6.3f % 6.3f % 6.3f]\nLocal Direction: [% 6.3f % 6.3f % 6.3f]"),
+				Direction.X, Direction.Y, Direction.Z,
+				LocalDirection.X, LocalDirection.Y, LocalDirection.Z);
+
+			UE_LEARNING_AGENTS_VLOG_PLANE(this, LogLearning, Display,
+				Actor->GetActorLocation(),
+				RelativeRotation,
+				FeatureObject->Axis0,
+				FeatureObject->Axis1,
+				VisualLogColor.ToFColor(true),
+				TEXT(""));
+
+			UE_LEARNING_AGENTS_VLOG_TRANSFORM(this, LogLearning, Display,
+				Actor->GetActorLocation(),
+				RelativeRotation,
+				VisualLogColor.ToFColor(true),
+				TEXT("Agent %i\nScale: [% 6.2f]\nEncoded: %s"),
+				Instance,
+				FeatureObject->Scale,
+				*UE::Learning::Array::FormatFloat(FeatureView[Instance]));
+		}
+	}
+}
+#endif
+
+UPlanarDirectionArrayObservation* UPlanarDirectionArrayObservation::AddPlanarDirectionArrayObservation(ULearningAgentsType* InAgentType, const FName Name, const int32 DirectionNum, const float Scale, const FVector Axis0, const FVector Axis1)
+{
+	return UE::Learning::Agents::Observations::Private::AddObservation<UPlanarDirectionArrayObservation, UE::Learning::FPlanarDirectionFeature>(InAgentType, Name, DirectionNum, Scale, Axis0, Axis1);
+}
+
+void UPlanarDirectionArrayObservation::SetPlanarDirectionArrayObservation(const int32 AgentId, const TArray<FVector>& Directions, const FRotator RelativeRotation)
+{
+	if (!AgentType->HasAgent(AgentId))
+	{
+		UE_LOG(LogLearning, Error, TEXT("%s: AgentId %d not found in the agents set."), *GetName(), AgentId);
+		return;
+	}
+
+	const TLearningArrayView<2, FVector> DirectionView = FeatureObject->InstanceData->View(FeatureObject->DirectionHandle);
+	const TLearningArrayView<1, FQuat> RelativeRotationView = FeatureObject->InstanceData->View(FeatureObject->RelativeRotationHandle);
+
+	if (Directions.Num() != DirectionView.Num<1>())
+	{
+		UE_LOG(LogLearning, Error, TEXT("%s: Got wrong number of elements in array. Expected %i, got %i."), *GetName(), DirectionView.Num<1>(), Directions.Num());
+		return;
+	}
+
+	UE::Learning::Array::Copy<1, FVector>(DirectionView[AgentId], Directions);
+	RelativeRotationView[AgentId] = RelativeRotation.Quaternion();
+}
+
+#if UE_LEARNING_AGENTS_ENABLE_VISUAL_LOG
+void UPlanarDirectionArrayObservation::VisualLog(const UE::Learning::FIndexSet Instances) const
+{
+	UE_LEARNING_TRACE_CPUPROFILER_EVENT_SCOPE(UPlanarDirectionArrayObservation::VisualLog);
+
+	const TLearningArrayView<2, const FVector> DirectionView = FeatureObject->InstanceData->ConstView(FeatureObject->DirectionHandle);
+	const TLearningArrayView<1, const FQuat> RelativeRotationView = FeatureObject->InstanceData->ConstView(FeatureObject->RelativeRotationHandle);
+	const TLearningArrayView<2, const float> FeatureView = FeatureObject->InstanceData->ConstView(FeatureObject->FeatureHandle);
+
+	const int32 DirectionNum = DirectionView.Num<1>();
+	
+	for (const int32 Instance : Instances)
+	{
+		if (const AActor* Actor = Cast<AActor>(AgentType->GetAgent(Instance)))
+		{
+			const FQuat RelativeRotation = RelativeRotationView[Instance];
+
+			for (int32 DirectionIdx = 0; DirectionIdx < DirectionNum; DirectionIdx++)
+			{
+				const FVector Direction = DirectionView[Instance][0];
+				const FVector LocalDirection = RelativeRotation.UnrotateVector(Direction);
+				const FVector Offset = UE::Learning::Agents::Debug::GridOffsetForIndex(DirectionIdx, DirectionNum);
+
+				UE_LEARNING_AGENTS_VLOG_ARROW(this, LogLearning, Display,
+					Actor->GetActorLocation() + Offset,
+					Actor->GetActorLocation() + Offset + 10.0f * Direction,
+					VisualLogColor.ToFColor(true),
+					TEXT(""));
+
+				UE_LEARNING_AGENTS_VLOG_STRING(this, LogLearning, Display,
+					Actor->GetActorLocation() + Offset + 10.0f * Direction,
+					VisualLogColor.ToFColor(true),
+					TEXT("Direction %i: [% 6.3f % 6.3f % 6.3f]\nLocal Direction: [% 6.3f % 6.3f % 6.3f]"),
+					DirectionIdx,
+					Direction.X, Direction.Y, Direction.Z,
+					LocalDirection.X, LocalDirection.Y, LocalDirection.Z);
+			}
+
+			UE_LEARNING_AGENTS_VLOG_PLANE(this, LogLearning, Display,
+				Actor->GetActorLocation() + FVector(0.0f, 0.0f, 20.0f),
+				RelativeRotation,
+				FeatureObject->Axis0,
+				FeatureObject->Axis1,
+				VisualLogColor.ToFColor(true),
+				TEXT(""));
+
+			UE_LEARNING_AGENTS_VLOG_TRANSFORM(this, LogLearning, Display,
+				Actor->GetActorLocation(),
+				RelativeRotation,
+				VisualLogColor.ToFColor(true),
+				TEXT("Agent %i\nScale: [% 6.2f]\nEncoded: %s"),
+				Instance,
+				FeatureObject->Scale,
+				*UE::Learning::Array::FormatFloat(FeatureView[Instance]));
+		}
+	}
+}
+#endif
+
+
+
+//------------------------------------------------------------------
+
+UPositionObservation* UPositionObservation::AddPositionObservation(ULearningAgentsType* InAgentType, const FName Name, const float Scale)
+{
+	return UE::Learning::Agents::Observations::Private::AddObservation<UPositionObservation, UE::Learning::FPositionFeature>(InAgentType, Name, 1, Scale);
+}
+
+void UPositionObservation::SetPositionObservation(const int32 AgentId, const FVector Position, const FVector RelativePosition, const FRotator RelativeRotation)
+{
+	if (!AgentType->HasAgent(AgentId))
+	{
+		UE_LOG(LogLearning, Error, TEXT("%s: AgentId %d not found in the agents set."), *GetName(), AgentId);
+		return;
+	}
+
+	FeatureObject->InstanceData->View(FeatureObject->PositionHandle)[AgentId][0] = Position;
+	FeatureObject->InstanceData->View(FeatureObject->RelativePositionHandle)[AgentId] = RelativePosition;
+	FeatureObject->InstanceData->View(FeatureObject->RelativeRotationHandle)[AgentId] = RelativeRotation.Quaternion();
+}
+
+#if UE_LEARNING_AGENTS_ENABLE_VISUAL_LOG
+void UPositionObservation::VisualLog(const UE::Learning::FIndexSet Instances) const
+{
+	UE_LEARNING_TRACE_CPUPROFILER_EVENT_SCOPE(UPositionObservation::VisualLog);
+
+	const TLearningArrayView<2, const FVector> PositionView = FeatureObject->InstanceData->ConstView(FeatureObject->PositionHandle);
+	const TLearningArrayView<1, const FVector> RelativePositionView = FeatureObject->InstanceData->ConstView(FeatureObject->RelativePositionHandle);
+	const TLearningArrayView<1, const FQuat> RelativeRotationView = FeatureObject->InstanceData->ConstView(FeatureObject->RelativeRotationHandle);
+	const TLearningArrayView<2, const float> FeatureView = FeatureObject->InstanceData->ConstView(FeatureObject->FeatureHandle);
+
+	for (const int32 Instance : Instances)
+	{
+		const FVector Position = PositionView[Instance][0];
+		const FVector RelativePosition = RelativePositionView[Instance];
+		const FQuat RelativeRotation = RelativeRotationView[Instance];
+		const FVector LocalPosition = RelativeRotation.UnrotateVector(Position - RelativePosition);
+
+		UE_LEARNING_AGENTS_VLOG_LOCATION(this, LogLearning, Display,
+			Position,
+			10.0f,
+			VisualLogColor.ToFColor(true),
+			TEXT("Position: [% 6.1f % 6.1f % 6.1f]\nLocal Position: [% 6.1f % 6.1f % 6.1f]"),
+			Position.X, Position.Y, Position.Z,
+			LocalPosition.X, LocalPosition.Y, LocalPosition.Z);
+
+		UE_LEARNING_AGENTS_VLOG_SEGMENT(this, LogLearning, Display,
+			RelativePosition,
+			Position,
+			VisualLogColor.ToFColor(true),
+			TEXT(""));
+
+		UE_LEARNING_AGENTS_VLOG_TRANSFORM(this, LogLearning, Display,
+			RelativePosition,
+			RelativeRotation,
+			VisualLogColor.ToFColor(true),
+			TEXT("Agent %i\nScale: [% 6.2f]\nEncoded: %s"),
+			Instance,
+			FeatureObject->Scale,
+			*UE::Learning::Array::FormatFloat(FeatureView[Instance]));
+	}
+}
+#endif
+
+UPositionArrayObservation* UPositionArrayObservation::AddPositionArrayObservation(ULearningAgentsType* InAgentType, const FName Name, const int32 PositionNum, const float Scale)
+{
+	if (PositionNum < 1)
+	{
+		UE_LOG(LogLearning, Error, TEXT("Number of elements in array must be at least 1, got %i."), PositionNum);
+		return nullptr;
+	}
+
+	return UE::Learning::Agents::Observations::Private::AddObservation<UPositionArrayObservation, UE::Learning::FPositionFeature>(InAgentType, Name, PositionNum, Scale);
+}
+
+void UPositionArrayObservation::SetPositionArrayObservation(const int32 AgentId, const TArray<FVector>& Positions, const FVector RelativePosition, const FRotator RelativeRotation)
+{
+	if (!AgentType->HasAgent(AgentId))
+	{
+		UE_LOG(LogLearning, Error, TEXT("%s: AgentId %d not found in the agents set."), *GetName(), AgentId);
+		return;
+	}
+
+	const TLearningArrayView<2, FVector> PositionView = FeatureObject->InstanceData->View(FeatureObject->PositionHandle);
+	const TLearningArrayView<1, FVector> RelativePositionView = FeatureObject->InstanceData->View(FeatureObject->RelativePositionHandle);
+	const TLearningArrayView<1, FQuat> RelativeRotationView = FeatureObject->InstanceData->View(FeatureObject->RelativeRotationHandle);
+
+	if (Positions.Num() != PositionView.Num<1>())
+	{
+		UE_LOG(LogLearning, Error, TEXT("%s: Got wrong number of elements in array. Expected %i, got %i."), *GetName(), PositionView.Num<1>(), Positions.Num());
+		return;
+	}
+
+	RelativePositionView[AgentId] = RelativePosition;
+	RelativeRotationView[AgentId] = RelativeRotation.Quaternion();
+	UE::Learning::Array::Copy<1, FVector>(PositionView[AgentId], Positions);
+}
+
+#if UE_LEARNING_AGENTS_ENABLE_VISUAL_LOG
+void UPositionArrayObservation::VisualLog(const UE::Learning::FIndexSet Instances) const
+{
+	UE_LEARNING_TRACE_CPUPROFILER_EVENT_SCOPE(UPositionArrayObservation::VisualLog);
+
+	const TLearningArrayView<2, const FVector> PositionView = FeatureObject->InstanceData->ConstView(FeatureObject->PositionHandle);
+	const TLearningArrayView<1, const FVector> RelativePositionView = FeatureObject->InstanceData->ConstView(FeatureObject->RelativePositionHandle);
+	const TLearningArrayView<1, const FQuat> RelativeRotationView = FeatureObject->InstanceData->ConstView(FeatureObject->RelativeRotationHandle);
+	const TLearningArrayView<2, const float> FeatureView = FeatureObject->InstanceData->ConstView(FeatureObject->FeatureHandle);
+
+	const int32 PositionNum = PositionView.Num<1>();
+
+	for (const int32 Instance : Instances)
+	{
+		const FVector RelativePosition = RelativePositionView[Instance];
+		const FQuat RelativeRotation = RelativeRotationView[Instance];
+
+		for (int32 PositionIdx = 0; PositionIdx < PositionNum; PositionIdx++)
+		{
+			const FVector LocalPosition = RelativeRotation.UnrotateVector(PositionView[Instance][PositionIdx] - RelativePosition);
+
+			UE_LEARNING_AGENTS_VLOG_LOCATION(this, LogLearning, Display,
+				PositionView[Instance][PositionIdx],
+				10.0f,
+				VisualLogColor.ToFColor(true),
+				TEXT("Position: [% 6.1f % 6.1f % 6.1f]\nLocal Position: [% 6.1f % 6.1f % 6.1f]"),
+				PositionView[Instance][PositionIdx].X,
+				PositionView[Instance][PositionIdx].Y,
+				PositionView[Instance][PositionIdx].Z,
+				LocalPosition.X, LocalPosition.Y, LocalPosition.Z);
+
+			UE_LEARNING_AGENTS_VLOG_SEGMENT(this, LogLearning, Display,
+				RelativePosition,
+				PositionView[Instance][PositionIdx],
+				VisualLogColor.ToFColor(true),
+				TEXT(""));
+		}
+
+		UE_LEARNING_AGENTS_VLOG_TRANSFORM(this, LogLearning, Display,
+			RelativePosition,
+			RelativeRotation,
+			VisualLogColor.ToFColor(true),
+			TEXT("Agent %i\nScale: [% 6.2f]\nEncoded: %s"),
+			Instance,
+			FeatureObject->Scale,
+			*UE::Learning::Array::FormatFloat(FeatureView[Instance]));
+	}
+}
+#endif
+
+UScalarPositionObservation* UScalarPositionObservation::AddScalarPositionObservation(ULearningAgentsType* InAgentType, const FName Name, const float Scale)
+{
+	return UE::Learning::Agents::Observations::Private::AddObservation<UScalarPositionObservation, UE::Learning::FScalarPositionFeature>(InAgentType, Name, 1, Scale);
+}
+
+void UScalarPositionObservation::SetScalarPositionObservation(const int32 AgentId, const float Position, const float RelativePosition)
+{
+	if (!AgentType->HasAgent(AgentId))
+	{
+		UE_LOG(LogLearning, Error, TEXT("%s: AgentId %d not found in the agents set."), *GetName(), AgentId);
+		return;
+	}
+
+	FeatureObject->InstanceData->View(FeatureObject->PositionHandle)[AgentId][0] = Position;
+	FeatureObject->InstanceData->View(FeatureObject->RelativePositionHandle)[AgentId] = RelativePosition;
+}
+
+void UScalarPositionObservation::SetScalarPositionObservationWithAxis(const int32 AgentId, const FVector Position, const FVector RelativePosition, const FVector Axis)
+{
+	SetScalarPositionObservation(Position.Dot(Axis), RelativePosition.Dot(Axis));
+}
+
+#if UE_LEARNING_AGENTS_ENABLE_VISUAL_LOG
+void UScalarPositionObservation::VisualLog(const UE::Learning::FIndexSet Instances) const
+{
+	UE_LEARNING_TRACE_CPUPROFILER_EVENT_SCOPE(UScalarPositionObservation::VisualLog);
+
+	const TLearningArrayView<2, const float> PositionView = FeatureObject->InstanceData->ConstView(FeatureObject->PositionHandle);
+	const TLearningArrayView<1, const float> RelativePositionView = FeatureObject->InstanceData->ConstView(FeatureObject->RelativePositionHandle);
+	const TLearningArrayView<2, const float> FeatureView = FeatureObject->InstanceData->ConstView(FeatureObject->FeatureHandle);
+
+	for (const int32 Instance : Instances)
+	{
+		if (const AActor* Actor = Cast<AActor>(AgentType->GetAgent(Instance)))
+		{
+			UE_LEARNING_AGENTS_VLOG_STRING(this, LogLearning, Display,
+				Actor->GetActorLocation(),
+				VisualLogColor.ToFColor(true),
+				TEXT("Agent %i\nScale: [% 6.2f]\nRelative Position [% 6.1f]\nLocal Position: [% 6.1f]\nPosition: [% 6.1f]\nEncoded: %s"),
+				Instance,
+				FeatureObject->Scale,
+				RelativePositionView[Instance],
+				PositionView[Instance][0] - RelativePositionView[Instance],
+				PositionView[Instance][0],
+				*UE::Learning::Array::FormatFloat(FeatureView[Instance]));
+		}
+	}
+}
+#endif
+
+UScalarPositionArrayObservation* UScalarPositionArrayObservation::AddScalarPositionArrayObservation(ULearningAgentsType* InAgentType, const FName Name, const int32 PositionNum, const float Scale)
+{
+	if (PositionNum < 1)
+	{
+		UE_LOG(LogLearning, Error, TEXT("Number of elements in array must be at least 1, got %i."), PositionNum);
+		return nullptr;
+	}
+
+	return UE::Learning::Agents::Observations::Private::AddObservation<UScalarPositionArrayObservation, UE::Learning::FScalarPositionFeature>(InAgentType, Name, PositionNum, Scale);
+}
+
+void UScalarPositionArrayObservation::SetScalarPositionArrayObservation(const int32 AgentId, const TArray<float>& Positions, const float RelativePosition)
+{
+	if (!AgentType->HasAgent(AgentId))
+	{
+		UE_LOG(LogLearning, Error, TEXT("%s: AgentId %d not found in the agents set."), *GetName(), AgentId);
+		return;
+	}
+
+	const TLearningArrayView<2, float> PositionView = FeatureObject->InstanceData->View(FeatureObject->PositionHandle);
+	const TLearningArrayView<1, float> RelativePositionView = FeatureObject->InstanceData->View(FeatureObject->RelativePositionHandle);
+
+	if (Positions.Num() != PositionView.Num<1>())
+	{
+		UE_LOG(LogLearning, Error, TEXT("%s: Got wrong number of elements in array. Expected %i, got %i."), *GetName(), PositionView.Num<1>(), Positions.Num());
+		return;
+	}
+
+	RelativePositionView[AgentId] = RelativePosition;
+	UE::Learning::Array::Copy<1, float>(PositionView[AgentId], Positions);
+}
+
+void UScalarPositionArrayObservation::SetScalarPositionArrayObservationWithAxis(const int32 AgentId, const TArray<FVector>& Positions, const FVector RelativePosition, const FVector Axis)
+{
+	if (!AgentType->HasAgent(AgentId))
+	{
+		UE_LOG(LogLearning, Error, TEXT("%s: AgentId %d not found in the agents set."), *GetName(), AgentId);
+		return;
+	}
+
+	const TLearningArrayView<2, float> PositionView = FeatureObject->InstanceData->View(FeatureObject->PositionHandle);
+	const TLearningArrayView<1, float> RelativePositionView = FeatureObject->InstanceData->View(FeatureObject->RelativePositionHandle);
+
+	if (Positions.Num() != PositionView.Num<1>())
+	{
+		UE_LOG(LogLearning, Error, TEXT("%s: Got wrong number of elements in array. Expected %i, got %i."), *GetName(), PositionView.Num<1>(), Positions.Num());
+		return;
+	}
+
+	RelativePositionView[AgentId] = RelativePosition.Dot(Axis);
+	for (int32 PositionIdx = 0; PositionIdx < PositionView.Num<1>(); PositionIdx++)
+	{
+		PositionView[AgentId][PositionIdx] = Positions[PositionIdx].Dot(Axis);
+	}
+}
+
+#if UE_LEARNING_AGENTS_ENABLE_VISUAL_LOG
+void UScalarPositionArrayObservation::VisualLog(const UE::Learning::FIndexSet Instances) const
+{
+	UE_LEARNING_TRACE_CPUPROFILER_EVENT_SCOPE(UScalarPositionArrayObservation::VisualLog);
+
+	const TLearningArrayView<2, const float> PositionView = FeatureObject->InstanceData->ConstView(FeatureObject->PositionHandle);
+	const TLearningArrayView<1, const float> RelativePositionView = FeatureObject->InstanceData->ConstView(FeatureObject->RelativePositionHandle);
+	const TLearningArrayView<2, const float> FeatureView = FeatureObject->InstanceData->ConstView(FeatureObject->FeatureHandle);
+
+	const int32 PositionNum = PositionView.Num<1>();
+
+	for (const int32 Instance : Instances)
+	{
+		if (const AActor* Actor = Cast<AActor>(AgentType->GetAgent(Instance)))
+		{
+			UE_LEARNING_AGENTS_VLOG_STRING(this, LogLearning, Display,
+				Actor->GetActorLocation(),
+				VisualLogColor.ToFColor(true),
+				TEXT("Agent %i\nScale: [% 6.2f]\nRelative Position: [% 6.2f]\nPositions: %s\nEncoded: %s"),
+				Instance,
+				FeatureObject->Scale,
+				RelativePositionView[Instance],
+				*UE::Learning::Array::FormatFloat(PositionView[Instance]),
+				*UE::Learning::Array::FormatFloat(FeatureView[Instance]));
+		}
+	}
+}
+#endif
+
+UPlanarPositionObservation* UPlanarPositionObservation::AddPlanarPositionObservation(ULearningAgentsType* InAgentType, const FName Name, const float Scale, const FVector Axis0, const FVector Axis1)
+{
+	return UE::Learning::Agents::Observations::Private::AddObservation<UPlanarPositionObservation, UE::Learning::FPlanarPositionFeature>(InAgentType, Name, 1, Scale, Axis0, Axis1);
+}
+
+void UPlanarPositionObservation::SetPlanarPositionObservation(const int32 AgentId, const FVector Position, const FVector RelativePosition, const FRotator RelativeRotation)
+{
+	if (!AgentType->HasAgent(AgentId))
+	{
+		UE_LOG(LogLearning, Error, TEXT("%s: AgentId %d not found in the agents set."), *GetName(), AgentId);
+		return;
+	}
+
+	FeatureObject->InstanceData->View(FeatureObject->PositionHandle)[AgentId][0] = Position;
+	FeatureObject->InstanceData->View(FeatureObject->RelativePositionHandle)[AgentId] = RelativePosition;
+	FeatureObject->InstanceData->View(FeatureObject->RelativeRotationHandle)[AgentId] = RelativeRotation.Quaternion();
+}
+
+#if UE_LEARNING_AGENTS_ENABLE_VISUAL_LOG
+void UPlanarPositionObservation::VisualLog(const UE::Learning::FIndexSet Instances) const
+{
+	UE_LEARNING_TRACE_CPUPROFILER_EVENT_SCOPE(UPlanarPositionObservation::VisualLog);
+
+	const TLearningArrayView<2, const FVector> PositionView = FeatureObject->InstanceData->ConstView(FeatureObject->PositionHandle);
+	const TLearningArrayView<1, const FVector> RelativePositionView = FeatureObject->InstanceData->ConstView(FeatureObject->RelativePositionHandle);
+	const TLearningArrayView<1, const FQuat> RelativeRotationView = FeatureObject->InstanceData->ConstView(FeatureObject->RelativeRotationHandle);
+	const TLearningArrayView<2, const float> FeatureView = FeatureObject->InstanceData->ConstView(FeatureObject->FeatureHandle);
+
+	for (const int32 Instance : Instances)
+	{
+		const FVector Position = PositionView[Instance][0];
+		const FVector RelativePosition = RelativePositionView[Instance];
+		const FQuat RelativeRotation = RelativeRotationView[Instance];
+		const FVector LocalPosition = RelativeRotation.UnrotateVector(Position - RelativePosition);
+
+		UE_LEARNING_AGENTS_VLOG_LOCATION(this, LogLearning, Display,
+			Position,
+			10.0f,
+			VisualLogColor.ToFColor(true),
+			TEXT("Position: [% 6.1f % 6.1f % 6.1f]\nLocal Position: [% 6.1f % 6.1f % 6.1f]"),
+			Position.X, Position.Y, Position.Z,
+			LocalPosition.X, LocalPosition.Y, LocalPosition.Z);
+
+		UE_LEARNING_AGENTS_VLOG_SEGMENT(this, LogLearning, Display,
+			RelativePosition,
+			Position,
+			VisualLogColor.ToFColor(true),
+			TEXT(""));
+
+		UE_LEARNING_AGENTS_VLOG_PLANE(this, LogLearning, Display,
+			RelativePosition,
+			RelativeRotation,
+			FeatureObject->Axis0,
+			FeatureObject->Axis1,
+			VisualLogColor.ToFColor(true),
+			TEXT(""));
+
+		UE_LEARNING_AGENTS_VLOG_TRANSFORM(this, LogLearning, Display,
+			RelativePosition,
+			RelativeRotation,
+			VisualLogColor.ToFColor(true),
+			TEXT("Agent %i\nScale: [% 6.2f]\nEncoded: %s"),
+			Instance,
+			FeatureObject->Scale,
+			*UE::Learning::Array::FormatFloat(FeatureView[Instance]));
+	}
+}
+#endif
+
+UPlanarPositionArrayObservation* UPlanarPositionArrayObservation::AddPlanarPositionArrayObservation(ULearningAgentsType* InAgentType, const FName Name, const int32 PositionNum, const float Scale, const FVector Axis0, const FVector Axis1)
+{
+	if (PositionNum < 1)
+	{
+		UE_LOG(LogLearning, Error, TEXT("Number of elements in array must be at least 1, got %i."), PositionNum);
+		return nullptr;
+	}
+
+	return UE::Learning::Agents::Observations::Private::AddObservation<UPlanarPositionArrayObservation, UE::Learning::FPlanarPositionFeature>(InAgentType, Name, PositionNum, Scale, Axis0, Axis1);
+}
+
+void UPlanarPositionArrayObservation::SetPlanarPositionArrayObservation(const int32 AgentId, const TArray<FVector>& Positions, const FVector RelativePosition, const FRotator RelativeRotation)
+{
+	if (!AgentType->HasAgent(AgentId))
+	{
+		UE_LOG(LogLearning, Error, TEXT("%s: AgentId %d not found in the agents set."), *GetName(), AgentId);
+		return;
+	}
+
+	const TLearningArrayView<2, FVector> PositionView = FeatureObject->InstanceData->View(FeatureObject->PositionHandle);
+	const TLearningArrayView<1, FVector> RelativePositionView = FeatureObject->InstanceData->View(FeatureObject->RelativePositionHandle);
+	const TLearningArrayView<1, FQuat> RelativeRotationView = FeatureObject->InstanceData->View(FeatureObject->RelativeRotationHandle);
+
+	if (Positions.Num() != PositionView.Num<1>())
+	{
+		UE_LOG(LogLearning, Error, TEXT("%s: Got wrong number of elements in array. Expected %i, got %i."), *GetName(), PositionView.Num<1>(), Positions.Num());
+		return;
+	}
+
+	RelativePositionView[AgentId] = RelativePosition;
+	RelativeRotationView[AgentId] = RelativeRotation.Quaternion();
+	UE::Learning::Array::Copy<1, FVector>(PositionView[AgentId], Positions);
+}
+
+#if UE_LEARNING_AGENTS_ENABLE_VISUAL_LOG
+void UPlanarPositionArrayObservation::VisualLog(const UE::Learning::FIndexSet Instances) const
+{
+	UE_LEARNING_TRACE_CPUPROFILER_EVENT_SCOPE(UPlanarPositionArrayObservation::VisualLog);
+
+	const TLearningArrayView<2, const FVector> PositionView = FeatureObject->InstanceData->ConstView(FeatureObject->PositionHandle);
+	const TLearningArrayView<1, const FVector> RelativePositionView = FeatureObject->InstanceData->ConstView(FeatureObject->RelativePositionHandle);
+	const TLearningArrayView<1, const FQuat> RelativeRotationView = FeatureObject->InstanceData->ConstView(FeatureObject->RelativeRotationHandle);
+	const TLearningArrayView<2, const float> FeatureView = FeatureObject->InstanceData->ConstView(FeatureObject->FeatureHandle);
+
+	const int32 PositionNum = PositionView.Num<1>();
+
+	for (const int32 Instance : Instances)
+	{
+		const FVector RelativePosition = RelativePositionView[Instance];
+		const FQuat RelativeRotation = RelativeRotationView[Instance];
+
+		for (int32 PositionIdx = 0; PositionIdx < PositionNum; PositionIdx++)
+		{
+			const FVector LocalPosition = RelativeRotation.UnrotateVector(PositionView[Instance][PositionIdx] - RelativePosition);
+
+			UE_LEARNING_AGENTS_VLOG_LOCATION(this, LogLearning, Display,
+				PositionView[Instance][PositionIdx],
+				10.0f,
+				VisualLogColor.ToFColor(true),
+				TEXT("Position: [% 6.1f % 6.1f % 6.1f]\nLocal Position: [% 6.1f % 6.1f % 6.1f]"),
+				PositionView[Instance][PositionIdx].X, 
+				PositionView[Instance][PositionIdx].Y, 
+				PositionView[Instance][PositionIdx].Z,
+				LocalPosition.X, LocalPosition.Y, LocalPosition.Z);
+
+			UE_LEARNING_AGENTS_VLOG_SEGMENT(this, LogLearning, Display,
+				RelativePosition,
+				PositionView[Instance][PositionIdx],
+				VisualLogColor.ToFColor(true),
+				TEXT(""));
+		}
+
+		UE_LEARNING_AGENTS_VLOG_PLANE(this, LogLearning, Display,
+			RelativePosition,
+			RelativeRotation,
+			FeatureObject->Axis0,
+			FeatureObject->Axis1,
+			VisualLogColor.ToFColor(true),
+			TEXT(""));
+
+
+		UE_LEARNING_AGENTS_VLOG_TRANSFORM(this, LogLearning, Display,
+			RelativePosition,
+			RelativeRotation,
+			VisualLogColor.ToFColor(true),
+			TEXT("Agent %i\nScale: [% 6.2f]\nEncoded: %s"),
+			Instance,
+			FeatureObject->Scale,
+			*UE::Learning::Array::FormatFloat(FeatureView[Instance]));
+	}
+}
+#endif
+
+//------------------------------------------------------------------
+
+UVelocityObservation* UVelocityObservation::AddVelocityObservation(ULearningAgentsType* InAgentType, const FName Name, const float Scale)
+{
+	return UE::Learning::Agents::Observations::Private::AddObservation<UVelocityObservation, UE::Learning::FVelocityFeature>(InAgentType, Name, 1, Scale);
+}
+
+void UVelocityObservation::SetVelocityObservation(const int32 AgentId, const FVector Velocity, const FRotator RelativeRotation)
+{
+	if (!AgentType->HasAgent(AgentId))
+	{
+		UE_LOG(LogLearning, Error, TEXT("%s: AgentId %d not found in the agents set."), *GetName(), AgentId);
+		return;
+	}
+
+	FeatureObject->InstanceData->View(FeatureObject->VelocityHandle)[AgentId][0] = Velocity;
+	FeatureObject->InstanceData->View(FeatureObject->RelativeRotationHandle)[AgentId] = RelativeRotation.Quaternion();
+}
+
+#if UE_LEARNING_AGENTS_ENABLE_VISUAL_LOG
+void UVelocityObservation::VisualLog(const UE::Learning::FIndexSet Instances) const
+{
+	UE_LEARNING_TRACE_CPUPROFILER_EVENT_SCOPE(UVelocityObservation::VisualLog);
+
+	const TLearningArrayView<2, const FVector> VelocityView = FeatureObject->InstanceData->ConstView(FeatureObject->VelocityHandle);
+	const TLearningArrayView<1, const FQuat> RelativeRotationView = FeatureObject->InstanceData->ConstView(FeatureObject->RelativeRotationHandle);
+	const TLearningArrayView<2, const float> FeatureView = FeatureObject->InstanceData->ConstView(FeatureObject->FeatureHandle);
+
+	for (const int32 Instance : Instances)
+	{
+		if (const AActor* Actor = Cast<AActor>(AgentType->GetAgent(Instance)))
+		{
+			const FVector Velocity = VelocityView[Instance][0];
+			const FQuat RelativeRotation = RelativeRotationView[Instance];
+			const FVector LocalVelocity = RelativeRotation.UnrotateVector(Velocity);
+
+			UE_LEARNING_AGENTS_VLOG_ARROW(this, LogLearning, Display,
+				Actor->GetActorLocation(),
+				Actor->GetActorLocation() + Velocity,
+				VisualLogColor.ToFColor(true),
+				TEXT(""));
+
+			UE_LEARNING_AGENTS_VLOG_STRING(this, LogLearning, Display,
+				Actor->GetActorLocation() + Velocity,
+				VisualLogColor.ToFColor(true),
+				TEXT("Velocity: [% 6.3f % 6.3f % 6.3f]\nLocal Velocity: [% 6.3f % 6.3f % 6.3f]"),
+				Velocity.X, Velocity.Y, Velocity.Z,
+				LocalVelocity.X, LocalVelocity.Y, LocalVelocity.Z);
+
+			UE_LEARNING_AGENTS_VLOG_TRANSFORM(this, LogLearning, Display,
+				Actor->GetActorLocation(),
+				RelativeRotation,
+				VisualLogColor.ToFColor(true),
+				TEXT("Agent %i\nScale: [% 6.2f]\nEncoded: %s"),
+				Instance,
+				FeatureObject->Scale,
+				*UE::Learning::Array::FormatFloat(FeatureView[Instance]));
+		}
+	}
+}
+#endif
+
+UVelocityArrayObservation* UVelocityArrayObservation::AddVelocityArrayObservation(ULearningAgentsType* InAgentType, const FName Name, const int32 VelocityNum, const float Scale)
+{
+	if (VelocityNum < 1)
+	{
+		UE_LOG(LogLearning, Error, TEXT("Number of elements in array must be at least 1, got %i."), VelocityNum);
+		return nullptr;
+	}
+
+	return UE::Learning::Agents::Observations::Private::AddObservation<UVelocityArrayObservation, UE::Learning::FVelocityFeature>(InAgentType, Name, VelocityNum, Scale);
+}
+
+void UVelocityArrayObservation::SetVelocityArrayObservation(const int32 AgentId, const TArray<FVector>& Velocities, const FRotator RelativeRotation)
+{
+	if (!AgentType->HasAgent(AgentId))
+	{
+		UE_LOG(LogLearning, Error, TEXT("%s: AgentId %d not found in the agents set."), *GetName(), AgentId);
+		return;
+	}
 
-				UE_LEARNING_AGENTS_VLOG_PLANE(this, LogLearning, Display,
-					RelativePosition,
-					RelativeRotation,
-					FeatureObject->Axis0,
-					FeatureObject->Axis1,
+	const TLearningArrayView<2, FVector> VelocityView = FeatureObject->InstanceData->View(FeatureObject->VelocityHandle);
+	const TLearningArrayView<1, FQuat> RelativeRotationView = FeatureObject->InstanceData->View(FeatureObject->RelativeRotationHandle);
+
+	if (Velocities.Num() != VelocityView.Num<1>())
+	{
+		UE_LOG(LogLearning, Error, TEXT("%s: Got wrong number of elements in array. Expected %i, got %i."), *GetName(), VelocityView.Num<1>(), Velocities.Num());
+		return;
+	}
+
+	UE::Learning::Array::Copy<1, FVector>(VelocityView[AgentId], Velocities);
+	RelativeRotationView[AgentId] = RelativeRotation.Quaternion();
+}
+
+#if UE_LEARNING_AGENTS_ENABLE_VISUAL_LOG
+void UVelocityArrayObservation::VisualLog(const UE::Learning::FIndexSet Instances) const
+{
+	UE_LEARNING_TRACE_CPUPROFILER_EVENT_SCOPE(UVelocityArrayObservation::VisualLog);
+
+	const TLearningArrayView<2, const FVector> VelocityView = FeatureObject->InstanceData->ConstView(FeatureObject->VelocityHandle);
+	const TLearningArrayView<1, const FQuat> RelativeRotationView = FeatureObject->InstanceData->ConstView(FeatureObject->RelativeRotationHandle);
+	const TLearningArrayView<2, const float> FeatureView = FeatureObject->InstanceData->ConstView(FeatureObject->FeatureHandle);
+
+	const int32 VelocityNum = VelocityView.Num<1>();
+
+	for (const int32 Instance : Instances)
+	{
+		if (const AActor* Actor = Cast<AActor>(AgentType->GetAgent(Instance)))
+		{
+			const FQuat RelativeRotation = RelativeRotationView[Instance];
+				
+			for (int32 VelocityIdx = 0; VelocityIdx < VelocityNum; VelocityIdx++)
+			{
+				const FVector Velocity = VelocityView[Instance][0];
+				const FVector LocalVelocity = RelativeRotation.UnrotateVector(Velocity);
+				const FVector Offset = UE::Learning::Agents::Debug::GridOffsetForIndex(VelocityIdx, VelocityNum);
+
+				UE_LEARNING_AGENTS_VLOG_ARROW(this, LogLearning, Display,
+					Actor->GetActorLocation() + Offset,
+					Actor->GetActorLocation() + Offset + Velocity,
 					VisualLogColor.ToFColor(true),
 					TEXT(""));
 
-				if (PositionNum > 0)
-				{
-					UE_LEARNING_AGENTS_VLOG_TRANSFORM(this, LogLearning, Display,
-						RelativePosition,
-						RelativeRotation,
-						VisualLogColor.ToFColor(true),
-						TEXT("Agent %i\nEncoded: [% 6.3f % 6.3f ...]"),
-						Instance,
-						FeatureObject->Scale,
-						FeatureView[Instance][0], FeatureView[Instance][1]);
-				}
-				else
-				{
-					UE_LEARNING_AGENTS_VLOG_TRANSFORM(this, LogLearning, Display,
-						RelativePosition,
-						RelativeRotation,
-						VisualLogColor.ToFColor(true),
-						TEXT("Agent %i\nScale: [% 6.2f]\nEncoded: []"),
-						Instance,
-						FeatureObject->Scale);
-				}
+				UE_LEARNING_AGENTS_VLOG_STRING(this, LogLearning, Display,
+					Actor->GetActorLocation() + Offset + Velocity,
+					VisualLogColor.ToFColor(true),
+					TEXT("Velocity %i: [% 6.3f % 6.3f % 6.3f]\nLocal Velocity: [% 6.3f % 6.3f % 6.3f]"),
+					VelocityIdx,
+					Velocity.X, Velocity.Y, Velocity.Z,
+					LocalVelocity.X, LocalVelocity.Y, LocalVelocity.Z);
 			}
+
+			UE_LEARNING_AGENTS_VLOG_TRANSFORM(this, LogLearning, Display,
+				Actor->GetActorLocation(),
+				RelativeRotation,
+				VisualLogColor.ToFColor(true),
+				TEXT("Agent %i\nScale: [% 6.2f]\nEncoded: %s"),
+				Instance,
+				FeatureObject->Scale,
+				*UE::Learning::Array::FormatFloat(FeatureView[Instance]));
 		}
 	}
 }
 #endif
 
-UPositionArrayObservation* UPositionArrayObservation::AddPositionArrayObservation(ULearningAgentsType* AgentType, const FName Name, const int32 PositionNum, const float Scale)
+UScalarVelocityObservation* UScalarVelocityObservation::AddScalarVelocityObservation(ULearningAgentsType* InAgentType, const FName Name, const float Scale)
 {
-	return UE::Learning::Agents::Private::AddObservation<UPositionArrayObservation, UE::Learning::FPositionFeature>(AgentType, Name, PositionNum, Scale);
+	return UE::Learning::Agents::Observations::Private::AddObservation<UScalarVelocityObservation, UE::Learning::FScalarVelocityFeature>(InAgentType, Name, 1, Scale);
 }
 
-void UPositionArrayObservation::SetPositionArrayObservation(const int32 AgentId, const TArray<FVector>& Positions, const FVector RelativePosition, const FRotator RelativeRotation)
+void UScalarVelocityObservation::SetScalarVelocityObservation(const int32 AgentId, const float Velocity)
 {
-	const TLearningArrayView<2, FVector> PositionView = FeatureObject->InstanceData->View(FeatureObject->PositionHandle);
-	const TLearningArrayView<1, FVector> RelativePositionView = FeatureObject->InstanceData->View(FeatureObject->RelativePositionHandle);
-	const TLearningArrayView<1, FQuat> RelativeRotationView = FeatureObject->InstanceData->View(FeatureObject->RelativeRotationHandle);
-	
-	if (AgentId == INDEX_NONE)
+	if (!AgentType->HasAgent(AgentId))
 	{
-		UE_LOG(LogLearning, Error, TEXT("AgentId is invalid (INDEX_NONE)"));
+		UE_LOG(LogLearning, Error, TEXT("%s: AgentId %d not found in the agents set."), *GetName(), AgentId);
 		return;
 	}
 
-	if (AgentId < 0 || AgentId >= PositionView.Num<0>())
+	FeatureObject->InstanceData->View(FeatureObject->VelocityHandle)[AgentId][0] = Velocity;
+}
+
+void UScalarVelocityObservation::SetScalarVelocityObservationWithAxis(const int32 AgentId, const FVector Velocity, const FVector Axis)
+{
+	SetScalarVelocityObservation(AgentId, Velocity.Dot(Axis));
+}
+
+#if UE_LEARNING_AGENTS_ENABLE_VISUAL_LOG
+void UScalarVelocityObservation::VisualLog(const UE::Learning::FIndexSet Instances) const
+{
+	UE_LEARNING_TRACE_CPUPROFILER_EVENT_SCOPE(UScalarVelocityObservation::VisualLog);
+
+	const TLearningArrayView<2, const float> VelocityView = FeatureObject->InstanceData->ConstView(FeatureObject->VelocityHandle);
+	const TLearningArrayView<2, const float> FeatureView = FeatureObject->InstanceData->ConstView(FeatureObject->FeatureHandle);
+
+	for (const int32 Instance : Instances)
+	{
+		if (const AActor* Actor = Cast<AActor>(AgentType->GetAgent(Instance)))
+		{
+			UE_LEARNING_AGENTS_VLOG_STRING(this, LogLearning, Display,
+				Actor->GetActorLocation(),
+				VisualLogColor.ToFColor(true),
+				TEXT("Agent %i\nScale: [% 6.2f]\nVelocity: [% 6.1f]\nEncoded: %s"),
+				Instance,
+				FeatureObject->Scale,
+				VelocityView[Instance][0],
+				*UE::Learning::Array::FormatFloat(FeatureView[Instance]));
+		}
+	}
+}
+#endif
+
+
+UScalarVelocityArrayObservation* UScalarVelocityArrayObservation::AddScalarVelocityArrayObservation(ULearningAgentsType* InAgentType, const FName Name, const int32 VelocityNum, const float Scale)
+{
+	if (VelocityNum < 1)
+	{
+		UE_LOG(LogLearning, Error, TEXT("Number of elements in array must be at least 1, got %i."), VelocityNum);
+		return nullptr;
+	}
+
+	return UE::Learning::Agents::Observations::Private::AddObservation<UScalarVelocityArrayObservation, UE::Learning::FScalarVelocityFeature>(InAgentType, Name, VelocityNum, Scale);
+}
+
+void UScalarVelocityArrayObservation::SetScalarVelocityArrayObservation(const int32 AgentId, const TArray<float>& Velocities)
+{
+	if (!AgentType->HasAgent(AgentId))
 	{
-		UE_LOG(LogLearning, Error, TEXT("AgentId %d is out of index. Valid range [0, %d]."), AgentId, PositionView.Num<0>() - 1);
+		UE_LOG(LogLearning, Error, TEXT("%s: AgentId %d not found in the agents set."), *GetName(), AgentId);
 		return;
 	}
 
-	const int32 PositionNum = PositionView.Num<1>();
+	const TLearningArrayView<2, float> VelocityView = FeatureObject->InstanceData->View(FeatureObject->VelocityHandle);
 
-	if (Positions.Num() != PositionNum)
+	if (Velocities.Num() != VelocityView.Num<1>())
 	{
-		UE_LOG(LogLearning, Error, TEXT("Incorrect number of positions in array. Got %i, expected %i."), Positions.Num(), PositionNum);
+		UE_LOG(LogLearning, Error, TEXT("%s: Got wrong number of elements in array. Expected %i, got %i."), *GetName(), VelocityView.Num<1>(), Velocities.Num());
 		return;
 	}
 
-	RelativePositionView[AgentId] = RelativePosition;
-	RelativeRotationView[AgentId] = FQuat::MakeFromRotator(RelativeRotation);
-	UE::Learning::Array::Copy(PositionView[AgentId], TLearningArrayView<1, const FVector>(Positions));
+	UE::Learning::Array::Copy<1, float>(VelocityView[AgentId], Velocities);
 }
 
-#if ENABLE_VISUAL_LOG
-void UPositionArrayObservation::VisualLog(const UE::Learning::FIndexSet Instances) const
+void UScalarVelocityArrayObservation::SetScalarVelocityArrayObservationWithAxis(const int32 AgentId, const TArray<FVector>& Velocities, const FVector Axis)
 {
-	UE_LEARNING_TRACE_CPUPROFILER_EVENT_SCOPE(UPositionArrayObservation::VisualLog);
+	if (!AgentType->HasAgent(AgentId))
+	{
+		UE_LOG(LogLearning, Error, TEXT("%s: AgentId %d not found in the agents set."), *GetName(), AgentId);
+		return;
+	}
 
-	const TLearningArrayView<2, const FVector> PositionView = FeatureObject->InstanceData->ConstView(FeatureObject->PositionHandle);
-	const TLearningArrayView<1, const FVector> RelativePositionView = FeatureObject->InstanceData->ConstView(FeatureObject->RelativePositionHandle);
-	const TLearningArrayView<1, const FQuat> RelativeRotationView = FeatureObject->InstanceData->ConstView(FeatureObject->RelativeRotationHandle);
-	const TLearningArrayView<2, const float> FeatureView = FeatureObject->InstanceData->ConstView(FeatureObject->FeatureHandle);
+	const TLearningArrayView<2, float> VelocityView = FeatureObject->InstanceData->View(FeatureObject->VelocityHandle);
 
-	const int32 PositionNum = PositionView.Num<1>();
+	if (Velocities.Num() != VelocityView.Num<1>())
+	{
+		UE_LOG(LogLearning, Error, TEXT("%s: Got wrong number of elements in array. Expected %i, got %i."), *GetName(), VelocityView.Num<1>(), Velocities.Num());
+		return;
+	}
 
-	if (const ULearningAgentsType* AgentType = Cast<ULearningAgentsType>(GetOuter()))
+	for (int32 VelocityIdx = 0; VelocityIdx < VelocityView.Num<1>(); VelocityIdx++)
 	{
-		for (const int32 Instance : Instances)
-		{
-			if (const AActor* Actor = Cast<AActor>(AgentType->GetAgent(Instance)))
-			{
-				const FVector RelativePosition = RelativePositionView[Instance];
-				const FQuat RelativeRotation = RelativeRotationView[Instance];
+		VelocityView[AgentId][VelocityIdx] = Velocities[VelocityIdx].Dot(Axis);
+	}
+}
 
-				for (int32 PositionIdx = 0; PositionIdx < PositionNum; PositionIdx++)
-				{
-					const FVector LocalPosition = RelativeRotation.UnrotateVector(PositionView[Instance][PositionIdx] - RelativePosition);
-
-					UE_VLOG_LOCATION(this, LogLearning, Display,
-						PositionView[Instance][PositionIdx],
-						10.0f,
-						VisualLogColor.ToFColor(true),
-						TEXT("Position: [% 6.1f % 6.1f % 6.1f]\nLocal Position: [% 6.1f % 6.1f % 6.1f]"),
-						PositionView[Instance][PositionIdx].X, 
-						PositionView[Instance][PositionIdx].Y, 
-						PositionView[Instance][PositionIdx].Z,
-						LocalPosition.X, LocalPosition.Y, LocalPosition.Z);
-
-					UE_VLOG_SEGMENT(this, LogLearning, Display,
-						RelativePosition,
-						PositionView[Instance][PositionIdx],
-						VisualLogColor.ToFColor(true),
-						TEXT(""));
-				}
+#if UE_LEARNING_AGENTS_ENABLE_VISUAL_LOG
+void UScalarVelocityArrayObservation::VisualLog(const UE::Learning::FIndexSet Instances) const
+{
+	UE_LEARNING_TRACE_CPUPROFILER_EVENT_SCOPE(UScalarVelocityArrayObservation::VisualLog);
 
-				if (PositionNum > 0)
-				{
-					UE_LEARNING_AGENTS_VLOG_TRANSFORM(this, LogLearning, Display,
-						RelativePosition,
-						RelativeRotation,
-						VisualLogColor.ToFColor(true),
-						TEXT("Agent %i\nEncoded: [% 6.3f % 6.3f % 6.3f ...]"),
-						Instance,
-						FeatureObject->Scale,
-						FeatureView[Instance][0], FeatureView[Instance][1], FeatureView[Instance][2]);
-				}
-				else
-				{
-					UE_LEARNING_AGENTS_VLOG_TRANSFORM(this, LogLearning, Display,
-						RelativePosition,
-						RelativeRotation,
-						VisualLogColor.ToFColor(true),
-						TEXT("Agent %i\nScale: [% 6.2f]\nEncoded: []"),
-						Instance,
-						FeatureObject->Scale);
-				}
-			}
+	const TLearningArrayView<2, const float> VelocityView = FeatureObject->InstanceData->ConstView(FeatureObject->VelocityHandle);
+	const TLearningArrayView<2, const float> FeatureView = FeatureObject->InstanceData->ConstView(FeatureObject->FeatureHandle);
+
+	const int32 VelocityNum = VelocityView.Num<1>();
+
+	for (const int32 Instance : Instances)
+	{
+		if (const AActor* Actor = Cast<AActor>(AgentType->GetAgent(Instance)))
+		{
+			UE_LEARNING_AGENTS_VLOG_STRING(this, LogLearning, Display,
+				Actor->GetActorLocation(),
+				VisualLogColor.ToFColor(true),
+				TEXT("Agent %i\nScale: [% 6.2f]\nVelocities: %s\nEncoded: %s"),
+				Instance,
+				FeatureObject->Scale,
+				*UE::Learning::Array::FormatFloat(VelocityView[Instance]),
+				*UE::Learning::Array::FormatFloat(FeatureView[Instance]));
 		}
 	}
 }
 #endif
 
-//------------------------------------------------------------------
 
-UPlanarVelocityObservation* UPlanarVelocityObservation::AddPlanarVelocityObservation(ULearningAgentsType* AgentType, const FName Name, const float Scale, const FVector Axis0, const FVector Axis1)
+UPlanarVelocityObservation* UPlanarVelocityObservation::AddPlanarVelocityObservation(ULearningAgentsType* InAgentType, const FName Name, const float Scale, const FVector Axis0, const FVector Axis1)
 {
-	return UE::Learning::Agents::Private::AddObservation<UPlanarVelocityObservation, UE::Learning::FPlanarVelocityFeature>(AgentType, Name, 1, Scale, Axis0, Axis1);
+	return UE::Learning::Agents::Observations::Private::AddObservation<UPlanarVelocityObservation, UE::Learning::FPlanarVelocityFeature>(InAgentType, Name, 1, Scale, Axis0, Axis1);
 }
 
 void UPlanarVelocityObservation::SetPlanarVelocityObservation(const int32 AgentId, const FVector Velocity, const FRotator RelativeRotation)
 {
-	const TLearningArrayView<2, FVector> VelocityView = FeatureObject->InstanceData->View(FeatureObject->VelocityHandle);
-	const TLearningArrayView<1, FQuat> RelativeRotationView = FeatureObject->InstanceData->View(FeatureObject->RelativeRotationHandle);
+	if (!AgentType->HasAgent(AgentId))
+	{
+		UE_LOG(LogLearning, Error, TEXT("%s: AgentId %d not found in the agents set."), *GetName(), AgentId);
+		return;
+	}
+
+	FeatureObject->InstanceData->View(FeatureObject->VelocityHandle)[AgentId][0] = Velocity;
+	FeatureObject->InstanceData->View(FeatureObject->RelativeRotationHandle)[AgentId] = RelativeRotation.Quaternion();
+}
+
+#if UE_LEARNING_AGENTS_ENABLE_VISUAL_LOG
+void UPlanarVelocityObservation::VisualLog(const UE::Learning::FIndexSet Instances) const
+{
+	UE_LEARNING_TRACE_CPUPROFILER_EVENT_SCOPE(UPlanarVelocityObservation::VisualLog);
+
+	const TLearningArrayView<2, const FVector> VelocityView = FeatureObject->InstanceData->ConstView(FeatureObject->VelocityHandle);
+	const TLearningArrayView<1, const FQuat> RelativeRotationView = FeatureObject->InstanceData->ConstView(FeatureObject->RelativeRotationHandle);
+	const TLearningArrayView<2, const float> FeatureView = FeatureObject->InstanceData->ConstView(FeatureObject->FeatureHandle);
+
+	for (const int32 Instance : Instances)
+	{
+		if (const AActor* Actor = Cast<AActor>(AgentType->GetAgent(Instance)))
+		{
+			const FVector Velocity = VelocityView[Instance][0];
+			const FQuat RelativeRotation = RelativeRotationView[Instance];
+			const FVector LocalVelocity = RelativeRotation.UnrotateVector(Velocity);
+
+			UE_LEARNING_AGENTS_VLOG_ARROW(this, LogLearning, Display,
+				Actor->GetActorLocation(),
+				Actor->GetActorLocation() + Velocity,
+				VisualLogColor.ToFColor(true),
+				TEXT(""));
+
+			UE_LEARNING_AGENTS_VLOG_STRING(this, LogLearning, Display,
+				Actor->GetActorLocation() + Velocity,
+				VisualLogColor.ToFColor(true),
+				TEXT("Velocity: [% 6.3f % 6.3f % 6.3f]\nLocal Velocity: [% 6.3f % 6.3f % 6.3f]"),
+				Velocity.X, Velocity.Y, Velocity.Z,
+				LocalVelocity.X, LocalVelocity.Y, LocalVelocity.Z);
+
+			UE_LEARNING_AGENTS_VLOG_PLANE(this, LogLearning, Display,
+				Actor->GetActorLocation(),
+				RelativeRotation,
+				FeatureObject->Axis0,
+				FeatureObject->Axis1,
+				VisualLogColor.ToFColor(true),
+				TEXT(""));
+
+			UE_LEARNING_AGENTS_VLOG_TRANSFORM(this, LogLearning, Display,
+				Actor->GetActorLocation(),
+				RelativeRotation,
+				VisualLogColor.ToFColor(true),
+				TEXT("Agent %i\nScale: [% 6.2f]\nEncoded: %s"),
+				Instance,
+				FeatureObject->Scale,
+				*UE::Learning::Array::FormatFloat(FeatureView[Instance]));
+		}
+	}
+}
+#endif
+
+
+UPlanarVelocityArrayObservation* UPlanarVelocityArrayObservation::AddPlanarVelocityArrayObservation(ULearningAgentsType* InAgentType, const FName Name, const int32 VelocityNum, const float Scale, const FVector Axis0, const FVector Axis1)
+{
+	if (VelocityNum < 1)
+	{
+		UE_LOG(LogLearning, Error, TEXT("Number of elements in array must be at least 1, got %i."), VelocityNum);
+		return nullptr;
+	}
 
-	if (AgentId == INDEX_NONE)
+	return UE::Learning::Agents::Observations::Private::AddObservation<UPlanarVelocityArrayObservation, UE::Learning::FPlanarVelocityFeature>(InAgentType, Name, VelocityNum, Scale, Axis0, Axis1);
+}
+
+void UPlanarVelocityArrayObservation::SetPlanarVelocityArrayObservation(const int32 AgentId, const TArray<FVector>& Velocities, const FRotator RelativeRotation)
+{
+	if (!AgentType->HasAgent(AgentId))
 	{
-		UE_LOG(LogLearning, Error, TEXT("AgentId is invalid (INDEX_NONE)"));
+		UE_LOG(LogLearning, Error, TEXT("%s: AgentId %d not found in the agents set."), *GetName(), AgentId);
 		return;
 	}
 
-	if (AgentId < 0 || AgentId >= VelocityView.Num<0>())
+	const TLearningArrayView<2, FVector> VelocityView = FeatureObject->InstanceData->View(FeatureObject->VelocityHandle);
+	const TLearningArrayView<1, FQuat> RelativeRotationView = FeatureObject->InstanceData->View(FeatureObject->RelativeRotationHandle);
+
+	if (Velocities.Num() != VelocityView.Num<1>())
 	{
-		UE_LOG(LogLearning, Error, TEXT("AgentId %d is out of index. Valid range [0, %d]."), AgentId, VelocityView.Num<0>() - 1);
+		UE_LOG(LogLearning, Error, TEXT("%s: Got wrong number of elements in array. Expected %i, got %i."), *GetName(), VelocityView.Num<1>(), Velocities.Num());
 		return;
 	}
 
-	VelocityView[AgentId][0] = Velocity;
-	RelativeRotationView[AgentId] = FQuat::MakeFromRotator(RelativeRotation);
+	UE::Learning::Array::Copy<1, FVector>(VelocityView[AgentId], Velocities);
+	RelativeRotationView[AgentId] = RelativeRotation.Quaternion();
 }
 
-#if ENABLE_VISUAL_LOG
-void UPlanarVelocityObservation::VisualLog(const UE::Learning::FIndexSet Instances) const
+#if UE_LEARNING_AGENTS_ENABLE_VISUAL_LOG
+void UPlanarVelocityArrayObservation::VisualLog(const UE::Learning::FIndexSet Instances) const
 {
-	UE_LEARNING_TRACE_CPUPROFILER_EVENT_SCOPE(UPlanarVelocityObservation::VisualLog);
+	UE_LEARNING_TRACE_CPUPROFILER_EVENT_SCOPE(UPlanarVelocityArrayObservation::VisualLog);
 
 	const TLearningArrayView<2, const FVector> VelocityView = FeatureObject->InstanceData->ConstView(FeatureObject->VelocityHandle);
 	const TLearningArrayView<1, const FQuat> RelativeRotationView = FeatureObject->InstanceData->ConstView(FeatureObject->RelativeRotationHandle);
 	const TLearningArrayView<2, const float> FeatureView = FeatureObject->InstanceData->ConstView(FeatureObject->FeatureHandle);
 
-	if (const ULearningAgentsType* AgentType = Cast<ULearningAgentsType>(GetOuter()))
+	const int32 VelocityNum = VelocityView.Num<1>();
+
+	for (const int32 Instance : Instances)
 	{
-		for (const int32 Instance : Instances)
+		if (const AActor* Actor = Cast<AActor>(AgentType->GetAgent(Instance)))
 		{
-			if (const AActor* Actor = Cast<AActor>(AgentType->GetAgent(Instance)))
+			const FQuat RelativeRotation = RelativeRotationView[Instance];
+
+			for (int32 VelocityIdx = 0; VelocityIdx < VelocityNum; VelocityIdx++)
 			{
 				const FVector Velocity = VelocityView[Instance][0];
-				const FQuat RelativeRotation = RelativeRotationView[Instance];
 				const FVector LocalVelocity = RelativeRotation.UnrotateVector(Velocity);
+				const FVector Offset = UE::Learning::Agents::Debug::GridOffsetForIndex(VelocityIdx, VelocityNum);
 
-				UE_VLOG_ARROW(this, LogLearning, Display,
+				UE_LEARNING_AGENTS_VLOG_ARROW(this, LogLearning, Display,
 					Actor->GetActorLocation(),
 					Actor->GetActorLocation() + Velocity,
 					VisualLogColor.ToFColor(true),
@@ -858,106 +2010,298 @@ void UPlanarVelocityObservation::VisualLog(const UE::Learning::FIndexSet Instanc
 				UE_LEARNING_AGENTS_VLOG_STRING(this, LogLearning, Display,
 					Actor->GetActorLocation() + Velocity,
 					VisualLogColor.ToFColor(true),
-					TEXT("Velocity: [% 6.3f % 6.3f % 6.3f]\nLocal Velocity: [% 6.3f % 6.3f % 6.3f]"),
+					TEXT("Velocity %i: [% 6.3f % 6.3f % 6.3f]\nLocal Velocity: [% 6.3f % 6.3f % 6.3f]"),
+					VelocityIdx,
 					Velocity.X, Velocity.Y, Velocity.Z,
 					LocalVelocity.X, LocalVelocity.Y, LocalVelocity.Z);
-
-				UE_LEARNING_AGENTS_VLOG_PLANE(this, LogLearning, Display,
-					Actor->GetActorLocation(),
-					RelativeRotation,
-					FeatureObject->Axis0,
-					FeatureObject->Axis1,
-					VisualLogColor.ToFColor(true),
-					TEXT(""));
-
-				UE_LEARNING_AGENTS_VLOG_TRANSFORM(this, LogLearning, Display,
-					Actor->GetActorLocation(),
-					RelativeRotation,
-					VisualLogColor.ToFColor(true),
-					TEXT("Agent %i\nScale: [% 6.2f]\nEncoded: [% 6.3f % 6.3f]"),
-					Instance,
-					FeatureObject->Scale,
-					FeatureView[Instance][0], FeatureView[Instance][1]);
 			}
+
+			UE_LEARNING_AGENTS_VLOG_PLANE(this, LogLearning, Display,
+				Actor->GetActorLocation(),
+				RelativeRotation,
+				FeatureObject->Axis0,
+				FeatureObject->Axis1,
+				VisualLogColor.ToFColor(true),
+				TEXT(""));
+
+			UE_LEARNING_AGENTS_VLOG_TRANSFORM(this, LogLearning, Display,
+				Actor->GetActorLocation(),
+				RelativeRotation,
+				VisualLogColor.ToFColor(true),
+				TEXT("Agent %i\nScale: [% 6.2f]\nEncoded: %s"),
+				Instance,
+				FeatureObject->Scale,
+				*UE::Learning::Array::FormatFloat(FeatureView[Instance]));
 		}
 	}
 }
 #endif
 
-UVelocityObservation* UVelocityObservation::AddVelocityObservation(ULearningAgentsType* AgentType, const FName Name, const float Scale)
+//------------------------------------------------------------------
+
+UAngularVelocityObservation* UAngularVelocityObservation::AddAngularVelocityObservation(ULearningAgentsType* InAgentType, const FName Name, const float Scale)
 {
-	return UE::Learning::Agents::Private::AddObservation<UVelocityObservation, UE::Learning::FVelocityFeature>(AgentType, Name, 1, Scale);
+	return UE::Learning::Agents::Observations::Private::AddObservation<UAngularVelocityObservation, UE::Learning::FAngularVelocityFeature>(InAgentType, Name, 1, Scale);
 }
 
-void UVelocityObservation::SetVelocityObservation(const int32 AgentId, const FVector Velocity, const FRotator RelativeRotation)
+void UAngularVelocityObservation::SetAngularVelocityObservation(const int32 AgentId, const FVector AngularVelocity, const FRotator RelativeRotation)
 {
-	const TLearningArrayView<2, FVector> VelocityView = FeatureObject->InstanceData->View(FeatureObject->VelocityHandle);
-	const TLearningArrayView<1, FQuat> RelativeRotationView = FeatureObject->InstanceData->View(FeatureObject->RelativeRotationHandle);
+	if (!AgentType->HasAgent(AgentId))
+	{
+		UE_LOG(LogLearning, Error, TEXT("%s: AgentId %d not found in the agents set."), *GetName(), AgentId);
+		return;
+	}
+
+	FeatureObject->InstanceData->View(FeatureObject->AngularVelocityHandle)[AgentId][0] = AngularVelocity;
+	FeatureObject->InstanceData->View(FeatureObject->RelativeRotationHandle)[AgentId] = RelativeRotation.Quaternion();
+}
+
+#if UE_LEARNING_AGENTS_ENABLE_VISUAL_LOG
+void UAngularVelocityObservation::VisualLog(const UE::Learning::FIndexSet Instances) const
+{
+	UE_LEARNING_TRACE_CPUPROFILER_EVENT_SCOPE(UAngularVelocityObservation::VisualLog);
+
+	const TLearningArrayView<2, const FVector> AngularVelocityView = FeatureObject->InstanceData->ConstView(FeatureObject->AngularVelocityHandle);
+	const TLearningArrayView<1, const FQuat> RelativeRotationView = FeatureObject->InstanceData->ConstView(FeatureObject->RelativeRotationHandle);
+	const TLearningArrayView<2, const float> FeatureView = FeatureObject->InstanceData->ConstView(FeatureObject->FeatureHandle);
+
+	for (const int32 Instance : Instances)
+	{
+		if (const AActor* Actor = Cast<AActor>(AgentType->GetAgent(Instance)))
+		{
+			const FVector AngularVelocity = AngularVelocityView[Instance][0];
+			const FQuat RelativeRotation = RelativeRotationView[Instance];
+			const FVector LocalAngularVelocity = RelativeRotation.UnrotateVector(AngularVelocity);
+
+			UE_LEARNING_AGENTS_VLOG_ARROW(this, LogLearning, Display,
+				Actor->GetActorLocation(),
+				Actor->GetActorLocation() + AngularVelocity,
+				VisualLogColor.ToFColor(true),
+				TEXT(""));
+
+			UE_LEARNING_AGENTS_VLOG_STRING(this, LogLearning, Display,
+				Actor->GetActorLocation() + AngularVelocity,
+				VisualLogColor.ToFColor(true),
+				TEXT("Angular Velocity: [% 6.3f % 6.3f % 6.3f]\nLocal Angular Velocity: [% 6.3f % 6.3f % 6.3f]"),
+				AngularVelocity.X, AngularVelocity.Y, AngularVelocity.Z,
+				LocalAngularVelocity.X, LocalAngularVelocity.Y, LocalAngularVelocity.Z);
+
+			UE_LEARNING_AGENTS_VLOG_TRANSFORM(this, LogLearning, Display,
+				Actor->GetActorLocation(),
+				RelativeRotation,
+				VisualLogColor.ToFColor(true),
+				TEXT("Agent %i\nScale: [% 6.2f]\nEncoded: %s"),
+				Instance,
+				FeatureObject->Scale,
+				*UE::Learning::Array::FormatFloat(FeatureView[Instance]));
+		}
+	}
+}
+#endif
+
+UAngularVelocityArrayObservation* UAngularVelocityArrayObservation::AddAngularVelocityArrayObservation(ULearningAgentsType* InAgentType, const FName Name, const int32 AngularVelocityNum, const float Scale)
+{
+	if (AngularVelocityNum < 1)
+	{
+		UE_LOG(LogLearning, Error, TEXT("Number of elements in array must be at least 1, got %i."), AngularVelocityNum);
+		return nullptr;
+	}
+
+	return UE::Learning::Agents::Observations::Private::AddObservation<UAngularVelocityArrayObservation, UE::Learning::FAngularVelocityFeature>(InAgentType, Name, AngularVelocityNum, Scale);
+}
 
-	if (AgentId == INDEX_NONE)
+void UAngularVelocityArrayObservation::SetAngularVelocityArrayObservation(const int32 AgentId, const TArray<FVector>& AngularVelocities, const FRotator RelativeRotation)
+{
+	if (!AgentType->HasAgent(AgentId))
 	{
-		UE_LOG(LogLearning, Error, TEXT("AgentId is invalid (INDEX_NONE)"));
+		UE_LOG(LogLearning, Error, TEXT("%s: AgentId %d not found in the agents set."), *GetName(), AgentId);
 		return;
 	}
 
-	if (AgentId < 0 || AgentId >= VelocityView.Num<0>())
+	const TLearningArrayView<2, FVector> AngularVelocityView = FeatureObject->InstanceData->View(FeatureObject->AngularVelocityHandle);
+	const TLearningArrayView<1, FQuat> RelativeRotationView = FeatureObject->InstanceData->View(FeatureObject->RelativeRotationHandle);
+
+	if (AngularVelocities.Num() != AngularVelocityView.Num<1>())
 	{
-		UE_LOG(LogLearning, Error, TEXT("AgentId %d is out of index. Valid range [0, %d]."), AgentId, VelocityView.Num<0>() - 1);
+		UE_LOG(LogLearning, Error, TEXT("%s: Got wrong number of elements in array. Expected %i, got %i."), *GetName(), AngularVelocityView.Num<1>(), AngularVelocities.Num());
 		return;
 	}
 
-	VelocityView[AgentId][0] = Velocity;
-	RelativeRotationView[AgentId] = FQuat::MakeFromRotator(RelativeRotation);
+	UE::Learning::Array::Copy<1, FVector>(AngularVelocityView[AgentId], AngularVelocities);
+	RelativeRotationView[AgentId] = RelativeRotation.Quaternion();
 }
 
-#if ENABLE_VISUAL_LOG
-void UVelocityObservation::VisualLog(const UE::Learning::FIndexSet Instances) const
+#if UE_LEARNING_AGENTS_ENABLE_VISUAL_LOG
+void UAngularVelocityArrayObservation::VisualLog(const UE::Learning::FIndexSet Instances) const
 {
-	UE_LEARNING_TRACE_CPUPROFILER_EVENT_SCOPE(UVelocityObservation::VisualLog);
+	UE_LEARNING_TRACE_CPUPROFILER_EVENT_SCOPE(UAngularVelocityArrayObservation::VisualLog);
 
-	const TLearningArrayView<2, const FVector> VelocityView = FeatureObject->InstanceData->ConstView(FeatureObject->VelocityHandle);
+	const TLearningArrayView<2, const FVector> AngularVelocityView = FeatureObject->InstanceData->ConstView(FeatureObject->AngularVelocityHandle);
 	const TLearningArrayView<1, const FQuat> RelativeRotationView = FeatureObject->InstanceData->ConstView(FeatureObject->RelativeRotationHandle);
 	const TLearningArrayView<2, const float> FeatureView = FeatureObject->InstanceData->ConstView(FeatureObject->FeatureHandle);
 
-	if (const ULearningAgentsType* AgentType = Cast<ULearningAgentsType>(GetOuter()))
+	const int32 AngularVelocityNum = AngularVelocityView.Num<1>();
+
+	for (const int32 Instance : Instances)
 	{
-		for (const int32 Instance : Instances)
+		if (const AActor* Actor = Cast<AActor>(AgentType->GetAgent(Instance)))
 		{
-			if (const AActor* Actor = Cast<AActor>(AgentType->GetAgent(Instance)))
+			const FQuat RelativeRotation = RelativeRotationView[Instance];
+
+			for (int32 AngularVelocityIdx = 0; AngularVelocityIdx < AngularVelocityNum; AngularVelocityIdx++)
 			{
-				const FVector Velocity = VelocityView[Instance][0];
-				const FQuat RelativeRotation = RelativeRotationView[Instance];
-				const FVector LocalVelocity = RelativeRotation.UnrotateVector(Velocity);
+				const FVector AngularVelocity = AngularVelocityView[Instance][0];
+				const FVector LocalAngularVelocity = RelativeRotation.UnrotateVector(AngularVelocity);
+				const FVector Offset = UE::Learning::Agents::Debug::GridOffsetForIndex(AngularVelocityIdx, AngularVelocityNum);
 
-				UE_VLOG_ARROW(this, LogLearning, Display,
-					Actor->GetActorLocation(),
-					Actor->GetActorLocation() + Velocity,
+				UE_LEARNING_AGENTS_VLOG_ARROW(this, LogLearning, Display,
+					Actor->GetActorLocation() + Offset,
+					Actor->GetActorLocation() + Offset + AngularVelocity,
 					VisualLogColor.ToFColor(true),
 					TEXT(""));
 
 				UE_LEARNING_AGENTS_VLOG_STRING(this, LogLearning, Display,
-					Actor->GetActorLocation() + Velocity,
-					VisualLogColor.ToFColor(true),
-					TEXT("Velocity: [% 6.3f % 6.3f % 6.3f]\nLocal Velocity: [% 6.3f % 6.3f % 6.3f]"),
-					Velocity.X, Velocity.Y, Velocity.Z,
-					LocalVelocity.X, LocalVelocity.Y, LocalVelocity.Z);
-
-				UE_LEARNING_AGENTS_VLOG_TRANSFORM(this, LogLearning, Display,
-					Actor->GetActorLocation(),
-					RelativeRotation,
+					Actor->GetActorLocation() + Offset + AngularVelocity,
 					VisualLogColor.ToFColor(true),
-					TEXT("Agent %i\nScale: [% 6.2f]\nEncoded: [% 6.3f % 6.3f % 6.3f]"),
-					Instance,
-					FeatureObject->Scale,
-					FeatureView[Instance][0], FeatureView[Instance][1], FeatureView[Instance][2]);
+					TEXT("Angular Velocity %i: [% 6.3f % 6.3f % 6.3f]\nLocal Angular Velocity: [% 6.3f % 6.3f % 6.3f]"),
+					AngularVelocityIdx,
+					AngularVelocity.X, AngularVelocity.Y, AngularVelocity.Z,
+					LocalAngularVelocity.X, LocalAngularVelocity.Y, LocalAngularVelocity.Z);
 			}
+
+			UE_LEARNING_AGENTS_VLOG_TRANSFORM(this, LogLearning, Display,
+				Actor->GetActorLocation(),
+				RelativeRotation,
+				VisualLogColor.ToFColor(true),
+				TEXT("Agent %i\nScale: [% 6.2f]\nEncoded: %s"),
+				Instance,
+				FeatureObject->Scale,
+				*UE::Learning::Array::FormatFloat(FeatureView[Instance]));
 		}
 	}
 }
 #endif
 
-//------------------------------------------------------------------
+UScalarAngularVelocityObservation* UScalarAngularVelocityObservation::AddScalarAngularVelocityObservation(ULearningAgentsType* InAgentType, const FName Name, const float Scale)
+{
+	return UE::Learning::Agents::Observations::Private::AddObservation<UScalarAngularVelocityObservation, UE::Learning::FScalarAngularVelocityFeature>(InAgentType, Name, 1, Scale);
+}
 
-#undef UE_LEARNING_AGENTS_VLOG_STRING
-#undef UE_LEARNING_AGENTS_VLOG_TRANSFORM
-#undef UE_LEARNING_AGENTS_VLOG_PLANE
+void UScalarAngularVelocityObservation::SetScalarAngularVelocityObservation(const int32 AgentId, const float AngularVelocity)
+{
+	if (!AgentType->HasAgent(AgentId))
+	{
+		UE_LOG(LogLearning, Error, TEXT("%s: AgentId %d not found in the agents set."), *GetName(), AgentId);
+		return;
+	}
+
+	FeatureObject->InstanceData->View(FeatureObject->AngularVelocityHandle)[AgentId][0] = AngularVelocity;
+}
+
+void UScalarAngularVelocityObservation::SetScalarAngularVelocityObservationWithAxis(const int32 AgentId, const FVector AngularVelocity, const FVector Axis)
+{
+	SetScalarAngularVelocityObservation(AgentId, AngularVelocity.Dot(Axis));
+}
+
+#if UE_LEARNING_AGENTS_ENABLE_VISUAL_LOG
+void UScalarAngularVelocityObservation::VisualLog(const UE::Learning::FIndexSet Instances) const
+{
+	UE_LEARNING_TRACE_CPUPROFILER_EVENT_SCOPE(UScalarAngularVelocityObservation::VisualLog);
+
+	const TLearningArrayView<2, const float> AngularVelocityView = FeatureObject->InstanceData->ConstView(FeatureObject->AngularVelocityHandle);
+	const TLearningArrayView<2, const float> FeatureView = FeatureObject->InstanceData->ConstView(FeatureObject->FeatureHandle);
+
+	for (const int32 Instance : Instances)
+	{
+		if (const AActor* Actor = Cast<AActor>(AgentType->GetAgent(Instance)))
+		{
+			UE_LEARNING_AGENTS_VLOG_STRING(this, LogLearning, Display,
+				Actor->GetActorLocation(),
+				VisualLogColor.ToFColor(true),
+				TEXT("Agent %i\nScale: [% 6.2f]\nAngular Velocity: [% 6.1f]\nEncoded: %s"),
+				Instance,
+				FeatureObject->Scale,
+				AngularVelocityView[Instance][0],
+				*UE::Learning::Array::FormatFloat(FeatureView[Instance]));
+		}
+	}
+}
+#endif
+
+
+UScalarAngularVelocityArrayObservation* UScalarAngularVelocityArrayObservation::AddScalarAngularVelocityArrayObservation(ULearningAgentsType* InAgentType, const FName Name, const int32 AngularVelocityNum, const float Scale)
+{
+	if (AngularVelocityNum < 1)
+	{
+		UE_LOG(LogLearning, Error, TEXT("Number of elements in array must be at least 1, got %i."), AngularVelocityNum);
+		return nullptr;
+	}
+
+	return UE::Learning::Agents::Observations::Private::AddObservation<UScalarAngularVelocityArrayObservation, UE::Learning::FScalarAngularVelocityFeature>(InAgentType, Name, AngularVelocityNum, Scale);
+}
+
+void UScalarAngularVelocityArrayObservation::SetScalarAngularVelocityArrayObservation(const int32 AgentId, const TArray<float>& AngularVelocities)
+{
+	if (!AgentType->HasAgent(AgentId))
+	{
+		UE_LOG(LogLearning, Error, TEXT("%s: AgentId %d not found in the agents set."), *GetName(), AgentId);
+		return;
+	}
+
+	const TLearningArrayView<2, float> AngularVelocityView = FeatureObject->InstanceData->View(FeatureObject->AngularVelocityHandle);
+
+	if (AngularVelocities.Num() != AngularVelocityView.Num<1>())
+	{
+		UE_LOG(LogLearning, Error, TEXT("%s: Got wrong number of elements in array. Expected %i, got %i."), *GetName(), AngularVelocityView.Num<1>(), AngularVelocities.Num());
+		return;
+	}
+
+	UE::Learning::Array::Copy<1, float>(AngularVelocityView[AgentId], AngularVelocities);
+}
+
+void UScalarAngularVelocityArrayObservation::SetScalarAngularVelocityArrayObservationWithAxis(const int32 AgentId, const TArray<FVector>& AngularVelocities, const FVector Axis)
+{
+	if (!AgentType->HasAgent(AgentId))
+	{
+		UE_LOG(LogLearning, Error, TEXT("%s: AgentId %d not found in the agents set."), *GetName(), AgentId);
+		return;
+	}
+
+	const TLearningArrayView<2, float> AngularVelocityView = FeatureObject->InstanceData->View(FeatureObject->AngularVelocityHandle);
+
+	if (AngularVelocities.Num() != AngularVelocityView.Num<1>())
+	{
+		UE_LOG(LogLearning, Error, TEXT("%s: Got wrong number of elements in array. Expected %i, got %i."), *GetName(), AngularVelocityView.Num<1>(), AngularVelocities.Num());
+		return;
+	}
+
+	for (int32 AngularVelocityIdx = 0; AngularVelocityIdx < AngularVelocityView.Num<1>(); AngularVelocityIdx++)
+	{
+		AngularVelocityView[AgentId][AngularVelocityIdx] = AngularVelocities[AngularVelocityIdx].Dot(Axis);
+	}
+}
+
+#if UE_LEARNING_AGENTS_ENABLE_VISUAL_LOG
+void UScalarAngularVelocityArrayObservation::VisualLog(const UE::Learning::FIndexSet Instances) const
+{
+	UE_LEARNING_TRACE_CPUPROFILER_EVENT_SCOPE(UScalarAngularVelocityArrayObservation::VisualLog);
+
+	const TLearningArrayView<2, const float> AngularVelocityView = FeatureObject->InstanceData->ConstView(FeatureObject->AngularVelocityHandle);
+	const TLearningArrayView<2, const float> FeatureView = FeatureObject->InstanceData->ConstView(FeatureObject->FeatureHandle);
+
+	for (const int32 Instance : Instances)
+	{
+		if (const AActor* Actor = Cast<AActor>(AgentType->GetAgent(Instance)))
+		{
+			UE_LEARNING_AGENTS_VLOG_STRING(this, LogLearning, Display,
+				Actor->GetActorLocation(),
+				VisualLogColor.ToFColor(true),
+				TEXT("Agent %i\nScale: [% 6.2f]\nAngular Velocities: %s\nEncoded: %s"),
+				Instance,
+				FeatureObject->Scale,
+				*UE::Learning::Array::FormatFloat(AngularVelocityView[Instance]),
+				*UE::Learning::Array::FormatFloat(FeatureView[Instance]));
+		}
+	}
+}
+#endif
