LearningAgents: Added some additional observations, actions, rewards, completions and helpers. Cleaned up some of the error messaging to make it more consistent.

diff --git a/Engine/Plugins/Experimental/LearningAgents/Source/LearningAgents/Private/LearningAgentsPolicy.cpp b/Engine/Plugins/Experimental/LearningAgents/Source/LearningAgents/Private/LearningAgentsPolicy.cpp
--- a/Engine/Plugins/Experimental/LearningAgents/Source/LearningAgents/Private/LearningAgentsPolicy.cpp
+++ b/Engine/Plugins/Experimental/LearningAgents/Source/LearningAgents/Private/LearningAgentsPolicy.cpp
@@ -13,62 +13,6 @@
 #include "Misc/FileHelper.h"
 #include "Misc/Paths.h"
 #include "GameFramework/Actor.h"
-#include "VisualLogger/VisualLogger.h"
-
-#define UE_LEARNING_AGENTS_VLOG_STRING(Owner, Category, Verbosity, Location, Color, Format, ...) \
-	UE_VLOG_LOCATION(Owner, Category, Verbosity, Location, 0.0f, Color, Format, ##__VA_ARGS__)
-
-namespace UE::Learning::Agents::Policy::Private
-{
-	static inline FString FloatArrayToString(const TLearningArrayView<1, const float> Array)
-	{
-		const int32 ItemNum = Array.Num();
-		const int32 MaxItemNum = 32;
-		const int32 OutputItemNum = FMath::Min(ItemNum, MaxItemNum);
-
-		FString Output = TEXT("[");
-
-		for (int32 Idx = 0; Idx < OutputItemNum; Idx++)
-		{
-			Output.Appendf(TEXT("% 6.3f"), Array[Idx]);
-
-			if (Idx < OutputItemNum - 1)
-			{
-				Output += TEXT(" ");
-			}
-			else if (Idx == MaxItemNum - 1)
-			{
-				Output += TEXT("...");
-			}
-		}
-
-		Output += TEXT("]");
-
-		return Output;
-	}
-
-	static inline FString FloatArrayToStatsString(const TLearningArrayView<1, const float> Array)
-	{
-		const int32 ItemNum = Array.Num();
-
-		float Min = +FLT_MAX, Max = -FLT_MAX, Mean = 0.0f;
-		for (int32 Idx = 0; Idx < ItemNum; Idx++)
-		{
-			Min = FMath::Min(Min, Array[Idx]);
-			Max = FMath::Max(Max, Array[Idx]);
-			Mean += Array[Idx] / ItemNum;
-		}
-
-		float Var = 0.0f;
-		for (int32 Idx = 0; Idx < ItemNum; Idx++)
-		{
-			Var += FMath::Square(Array[Idx] - Mean) / ItemNum;
-		}
-
-		return FString::Printf(TEXT("[% 6.3f/% 6.3f/% 6.3f/% 6.3f]"),
-			Min, Max, Mean, FMath::Sqrt(Var));
-	}
-}
 
 ULearningAgentsPolicy::ULearningAgentsPolicy() : ULearningAgentsManagerComponent() {}
 ULearningAgentsPolicy::ULearningAgentsPolicy(FVTableHelper& Helper) : ULearningAgentsPolicy() {}
@@ -78,7 +22,7 @@ void ULearningAgentsPolicy::SetupPolicy(ALearningAgentsManager* InAgentManager,
 {
 	if (IsSetup())
 	{
-		UE_LOG(LogLearning, Error, TEXT("%s: Setup already performed!"), *GetName());
+		UE_LOG(LogLearning, Error, TEXT("%s: Setup already run!"), *GetName());
 		return;
 	}
 
@@ -90,7 +34,7 @@ void ULearningAgentsPolicy::SetupPolicy(ALearningAgentsManager* InAgentManager,
 
 	if (!InAgentManager->IsManagerSetup())
 	{
-		UE_LOG(LogLearning, Error, TEXT("%s's SetupManager() must be run before %s can be setup."), *InAgentManager->GetName(), *GetName());
+		UE_LOG(LogLearning, Error, TEXT("%s: %s's SetupManager must be run before it can be used."), *GetName(), *InAgentManager->GetName());
 		return;
 	}
 
@@ -136,9 +80,24 @@ void ULearningAgentsPolicy::SetupPolicy(ALearningAgentsManager* InAgentManager,
 
 	AgentManager->GetInstanceData()->Link(AgentType->GetObservationFeature().FeatureHandle, PolicyObject->InputHandle);
 
+	InitialActionNoiseScale = PolicySettings.InitialActionNoiseScale;
+	
 	bIsSetup = true;
 }
 
+bool ULearningAgentsPolicy::AddAgent(const int32 AgentId)
+{
+	bool bSuccess = Super::AddAgent(AgentId);
+
+	if (bSuccess)
+	{
+		// Reset the noise scale for this agent back to the initial value.
+		SetAgentActionNoiseScale(AgentId, InitialActionNoiseScale);
+	}
+
+	return bSuccess;
+}
+
 UE::Learning::FNeuralNetwork& ULearningAgentsPolicy::GetPolicyNetwork()
 {
 	return *Network->NeuralNetwork;
@@ -153,7 +112,7 @@ void ULearningAgentsPolicy::LoadPolicyFromSnapshot(const FDirectoryPath& Directo
 {
 	if (!IsSetup())
 	{
-		UE_LOG(LogLearning, Error, TEXT("Policy setup must be run before network can be loaded."));
+		UE_LOG(LogLearning, Error, TEXT("%s: Setup not run."), *GetName());
 		return;
 	}
 
@@ -169,7 +128,7 @@ void ULearningAgentsPolicy::LoadPolicyFromSnapshot(const FDirectoryPath& Directo
 
 		if (NetworkData.Num() != TotalByteNum)
 		{
-			UE_LOG(LogLearning, Error, TEXT("Failed to load network from file %s. File size incorrect."), *FilePath);
+			UE_LOG(LogLearning, Error, TEXT("%s: Failed to load network from file %s. File size incorrect."), *GetName(), *FilePath);
 			return;
 		}
 
@@ -177,7 +136,7 @@ void ULearningAgentsPolicy::LoadPolicyFromSnapshot(const FDirectoryPath& Directo
 	}
 	else
 	{
-		UE_LOG(LogLearning, Warning, TEXT("Failed to load network. File not found: %s"), *FilePath);
+		UE_LOG(LogLearning, Error, TEXT("%s: Failed to load network. File not found: %s"), *GetName(), *FilePath);
 	}
 }
 
@@ -185,7 +144,7 @@ void ULearningAgentsPolicy::SavePolicyToSnapshot(const FDirectoryPath& Directory
 {
 	if (!IsSetup())
 	{
-		UE_LOG(LogLearning, Error, TEXT("Policy setup must be run before network can be saved."));
+		UE_LOG(LogLearning, Error, TEXT("%s: Setup not run."), *GetName());
 		return;
 	}
 
@@ -201,7 +160,7 @@ void ULearningAgentsPolicy::SavePolicyToSnapshot(const FDirectoryPath& Directory
 	const FString FilePath = Directory.Path + FGenericPlatformMisc::GetDefaultPathSeparator() + Filename;
 	if (!FFileHelper::SaveArrayToFile(NetworkData, *FilePath))
 	{
-		UE_LOG(LogLearning, Error, TEXT("Failed to save network to file: %s"), *FilePath);
+		UE_LOG(LogLearning, Error, TEXT("%s: Failed to save network to file: %s"), *GetName(), *FilePath);
 	}
 }
 
@@ -209,20 +168,20 @@ void ULearningAgentsPolicy::LoadPolicyFromAsset(const ULearningAgentsNeuralNetwo
 {
 	if (!IsSetup())
 	{
-		UE_LOG(LogLearning, Error, TEXT("Policy setup must be run before network can be loaded."));
+		UE_LOG(LogLearning, Error, TEXT("%s: Setup not run."), *GetName());
 		return;
 	}
 
 	if (!NeuralNetworkAsset || !NeuralNetworkAsset->NeuralNetwork)
 	{
-		UE_LOG(LogLearning, Error, TEXT("Cannot load policy from invalid asset."));
+		UE_LOG(LogLearning, Error, TEXT("%s: Asset is invalid."), *GetName());
 		return;
 	}
 
 	if (NeuralNetworkAsset->NeuralNetwork->GetInputNum() != Network->NeuralNetwork->GetInputNum() ||
 		NeuralNetworkAsset->NeuralNetwork->GetOutputNum() != Network->NeuralNetwork->GetOutputNum())
 	{
-		UE_LOG(LogLearning, Error, TEXT("Failed to load policy from asset. Network Asset inputs and outputs don't match."));
+		UE_LOG(LogLearning, Error, TEXT("%s: Failed to load network from asset. Inputs and outputs don't match."), *GetName());
 		return;
 	}
 
@@ -233,13 +192,13 @@ void ULearningAgentsPolicy::SavePolicyToAsset(ULearningAgentsNeuralNetwork* Neur
 {
 	if (!IsSetup())
 	{
-		UE_LOG(LogLearning, Error, TEXT("Policy setup must be run before network can be saved."));
+		UE_LOG(LogLearning, Error, TEXT("%s: Setup not run."), *GetName());
 		return;
 	}
 
 	if (!NeuralNetworkAsset)
 	{
-		UE_LOG(LogLearning, Error, TEXT("Cannot save policy to invalid asset."));
+		UE_LOG(LogLearning, Error, TEXT("%s: Asset is invalid."), *GetName());
 		return;
 	}
 
@@ -272,7 +231,7 @@ void ULearningAgentsPolicy::EvaluatePolicy()
 
 	if (!IsSetup())
 	{
-		UE_LOG(LogLearning, Error, TEXT("Setup must be run before the policy can be evaluated."));
+		UE_LOG(LogLearning, Error, TEXT("%s: Setup not run."), *GetName());
 		return;
 	}
 
@@ -289,22 +248,36 @@ void ULearningAgentsPolicy::EvaluatePolicy()
 		PolicyObject->InstanceData->ConstView(PolicyObject->OutputHandle),
 		AddedAgentIds);
 
-#if ENABLE_VISUAL_LOG
+#if UE_LEARNING_AGENTS_ENABLE_VISUAL_LOG
 	VisualLog(AddedAgentIds);
 #endif
 }
 
+void ULearningAgentsPolicy::RunInference()
+{
+	UE_LEARNING_TRACE_CPUPROFILER_EVENT_SCOPE(ULearningAgentsPolicy::RunInference);
+
+	if (!IsSetup())
+	{
+		UE_LOG(LogLearning, Error, TEXT("%s: Setup not run."), *GetName());
+		return;
+	}
+
+	AgentType->EncodeObservations();
+	EvaluatePolicy();
+	AgentType->DecodeActions();
+}
 float ULearningAgentsPolicy::GetAgentActionNoiseScale(const int32 AgentId) const
 {
 	if (!IsSetup())
 	{
-		UE_LOG(LogLearning, Error, TEXT("Setup must be run before getting the action noise."));
+		UE_LOG(LogLearning, Error, TEXT("%s: Setup not run."), *GetName());
 		return 0.0f;
 	}
 
 	if (!HasAgent(AgentId))
 	{
-		UE_LOG(LogLearning, Error, TEXT("Unable to get action noise for agent - AgentId %d not found in the added agents set."), AgentId);
+		UE_LOG(LogLearning, Error, TEXT("%s: AgentId %d not found in the agents set."), *GetName(), AgentId);
 		return 0.0f;
 	}
 
@@ -316,13 +289,13 @@ void ULearningAgentsPolicy::SetAgentActionNoiseScale(const int32 AgentId, const
 {
 	if (!IsSetup())
 	{
-		UE_LOG(LogLearning, Error, TEXT("Setup must be run before setting the action noise."));
+		UE_LOG(LogLearning, Error, TEXT("%s: Setup not run."), *GetName());
 		return;
 	}
 
 	if (!HasAgent(AgentId))
 	{
-		UE_LOG(LogLearning, Error, TEXT("Unable to set action noise for agent - AgentId %d not found in the added agents set."), AgentId);
+		UE_LOG(LogLearning, Error, TEXT("%s: AgentId %d not found in the agents set."), *GetName(), AgentId);
 		return;
 	}
 
@@ -330,7 +303,19 @@ void ULearningAgentsPolicy::SetAgentActionNoiseScale(const int32 AgentId, const
 	ActionNoiseScaleView[AgentId] = ActionNoiseScale;
 }
 
-#if ENABLE_VISUAL_LOG
+void ULearningAgentsPolicy::SetAllAgentsActionNoiseScale(const float ActionNoiseScale)
+{
+	if (!IsSetup())
+	{
+		UE_LOG(LogLearning, Error, TEXT("%s: Setup not run."), *GetName());
+		return;
+	}
+
+	const TLearningArrayView<1, float> ActionNoiseScaleView = PolicyObject->InstanceData->View(PolicyObject->ActionNoiseScaleHandle);
+	UE::Learning::Array::Set(ActionNoiseScaleView, ActionNoiseScale, AddedAgentSet);
+}
+
+#if UE_LEARNING_AGENTS_ENABLE_VISUAL_LOG
 void ULearningAgentsPolicy::VisualLog(const UE::Learning::FIndexSet Instances) const
 {
 	UE_LEARNING_TRACE_CPUPROFILER_EVENT_SCOPE(ULearningAgentsPolicy::VisualLog);
@@ -345,12 +330,6 @@ void ULearningAgentsPolicy::VisualLog(const UE::Learning::FIndexSet Instances) c
 	{
 		if (const AActor* Actor = Cast<AActor>(AgentType->GetAgent(Instance)))
 		{
-			const FString InputArrayString = UE::Learning::Agents::Policy::Private::FloatArrayToString(InputView[Instance]);
-			const FString OutputArrayString = UE::Learning::Agents::Policy::Private::FloatArrayToString(OutputView[Instance]);
-			const FString OutputMeanArrayString = UE::Learning::Agents::Policy::Private::FloatArrayToString(OutputMeanView[Instance]);
-			const FString OutputStdArrayString = UE::Learning::Agents::Policy::Private::FloatArrayToString(OutputStdView[Instance]);
-			const FString InputStatsString = UE::Learning::Agents::Policy::Private::FloatArrayToStatsString(InputView[Instance]);
-			const FString OutputStatsString = UE::Learning::Agents::Policy::Private::FloatArrayToStatsString(OutputView[Instance]);
 
 			UE_LEARNING_AGENTS_VLOG_STRING(this, LogLearning, Display,
 				Actor->GetActorLocation(),
@@ -358,15 +337,14 @@ void ULearningAgentsPolicy::VisualLog(const UE::Learning::FIndexSet Instances) c
 				TEXT("Agent %i\nAction Noise Scale: [% 6.3f]\nInput: %s\nInput Stats (Min/Max/Mean/Std): %s\nOutput Mean: %s\nOutput Std: %s\nOutput Sample: %s\nOutput Stats (Min/Max/Mean/Std): %s"),
 				Instance,
 				ActionNoiseScaleView[Instance],
-				*InputArrayString,
-				*InputStatsString,
-				*OutputMeanArrayString,
-				*OutputStdArrayString,
-				*OutputArrayString,
-				*OutputStatsString);
+				*UE::Learning::Array::FormatFloat(InputView[Instance]),
+				*UE::Learning::Array::FormatFloat(OutputView[Instance]),
+				*UE::Learning::Array::FormatFloat(OutputMeanView[Instance]),
+				*UE::Learning::Array::FormatFloat(OutputStdView[Instance]),
+				*UE::Learning::Agents::Debug::FloatArrayToStatsString(InputView[Instance]),
+				*UE::Learning::Agents::Debug::FloatArrayToStatsString(OutputView[Instance]));
+
 		}
 	}
 }
 #endif
-
-#undef UE_LEARNING_AGENTS_VLOG_STRING
diff --git a/Engine/Plugins/Experimental/LearningAgents/Source/LearningAgents/Private/LearningAgentsType.cpp b/Engine/Plugins/Experimental/LearningAgents/Source/LearningAgents/Private/LearningAgentsType.cpp
--- a/Engine/Plugins/Experimental/LearningAgents/Source/LearningAgents/Private/LearningAgentsType.cpp
+++ b/Engine/Plugins/Experimental/LearningAgents/Source/LearningAgents/Private/LearningAgentsType.cpp
@@ -20,7 +20,7 @@ void ULearningAgentsType::SetupAgentType(ALearningAgentsManager* InAgentManager)
 {
 	if (IsSetup())
 	{
-		UE_LOG(LogLearning, Error, TEXT("%s: Setup already performed!"), *GetName());
+		UE_LOG(LogLearning, Error, TEXT("%s: Setup already run!"), *GetName());
 		return;
 	}
 
@@ -32,7 +32,7 @@ void ULearningAgentsType::SetupAgentType(ALearningAgentsManager* InAgentManager)
 
 	if (!InAgentManager->IsManagerSetup())
 	{
-		UE_LOG(LogLearning, Error, TEXT("%s's SetupManager() must be run before %s can be setup."), *InAgentManager->GetName(), *GetName());
+		UE_LOG(LogLearning, Error, TEXT("%s: %s's SetupManager must be run before it can be used."), *GetName(), *InAgentManager->GetName());
 		return;
 	}
 
@@ -119,7 +119,7 @@ void ULearningAgentsType::EncodeObservations()
 
 	if (!IsSetup())
 	{
-		UE_LOG(LogLearning, Error, TEXT("Setup must be run before observations can be encoded."));
+		UE_LOG(LogLearning, Error, TEXT("%s: Setup not run."), *GetName());
 		return;
 	}
 
@@ -144,7 +144,7 @@ void ULearningAgentsType::DecodeActions()
 
 	if (!IsSetup())
 	{
-		UE_LOG(LogLearning, Error, TEXT("Setup must be run before actions can be decoded."));
+		UE_LOG(LogLearning, Error, TEXT("%s: Setup not run."), *GetName());
 		return;
 	}
 
@@ -162,3 +162,42 @@ void ULearningAgentsType::DecodeActions()
 	}
 #endif
 }
+void ULearningAgentsType::GetObservationVector(const int32 AgentId, TArray<float>& OutObservationVector) const
+{
+	if (!IsSetup())
+	{
+		UE_LOG(LogLearning, Error, TEXT("%s: Setup not run."), *GetName());
+		OutObservationVector.Empty();
+		return;
+	}
+
+	if (!HasAgent(AgentId))
+	{
+		UE_LOG(LogLearning, Error, TEXT("%s: AgentId %d not found in the agents set."), *GetName(), AgentId);
+		OutObservationVector.Empty();
+		return;
+	}
+
+	OutObservationVector.SetNumUninitialized(Observations->DimNum());
+	UE::Learning::Array::Copy<1, float>(OutObservationVector, Observations->FeatureBuffer()[AgentId]);
+}
+
+void ULearningAgentsType::GetActionVector(const int32 AgentId, TArray<float>& OutActionVector) const
+{
+	if (!IsSetup())
+	{
+		UE_LOG(LogLearning, Error, TEXT("%s: Setup not run."), *GetName());
+		OutActionVector.Empty();
+		return;
+	}
+
+	if (!HasAgent(AgentId))
+	{
+		UE_LOG(LogLearning, Error, TEXT("%s: AgentId %d not found in the agents set."), *GetName(), AgentId);
+		OutActionVector.Empty();
+		return;
+	}
+
+	OutActionVector.SetNumUninitialized(Actions->DimNum());
+	UE::Learning::Array::Copy<1, float>(OutActionVector, Actions->FeatureBuffer()[AgentId]);
+}
\ No newline at end of file
diff --git a/Engine/Plugins/Experimental/LearningAgents/Source/LearningAgents/Public/LearningAgentsActions.h b/Engine/Plugins/Experimental/LearningAgents/Source/LearningAgents/Public/LearningAgentsActions.h
--- a/Engine/Plugins/Experimental/LearningAgents/Source/LearningAgents/Public/LearningAgentsActions.h
+++ b/Engine/Plugins/Experimental/LearningAgents/Source/LearningAgents/Public/LearningAgentsActions.h
@@ -3,10 +3,10 @@
 #pragma once
 
 #include "LearningArray.h"
+#include "LearningAgentsDebug.h"
 
 #include "Templates/SharedPointer.h"
 #include "UObject/Object.h"
-#include "EngineDefines.h"
 
 #include "LearningAgentsActions.generated.h"
 
@@ -22,7 +22,13 @@ class ULearningAgentsType;
 // For functions in this file, we are favoring having more verbose names such as "AddFloatAction" vs simply "Add" in 
 // order to keep it easy to find the correct function in blueprints.
 
-/** The base class for all actions. Actions define the outputs from your agents. */
+//------------------------------------------------------------------
+
+/**
+* The base class for all actions. Actions define the outputs from your agents. Action getters are marked non-pure by
+* convention as many of them do non-trivial amounts of work that can cause performance issues when marked pure in 
+* blueprints.
+*/
 UCLASS(Abstract, BlueprintType)
 class LEARNINGAGENTS_API ULearningAgentsAction : public UObject
 {
@@ -30,15 +36,20 @@ class LEARNINGAGENTS_API ULearningAgentsAction : public UObject
 
 public:
 
-#if ENABLE_VISUAL_LOG
+	UPROPERTY(VisibleAnywhere, Transient, Category = "LearningAgents")
+	TObjectPtr<ULearningAgentsType> AgentType;
+
+#if UE_LEARNING_AGENTS_ENABLE_VISUAL_LOG
 	/** Color used to draw this action in the visual log */
 	FLinearColor VisualLogColor = FColor::Blue;
 
-	/** Describes this observation to the visual logger for debugging purposes. */
+	/** Describes this action to the visual logger for debugging purposes. */
 	virtual void VisualLog(const UE::Learning::FIndexSet Instances) const {}
 #endif
 };
 
+//------------------------------------------------------------------
+
 /** A simple float action. Used as a catch-all for situations where a more type-specific action does not exist yet. */
 UCLASS()
 class LEARNINGAGENTS_API UFloatAction : public ULearningAgentsAction
@@ -49,38 +60,83 @@ public:
 
 	/**
 	* Adds a new float action to the given agent type. Call during ULearningAgentsType::SetupActions event.
-	* @param AgentType The agent type to add this action to.
+	* @param InAgentType The agent type to add this action to.
 	* @param Name The name of this new action. Used for debugging.
 	* @param Scale Used to normalize the data for the action.
 	* @return The newly created action.
 	*/
 	UFUNCTION(BlueprintCallable, Category = "LearningAgents")
-	static UFloatAction* AddFloatAction(ULearningAgentsType* AgentType, const FName Name = NAME_None, const float Scale = 1.0f);
+	static UFloatAction* AddFloatAction(ULearningAgentsType* InAgentType, const FName Name = NAME_None, const float Scale = 1.0f);
 
 	/**
 	* Gets the data for this action. Call during ULearningAgentsType::GetActions event.
 	* @param AgentId The agent id to get data for.
 	* @return The current action value.
 	*/
-	UFUNCTION(BlueprintPure, Category = "LearningAgents")
-	float GetFloatAction(const int32 AgentId);
+	UFUNCTION(BlueprintPure = false, Category = "LearningAgents", meta = (AgentId = "-1"))
+	float GetFloatAction(const int32 AgentId) const;
 
 	/**
 	* Sets the data for this action. Call during ULearningAgentsController::EncodeActions event.
 	* @param AgentId The agent id to set data for.
 	* @param Value The current action value.
 	*/	
-	UFUNCTION(BlueprintCallable, Category = "LearningAgents")
+	UFUNCTION(BlueprintCallable, Category = "LearningAgents", meta = (AgentId = "-1"))
 	void SetFloatAction(const int32 AgentId, const float Value);
 
-#if ENABLE_VISUAL_LOG
-	/** Describes this observation to the visual logger for debugging purposes. */
+#if UE_LEARNING_AGENTS_ENABLE_VISUAL_LOG
+	/** Describes this action to the visual logger for debugging purposes. */
 	virtual void VisualLog(const UE::Learning::FIndexSet Instances) const override;
 #endif
 
 	TSharedPtr<UE::Learning::FFloatFeature> FeatureObject;
 };
 
+/** A simple float array action. Used as a catch-all for situations where a more type-specific action does not exist yet. */
+UCLASS()
+class LEARNINGAGENTS_API UFloatArrayAction : public ULearningAgentsAction
+{
+	GENERATED_BODY()
+
+public:
+
+	/**
+	* Adds a new float array action to the given agent type. Call during ULearningAgentsType::SetupActions event.
+	* @param InAgentType The agent type to add this action to.
+	* @param Name The name of this new action. Used for debugging.
+	* @param Num The number of floats in the array
+	* @param Scale Used to normalize the data for the action.
+	* @return The newly created action.
+	*/
+	UFUNCTION(BlueprintCallable, Category = "LearningAgents")
+	static UFloatArrayAction* AddFloatArrayAction(ULearningAgentsType* InAgentType, const FName Name = NAME_None, const int32 Num = 1, const float Scale = 1.0f);
+
+	/**
+	* Gets the data for this action. Call during ULearningAgentsType::GetActions event.
+	* @param AgentId The agent id to get data for.
+	* @param OutValues The output array of floats
+	*/
+	UFUNCTION(BlueprintPure = false, Category = "LearningAgents", meta = (AgentId = "-1"))
+	void GetFloatArrayAction(const int32 AgentId, TArray<float>& OutValues) const;
+
+	/**
+	* Sets the data for this action. Call during ULearningAgentsController::EncodeActions event.
+	* @param AgentId The agent id to set data for.
+	* @param Values The current action values.
+	*/
+	UFUNCTION(BlueprintCallable, Category = "LearningAgents", meta = (AgentId = "-1"))
+	void SetFloatArrayAction(const int32 AgentId, const TArray<float>& Values);
+
+#if UE_LEARNING_AGENTS_ENABLE_VISUAL_LOG
+	/** Describes this action to the visual logger for debugging purposes. */
+	virtual void VisualLog(const UE::Learning::FIndexSet Instances) const override;
+#endif
+
+	TSharedPtr<UE::Learning::FFloatFeature> FeatureObject;
+};
+
+//------------------------------------------------------------------
+
 /** A simple FVector action. */
 UCLASS()
 class LEARNINGAGENTS_API UVectorAction : public ULearningAgentsAction
@@ -91,83 +147,178 @@ public:
 
 	/**
 	* Adds a new vector action to the given agent type. Call during ULearningAgentsType::SetupActions event.
-	* @param AgentType The agent type to add this action to.
+	* @param InAgentType The agent type to add this action to.
 	* @param Name The name of this new action. Used for debugging.
 	* @param Scale Used to normalize the data for the action.
 	* @return The newly created action.
 	*/
 	UFUNCTION(BlueprintCallable, Category = "LearningAgents")
-	static UVectorAction* AddVectorAction(ULearningAgentsType* AgentType, const FName Name = NAME_None, const float Scale = 1.0f);
+	static UVectorAction* AddVectorAction(ULearningAgentsType* InAgentType, const FName Name = NAME_None, const float Scale = 1.0f);
 
 	/**
 	* Gets the data for this action. Call during ULearningAgentsType::GetActions event.
 	* @param AgentId The agent id to get data for.
 	* @return The current action values.
 	*/
-	UFUNCTION(BlueprintPure, Category = "LearningAgents")
-	FVector GetVectorAction(const int32 AgentId);
+	UFUNCTION(BlueprintPure = false, Category = "LearningAgents", meta = (AgentId = "-1"))
+	FVector GetVectorAction(const int32 AgentId) const;
 
 	/**
 	* Sets the data for this action. Call during ULearningAgentsController::EncodeActions event.
 	* @param AgentId The agent id to set data for.
 	* @param Value The current action values.
 	*/
+	UFUNCTION(BlueprintCallable, Category = "LearningAgents", meta = (AgentId = "-1"))
+	void SetVectorAction(const int32 AgentId, const FVector Value);
+
+#if UE_LEARNING_AGENTS_ENABLE_VISUAL_LOG
+	/** Describes this action to the visual logger for debugging purposes. */
+	virtual void VisualLog(const UE::Learning::FIndexSet Instances) const override;
+#endif
+
+	TSharedPtr<UE::Learning::FFloatFeature> FeatureObject;
+};
+
+/** A simple array of FVector action. */
+UCLASS()
+class LEARNINGAGENTS_API UVectorArrayAction : public ULearningAgentsAction
+{
+	GENERATED_BODY()
+
+public:
+
+	/**
+	* Adds a new vector action to the given agent type. Call during ULearningAgentsType::SetupActions event.
+	* @param InAgentType The agent type to add this action to.
+	* @param Name The name of this new action. Used for debugging.
+	* @param Num The number of vectors in the array
+	* @param Scale Used to normalize the data for the action.
+	* @return The newly created action.
+	*/
 	UFUNCTION(BlueprintCallable, Category = "LearningAgents")
-	void SetVectorAction(const int32 AgentId, const FVector Action);
+	static UVectorArrayAction* AddVectorArrayAction(ULearningAgentsType* InAgentType, const FName Name = NAME_None, const int32 Num = 1, const float Scale = 1.0f);
 
-#if ENABLE_VISUAL_LOG
-	/** Describes this observation to the visual logger for debugging purposes. */
+	/**
+	* Gets the data for this action. Call during ULearningAgentsType::GetActions event.
+	* @param AgentId The agent id to get data for.
+	* @param OutVectors The current action values.
+	*/
+	UFUNCTION(BlueprintPure = false, Category = "LearningAgents", meta = (AgentId = "-1"))
+	void GetVectorArrayAction(const int32 AgentId, TArray<FVector>& OutVectors) const;
+
+	/**
+	* Sets the data for this action. Call during ULearningAgentsController::EncodeActions event.
+	* @param AgentId The agent id to set data for.
+	* @param Vectors The current action values.
+	*/
+	UFUNCTION(BlueprintCallable, Category = "LearningAgents", meta = (AgentId = "-1"))
+	void SetVectorArrayAction(const int32 AgentId, const TArray<FVector>& Vectors);
+
+#if UE_LEARNING_AGENTS_ENABLE_VISUAL_LOG
+	/** Describes this action to the visual logger for debugging purposes. */
 	virtual void VisualLog(const UE::Learning::FIndexSet Instances) const override;
 #endif
 
 	TSharedPtr<UE::Learning::FFloatFeature> FeatureObject;
 };
 
-/** An array of rotation vector actions. */
+//------------------------------------------------------------------
+
+/** A rotation action. */
 UCLASS()
-class LEARNINGAGENTS_API URotationVectorArrayAction : public ULearningAgentsAction
+class LEARNINGAGENTS_API URotationAction : public ULearningAgentsAction
 {
 	GENERATED_BODY()
 
 public:
 
 	/**
-	* Adds a new rotation vector array action to the given agent type. Call during ULearningAgentsType::SetupActions event.
-	* @param AgentType The agent type to add this action to.
+	* Adds a new rotation action to the given agent type. Call during ULearningAgentsType::SetupActions event.
+	* @param InAgentType The agent type to add this action to.
 	* @param Name The name of this new action. Used for debugging.
-	* @param RotationVectorNum The number of rotations in the array.
 	* @param Scale Used to normalize the data for the action.
 	* @return The newly created action.
 	*/
 	UFUNCTION(BlueprintCallable, Category = "LearningAgents")
-	static URotationVectorArrayAction* AddRotationVectorArrayAction(ULearningAgentsType* AgentType, const FName Name = NAME_None, const int32 RotationVectorNum = 0, const float Scale = 180.0f);
+	static URotationAction* AddRotationAction(ULearningAgentsType* InAgentType, const FName Name = NAME_None, const float Scale = 180.0f);
 
 	/**
-	* Gets the data for this action as rotation vectors. Call during ULearningAgentsType::GetActions event.
+	* Gets the data for this action as a rotator. Call during ULearningAgentsType::GetActions event.
 	* @param AgentId The agent id to get data for.
-	* @param OutRotationVectors The current action values.
+	* @return The current action value.
 	*/
-	UFUNCTION(BlueprintPure, Category = "LearningAgents")
-	void GetRotationVectorArrayAction(const int32 AgentId, TArray<FVector>& OutRotationVectors);
+	UFUNCTION(BlueprintPure = false, Category = "LearningAgents", meta = (AgentId = "-1"))
+	FRotator GetRotationAction(const int32 AgentId) const;
 
 	/**
-	* Gets the data for this action as quaternions. Call during ULearningAgentsType::GetActions event.
+	* Gets the data for this action as a rotation vector. Call during ULearningAgentsType::GetActions event.
 	* @param AgentId The agent id to get data for.
-	* @param OutRotations The current action values.
+	* @return The current action value.
 	*/
-	UFUNCTION(BlueprintPure, Category = "LearningAgents")
-	void GetRotationVectorArrayActionAsQuats(const int32 AgentId, TArray<FQuat>& OutRotations);
+	UFUNCTION(BlueprintPure = false, Category = "LearningAgents", meta = (AgentId = "-1"))
+	FVector GetRotationActionAsRotationVector(const int32 AgentId) const;
+
+	/**
+	* Gets the data for this action as a quaternion. Call during ULearningAgentsType::GetActions event.
+	* @param AgentId The agent id to get data for.
+	* @return The current action value.
+	*/
+	UFUNCTION(BlueprintPure = false, Category = "LearningAgents", meta = (AgentId = "-1"))
+	FQuat GetRotationActionAsQuat(const int32 AgentId) const;
+
+#if UE_LEARNING_AGENTS_ENABLE_VISUAL_LOG
+	/** Describes this action to the visual logger for debugging purposes. */
+	virtual void VisualLog(const UE::Learning::FIndexSet Instances) const override;
+#endif
+
+	TSharedPtr<UE::Learning::FRotationVectorFeature> FeatureObject;
+};
+
+/** An array of rotation actions. */
+UCLASS()
+class LEARNINGAGENTS_API URotationArrayAction : public ULearningAgentsAction
+{
+	GENERATED_BODY()
+
+public:
+
+	/**
+	* Adds a new rotation array action to the given agent type. Call during ULearningAgentsType::SetupActions event.
+	* @param InAgentType The agent type to add this action to.
+	* @param Name The name of this new action. Used for debugging.
+	* @param RotationNum The number of rotations in the array.
+	* @param Scale Used to normalize the data for the action.
+	* @return The newly created action.
+	*/
+	UFUNCTION(BlueprintCallable, Category = "LearningAgents")
+	static URotationArrayAction* AddRotationArrayAction(ULearningAgentsType* InAgentType, const FName Name = NAME_None, const int32 RotationNum = 1, const float Scale = 180.0f);
 
 	/**
 	* Gets the data for this action as rotators. Call during ULearningAgentsType::GetActions event.
 	* @param AgentId The agent id to get data for.
 	* @param OutRotations The current action values.
 	*/
-	UFUNCTION(BlueprintPure, Category = "LearningAgents")
-	void GetRotationVectorArrayActionAsRotators(const int32 AgentId, TArray<FRotator>& OutRotations);
+	UFUNCTION(BlueprintPure = false, Category = "LearningAgents", meta = (AgentId = "-1"))
+	void GetRotationArrayAction(const int32 AgentId, TArray<FRotator>& OutRotations) const;
+
+	/**
+	* Gets the data for this action as rotation vectors. Call during ULearningAgentsType::GetActions event.
+	* @param AgentId The agent id to get data for.
+	* @param OutRotationVectors The current action values.
+	*/
+	UFUNCTION(BlueprintPure = false, Category = "LearningAgents", meta = (AgentId = "-1"))
+	void GetRotationArrayActionAsRotationVectors(const int32 AgentId, TArray<FVector>& OutRotationVectors) const;
+
+	/**
+	* Gets the data for this action as quaternions. Call during ULearningAgentsType::GetActions event.
+	* @param AgentId The agent id to get data for.
+	* @param OutRotations The current action values.
+	*/
+	UFUNCTION(BlueprintPure = false, Category = "LearningAgents", meta = (AgentId = "-1"))
+	void GetRotationArrayActionAsQuats(const int32 AgentId, TArray<FQuat>& OutRotations) const;
 
-#if ENABLE_VISUAL_LOG
-	/** Describes this observation to the visual logger for debugging purposes. */
+#if UE_LEARNING_AGENTS_ENABLE_VISUAL_LOG
+	/** Describes this action to the visual logger for debugging purposes. */
 	virtual void VisualLog(const UE::Learning::FIndexSet Instances) const override;
 #endif
 
diff --git a/Engine/Plugins/Experimental/LearningAgents/Source/LearningAgents/Public/LearningAgentsController.h b/Engine/Plugins/Experimental/LearningAgents/Source/LearningAgents/Public/LearningAgentsController.h
--- a/Engine/Plugins/Experimental/LearningAgents/Source/LearningAgents/Public/LearningAgentsController.h
+++ b/Engine/Plugins/Experimental/LearningAgents/Source/LearningAgents/Public/LearningAgentsController.h
@@ -43,6 +43,12 @@ public:
 	UFUNCTION(BlueprintCallable, Category = "LearningAgents")
 	void EncodeActions();
 
+	/**
+	* Calls EncodeObservations, followed by EncodeActions, followed by DecodeActions
+	*/
+	UFUNCTION(BlueprintCallable, Category = "LearningAgents")
+	void RunController();
+
 protected:
 
 	/**
diff --git a/Engine/Plugins/Experimental/LearningAgents/Source/LearningAgents/Public/LearningAgentsCritic.h b/Engine/Plugins/Experimental/LearningAgents/Source/LearningAgents/Public/LearningAgentsCritic.h
--- a/Engine/Plugins/Experimental/LearningAgents/Source/LearningAgents/Public/LearningAgentsCritic.h
+++ b/Engine/Plugins/Experimental/LearningAgents/Source/LearningAgents/Public/LearningAgentsCritic.h
@@ -5,9 +5,9 @@
 #include "LearningAgentsManagerComponent.h"
 
 #include "LearningAgentsNeuralNetwork.h" // Included for ELearningAgentsActivationFunction
+#include "LearningAgentsDebug.h"
 #include "LearningArray.h"
 #include "UObject/ObjectPtr.h"
-#include "EngineDefines.h"
 
 #include "LearningAgentsCritic.generated.h"
 
@@ -108,7 +108,7 @@ public:
 	* Gets an estimate of the average discounted return expected by an agent according to the critic. I.E. the total
 	* sum of future rewards, scaled by the discount factor that was used during training. This value can be useful if 
 	* you want to make some decision based on how well the agent thinks they are doing at achieving their task. This 
-	* should be called after EvaluateCritic.
+	* should be called only after EvaluateCritic.
 	* 
 	* @param AgentId	The AgentId to look-up the estimated discounted return for
 	* @returns			The estimated average discounted return according to the critic
@@ -140,7 +140,7 @@ private:
 
 	TSharedPtr<UE::Learning::FNeuralNetworkCriticFunction> CriticObject;
 
-#if ENABLE_VISUAL_LOG
+#if UE_LEARNING_AGENTS_ENABLE_VISUAL_LOG
 	/** Color used to draw this action in the visual log */
 	FLinearColor VisualLogColor = FColor::Orange;
 
diff --git a/Engine/Plugins/Experimental/LearningAgents/Source/LearningAgents/Public/LearningAgentsHelpers.h b/Engine/Plugins/Experimental/LearningAgents/Source/LearningAgents/Public/LearningAgentsHelpers.h
--- a/Engine/Plugins/Experimental/LearningAgents/Source/LearningAgents/Public/LearningAgentsHelpers.h
+++ b/Engine/Plugins/Experimental/LearningAgents/Source/LearningAgents/Public/LearningAgentsHelpers.h
@@ -2,51 +2,145 @@
 
 #pragma once
 
+#include "LearningAgentsDebug.h"
+
 #include "Templates/SharedPointer.h"
 #include "UObject/Object.h"
-#include "Kismet/BlueprintFunctionLibrary.h"
-
 #include "Components/SplineComponent.h" // Required for ESplineCoordinateSpace::World
 
 #include "LearningAgentsHelpers.generated.h"
 
+class ULearningAgentsManagerComponent;
 class UMeshComponent;
 
+//------------------------------------------------------------------
+
+/**
+* The base class for all helpers. Helpers are additional objects that can be used in getting or setting observations,
+* actions, rewards, and completions.
+*/
+UCLASS(Abstract, BlueprintType)
+class LEARNINGAGENTS_API ULearningAgentsHelper : public UObject
+{
+	GENERATED_BODY()
+
+public:
+
+	UPROPERTY(VisibleAnywhere, Transient, Category = "LearningAgents")
+	TObjectPtr<ULearningAgentsManagerComponent> ManagerComponent;
+
+#if UE_LEARNING_AGENTS_ENABLE_VISUAL_LOG
+	/** Color used to draw this helper in the visual log */
+	FLinearColor VisualLogColor = FColor::Magenta;
+#endif
+};
+
+//------------------------------------------------------------------
+
+/** A helper for computing various properties from a SplineComponent. */
+UCLASS()
+class LEARNINGAGENTS_API USplineComponentHelper : public ULearningAgentsHelper
+{
+	GENERATED_BODY()
+
+public:
+
+	UFUNCTION(BlueprintCallable, Category = "LearningAgents")
+	static USplineComponentHelper* AddSplineComponentHelper(ULearningAgentsManagerComponent* InManagerComponent, const FName Name = NAME_None);
+
+	UFUNCTION(BlueprintPure = false, Category = "LearningAgents", meta = (AgentId = "-1"))
+	float GetProportionAlongSpline(const int32 AgentId, const USplineComponent* SplineComponent, const float DistanceAlongSpline) const;
+
+	UFUNCTION(BlueprintPure = false, Category = "LearningAgents", meta = (AgentId = "-1"))
+	float GetProportionAlongSplineAsAngle(const int32 AgentId, const USplineComponent* SplineComponent, const float DistanceAlongSpline) const;
+
+	UFUNCTION(BlueprintPure = false, Category = "LearningAgents", meta = (AgentId = "-1"))
+	void GetPositionsAlongSpline(
+		TArray<FVector>& OutPositions,
+		const int32 AgentId,
+		const USplineComponent* SplineComponent,
+		const int32 PositionNum,
+		const float StartDistanceAlongSpline,
+		const float StopDistanceAlongSpline,
+		const ESplineCoordinateSpace::Type CoordinateSpace = ESplineCoordinateSpace::World) const;
+
+	UFUNCTION(BlueprintPure = false, Category = "LearningAgents", meta = (AgentId = "-1"))
+	float GetVelocityAlongSpline(
+		const int32 AgentId,
+		const USplineComponent* SplineComponent,
+		const FVector Position,
+		const FVector Velocity,
+		const float FiniteDifferenceDelta = 10.0f,
+		const ESplineCoordinateSpace::Type CoordinateSpace = ESplineCoordinateSpace::World) const;
+};
+
+//------------------------------------------------------------------
+
+/** A helper for projecting onto surfaces. */
 UCLASS()
-class LEARNINGAGENTS_API ULearningAgentsHelpers : public UBlueprintFunctionLibrary
+class LEARNINGAGENTS_API UProjectionHelper : public ULearningAgentsHelper
 {
 	GENERATED_BODY()
 
-	//------------------------------------------------------------------
-	// General Math Helpers
-	//------------------------------------------------------------------
+public:
+
+	UFUNCTION(BlueprintCallable, Category = "LearningAgents")
+	static UProjectionHelper* AddProjectionHelper(ULearningAgentsManagerComponent* InManagerComponent, const FName Name = NAME_None);
+
+	UFUNCTION(BlueprintPure = false, Category = "LearningAgents", meta = (AgentId = "-1"))
+	FTransform ProjectTransformOntoGroundPlane(const int32 AgentId, const FTransform Transform, const FVector LocalForwardVector = FVector::ForwardVector) const;
+
+	UFUNCTION(BlueprintPure = false, Category = "LearningAgents", meta = (AgentId = "-1"))
+	void ProjectPositionRotationOntoGroundPlane(
+		FVector& OutPosition, 
+		FRotator& OutRotation, 
+		const int32 AgentId, 
+		const FVector InPosition, 
+		const FRotator InRotation, 
+		const FVector LocalForwardVector = FVector::ForwardVector) const;
+};
+
+//------------------------------------------------------------------
 
-	UFUNCTION(BlueprintPure, Category = "LearningAgents|Helpers")
-	static FTransform ProjectTransformOntoGroundPlane(const FTransform Transform, const FVector LocalForwardVector = FVector::ForwardVector);
+/** A helper for getting various properties from a MeshComponent. */
+UCLASS()
+class LEARNINGAGENTS_API UMeshComponentHelper : public ULearningAgentsHelper
+{
+	GENERATED_BODY()
 
-	UFUNCTION(BlueprintPure, Category = "LearningAgents|Helpers")
-	static void ProjectPositionRotationOntoGroundPlane(FVector& OutPosition, FRotator& OutRotation, const FVector InPosition, const FRotator InRotation, const FVector LocalForwardVector = FVector::ForwardVector);
+public:
 
-	//------------------------------------------------------------------
-	// MeshComponent Helpers 
-	//------------------------------------------------------------------
+	UFUNCTION(BlueprintCallable, Category = "LearningAgents")
+	static UMeshComponentHelper* AddMeshComponentHelper(ULearningAgentsManagerComponent* InManagerComponent, const FName Name = NAME_None);
 
-	UFUNCTION(BlueprintPure, Category = "LearningAgents|Helpers")
-	static void GetMeshBonePositions(const UMeshComponent* MeshComponent, const TArray<FName>& BoneNames, TArray<FVector>& OutBonePositions);
+	UFUNCTION(BlueprintPure = false, Category = "LearningAgents", meta = (AgentId = "-1"))
+	void GetMeshBonePositions(TArray<FVector>& OutBonePositions, const int32 AgentId, const UMeshComponent* MeshComponent, const TArray<FName>& BoneNames) const;
+};
 
-	//------------------------------------------------------------------
-	// SplineComponent Helpers 
-	//------------------------------------------------------------------
+//------------------------------------------------------------------
+
+/** A helper for performing various kinds of ray cast. */
+UCLASS()
+class LEARNINGAGENTS_API URayCastHelper : public ULearningAgentsHelper
+{
+	GENERATED_BODY()
 
-	UFUNCTION(BlueprintPure, Category = "LearningAgents|Helpers")
-	static float GetProportionAlongSpline(const USplineComponent* SplineComponent, const float DistanceAlongSpline);
+public:
 
-	UFUNCTION(BlueprintPure, Category = "LearningAgents|Helpers")
-	static float GetProportionAlongSplineAsAngle(const USplineComponent* SplineComponent, const float DistanceAlongSpline);
+	UFUNCTION(BlueprintCallable, Category = "LearningAgents")
+	static URayCastHelper* AddRayCastHelper(ULearningAgentsManagerComponent* InManagerComponent, const FName Name = NAME_None);
 
-	UFUNCTION(BlueprintPure, Category = "LearningAgents|Helpers")
-	static void GetPositionsAlongSpline(TArray<FVector>& OutPositions, const USplineComponent* SplineComponent, const int32 PositionNum, const float StartDistanceAlongSpline, const float StopDistanceAlongSpline, const ESplineCoordinateSpace::Type CoordinateSpace = ESplineCoordinateSpace::World);
-	
-	UFUNCTION(BlueprintPure, Category = "LearningAgents|Helpers")
-	static float GetVelocityAlongSpline(const USplineComponent* SplineComponent, const FVector Position, const FVector Velocity, const float FiniteDifferenceDelta = 1.0f, const ESplineCoordinateSpace::Type CoordinateSpace = ESplineCoordinateSpace::World);
+	UFUNCTION(BlueprintPure=false, Category = "LearningAgents", meta = (AgentId = "-1"))
+	void RayCastGridHeights(
+		TArray<float>& OutHeights,
+		const int32 AgentId,
+		const FVector Position,
+		const FRotator Rotation,
+		const int32 RowNum = 5,
+		const int32 ColNum = 5,
+		const float RowWidth = 1000.0f,
+		const float ColWidth = 1000.0f,
+		const float MaxHeight = 10000.0f,
+		const float MinHeight = -10000.0f,
+		const ECollisionChannel CollisionChannel = ECollisionChannel::ECC_WorldStatic) const;
 };
diff --git a/Engine/Plugins/Experimental/LearningAgents/Source/LearningAgents/Public/LearningAgentsManager.h b/Engine/Plugins/Experimental/LearningAgents/Source/LearningAgents/Public/LearningAgentsManager.h
--- a/Engine/Plugins/Experimental/LearningAgents/Source/LearningAgents/Public/LearningAgentsManager.h
+++ b/Engine/Plugins/Experimental/LearningAgents/Source/LearningAgents/Public/LearningAgentsManager.h
@@ -122,6 +122,10 @@ public:
 	UFUNCTION(BlueprintPure, Category = "LearningAgents", meta = (AgentId = "-1"))
 	bool HasAgentById(const int32 AgentId) const;
 
+	/** Get the current added agent ids. */
+	UFUNCTION(BlueprintPure = false, Category = "LearningAgents")
+	void GetAgentIds(TArray<int32>& OutAgentIds) const;
+
 // ----- Non-blueprint public interface -----
 public:
 
diff --git a/Engine/Plugins/Experimental/LearningAgents/Source/LearningAgents/Public/LearningAgentsManagerComponent.h b/Engine/Plugins/Experimental/LearningAgents/Source/LearningAgents/Public/LearningAgentsManagerComponent.h
--- a/Engine/Plugins/Experimental/LearningAgents/Source/LearningAgents/Public/LearningAgentsManagerComponent.h
+++ b/Engine/Plugins/Experimental/LearningAgents/Source/LearningAgents/Public/LearningAgentsManagerComponent.h
@@ -12,6 +12,7 @@
 #include "LearningAgentsManagerComponent.generated.h"
 
 class ALearningAgentsManager;
+class ULearningAgentsHelper;
 
 /**
  * Base class for components which can be attached to an ALearningAgentsManager.
@@ -84,6 +85,15 @@ public:
 
 	ALearningAgentsManager* GetAgentManager() const;
 
+// ----- Helpers -----
+public:
+
+	/**
+	* Used by objects derived from ULearningAgentsHelper to add themselves to this component during their creation.
+	* You shouldn't need to call this directly.
+	*/
+	void AddHelper(TObjectPtr<ULearningAgentsHelper> Object);
+
 protected:
 
 	/** True if this component has been setup. Otherwise, false. */
@@ -99,4 +109,8 @@ protected:
 	TArray<int32> AddedAgentIds;
 
 	UE::Learning::FIndexSet AddedAgentSet;
+
+	/** The list of current helper objects. */
+	UPROPERTY(VisibleAnywhere, Transient, Category = "LearningAgents")
+	TArray<TObjectPtr<ULearningAgentsHelper>> HelperObjects;
 };
