LearningAgents: Added some additional observations, actions, rewards, completions and helpers. Cleaned up some of the error messaging to make it more consistent.

diff --git a/Engine/Plugins/Experimental/LearningAgents/Source/LearningAgents/Public/LearningAgentsObservations.h b/Engine/Plugins/Experimental/LearningAgents/Source/LearningAgents/Public/LearningAgentsObservations.h
--- a/Engine/Plugins/Experimental/LearningAgents/Source/LearningAgents/Public/LearningAgentsObservations.h
+++ b/Engine/Plugins/Experimental/LearningAgents/Source/LearningAgents/Public/LearningAgentsObservations.h
@@ -3,10 +3,10 @@
 #pragma once
 
 #include "LearningArray.h"
+#include "LearningAgentsDebug.h"
 
 #include "Templates/SharedPointer.h"
 #include "UObject/Object.h"
-#include "EngineDefines.h"
 
 #include "LearningAgentsObservations.generated.h"
 
@@ -14,13 +14,19 @@
 {
 	struct FFeatureObject;
 	struct FFloatFeature;
+	struct FTimeFeature;
 	struct FAngleFeature;
-	struct FPlanarDirectionFeature;
+	struct FRotationFeature;
 	struct FDirectionFeature;
-	struct FPlanarPositionFeature;
+	struct FPlanarDirectionFeature;
 	struct FPositionFeature;
-	struct FPlanarVelocityFeature;
+	struct FScalarPositionFeature;
+	struct FPlanarPositionFeature;
 	struct FVelocityFeature;
+	struct FScalarVelocityFeature;
+	struct FPlanarVelocityFeature;
+	struct FAngularVelocityFeature;
+	struct FScalarAngularVelocityFeature;
 }
 
 class ULearningAgentsType;
@@ -28,6 +34,8 @@
 // For functions in this file, we are favoring having more verbose names such as "AddFloatObservation" vs simply "Add" in 
 // order to keep it easy to find the correct function in blueprints.
 
+//------------------------------------------------------------------
+
 /** The base class for all observations. Observations define the inputs to your agents. */
 UCLASS(Abstract, BlueprintType)
 class LEARNINGAGENTS_API ULearningAgentsObservation : public UObject
@@ -36,7 +44,10 @@
 
 public:
 
-#if ENABLE_VISUAL_LOG
+	UPROPERTY(VisibleAnywhere, Transient, Category = "LearningAgents")
+	TObjectPtr<ULearningAgentsType> AgentType;
+
+#if UE_LEARNING_AGENTS_ENABLE_VISUAL_LOG
 	/** Color used to draw this observation in the visual log */
 	FLinearColor VisualLogColor = FColor::Red;
 
@@ -45,6 +56,8 @@
 #endif
 };
 
+//------------------------------------------------------------------
+
 /** A simple float observation. Used as a catch-all for situations where a more type-specific observation does not exist yet. */
 UCLASS()
 class LEARNINGAGENTS_API UFloatObservation : public ULearningAgentsObservation
@@ -55,23 +68,58 @@
 
 	/**
 	* Adds a new float observation to the given agent type. Call during ULearningAgentsType::SetupObservations event.
-	* @param AgentType The agent type to add this observation to.
+	* @param InAgentType The agent type to add this observation to.
 	* @param Name The name of this new observation. Used for debugging.
 	* @param Scale Used to normalize the data for the observation.
 	* @return The newly created observation.
 	*/
 	UFUNCTION(BlueprintCallable, Category = "LearningAgents")
-	static UFloatObservation* AddFloatObservation(ULearningAgentsType* AgentType, const FName Name = NAME_None, const float Scale = 1.0f);
+	static UFloatObservation* AddFloatObservation(ULearningAgentsType* InAgentType, const FName Name = NAME_None, const float Scale = 1.0f);
 
 	/**
 	* Sets the data for this observation. Call during ULearningAgentsType::SetObservations event.
 	* @param AgentId The agent id this data corresponds to.
 	* @param Observation The value currently being observed.
 	*/
-	UFUNCTION(BlueprintCallable, Category = "LearningAgents")
+	UFUNCTION(BlueprintCallable, Category = "LearningAgents", meta = (AgentId = "-1"))
 	void SetFloatObservation(const int32 AgentId, const float Observation);
 
-#if ENABLE_VISUAL_LOG
+#if UE_LEARNING_AGENTS_ENABLE_VISUAL_LOG
 	/** Describes this observation to the visual logger for debugging purposes. */
 	virtual void VisualLog(const UE::Learning::FIndexSet Instances) const override;
 #endif
+
+	TSharedPtr<UE::Learning::FFloatFeature> FeatureObject;
+};
+
+/** A simple array of floats observation. Used as a catch-all for situations where a more type-specific observation does not exist yet. */
+UCLASS()
+class LEARNINGAGENTS_API UFloatArrayObservation : public ULearningAgentsObservation
+{
+	GENERATED_BODY()
+
+public:
+
+	/**
+	* Adds a new float array observation to the given agent type. Call during ULearningAgentsType::SetupObservations event.
+	* @param InAgentType The agent type to add this observation to.
+	* @param Name The name of this new observation. Used for debugging.
+	* @param Num The number of floats in the array
+	* @param Scale Used to normalize the data for the observation.
+	* @return The newly created observation.
+	*/
+	UFUNCTION(BlueprintCallable, Category = "LearningAgents")
+	static UFloatArrayObservation* AddFloatArrayObservation(ULearningAgentsType* InAgentType, const FName Name = NAME_None, const int32 Num = 1, const float Scale = 1.0f);
+
+	/**
+	* Sets the data for this observation. Call during ULearningAgentsType::SetObservations event.
+	* @param AgentId The agent id this data corresponds to.
+	* @param Observation The value currently being observed.
+	*/
+	UFUNCTION(BlueprintCallable, Category = "LearningAgents", meta = (AgentId = "-1"))
+	void SetFloatArrayObservation(const int32 AgentId, const TArray<float>& Observation);
+
+#if UE_LEARNING_AGENTS_ENABLE_VISUAL_LOG
+	/** Describes this observation to the visual logger for debugging purposes. */
+	virtual void VisualLog(const UE::Learning::FIndexSet Instances) const override;
+#endif
@@ -79,6 +127,8 @@
 	TSharedPtr<UE::Learning::FFloatFeature> FeatureObject;
 };
 
+//------------------------------------------------------------------
+
 /** A simple observation for an FVector. */
 UCLASS()
 class LEARNINGAGENTS_API UVectorObservation : public ULearningAgentsObservation
@@ -89,23 +139,23 @@
 
 	/**
 	* Adds a new vector observation to the given agent type. Call during ULearningAgentsType::SetupObservations event.
-	* @param AgentType The agent type to add this observation to.
+	* @param InAgentType The agent type to add this observation to.
 	* @param Name The name of this new observation. Used for debugging.
 	* @param Scale Used to normalize the data for the observation.
 	* @return The newly created observation.
 	*/
 	UFUNCTION(BlueprintCallable, Category = "LearningAgents")
-	static UVectorObservation* AddVectorObservation(ULearningAgentsType* AgentType, const FName Name = NAME_None, const float Scale = 1.0f);
+	static UVectorObservation* AddVectorObservation(ULearningAgentsType* InAgentType, const FName Name = NAME_None, const float Scale = 1.0f);
 
 	/**
 	* Sets the data for this observation. Call during ULearningAgentsType::SetObservations event.
 	* @param AgentId The agent id this data corresponds to.
 	* @param Observation The values currently being observed.
 	*/
-	UFUNCTION(BlueprintCallable, Category = "LearningAgents")
+	UFUNCTION(BlueprintCallable, Category = "LearningAgents", meta = (AgentId = "-1"))
 	void SetVectorObservation(const int32 AgentId, const FVector Observation);
 
-#if ENABLE_VISUAL_LOG
+#if UE_LEARNING_AGENTS_ENABLE_VISUAL_LOG
 	/** Describes this observation to the visual logger for debugging purposes. */
 	virtual void VisualLog(const UE::Learning::FIndexSet Instances) const override;
 #endif
@@ -113,6 +163,189 @@
 	TSharedPtr<UE::Learning::FFloatFeature> FeatureObject;
 };
 
+/** A simple observation for an array of FVector. */
+UCLASS()
+class LEARNINGAGENTS_API UVectorArrayObservation : public ULearningAgentsObservation
+{
+	GENERATED_BODY()
+
+public:
+
+	/**
+	* Adds a new vector array observation to the given agent type. Call during ULearningAgentsType::SetupObservations event.
+	* @param InAgentType The agent type to add this observation to.
+	* @param Name The name of this new observation. Used for debugging.
+	* @param Num The number of vectors in the array
+	* @param Scale Used to normalize the data for the observation.
+	* @return The newly created observation.
+	*/
+	UFUNCTION(BlueprintCallable, Category = "LearningAgents")
+	static UVectorArrayObservation* AddVectorArrayObservation(ULearningAgentsType* InAgentType, const FName Name = NAME_None, const int32 Num = 1, const float Scale = 1.0f);
+
+	/**
+	* Sets the data for this observation. Call during ULearningAgentsType::SetObservations event.
+	* @param AgentId The agent id this data corresponds to.
+	* @param Observation The values currently being observed.
+	*/
+	UFUNCTION(BlueprintCallable, Category = "LearningAgents", meta = (AgentId = "-1"))
+	void SetVectorArrayObservation(const int32 AgentId, const TArray<FVector>& Observation);
+
+#if UE_LEARNING_AGENTS_ENABLE_VISUAL_LOG
+	/** Describes this observation to the visual logger for debugging purposes. */
+	virtual void VisualLog(const UE::Learning::FIndexSet Instances) const override;
+#endif
+
+	TSharedPtr<UE::Learning::FFloatFeature> FeatureObject;
+};
+
+//------------------------------------------------------------------
+
+/** An observation of an enumeration. */
+UCLASS()
+class LEARNINGAGENTS_API UEnumObservation : public ULearningAgentsObservation
+{
+	GENERATED_BODY()
+
+public:
+
+	/**
+	* Adds a new enum observation to the given agent type. Call during ULearningAgentsType::SetupObservations event.
+	* @param InAgentType The agent type to add this observation to.
+	* @param EnumType The type of enum to use
+	* @param Name The name of this new observation. Used for debugging.
+	* @return The newly created observation.
+	*/
+	UFUNCTION(BlueprintCallable, Category = "LearningAgents")
+	static UEnumObservation* AddEnumObservation(ULearningAgentsType* InAgentType, const UEnum* EnumType, const FName Name = NAME_None);
+
+	/**
+	* Sets the data for this observation. Call during ULearningAgentsType::SetObservations event.
+	* @param AgentId The agent id this data corresponds to.
+	* @param Value The enum value currently being observed.
+	*/
+	UFUNCTION(BlueprintCallable, Category = "LearningAgents", meta = (AgentId = "-1"))
+	void SetEnumObservation(const int32 AgentId, const uint8 Value);
+
+#if UE_LEARNING_AGENTS_ENABLE_VISUAL_LOG
+	/** Describes this observation to the visual logger for debugging purposes. */
+	virtual void VisualLog(const UE::Learning::FIndexSet Instances) const override;
+#endif
+
+	const UEnum* Enum = nullptr;
+	TSharedPtr<UE::Learning::FFloatFeature> FeatureObject;
+};
+
+/** An observation of an array of enumerations. */
+UCLASS()
+class LEARNINGAGENTS_API UEnumArrayObservation : public ULearningAgentsObservation
+{
+	GENERATED_BODY()
+
+public:
+
+	/**
+	* Adds a new enum array observation to the given agent type. Call during ULearningAgentsType::SetupObservations event.
+	* @param InAgentType The agent type to add this observation to.
+	* @param EnumType The type of enum to use
+	* @param Name The name of this new observation. Used for debugging.
+	* @param EnumNum The number of enum observations in the array
+	* @return The newly created observation.
+	*/
+	UFUNCTION(BlueprintCallable, Category = "LearningAgents")
+	static UEnumArrayObservation* AddEnumArrayObservation(ULearningAgentsType* InAgentType, const UEnum* EnumType, const FName Name = NAME_None, const int32 EnumNum = 1);
+
+	/**
+	* Sets the data for this observation. Call during ULearningAgentsType::SetObservations event.
+	* @param AgentId The agent id this data corresponds to.
+	* @param Values The enum values currently being observed.
+	*/
+	UFUNCTION(BlueprintCallable, Category = "LearningAgents", meta = (AgentId = "-1"))
+	void SetEnumArrayObservation(const int32 AgentId, const TArray<uint8>& Values);
+
+#if UE_LEARNING_AGENTS_ENABLE_VISUAL_LOG
+	/** Describes this observation to the visual logger for debugging purposes. */
+	virtual void VisualLog(const UE::Learning::FIndexSet Instances) const override;
+#endif
+
+	const UEnum* Enum = nullptr;
+	TSharedPtr<UE::Learning::FFloatFeature> FeatureObject;
+};
+
+//------------------------------------------------------------------
+
+/** An observation of a time relative to another time. */
+UCLASS()
+class LEARNINGAGENTS_API UTimeObservation : public ULearningAgentsObservation
+{
+	GENERATED_BODY()
+
+public:
+
+	/**
+	* Adds a new time observation to the given agent type. Call during ULearningAgentsType::SetupObservations event.
+	* @param InAgentType The agent type to add this observation to.
+	* @param Name The name of this new observation. Used for debugging.
+	* @param Scale Used to normalize the data for the observation.
+	* @return The newly created observation.
+	*/
+	UFUNCTION(BlueprintCallable, Category = "LearningAgents")
+	static UTimeObservation* AddTimeObservation(ULearningAgentsType* InAgentType, const FName Name = NAME_None, const float Scale = 1.0f);
+
+	/**
+	* Sets the data for this observation. Call during ULearningAgentsType::SetObservations event.
+	* @param AgentId The agent id this data corresponds to.
+	* @param Time The time currently being observed.
+	* @param RelativeTime The time the provided time should be encoded relative to.
+	*/
+	UFUNCTION(BlueprintCallable, Category = "LearningAgents", meta = (AgentId = "-1"))
+	void SetTimeObservation(const int32 AgentId, const float Time, const float RelativeTime = 0.0f);
+
+#if UE_LEARNING_AGENTS_ENABLE_VISUAL_LOG
+	/** Describes this observation to the visual logger for debugging purposes. */
+	virtual void VisualLog(const UE::Learning::FIndexSet Instances) const override;
+#endif
+
+	TSharedPtr<UE::Learning::FTimeFeature> FeatureObject;
+};
+
+/** An observation of an array of times. */
+UCLASS()
+class LEARNINGAGENTS_API UTimeArrayObservation : public ULearningAgentsObservation
+{
+	GENERATED_BODY()
+
+public:
+
+	/**
+	* Adds a new angle observation to the given agent type. Call during ULearningAgentsType::SetupObservations event.
+	* @param InAgentType The agent type to add this observation to.
+	* @param Name The name of this new observation. Used for debugging.
+	* @param TimeNum The number of times in the array
+	* @param Scale Used to normalize the data for the observation.
+	* @return The newly created observation.
+	*/
+	UFUNCTION(BlueprintCallable, Category = "LearningAgents")
+	static UTimeArrayObservation* AddTimeArrayObservation(ULearningAgentsType* InAgentType, const FName Name = NAME_None, const int32 TimeNum = 1, const float Scale = 1.0f);
+
+	/**
+	* Sets the data for this observation. Call during ULearningAgentsType::SetObservations event.
+	* @param AgentId The agent id this data corresponds to.
+	* @param Time The times currently being observed.
+	* @param RelativeTime The time the provided time should be encoded relative to.
+	*/
+	UFUNCTION(BlueprintCallable, Category = "LearningAgents", meta = (AgentId = "-1"))
+	void SetTimeArrayObservation(const int32 AgentId, const TArray<float>& Times, const float RelativeTime = 0.0f);
+
+#if UE_LEARNING_AGENTS_ENABLE_VISUAL_LOG
+	/** Describes this observation to the visual logger for debugging purposes. */
+	virtual void VisualLog(const UE::Learning::FIndexSet Instances) const override;
+#endif
+
+	TSharedPtr<UE::Learning::FTimeFeature> FeatureObject;
+};
+
+//------------------------------------------------------------------
+
 /** An observation of an angle relative to another angle. */
 UCLASS()
 class LEARNINGAGENTS_API UAngleObservation : public ULearningAgentsObservation
@@ -123,13 +356,13 @@
 
 	/**
 	* Adds a new angle observation to the given agent type. Call during ULearningAgentsType::SetupObservations event.
-	* @param AgentType The agent type to add this observation to.
+	* @param InAgentType The agent type to add this observation to.
 	* @param Name The name of this new observation. Used for debugging.
-	* @param Scale Used to normalize the data for the observation.
+	* @param Scale Used to normalize the data for the observation. Angle observations are encoded as directions. 
 	* @return The newly created observation.
 	*/
 	UFUNCTION(BlueprintCallable, Category = "LearningAgents")
-	static UAngleObservation* AddAngleObservation(ULearningAgentsType* AgentType, const FName Name = NAME_None, const float Scale = 1.0f);
+	static UAngleObservation* AddAngleObservation(ULearningAgentsType* InAgentType, const FName Name = NAME_None, const float Scale = 1.0f);
 
 	/**
 	* Sets the data for this observation. Call during ULearningAgentsType::SetObservations event.
@@ -137,10 +370,10 @@
 	* @param Angle The angle currently being observed.
 	* @param RelativeAngle The frame of reference angle.
 	*/
-	UFUNCTION(BlueprintCallable, Category = "LearningAgents")
+	UFUNCTION(BlueprintCallable, Category = "LearningAgents", meta = (AgentId = "-1"))
 	void SetAngleObservation(const int32 AgentId, const float Angle, const float RelativeAngle = 0.0f);
 
-#if ENABLE_VISUAL_LOG
+#if UE_LEARNING_AGENTS_ENABLE_VISUAL_LOG
 	/** Describes this observation to the visual logger for debugging purposes. */
 	virtual void VisualLog(const UE::Learning::FIndexSet Instances) const override;
 #endif
@@ -148,45 +381,133 @@
 	TSharedPtr<UE::Learning::FAngleFeature> FeatureObject;
 };
 
-/** An observation of a direction vector projected onto a plane. */
+/** An observation of an array of angles. */
 UCLASS()
-class LEARNINGAGENTS_API UPlanarDirectionObservation : public ULearningAgentsObservation
+class LEARNINGAGENTS_API UAngleArrayObservation : public ULearningAgentsObservation
 {
 	GENERATED_BODY()
 
 public:
 
 	/**
-	* Adds a new planar direction observation to the given agent type. The axis parameters define the plane. Call
-	* during ULearningAgentsType::SetupObservations event.
-	* @param AgentType The agent type to add this observation to.
+	* Adds a new angle observation to the given agent type. Call during ULearningAgentsType::SetupObservations event.
+	* @param InAgentType The agent type to add this observation to.
 	* @param Name The name of this new observation. Used for debugging.
-	* @param Scale Used to normalize the data for the observation.
-	* @param Axis0 The forward axis of the plane.
-	* @param Axis1 The right axis of the plane.
+	* @param AngleNum The number of angles in the array
+	* @param Scale Used to normalize the data for the observation. Angle observations are encoded as directions. 
 	* @return The newly created observation.
 	*/
 	UFUNCTION(BlueprintCallable, Category = "LearningAgents")
-	static UPlanarDirectionObservation* AddPlanarDirectionObservation(ULearningAgentsType* AgentType, const FName Name = NAME_None, const float Scale = 1.0f, const FVector Axis0 = FVector::ForwardVector, const FVector Axis1 = FVector::RightVector);
+	static UAngleArrayObservation* AddAngleArrayObservation(ULearningAgentsType* InAgentType, const FName Name = NAME_None, const int32 AngleNum = 1, const float Scale = 1.0f);
 
 	/**
-	* Sets the data for this observation. The relative rotation can be used to make this observation relative to the
-	* agent's perspective, e.g. by passing the agent's forward rotation. Call during ULearningAgentsType::SetObservations event.
+	* Sets the data for this observation. Call during ULearningAgentsType::SetObservations event.
 	* @param AgentId The agent id this data corresponds to.
-	* @param Direction The direction currently being observed.
-	* @param RelativeRotation The frame of reference rotation.
+	* @param Angles The angles currently being observed.
+	* @param RelativeAngle The frame of reference angle.
 	*/
-	UFUNCTION(BlueprintCallable, Category = "LearningAgents")
-	void SetPlanarDirectionObservation(const int32 AgentId, const FVector Direction, const FRotator RelativeRotation = FRotator::ZeroRotator);
+	UFUNCTION(BlueprintCallable, Category = "LearningAgents", meta = (AgentId = "-1"))
+	void SetAngleArrayObservation(const int32 AgentId, const TArray<float>& Angles, const float RelativeAngle = 0.0f);
 
-#if ENABLE_VISUAL_LOG
+#if UE_LEARNING_AGENTS_ENABLE_VISUAL_LOG
 	/** Describes this observation to the visual logger for debugging purposes. */
 	virtual void VisualLog(const UE::Learning::FIndexSet Instances) const override;
 #endif
 
-	TSharedPtr<UE::Learning::FPlanarDirectionFeature> FeatureObject;
+	TSharedPtr<UE::Learning::FAngleFeature> FeatureObject;
 };
 
+/** An observation of a rotation relative to another rotation. */
+UCLASS()
+class LEARNINGAGENTS_API URotationObservation : public ULearningAgentsObservation
+{
+	GENERATED_BODY()
+
+public:
+
+	/**
+	* Adds a new rotation observation to the given agent type. Call during ULearningAgentsType::SetupObservations event.
+	* @param InAgentType The agent type to add this observation to.
+	* @param Name The name of this new observation. Used for debugging.
+	* @param Scale Used to normalize the data for the observation. Rotation observations are encoded as directions. 
+	* @return The newly created observation.
+	*/
+	UFUNCTION(BlueprintCallable, Category = "LearningAgents")
+	static URotationObservation* AddRotationObservation(ULearningAgentsType* InAgentType, const FName Name = NAME_None, const float Scale = 1.0f);
+
+	/**
+	* Sets the data for this observation. Call during ULearningAgentsType::SetObservations event.
+	* @param AgentId The agent id this data corresponds to.
+	* @param Rotation The rotation currently being observed.
+	* @param RelativeRotation The frame of reference rotation.
+	*/
+	UFUNCTION(BlueprintCallable, Category = "LearningAgents", meta = (AgentId = "-1"))
+	void SetRotationObservation(const int32 AgentId, const FRotator Rotation, const FRotator RelativeRotation = FRotator::ZeroRotator);
+
+	/**
+	* Sets the data for this observation. Call during ULearningAgentsType::SetObservations event.
+	* @param AgentId The agent id this data corresponds to.
+	* @param Rotation The rotation currently being observed.
+	* @param RelativeRotation The frame of reference rotation.
+	*/
+	UFUNCTION(BlueprintCallable, Category = "LearningAgents", meta = (AgentId = "-1"))
+	void SetRotationObservationFromQuat(const int32 AgentId, const FQuat Rotation, const FQuat RelativeRotation);
+
+#if UE_LEARNING_AGENTS_ENABLE_VISUAL_LOG
+	/** Describes this observation to the visual logger for debugging purposes. */
+	virtual void VisualLog(const UE::Learning::FIndexSet Instances) const override;
+#endif
+
+	TSharedPtr<UE::Learning::FRotationFeature> FeatureObject;
+};
+
+/** An observation of an array of rotations. */
+UCLASS()
+class LEARNINGAGENTS_API URotationArrayObservation : public ULearningAgentsObservation
+{
+	GENERATED_BODY()
+
+public:
+
+	/**
+	* Adds a new rotation observation to the given agent type. Call during ULearningAgentsType::SetupObservations event.
+	* @param InAgentType The agent type to add this observation to.
+	* @param Name The name of this new observation. Used for debugging.
+	* @param RotationNum The number of rotations in the array
+	* @param Scale Used to normalize the data for the observation. Rotation observations are encoded as directions. 
+	* @return The newly created observation.
+	*/
+	UFUNCTION(BlueprintCallable, Category = "LearningAgents")
+	static URotationArrayObservation* AddRotationArrayObservation(ULearningAgentsType* InAgentType, const FName Name = NAME_None, const int32 RotationNum = 1, const float Scale = 1.0f);
+
+	/**
+	* Sets the data for this observation. Call during ULearningAgentsType::SetObservations event.
+	* @param AgentId The agent id this data corresponds to.
+	* @param Rotations The rotations currently being observed.
+	* @param RelativeRotation The frame of reference rotation.
+	*/
+	UFUNCTION(BlueprintCallable, Category = "LearningAgents", meta = (AgentId = "-1"))
+	void SetRotationArrayObservation(const int32 AgentId, const TArray<FRotator>& Rotations, const FRotator RelativeRotation = FRotator::ZeroRotator);
+
+	/**
+	* Sets the data for this observation. Call during ULearningAgentsType::SetObservations event.
+	* @param AgentId The agent id this data corresponds to.
+	* @param Rotations The rotations currently being observed.
+	* @param RelativeRotation The frame of reference rotation.
+	*/
+	UFUNCTION(BlueprintCallable, Category = "LearningAgents", meta = (AgentId = "-1"))
+	void SetRotationArrayObservationFromQuats(const int32 AgentId, const TArray<FQuat>& Rotations, const FQuat RelativeRotation);
+
+#if UE_LEARNING_AGENTS_ENABLE_VISUAL_LOG
+	/** Describes this observation to the visual logger for debugging purposes. */
+	virtual void VisualLog(const UE::Learning::FIndexSet Instances) const override;
+#endif
+
+	TSharedPtr<UE::Learning::FRotationFeature> FeatureObject;
+};
+
+//------------------------------------------------------------------
+
 /** An observation of a direction vector. */
 UCLASS()
 class LEARNINGAGENTS_API UDirectionObservation : public ULearningAgentsObservation
@@ -197,13 +518,13 @@
 
 	/**
 	* Adds a new direction observation to the given agent type. Call during ULearningAgentsType::SetupObservations event.
-	* @param AgentType The agent type to add this observation to.
+	* @param InAgentType The agent type to add this observation to.
 	* @param Name The name of this new observation. Used for debugging.
 	* @param Scale Used to normalize the data for the observation.
 	* @return The newly created observation.
 	*/
 	UFUNCTION(BlueprintCallable, Category = "LearningAgents")
-	static UDirectionObservation* AddDirectionObservation(ULearningAgentsType* AgentType, const FName Name = NAME_None, const float Scale = 1.0f);
+	static UDirectionObservation* AddDirectionObservation(ULearningAgentsType* InAgentType, const FName Name = NAME_None, const float Scale = 1.0f);
 
 	/**
 	* Sets the data for this observation. The relative rotation can be used to make this observation relative to the
@@ -212,10 +533,10 @@
 	* @param Direction The direction currently being observed.
 	* @param RelativeRotation The frame of reference rotation.
 	*/
-	UFUNCTION(BlueprintCallable, Category = "LearningAgents")
+	UFUNCTION(BlueprintCallable, Category = "LearningAgents", meta = (AgentId = "-1"))
 	void SetDirectionObservation(const int32 AgentId, const FVector Direction, const FRotator RelativeRotation = FRotator::ZeroRotator);
 
-#if ENABLE_VISUAL_LOG
+#if UE_LEARNING_AGENTS_ENABLE_VISUAL_LOG
 	/** Describes this observation to the visual logger for debugging purposes. */
 	virtual void VisualLog(const UE::Learning::FIndexSet Instances) const override;
 #endif
@@ -223,18 +544,55 @@
 	TSharedPtr<UE::Learning::FDirectionFeature> FeatureObject;
 };
 
-/** An observation of a position projected onto a plane. */
+/** An observation of an array of direction vectors. */
 UCLASS()
-class LEARNINGAGENTS_API UPlanarPositionObservation : public ULearningAgentsObservation
+class LEARNINGAGENTS_API UDirectionArrayObservation : public ULearningAgentsObservation
 {
 	GENERATED_BODY()
 
 public:
 
 	/**
-	* Adds a new planar position observation to the given agent type. The axis parameters define the plane. Call
+	* Adds a new direction array observation to the given agent type. Call during ULearningAgentsType::SetupObservations event.
+	* @param InAgentType The agent type to add this observation to.
+	* @param Name The name of this new observation. Used for debugging.
+	* @param DirectionNum The number of directions in the array
+	* @param Scale Used to normalize the data for the observation.
+	* @return The newly created observation.
+	*/
+	UFUNCTION(BlueprintCallable, Category = "LearningAgents")
+	static UDirectionArrayObservation* AddDirectionArrayObservation(ULearningAgentsType* InAgentType, const FName Name = NAME_None, const int32 DirectionNum = 1, const float Scale = 1.0f);
+
+	/**
+	* Sets the data for this observation. The relative rotation can be used to make this observation relative to the
+	* agent's perspective, e.g. by passing the agent's forward rotation. Call during ULearningAgentsType::SetObservations event.
+	* @param AgentId The agent id this data corresponds to.
+	* @param Directions The directions currently being observed.
+	* @param RelativeRotation The frame of reference rotation.
+	*/
+	UFUNCTION(BlueprintCallable, Category = "LearningAgents", meta = (AgentId = "-1"))
+	void SetDirectionArrayObservation(const int32 AgentId, const TArray<FVector>& Directions, const FRotator RelativeRotation = FRotator::ZeroRotator);
+
+#if UE_LEARNING_AGENTS_ENABLE_VISUAL_LOG
+	/** Describes this observation to the visual logger for debugging purposes. */
+	virtual void VisualLog(const UE::Learning::FIndexSet Instances) const override;
+#endif
+
+	TSharedPtr<UE::Learning::FDirectionFeature> FeatureObject;
+};
+
+/** An observation of a direction vector projected onto a plane. */
+UCLASS()
+class LEARNINGAGENTS_API UPlanarDirectionObservation : public ULearningAgentsObservation
+{
+	GENERATED_BODY()
+
+public:
+
+	/**
+	* Adds a new planar direction observation to the given agent type. The axis parameters define the plane. Call
 	* during ULearningAgentsType::SetupObservations event.
-	* @param AgentType The agent type to add this observation to.
+	* @param InAgentType The agent type to add this observation to.
 	* @param Name The name of this new observation. Used for debugging.
 	* @param Scale Used to normalize the data for the observation.
 	* @param Axis0 The forward axis of the plane.
@@ -242,28 +600,68 @@
 	* @return The newly created observation.
 	*/
 	UFUNCTION(BlueprintCallable, Category = "LearningAgents")
-	static UPlanarPositionObservation* AddPlanarPositionObservation(ULearningAgentsType* AgentType, const FName Name = NAME_None, const float Scale = 100.0f, const FVector Axis0 = FVector::ForwardVector, const FVector Axis1 = FVector::RightVector);
+	static UPlanarDirectionObservation* AddPlanarDirectionObservation(ULearningAgentsType* InAgentType, const FName Name = NAME_None, const float Scale = 1.0f, const FVector Axis0 = FVector::ForwardVector, const FVector Axis1 = FVector::RightVector);
 
 	/**
-	* Sets the data for this observation. The relative position & rotation can be used to make this observation
-	* relative to the agent's perspective, e.g. by passing the agent's position & forward rotation. Call during
-	* ULearningAgentsType::SetObservations event.
+	* Sets the data for this observation. The relative rotation can be used to make this observation relative to the
+	* agent's perspective, e.g. by passing the agent's forward rotation. Call during ULearningAgentsType::SetObservations event.
 	* @param AgentId The agent id this data corresponds to.
-	* @param Position The position currently being observed.
-	* @param RelativePosition The vector Position will be offset from.
+	* @param Direction The direction currently being observed.
 	* @param RelativeRotation The frame of reference rotation.
 	*/
-	UFUNCTION(BlueprintCallable, Category = "LearningAgents")
-	void SetPlanarPositionObservation(const int32 AgentId, const FVector Position, const FVector RelativePosition = FVector::ZeroVector, const FRotator RelativeRotation = FRotator::ZeroRotator);
+	UFUNCTION(BlueprintCallable, Category = "LearningAgents", meta = (AgentId = "-1"))
+	void SetPlanarDirectionObservation(const int32 AgentId, const FVector Direction, const FRotator RelativeRotation = FRotator::ZeroRotator);
 
-#if ENABLE_VISUAL_LOG
+#if UE_LEARNING_AGENTS_ENABLE_VISUAL_LOG
 	/** Describes this observation to the visual logger for debugging purposes. */
 	virtual void VisualLog(const UE::Learning::FIndexSet Instances) const override;
 #endif
 
-	TSharedPtr<UE::Learning::FPlanarPositionFeature> FeatureObject;
+	TSharedPtr<UE::Learning::FPlanarDirectionFeature> FeatureObject;
 };
 
+/** An observation of an array of direction vectors projected onto a plane. */
+UCLASS()
+class LEARNINGAGENTS_API UPlanarDirectionArrayObservation : public ULearningAgentsObservation
+{
+	GENERATED_BODY()
+
+public:
+
+	/**
+	* Adds a new planar direction observation to the given agent type. The axis parameters define the plane. Call
+	* during ULearningAgentsType::SetupObservations event.
+	* @param InAgentType The agent type to add this observation to.
+	* @param Name The name of this new observation. Used for debugging.
+	* @param DirectionNum The number of directions in the array
+	* @param Scale Used to normalize the data for the observation.
+	* @param Axis0 The forward axis of the plane.
+	* @param Axis1 The right axis of the plane.
+	* @return The newly created observation.
+	*/
+	UFUNCTION(BlueprintCallable, Category = "LearningAgents")
+	static UPlanarDirectionArrayObservation* AddPlanarDirectionArrayObservation(ULearningAgentsType* InAgentType, const FName Name = NAME_None, const int32 DirectionNum = 1, const float Scale = 1.0f, const FVector Axis0 = FVector::ForwardVector, const FVector Axis1 = FVector::RightVector);
+
+	/**
+	* Sets the data for this observation. The relative rotation can be used to make this observation relative to the
+	* agent's perspective, e.g. by passing the agent's forward rotation. Call during ULearningAgentsType::SetObservations event.
+	* @param AgentId The agent id this data corresponds to.
+	* @param Directions The directions currently being observed.
+	* @param RelativeRotation The frame of reference rotation.
+	*/
+	UFUNCTION(BlueprintCallable, Category = "LearningAgents", meta = (AgentId = "-1"))
+	void SetPlanarDirectionArrayObservation(const int32 AgentId, const TArray<FVector>& Directions, const FRotator RelativeRotation = FRotator::ZeroRotator);
+
+#if UE_LEARNING_AGENTS_ENABLE_VISUAL_LOG
+	/** Describes this observation to the visual logger for debugging purposes. */
+	virtual void VisualLog(const UE::Learning::FIndexSet Instances) const override;
+#endif
+
+	TSharedPtr<UE::Learning::FPlanarDirectionFeature> FeatureObject;
+};
+
+//------------------------------------------------------------------
+
 /** An observation of a position vector. */
 UCLASS()
 class LEARNINGAGENTS_API UPositionObservation : public ULearningAgentsObservation
@@ -274,13 +672,13 @@
 
 	/**
 	* Adds a new position observation to the given agent type. Call during ULearningAgentsType::SetupObservations event.
-	* @param AgentType The agent type to add this observation to.
+	* @param InAgentType The agent type to add this observation to.
 	* @param Name The name of this new observation. Used for debugging.
 	* @param Scale Used to normalize the data for the observation.
 	* @return The newly created observation.
 	*/
 	UFUNCTION(BlueprintCallable, Category = "LearningAgents")
-	static UPositionObservation* AddPositionObservation(ULearningAgentsType* AgentType, const FName Name = NAME_None, const float Scale = 100.0f);
+	static UPositionObservation* AddPositionObservation(ULearningAgentsType* InAgentType, const FName Name = NAME_None, const float Scale = 100.0f);
 
 	/**
 	* Sets the data for this observation. The relative position & rotation can be used to make this observation
@@ -291,10 +689,10 @@
 	* @param RelativePosition The vector Position will be offset from.
 	* @param RelativeRotation The frame of reference rotation.
 	*/
-	UFUNCTION(BlueprintCallable, Category = "LearningAgents")
+	UFUNCTION(BlueprintCallable, Category = "LearningAgents", meta = (AgentId = "-1"))
 	void SetPositionObservation(const int32 AgentId, const FVector Position, const FVector RelativePosition = FVector::ZeroVector, const FRotator RelativeRotation = FRotator::ZeroRotator);
 
-#if ENABLE_VISUAL_LOG
+#if UE_LEARNING_AGENTS_ENABLE_VISUAL_LOG
 	/** Describes this observation to the visual logger for debugging purposes. */
 	virtual void VisualLog(const UE::Learning::FIndexSet Instances) const override;
 #endif
@@ -302,49 +700,6 @@
 	TSharedPtr<UE::Learning::FPositionFeature> FeatureObject;
 };
 
-/** An observation of an array of positions projected onto a plane. */
-UCLASS()
-class LEARNINGAGENTS_API UPlanarPositionArrayObservation : public ULearningAgentsObservation
-{
-	GENERATED_BODY()
-
-public:
-
-	/**
-	* Adds a new planar position array observation to the given agent type. The axis parameters define the plane.
-	* Call during ULearningAgentsType::SetupObservations event.
-	* @param AgentType The agent type to add this observation to.
-	* @param Name The name of this new observation. Used for debugging.
-	* @param PositionNum The number of positions in the array.
-	* @param Scale Used to normalize the data for the observation.
-	* @param Axis0 The forward axis of the plane.
-	* @param Axis1 The right axis of the plane.
-	* @return The newly created observation.
-	*/
-	UFUNCTION(BlueprintCallable, Category = "LearningAgents")
-	static UPlanarPositionArrayObservation* AddPlanarPositionArrayObservation(ULearningAgentsType* AgentType, const FName Name = NAME_None, const int32 PositionNum = 0, const float Scale = 100.0f, const FVector Axis0 = FVector::ForwardVector, const FVector Axis1 = FVector::RightVector);
-
-
-	/**
-	* Sets the data for this observation. The relative position & rotation can be used to make this observation
-	* relative to the agent's perspective, e.g. by passing the agent's position & forward rotation. Call during
-	* ULearningAgentsType::SetObservations event.
-	* @param AgentId The agent id this data corresponds to.
-	* @param Positions The positions currently being observed.
-	* @param RelativePosition The vector Positions will be offset from.
-	* @param RelativeRotation The frame of reference rotation.
-	*/
-	UFUNCTION(BlueprintCallable, Category = "LearningAgents")
-	void SetPlanarPositionArrayObservation(const int32 AgentId, const TArray<FVector>& Positions, const FVector RelativePosition = FVector::ZeroVector, const FRotator RelativeRotation = FRotator::ZeroRotator);
-
-#if ENABLE_VISUAL_LOG
-	/** Describes this observation to the visual logger for debugging purposes. */
-	virtual void VisualLog(const UE::Learning::FIndexSet Instances) const override;
-#endif
-
-	TSharedPtr<UE::Learning::FPlanarPositionFeature> FeatureObject;
-};
-
 /** An observation of an array of positions. */
 UCLASS()
 class LEARNINGAGENTS_API UPositionArrayObservation : public ULearningAgentsObservation
@@ -355,14 +710,14 @@
 
 	/**
 	* Adds a new position array observation to the given agent type. Call during ULearningAgentsType::SetupObservations event.
-	* @param AgentType The agent type to add this observation to.
+	* @param InAgentType The agent type to add this observation to.
 	* @param Name The name of this new observation. Used for debugging.
 	* @param PositionNum The number of positions in the array.
 	* @param Scale Used to normalize the data for the observation.
 	* @return The newly created observation.
 	*/
 	UFUNCTION(BlueprintCallable, Category = "LearningAgents")
-	static UPositionArrayObservation* AddPositionArrayObservation(ULearningAgentsType* AgentType, const FName Name = NAME_None, const int32 PositionNum = 0, const float Scale = 100.0f);
+	static UPositionArrayObservation* AddPositionArrayObservation(ULearningAgentsType* InAgentType, const FName Name = NAME_None, const int32 PositionNum = 1, const float Scale = 100.0f);
 
 	/**
 	* Sets the data for this observation. The relative position & rotation can be used to make this observation
@@ -373,10 +728,10 @@
 	* @param RelativePosition The vector Positions will be offset from.
 	* @param RelativeRotation The frame of reference rotation.
 	*/
-	UFUNCTION(BlueprintCallable, Category = "LearningAgents")
+	UFUNCTION(BlueprintCallable, Category = "LearningAgents", meta = (AgentId = "-1"))
 	void SetPositionArrayObservation(int32 AgentId, const TArray<FVector>& Positions, const FVector RelativePosition = FVector::ZeroVector, const FRotator RelativeRotation = FRotator::ZeroRotator);
 
-#if ENABLE_VISUAL_LOG
+#if UE_LEARNING_AGENTS_ENABLE_VISUAL_LOG
 	/** Describes this observation to the visual logger for debugging purposes. */
 	virtual void VisualLog(const UE::Learning::FIndexSet Instances) const override;
 #endif
@@ -384,18 +739,119 @@
 	TSharedPtr<UE::Learning::FPositionFeature> FeatureObject;
 };
 
-/** An observation of a velocity projected onto a plane. */
+/** An observation of a position along a single axis. Can be useful for providing information like object heights. */
 UCLASS()
-class LEARNINGAGENTS_API UPlanarVelocityObservation : public ULearningAgentsObservation
+class LEARNINGAGENTS_API UScalarPositionObservation : public ULearningAgentsObservation
 {
 	GENERATED_BODY()
 
 public:
 
 	/**
-	* Adds a new planar velocity observation to the given agent type. The axis parameters define the plane.
-	* Call during ULearningAgentsType::SetupObservations event.
-	* @param AgentType The agent type to add this observation to.
+	* Adds a new scalar position observation to the given agent type. Call during 
+	* ULearningAgentsType::SetupObservations event.
+	* @param InAgentType The agent type to add this observation to.
+	* @param Name The name of this new observation. Used for debugging.
+	* @param Scale Used to normalize the data for the observation.
+	* @return The newly created observation.
+	*/
+	UFUNCTION(BlueprintCallable, Category = "LearningAgents")
+	static UScalarPositionObservation* AddScalarPositionObservation(ULearningAgentsType* InAgentType, const FName Name = NAME_None, const float Scale = 100.0f);
+
+	/**
+	* Sets the data for this observation. The relative position can be used to make this observation
+	* relative to the agent's perspective, e.g. by passing the agent's position. Call during 
+	* ULearningAgentsType::SetObservations event.
+	* @param AgentId The agent id this data corresponds to.
+	* @param Position The position currently being observed.
+	* @param RelativePosition The vector Position will be offset from.
+	*/
+	UFUNCTION(BlueprintCallable, Category = "LearningAgents", meta = (AgentId = "-1"))
+	void SetScalarPositionObservation(const int32 AgentId, const float Position, const float RelativePosition = 0.0f);
+
+	/**
+	* Sets the data for this observation. The relative position can be used to make this observation
+	* relative to the agent's perspective, e.g. by passing the agent's position. Call during
+	* ULearningAgentsType::SetObservations event.
+	* @param AgentId The agent id this data corresponds to.
+	* @param Position The position currently being observed.
+	* @param RelativePosition The vector Position will be offset from.
+	* @param Axis The axis along which to encode the positions
+	*/
+	UFUNCTION(BlueprintCallable, Category = "LearningAgents", meta = (AgentId = "-1"))
+	void SetScalarPositionObservationWithAxis(const int32 AgentId, const FVector Position, const FVector RelativePosition = FVector::ZeroVector, const FVector Axis = FVector::UpVector);
+
+#if UE_LEARNING_AGENTS_ENABLE_VISUAL_LOG
+	/** Describes this observation to the visual logger for debugging purposes. */
+	virtual void VisualLog(const UE::Learning::FIndexSet Instances) const override;
+#endif
+
+	TSharedPtr<UE::Learning::FScalarPositionFeature> FeatureObject;
+};
+
+/** An observation of an array of positions along a single axis. */
+UCLASS()
+class LEARNINGAGENTS_API UScalarPositionArrayObservation : public ULearningAgentsObservation
+{
+	GENERATED_BODY()
+
+public:
+
+	/**
+	* Adds a new scalar position array observation to the given agent type. Call during 
+	* ULearningAgentsType::SetupObservations event.
+	* @param InAgentType The agent type to add this observation to.
+	* @param Name The name of this new observation. Used for debugging.
+	* @param PositionNum The number of positions in the array.
+	* @param Scale Used to normalize the data for the observation.
+	* @return The newly created observation.
+	*/
+	UFUNCTION(BlueprintCallable, Category = "LearningAgents")
+	static UScalarPositionArrayObservation* AddScalarPositionArrayObservation(ULearningAgentsType* InAgentType, const FName Name = NAME_None, const int32 PositionNum = 1, const float Scale = 100.0f);
+
+	/**
+	* Sets the data for this observation. The relative position can be used to make this observation
+	* relative to the agent's perspective, e.g. by passing the agent's position. Call during
+	* ULearningAgentsType::SetObservations event.
+	* @param AgentId The agent id this data corresponds to.
+	* @param Positions The positions currently being observed.
+	* @param RelativePosition The vector Positions will be offset from.
+	*/
+	UFUNCTION(BlueprintCallable, Category = "LearningAgents", meta = (AgentId = "-1"))
+	void SetScalarPositionArrayObservation(const int32 AgentId, const TArray<float>& Positions, const float RelativePosition = 0.0f);
+
+	/**
+	* Sets the data for this observation. The relative position can be used to make this observation
+	* relative to the agent's perspective, e.g. by passing the agent's position. Call during
+	* ULearningAgentsType::SetObservations event.
+	* @param AgentId The agent id this data corresponds to.
+	* @param Positions The positions currently being observed.
+	* @param RelativePosition The vector Positions will be offset from.
+	* @param Axis The axis along which to encode the positions
+	*/
+	UFUNCTION(BlueprintCallable, Category = "LearningAgents", meta = (AgentId = "-1"))
+	void SetScalarPositionArrayObservationWithAxis(const int32 AgentId, const TArray<FVector>& Positions, const FVector RelativePosition = FVector::ZeroVector, const FVector Axis = FVector::UpVector);
+
+#if UE_LEARNING_AGENTS_ENABLE_VISUAL_LOG
+	/** Describes this observation to the visual logger for debugging purposes. */
+	virtual void VisualLog(const UE::Learning::FIndexSet Instances) const override;
+#endif
+
+	TSharedPtr<UE::Learning::FScalarPositionFeature> FeatureObject;
+};
+
+/** An observation of a position projected onto a plane. */
+UCLASS()
+class LEARNINGAGENTS_API UPlanarPositionObservation : public ULearningAgentsObservation
+{
+	GENERATED_BODY()
+
+public:
+
+	/**
+	* Adds a new planar position observation to the given agent type. The axis parameters define the plane. Call
+	* during ULearningAgentsType::SetupObservations event.
+	* @param InAgentType The agent type to add this observation to.
 	* @param Name The name of this new observation. Used for debugging.
 	* @param Scale Used to normalize the data for the observation.
 	* @param Axis0 The forward axis of the plane.
@@ -403,26 +859,72 @@
 	* @return The newly created observation.
 	*/
 	UFUNCTION(BlueprintCallable, Category = "LearningAgents")
-	static UPlanarVelocityObservation* AddPlanarVelocityObservation(ULearningAgentsType* AgentType, const FName Name = NAME_None, const float Scale = 200.0f, const FVector Axis0 = FVector::ForwardVector, const FVector Axis1 = FVector::RightVector);
+	static UPlanarPositionObservation* AddPlanarPositionObservation(ULearningAgentsType* InAgentType, const FName Name = NAME_None, const float Scale = 100.0f, const FVector Axis0 = FVector::ForwardVector, const FVector Axis1 = FVector::RightVector);
 
 	/**
-	* Sets the data for this observation. The relative rotation can be used to make this observation relative to the
-	* agent's perspective, e.g. by passing the agent's forward rotation. Call during ULearningAgentsType::SetObservations event.
+	* Sets the data for this observation. The relative position & rotation can be used to make this observation
+	* relative to the agent's perspective, e.g. by passing the agent's position & forward rotation. Call during
+	* ULearningAgentsType::SetObservations event.
 	* @param AgentId The agent id this data corresponds to.
-	* @param Velocity The velocity currently being observed.
+	* @param Position The position currently being observed.
+	* @param RelativePosition The vector Position will be offset from.
 	* @param RelativeRotation The frame of reference rotation.
 	*/
-	UFUNCTION(BlueprintCallable, Category = "LearningAgents")
-	void SetPlanarVelocityObservation(const int32 AgentId, const FVector Velocity, const FRotator RelativeRotation = FRotator::ZeroRotator);
+	UFUNCTION(BlueprintCallable, Category = "LearningAgents", meta = (AgentId = "-1"))
+	void SetPlanarPositionObservation(const int32 AgentId, const FVector Position, const FVector RelativePosition = FVector::ZeroVector, const FRotator RelativeRotation = FRotator::ZeroRotator);
 
-#if ENABLE_VISUAL_LOG
+#if UE_LEARNING_AGENTS_ENABLE_VISUAL_LOG
 	/** Describes this observation to the visual logger for debugging purposes. */
 	virtual void VisualLog(const UE::Learning::FIndexSet Instances) const override;
 #endif
 
-	TSharedPtr<UE::Learning::FPlanarVelocityFeature> FeatureObject;
+	TSharedPtr<UE::Learning::FPlanarPositionFeature> FeatureObject;
 };
 
+/** An observation of an array of positions projected onto a plane. */
+UCLASS()
+class LEARNINGAGENTS_API UPlanarPositionArrayObservation : public ULearningAgentsObservation
+{
+	GENERATED_BODY()
+
+public:
+
+	/**
+	* Adds a new planar position array observation to the given agent type. The axis parameters define the plane.
+	* Call during ULearningAgentsType::SetupObservations event.
+	* @param InAgentType The agent type to add this observation to.
+	* @param Name The name of this new observation. Used for debugging.
+	* @param PositionNum The number of positions in the array.
+	* @param Scale Used to normalize the data for the observation.
+	* @param Axis0 The forward axis of the plane.
+	* @param Axis1 The right axis of the plane.
+	* @return The newly created observation.
+	*/
+	UFUNCTION(BlueprintCallable, Category = "LearningAgents")
+	static UPlanarPositionArrayObservation* AddPlanarPositionArrayObservation(ULearningAgentsType* InAgentType, const FName Name = NAME_None, const int32 PositionNum = 1, const float Scale = 100.0f, const FVector Axis0 = FVector::ForwardVector, const FVector Axis1 = FVector::RightVector);
+
+	/**
+	* Sets the data for this observation. The relative position & rotation can be used to make this observation
+	* relative to the agent's perspective, e.g. by passing the agent's position & forward rotation. Call during
+	* ULearningAgentsType::SetObservations event.
+	* @param AgentId The agent id this data corresponds to.
+	* @param Positions The positions currently being observed.
+	* @param RelativePosition The vector Positions will be offset from.
+	* @param RelativeRotation The frame of reference rotation.
+	*/
+	UFUNCTION(BlueprintCallable, Category = "LearningAgents", meta = (AgentId = "-1"))
+	void SetPlanarPositionArrayObservation(const int32 AgentId, const TArray<FVector>& Positions, const FVector RelativePosition = FVector::ZeroVector, const FRotator RelativeRotation = FRotator::ZeroRotator);
+
+#if UE_LEARNING_AGENTS_ENABLE_VISUAL_LOG
+	/** Describes this observation to the visual logger for debugging purposes. */
+	virtual void VisualLog(const UE::Learning::FIndexSet Instances) const override;
+#endif
+
+	TSharedPtr<UE::Learning::FPlanarPositionFeature> FeatureObject;
+};
+
+//------------------------------------------------------------------
+
 /** An observation of a velocity. */
 UCLASS()
 class LEARNINGAGENTS_API UVelocityObservation : public ULearningAgentsObservation
@@ -433,13 +935,13 @@
 
 	/**
 	* Adds a new velocity observation to the given agent type. Call during ULearningAgentsType::SetupObservations event.
-	* @param AgentType The agent type to add this observation to.
+	* @param InAgentType The agent type to add this observation to.
 	* @param Name The name of this new observation. Used for debugging.
 	* @param Scale Used to normalize the data for the observation.
 	* @return The newly created observation.
 	*/
 	UFUNCTION(BlueprintCallable, Category = "LearningAgents")
-	static UVelocityObservation* AddVelocityObservation(ULearningAgentsType* AgentType, const FName Name = NAME_None, const float Scale = 200.0f);
+	static UVelocityObservation* AddVelocityObservation(ULearningAgentsType* InAgentType, const FName Name = NAME_None, const float Scale = 200.0f);
 
 	/**
 	* Sets the data for this observation. The relative rotation can be used to make this observation relative to the
@@ -448,13 +950,382 @@
 	* @param Velocity The velocity currently being observed.
 	* @param RelativeRotation The frame of reference rotation.
 	*/
-	UFUNCTION(BlueprintCallable, Category = "LearningAgents")
+	UFUNCTION(BlueprintCallable, Category = "LearningAgents", meta = (AgentId = "-1"))
 	void SetVelocityObservation(const int32 AgentId, const FVector Velocity, const FRotator RelativeRotation = FRotator::ZeroRotator);
 
-#if ENABLE_VISUAL_LOG
+#if UE_LEARNING_AGENTS_ENABLE_VISUAL_LOG
 	/** Describes this observation to the visual logger for debugging purposes. */
 	virtual void VisualLog(const UE::Learning::FIndexSet Instances) const override;
 #endif
 
 	TSharedPtr<UE::Learning::FVelocityFeature> FeatureObject;
 };
+
+/** An observation of an array of velocities. */
+UCLASS()
+class LEARNINGAGENTS_API UVelocityArrayObservation : public ULearningAgentsObservation
+{
+	GENERATED_BODY()
+
+public:
+
+	/**
+	* Adds a new velocity observation to the given agent type. Call during ULearningAgentsType::SetupObservations event.
+	* @param InAgentType The agent type to add this observation to.
+	* @param Name The name of this new observation. Used for debugging.
+	* @param VelocityNum The number of velocities in the array.
+	* @param Scale Used to normalize the data for the observation.
+	* @return The newly created observation.
+	*/
+	UFUNCTION(BlueprintCallable, Category = "LearningAgents")
+	static UVelocityArrayObservation* AddVelocityArrayObservation(ULearningAgentsType* InAgentType, const FName Name = NAME_None, const int32 VelocityNum = 1, const float Scale = 200.0f);
+
+	/**
+	* Sets the data for this observation. The relative rotation can be used to make this observation relative to the
+	* agent's perspective, e.g. by passing the agent's forward rotation. Call during ULearningAgentsType::SetObservations event.
+	* @param AgentId The agent id this data corresponds to.
+	* @param Velocities The velocities currently being observed.
+	* @param RelativeRotation The frame of reference rotation.
+	*/
+	UFUNCTION(BlueprintCallable, Category = "LearningAgents", meta = (AgentId = "-1"))
+	void SetVelocityArrayObservation(const int32 AgentId, const TArray<FVector>& Velocities, const FRotator RelativeRotation = FRotator::ZeroRotator);
+
+#if UE_LEARNING_AGENTS_ENABLE_VISUAL_LOG
+	/** Describes this observation to the visual logger for debugging purposes. */
+	virtual void VisualLog(const UE::Learning::FIndexSet Instances) const override;
+#endif
+
+	TSharedPtr<UE::Learning::FVelocityFeature> FeatureObject;
+};
+
+/** An observation of a velocity along a single axis. */
+UCLASS()
+class LEARNINGAGENTS_API UScalarVelocityObservation : public ULearningAgentsObservation
+{
+	GENERATED_BODY()
+
+public:
+
+	/**
+	* Adds a new scalar velocity observation to the given agent type. 
+	* Call during ULearningAgentsType::SetupObservations event.
+	* @param InAgentType The agent type to add this observation to.
+	* @param Name The name of this new observation. Used for debugging.
+	* @param Scale Used to normalize the data for the observation.
+	* @return The newly created observation.
+	*/
+	UFUNCTION(BlueprintCallable, Category = "LearningAgents")
+	static UScalarVelocityObservation* AddScalarVelocityObservation(ULearningAgentsType* InAgentType, const FName Name = NAME_None, const float Scale = 200.0f);
+
+	/**
+	* Sets the data for this observation. Call during ULearningAgentsType::SetObservations event.
+	* @param AgentId The agent id this data corresponds to.
+	* @param Velocity The velocity currently being observed.
+	*/
+	UFUNCTION(BlueprintCallable, Category = "LearningAgents", meta = (AgentId = "-1"))
+	void SetScalarVelocityObservation(const int32 AgentId, const float Velocity);
+
+	/**
+	* Sets the data for this observation. Call during ULearningAgentsType::SetObservations event.
+	* @param AgentId The agent id this data corresponds to.
+	* @param Velocity The velocity currently being observed.
+	* @param Axis The axis to encode the velocity along
+	*/
+	UFUNCTION(BlueprintCallable, Category = "LearningAgents", meta = (AgentId = "-1"))
+	void SetScalarVelocityObservationWithAxis(const int32 AgentId, const FVector Velocity, const FVector Axis = FVector::UpVector);
+
+#if UE_LEARNING_AGENTS_ENABLE_VISUAL_LOG
+	/** Describes this observation to the visual logger for debugging purposes. */
+	virtual void VisualLog(const UE::Learning::FIndexSet Instances) const override;
+#endif
+
+	TSharedPtr<UE::Learning::FScalarVelocityFeature> FeatureObject;
+};
+
+/** An observation of an array of velocities along a single axis. */
+UCLASS()
+class LEARNINGAGENTS_API UScalarVelocityArrayObservation : public ULearningAgentsObservation
+{
+	GENERATED_BODY()
+
+public:
+
+	/**
+	* Adds a new scalar velocity observation to the given agent type.
+	* Call during ULearningAgentsType::SetupObservations event.
+	* @param InAgentType The agent type to add this observation to.
+	* @param Name The name of this new observation. Used for debugging.
+	* @param VelocityNum The number of velocities in the array.
+	* @param Scale Used to normalize the data for the observation.
+	* @return The newly created observation.
+	*/
+	UFUNCTION(BlueprintCallable, Category = "LearningAgents")
+	static UScalarVelocityArrayObservation* AddScalarVelocityArrayObservation(ULearningAgentsType* InAgentType, const FName Name = NAME_None, const int32 VelocityNum = 1, const float Scale = 200.0f);
+
+	/**
+	* Sets the data for this observation. Call during ULearningAgentsType::SetObservations event.
+	* @param AgentId The agent id this data corresponds to.
+	* @param Velocities The velocities currently being observed.
+	*/
+	UFUNCTION(BlueprintCallable, Category = "LearningAgents", meta = (AgentId = "-1"))
+	void SetScalarVelocityArrayObservation(const int32 AgentId, const TArray<float>& Velocities);
+
+	/**
+	* Sets the data for this observation. Call during ULearningAgentsType::SetObservations event.
+	* @param AgentId The agent id this data corresponds to.
+	* @param Velocities The velocities currently being observed.
+	* @param Axis The axis to encode the velocity along
+	*/
+	UFUNCTION(BlueprintCallable, Category = "LearningAgents", meta = (AgentId = "-1"))
+	void SetScalarVelocityArrayObservationWithAxis(const int32 AgentId, const TArray<FVector>& Velocities, const FVector Axis = FVector::UpVector);
+
+#if UE_LEARNING_AGENTS_ENABLE_VISUAL_LOG
+	/** Describes this observation to the visual logger for debugging purposes. */
+	virtual void VisualLog(const UE::Learning::FIndexSet Instances) const override;
+#endif
+
+	TSharedPtr<UE::Learning::FScalarVelocityFeature> FeatureObject;
+};
+
+/** An observation of a velocity projected onto a plane. */
+UCLASS()
+class LEARNINGAGENTS_API UPlanarVelocityObservation : public ULearningAgentsObservation
+{
+	GENERATED_BODY()
+
+public:
+
+	/**
+	* Adds a new planar velocity observation to the given agent type. The axis parameters define the plane.
+	* Call during ULearningAgentsType::SetupObservations event.
+	* @param InAgentType The agent type to add this observation to.
+	* @param Name The name of this new observation. Used for debugging.
+	* @param Scale Used to normalize the data for the observation.
+	* @param Axis0 The forward axis of the plane.
+	* @param Axis1 The right axis of the plane.
+	* @return The newly created observation.
+	*/
+	UFUNCTION(BlueprintCallable, Category = "LearningAgents")
+	static UPlanarVelocityObservation* AddPlanarVelocityObservation(ULearningAgentsType* InAgentType, const FName Name = NAME_None, const float Scale = 200.0f, const FVector Axis0 = FVector::ForwardVector, const FVector Axis1 = FVector::RightVector);
+
+	/**
+	* Sets the data for this observation. The relative rotation can be used to make this observation relative to the
+	* agent's perspective, e.g. by passing the agent's forward rotation. Call during ULearningAgentsType::SetObservations event.
+	* @param AgentId The agent id this data corresponds to.
+	* @param Velocity The velocity currently being observed.
+	* @param RelativeRotation The frame of reference rotation.
+	*/
+	UFUNCTION(BlueprintCallable, Category = "LearningAgents", meta = (AgentId = "-1"))
+	void SetPlanarVelocityObservation(const int32 AgentId, const FVector Velocity, const FRotator RelativeRotation = FRotator::ZeroRotator);
+
+#if UE_LEARNING_AGENTS_ENABLE_VISUAL_LOG
+	/** Describes this observation to the visual logger for debugging purposes. */
+	virtual void VisualLog(const UE::Learning::FIndexSet Instances) const override;
+#endif
+
+	TSharedPtr<UE::Learning::FPlanarVelocityFeature> FeatureObject;
+};
+
+/** An observation of an array of velocities projected onto a plane. */
+UCLASS()
+class LEARNINGAGENTS_API UPlanarVelocityArrayObservation : public ULearningAgentsObservation
+{
+	GENERATED_BODY()
+
+public:
+
+	/**
+	* Adds a new planar velocity observation to the given agent type. The axis parameters define the plane.
+	* Call during ULearningAgentsType::SetupObservations event.
+	* @param InAgentType The agent type to add this observation to.
+	* @param Name The name of this new observation. Used for debugging.
+	* @param VelocityNum The number of velocities in the array.
+	* @param Scale Used to normalize the data for the observation.
+	* @param Axis0 The forward axis of the plane.
+	* @param Axis1 The right axis of the plane.
+	* @return The newly created observation.
+	*/
+	UFUNCTION(BlueprintCallable, Category = "LearningAgents")
+	static UPlanarVelocityArrayObservation* AddPlanarVelocityArrayObservation(ULearningAgentsType* InAgentType, const FName Name = NAME_None, const int32 VelocityNum = 1, const float Scale = 200.0f, const FVector Axis0 = FVector::ForwardVector, const FVector Axis1 = FVector::RightVector);
+
+	/**
+	* Sets the data for this observation. The relative rotation can be used to make this observation relative to the
+	* agent's perspective, e.g. by passing the agent's forward rotation. Call during ULearningAgentsType::SetObservations event.
+	* @param AgentId The agent id this data corresponds to.
+	* @param Velocities The velocities currently being observed.
+	* @param RelativeRotation The frame of reference rotation.
+	*/
+	UFUNCTION(BlueprintCallable, Category = "LearningAgents", meta = (AgentId = "-1"))
+	void SetPlanarVelocityArrayObservation(const int32 AgentId, const TArray<FVector>& Velocities, const FRotator RelativeRotation = FRotator::ZeroRotator);
+
+#if UE_LEARNING_AGENTS_ENABLE_VISUAL_LOG
+	/** Describes this observation to the visual logger for debugging purposes. */
+	virtual void VisualLog(const UE::Learning::FIndexSet Instances) const override;
+#endif
+
+	TSharedPtr<UE::Learning::FPlanarVelocityFeature> FeatureObject;
+};
+
+//------------------------------------------------------------------
+
+/** An observation of an angular velocity. */
+UCLASS()
+class LEARNINGAGENTS_API UAngularVelocityObservation : public ULearningAgentsObservation
+{
+	GENERATED_BODY()
+
+public:
+
+	/**
+	* Adds a new angular velocity observation to the given agent type. Call during ULearningAgentsType::SetupObservations event.
+	* @param InAgentType The agent type to add this observation to.
+	* @param Name The name of this new observation. Used for debugging.
+	* @param Scale Used to normalize the data for the observation.
+	* @return The newly created observation.
+	*/
+	UFUNCTION(BlueprintCallable, Category = "LearningAgents")
+	static UAngularVelocityObservation* AddAngularVelocityObservation(ULearningAgentsType* InAgentType, const FName Name = NAME_None, const float Scale = 180.0f);
+
+	/**
+	* Sets the data for this observation. The relative rotation can be used to make this observation relative to the
+	* agent's perspective, e.g. by passing the agent's forward rotation. Call during ULearningAgentsType::SetObservations event.
+	* @param AgentId The agent id this data corresponds to.
+	* @param AngularVelocity The angular velocity currently being observed.
+	* @param RelativeRotation The frame of reference rotation.
+	*/
+	UFUNCTION(BlueprintCallable, Category = "LearningAgents", meta = (AgentId = "-1"))
+	void SetAngularVelocityObservation(const int32 AgentId, const FVector AngularVelocity, const FRotator RelativeRotation = FRotator::ZeroRotator);
+
+#if UE_LEARNING_AGENTS_ENABLE_VISUAL_LOG
+	/** Describes this observation to the visual logger for debugging purposes. */
+	virtual void VisualLog(const UE::Learning::FIndexSet Instances) const override;
+#endif
+
+	TSharedPtr<UE::Learning::FAngularVelocityFeature> FeatureObject;
+};
+
+/** An observation of an array of angular velocities. */
+UCLASS()
+class LEARNINGAGENTS_API UAngularVelocityArrayObservation : public ULearningAgentsObservation
+{
+	GENERATED_BODY()
+
+public:
+
+	/**
+	* Adds a new angular velocity array observation to the given agent type. Call during ULearningAgentsType::SetupObservations event.
+	* @param InAgentType The agent type to add this observation to.
+	* @param Name The name of this new observation. Used for debugging.
+	* @param AngularVelocityNum The number of angular velocities in the array.
+	* @param Scale Used to normalize the data for the observation.
+	* @return The newly created observation.
+	*/
+	UFUNCTION(BlueprintCallable, Category = "LearningAgents")
+	static UAngularVelocityArrayObservation* AddAngularVelocityArrayObservation(ULearningAgentsType* InAgentType, const FName Name = NAME_None, const int32 AngularVelocityNum = 1, const float Scale = 180.0f);
+
+	/**
+	* Sets the data for this observation. The relative rotation can be used to make this observation relative to the
+	* agent's perspective, e.g. by passing the agent's forward rotation. Call during ULearningAgentsType::SetObservations event.
+	* @param AgentId The agent id this data corresponds to.
+	* @param Velocities The angular velocities currently being observed.
+	* @param RelativeRotation The frame of reference rotation.
+	*/
+	UFUNCTION(BlueprintCallable, Category = "LearningAgents", meta = (AgentId = "-1"))
+	void SetAngularVelocityArrayObservation(const int32 AgentId, const TArray<FVector>& AngularVelocities, const FRotator RelativeRotation = FRotator::ZeroRotator);
+
+#if UE_LEARNING_AGENTS_ENABLE_VISUAL_LOG
+	/** Describes this observation to the visual logger for debugging purposes. */
+	virtual void VisualLog(const UE::Learning::FIndexSet Instances) const override;
+#endif
+
+	TSharedPtr<UE::Learning::FAngularVelocityFeature> FeatureObject;
+};
+
+/** An observation of a scalar angular velocity. */
+UCLASS()
+class LEARNINGAGENTS_API UScalarAngularVelocityObservation : public ULearningAgentsObservation
+{
+	GENERATED_BODY()
+
+public:
+
+	/**
+	* Adds a new scalar angular velocity observation to the given agent type.
+	* Call during ULearningAgentsType::SetupObservations event.
+	* @param InAgentType The agent type to add this observation to.
+	* @param Name The name of this new observation. Used for debugging.
+	* @param Scale Used to normalize the data for the observation.
+	* @return The newly created observation.
+	*/
+	UFUNCTION(BlueprintCallable, Category = "LearningAgents")
+	static UScalarAngularVelocityObservation* AddScalarAngularVelocityObservation(ULearningAgentsType* InAgentType, const FName Name = NAME_None, const float Scale = 180.0f);
+
+	/**
+	* Sets the data for this observation. Call during ULearningAgentsType::SetObservations event.
+	* @param AgentId The agent id this data corresponds to.
+	* @param AngularVelocity The angular velocity currently being observed.
+	*/
+	UFUNCTION(BlueprintCallable, Category = "LearningAgents", meta = (AgentId = "-1"))
+	void SetScalarAngularVelocityObservation(const int32 AgentId, const float AngularVelocity);
+
+	/**
+	* Sets the data for this observation. Call during ULearningAgentsType::SetObservations event.
+	* @param AgentId The agent id this data corresponds to.
+	* @param AngularVelocity The angular velocity currently being observed.
+	* @param Axis The axis to encode the angular velocity around
+	*/
+	UFUNCTION(BlueprintCallable, Category = "LearningAgents", meta = (AgentId = "-1"))
+	void SetScalarAngularVelocityObservationWithAxis(const int32 AgentId, const FVector AngularVelocity, const FVector Axis = FVector::UpVector);
+
+#if UE_LEARNING_AGENTS_ENABLE_VISUAL_LOG
+	/** Describes this observation to the visual logger for debugging purposes. */
+	virtual void VisualLog(const UE::Learning::FIndexSet Instances) const override;
+#endif
+
+	TSharedPtr<UE::Learning::FScalarAngularVelocityFeature> FeatureObject;
+};
+
+/** An observation of an array of scalar angular velocities. */
+UCLASS()
+class LEARNINGAGENTS_API UScalarAngularVelocityArrayObservation : public ULearningAgentsObservation
+{
+	GENERATED_BODY()
+
+public:
+
+	/**
+	* Adds a new scalar angular velocity array observation to the given agent type.
+	* Call during ULearningAgentsType::SetupObservations event.
+	* @param InAgentType The agent type to add this observation to.
+	* @param Name The name of this new observation. Used for debugging.
+	* @param AngularVelocityNum The number of angular velocities in the array.
+	* @param Scale Used to normalize the data for the observation.
+	* @return The newly created observation.
+	*/
+	UFUNCTION(BlueprintCallable, Category = "LearningAgents")
+	static UScalarAngularVelocityArrayObservation* AddScalarAngularVelocityArrayObservation(ULearningAgentsType* InAgentType, const FName Name = NAME_None, const int32 AngularVelocityNum = 1, const float Scale = 180.0f);
+
+	/**
+	* Sets the data for this observation. Call during ULearningAgentsType::SetObservations event.
+	* @param AgentId The agent id this data corresponds to.
+	* @param AngularVelocities The angular velocities currently being observed.
+	*/
+	UFUNCTION(BlueprintCallable, Category = "LearningAgents", meta = (AgentId = "-1"))
+	void SetScalarAngularVelocityArrayObservation(const int32 AgentId, const TArray<float>& AngularVelocities);
+
+	/**
+	* Sets the data for this observation. Call during ULearningAgentsType::SetObservations event.
+	* @param AgentId The agent id this data corresponds to.
+	* @param AngularVelocities The angular velocities currently being observed.
+	* @param Axis The axis to encode the angular velocity around
+	*/
+	UFUNCTION(BlueprintCallable, Category = "LearningAgents", meta = (AgentId = "-1"))
+	void SetScalarAngularVelocityArrayObservationWithAxis(const int32 AgentId, const TArray<FVector>& AngularVelocities, const FVector Axis = FVector::UpVector);
+
+#if UE_LEARNING_AGENTS_ENABLE_VISUAL_LOG
+	/** Describes this observation to the visual logger for debugging purposes. */
+	virtual void VisualLog(const UE::Learning::FIndexSet Instances) const override;
+#endif
+
+	TSharedPtr<UE::Learning::FScalarAngularVelocityFeature> FeatureObject;
+};
diff --git a/Engine/Plugins/Experimental/LearningAgents/Source/LearningAgents/Public/LearningAgentsPolicy.h b/Engine/Plugins/Experimental/LearningAgents/Source/LearningAgents/Public/LearningAgentsPolicy.h
--- a/Engine/Plugins/Experimental/LearningAgents/Source/LearningAgents/Public/LearningAgentsPolicy.h
+++ b/Engine/Plugins/Experimental/LearningAgents/Source/LearningAgents/Public/LearningAgentsPolicy.h
@@ -5,9 +5,9 @@
 #include "LearningAgentsManagerComponent.h"
 
 #include "LearningAgentsNeuralNetwork.h" // Included for ELearningAgentsActivationFunction
+#include "LearningAgentsDebug.h"
 #include "LearningArray.h"
 #include "UObject/ObjectPtr.h"
-#include "EngineDefines.h"
 
 #include "LearningAgentsPolicy.generated.h"
 
@@ -74,6 +74,12 @@ public:
 	UFUNCTION(BlueprintCallable, Category = "LearningAgents")
 	void SetupPolicy(ALearningAgentsManager* InAgentManager, ULearningAgentsType* InAgentType, const FLearningAgentsPolicySettings& PolicySettings = FLearningAgentsPolicySettings());
 
+public:
+
+	//~ Begin ULearningAgentsManagerComponent Interface
+	virtual bool AddAgent(const int32 AgentId) override;
+	//~ End ULearningAgentsManagerComponent Interface
+
 // ----- Load / Save -----
 public:
 
@@ -118,13 +124,19 @@ public:
 	UFUNCTION(BlueprintCallable, Category = "LearningAgents")
 	void EvaluatePolicy();
 
+	/**
+	* Calls EncodeObservations, followed by EvaluatePolicy, followed by DecodeActions
+	*/
+	UFUNCTION(BlueprintCallable, Category = "LearningAgents")
+	void RunInference();
+
 	/**
 	* Get the action noise scale used by an agent.
 	*
 	* @param AgentId			The AgentId to get the action noise scale for.
 	* @return					Action noise scale for that agent.
 	*/
-	UFUNCTION(BlueprintPure, Category = "LearningAgents")
+	UFUNCTION(BlueprintPure, Category = "LearningAgents", meta = (AgentId = "-1"))
 	float GetAgentActionNoiseScale(const int32 AgentId) const;
 
 	/**
@@ -136,9 +148,17 @@ public:
 	* @param AgentId			The AgentId to set the action noise scale for.
 	* @param ActionNoiseScale	Action noise scale for that agent.
 	*/
-	UFUNCTION(BlueprintCallable, Category = "LearningAgents")
+	UFUNCTION(BlueprintCallable, Category = "LearningAgents", meta = (AgentId = "-1"))
 	void SetAgentActionNoiseScale(const int32 AgentId, const float ActionNoiseScale);
 
+	/**
+	* Set the action noise scale used by all agents.
+	*
+	* @param ActionNoiseScale	Action noise scale to use for all agents
+	*/
+	UFUNCTION(BlueprintCallable, Category = "LearningAgents")
+	void SetAllAgentsActionNoiseScale(const float ActionNoiseScale);
+
 // ----- Non-blueprint public interface -----
 public:
 
@@ -162,8 +182,10 @@ private:
 private:
 
 	TSharedPtr<UE::Learning::FNeuralNetworkPolicyFunction> PolicyObject;
+	
+	float InitialActionNoiseScale = 1.0f;
 
-#if ENABLE_VISUAL_LOG
+#if UE_LEARNING_AGENTS_ENABLE_VISUAL_LOG
 	/** Color used to draw this action in the visual log */
 	FLinearColor VisualLogColor = FColor::Purple;
 
diff --git a/Engine/Plugins/Experimental/LearningAgents/Source/LearningAgents/Public/LearningAgentsType.h b/Engine/Plugins/Experimental/LearningAgents/Source/LearningAgents/Public/LearningAgentsType.h
--- a/Engine/Plugins/Experimental/LearningAgents/Source/LearningAgents/Public/LearningAgentsType.h
+++ b/Engine/Plugins/Experimental/LearningAgents/Source/LearningAgents/Public/LearningAgentsType.h
@@ -121,6 +121,24 @@ public:
 	UFUNCTION(BlueprintCallable, Category = "LearningAgents")
 	void DecodeActions();
 
+	/**
+	* Gets the observation vector used by a given agent. Should be called only after EncodeObservations.
+	*
+	* @param AgentId				The AgentId to look-up the observation vector for
+	* @param OutObservationVector	The observation vector for the given agent
+	*/
+	UFUNCTION(BlueprintPure, Category = "LearningAgents", meta = (AgentId = "-1"))
+	void GetObservationVector(const int32 AgentId, TArray<float>& OutObservationVector) const;
+
+	/**
+	* Gets the action vector used by a given agent. Should be called only after EncodeActions or EvaluatePolicy.
+	*
+	* @param AgentId				The AgentId to look-up the action vector for
+	* @param OutActionVector		The action vector for the given agent
+	*/
+	UFUNCTION(BlueprintPure, Category = "LearningAgents", meta = (AgentId = "-1"))
+	void GetActionVector(const int32 AgentId, TArray<float>& OutActionVector) const;
+
 // ----- Non-blueprint public interface -----
 public:
 
diff --git a/Engine/Plugins/Experimental/LearningAgents/Source/LearningAgentsEditor/Private/LearningAgentsEditor.cpp b/Engine/Plugins/Experimental/LearningAgents/Source/LearningAgentsEditor/Private/LearningAgentsEditor.cpp
--- a/Engine/Plugins/Experimental/LearningAgents/Source/LearningAgentsEditor/Private/LearningAgentsEditor.cpp
+++ b/Engine/Plugins/Experimental/LearningAgents/Source/LearningAgentsEditor/Private/LearningAgentsEditor.cpp
@@ -20,7 +20,7 @@ void FLearningAgentsEditorModule::StartupModule()
 	FKismetEditorUtilities::RegisterAutoGeneratedDefaultEvent(this, ULearningAgentsTrainer::StaticClass(), GET_FUNCTION_NAME_CHECKED(ULearningAgentsTrainer, SetRewards));
 	FKismetEditorUtilities::RegisterAutoGeneratedDefaultEvent(this, ULearningAgentsTrainer::StaticClass(), GET_FUNCTION_NAME_CHECKED(ULearningAgentsTrainer, SetupCompletions));
 	FKismetEditorUtilities::RegisterAutoGeneratedDefaultEvent(this, ULearningAgentsTrainer::StaticClass(), GET_FUNCTION_NAME_CHECKED(ULearningAgentsTrainer, SetCompletions));
-	FKismetEditorUtilities::RegisterAutoGeneratedDefaultEvent(this, ULearningAgentsTrainer::StaticClass(), GET_FUNCTION_NAME_CHECKED(ULearningAgentsTrainer, ResetInstance));
+	FKismetEditorUtilities::RegisterAutoGeneratedDefaultEvent(this, ULearningAgentsTrainer::StaticClass(), GET_FUNCTION_NAME_CHECKED(ULearningAgentsTrainer, ResetEpisodes));
 
 	// ULearningAgentsController Default Events
 	FKismetEditorUtilities::RegisterAutoGeneratedDefaultEvent(this, ULearningAgentsController::StaticClass(), GET_FUNCTION_NAME_CHECKED(ULearningAgentsController, SetActions));
diff --git a/Engine/Plugins/Experimental/LearningAgents/Source/LearningAgentsTraining/Private/LearningAgentsCompletions.cpp b/Engine/Plugins/Experimental/LearningAgents/Source/LearningAgentsTraining/Private/LearningAgentsCompletions.cpp
--- a/Engine/Plugins/Experimental/LearningAgents/Source/LearningAgentsTraining/Private/LearningAgentsCompletions.cpp
+++ b/Engine/Plugins/Experimental/LearningAgents/Source/LearningAgentsTraining/Private/LearningAgentsCompletions.cpp
@@ -11,40 +11,28 @@
 #include "LearningLog.h"
 
 #include "GameFramework/Actor.h"
-#include "VisualLogger/VisualLogger.h"
 
-#define UE_LEARNING_AGENTS_VLOG_STRING(Owner, Category, Verbosity, Location, Color, Format, ...) \
-	UE_VLOG_LOCATION(Owner, Category, Verbosity, Location, 0.0f, Color, Format, ##__VA_ARGS__)
-
-#define UE_LEARNING_AGENTS_VLOG_TRANSFORM(Owner, Category, Verbosity, Location, Rotation, Color, Format, ...) \
-	UE_VLOG_SEGMENT(Owner, Category, Verbosity, Location, Location + 15.0f * Rotation.RotateVector(FVector::ForwardVector), FColor::Red, TEXT("")); \
-	UE_VLOG_SEGMENT(Owner, Category, Verbosity, Location, Location + 15.0f * Rotation.RotateVector(FVector::RightVector), FColor::Green, TEXT("")); \
-	UE_VLOG_SEGMENT(Owner, Category, Verbosity, Location, Location + 15.0f * Rotation.RotateVector(FVector::UpVector), FColor::Blue, TEXT("")); \
-	UE_VLOG_OBOX(Owner, Category, Verbosity, FBox(10.0f * FVector(-1, -1, -1), 10.0f * FVector(1, 1, 1)), FTransform(Rotation, Location, FVector::OneVector).ToMatrixNoScale(), Color, TEXT("")); \
-	UE_LEARNING_AGENTS_VLOG_STRING(Owner, Category, Verbosity, Location + FVector(0.0f, 0.0f, 20.0f), Color, Format, ##__VA_ARGS__)
-
-#define UE_LEARNING_AGENTS_VLOG_PLANE(Owner, Category, Verbosity, Location, Rotation, Axis0, Axis1, Color, Format, ...) \
-	UE_VLOG_OBOX(Owner, Category, Verbosity, FBox(-25.0f * (Axis0 + Axis1), 25.0f * (Axis0 + Axis1)), FTransform(Rotation, Location, FVector::OneVector).ToMatrixNoScale(), Color, Format, ##__VA_ARGS__)
-
-namespace UE::Learning::Agents::Private
+namespace UE::Learning::Agents::Completions::Private
 {
 	template<typename CompletionUObject, typename CompletionFObject, typename... InArgTypes>
-	CompletionUObject* AddCompletion(ULearningAgentsTrainer* AgentTrainer, const FName Name, InArgTypes&& ...Args)
+	CompletionUObject* AddCompletion(ULearningAgentsTrainer* InAgentTrainer, const FName Name, InArgTypes&& ...Args)
 	{
-		if (!AgentTrainer)
+		if (!InAgentTrainer)
 		{
-			UE_LOG(LogLearning, Error, TEXT("AgentTrainer is nullptr")); return nullptr;
+			UE_LOG(LogLearning, Error, TEXT("InAgentTrainer is nullptr."));
+			return nullptr;
 		}
 
-		CompletionUObject* Completion = NewObject<CompletionUObject>(AgentTrainer, Name);
+		CompletionUObject* Completion = NewObject<CompletionUObject>(InAgentTrainer, Name);
 
+		Completion->AgentTrainer = InAgentTrainer;
 		Completion->CompletionObject = MakeShared<CompletionFObject>(
 			Completion->GetFName(),
-			AgentTrainer->GetAgentManager()->GetInstanceData().ToSharedRef(),
-			AgentTrainer->GetAgentManager()->GetMaxInstanceNum(),
+			InAgentTrainer->GetAgentManager()->GetInstanceData().ToSharedRef(),
+			InAgentTrainer->GetAgentManager()->GetMaxInstanceNum(),
 			Forward<InArgTypes>(Args)...);
 
-		AgentTrainer->AddCompletion(Completion, Completion->CompletionObject.ToSharedRef());
+		InAgentTrainer->AddCompletion(Completion, Completion->CompletionObject.ToSharedRef());
 
 		return Completion;
 	}
@@ -52,57 +40,85 @@ namespace UE::Learning::Agents::Private
 
 //------------------------------------------------------------------
 
-UConditionalCompletion* UConditionalCompletion::AddConditionalCompletion(ULearningAgentsTrainer* AgentTrainer, const FName Name, const ELearningAgentsCompletion InCompletionMode)
+UConditionalCompletion* UConditionalCompletion::AddConditionalCompletion(ULearningAgentsTrainer* InAgentTrainer, const FName Name, const ELearningAgentsCompletion InCompletionMode)
 {
-	return UE::Learning::Agents::Private::AddCompletion<UConditionalCompletion, UE::Learning::FConditionalCompletion>(
-		AgentTrainer,
-		Name,
-		InCompletionMode == ELearningAgentsCompletion::Termination ?
-		UE::Learning::ECompletionMode::Terminated :
-		UE::Learning::ECompletionMode::Truncated);
+	return UE::Learning::Agents::Completions::Private::AddCompletion<UConditionalCompletion, UE::Learning::FConditionalCompletion>(InAgentTrainer, Name, UE::Learning::Agents::GetCompletionMode(InCompletionMode));
 }
 
 void UConditionalCompletion::SetConditionalCompletion(const int32 AgentId, const bool bIsComplete)
 {
-	const TLearningArrayView<1, bool> View = CompletionObject->InstanceData->View(CompletionObject->ConditionHandle);
-
-	if (AgentId == INDEX_NONE)
+	if (!AgentTrainer->HasAgent(AgentId))
 	{
-		UE_LOG(LogLearning, Error, TEXT("AgentId is invalid (INDEX_NONE)"));
+		UE_LOG(LogLearning, Error, TEXT("%s: AgentId %d not found in the agents set."), *GetName(), AgentId);
 		return;
 	}
 
-	if (AgentId < 0 || AgentId >= View.Num<0>())
+	CompletionObject->InstanceData->View(CompletionObject->ConditionHandle)[AgentId] = bIsComplete;
+}
+
+#if UE_LEARNING_AGENTS_ENABLE_VISUAL_LOG
+void UConditionalCompletion::VisualLog(const UE::Learning::FIndexSet Instances) const
+{
+	UE_LEARNING_TRACE_CPUPROFILER_EVENT_SCOPE(UConditionalCompletion::VisualLog);
+
+	const TLearningArrayView<1, const bool> ConditionView = CompletionObject->InstanceData->ConstView(CompletionObject->ConditionHandle);
+	const TLearningArrayView<1, const UE::Learning::ECompletionMode> CompletionView = CompletionObject->InstanceData->ConstView(CompletionObject->CompletionHandle);
+
+	for (const int32 Instance : Instances)
 	{
-		UE_LOG(LogLearning, Error, TEXT("AgentId %d is out of index. Valid range [0, %d]."), AgentId, View.Num<0>() - 1);
+		if (const AActor* Actor = Cast<AActor>(AgentTrainer->GetAgent(Instance)))
+		{
+			UE_LEARNING_AGENTS_VLOG_STRING(this, LogLearning, Display,
+				Actor->GetActorLocation(),
+				VisualLogColor.ToFColor(true),
+				TEXT("Agent %i\nCondition: %s\nCompletion: %s"),
+				Instance,
+				ConditionView[Instance] ? TEXT("true") : TEXT("false"),
+				UE::Learning::Completion::CompletionModeString(CompletionView[Instance]));
+		}
+	}
+}
+#endif
+
+//------------------------------------------------------------------
+
+UTimeElapsedCompletion* UTimeElapsedCompletion::AddTimeElapsedCompletion(ULearningAgentsTrainer* InAgentTrainer, const FName Name, const float Threshold, const ELearningAgentsCompletion InCompletionMode)
+{
+	return UE::Learning::Agents::Completions::Private::AddCompletion<UTimeElapsedCompletion, UE::Learning::FTimeElapsedCompletion>(InAgentTrainer, Name, Threshold, UE::Learning::Agents::GetCompletionMode(InCompletionMode));
+}
+
+void UTimeElapsedCompletion::SetTimeElapsedCompletion(const int32 AgentId, const float Time)
+{
+	if (!AgentTrainer->HasAgent(AgentId))
+	{
+		UE_LOG(LogLearning, Error, TEXT("%s: AgentId %d not found in the agents set."), *GetName(), AgentId);
 		return;
 	}
 
-	View[AgentId] = bIsComplete;
+	CompletionObject->InstanceData->View(CompletionObject->TimeHandle)[AgentId] = Time;
 }
 
-#if ENABLE_VISUAL_LOG
-void UConditionalCompletion::VisualLog(const UE::Learning::FIndexSet Instances) const
+#if UE_LEARNING_AGENTS_ENABLE_VISUAL_LOG
+void UTimeElapsedCompletion::VisualLog(const UE::Learning::FIndexSet Instances) const
 {
-	UE_LEARNING_TRACE_CPUPROFILER_EVENT_SCOPE(UConditionalCompletion::VisualLog);
+	UE_LEARNING_TRACE_CPUPROFILER_EVENT_SCOPE(UTimeElapsedCompletion::VisualLog);
 
-	const TLearningArrayView<1, const bool> ConditionView = CompletionObject->InstanceData->ConstView(CompletionObject->ConditionHandle);
+	const TLearningArrayView<1, const float> TimeView = CompletionObject->InstanceData->ConstView(CompletionObject->TimeHandle);
+	const TLearningArrayView<1, const float> ThresholdView = CompletionObject->InstanceData->ConstView(CompletionObject->ThresholdHandle);
 	const TLearningArrayView<1, const UE::Learning::ECompletionMode> CompletionView = CompletionObject->InstanceData->ConstView(CompletionObject->CompletionHandle);
 
-	if (const ULearningAgentsTrainer* AgentTrainer = Cast<ULearningAgentsTrainer>(GetOuter()))
+	for (const int32 Instance : Instances)
 	{
-		for (const int32 Instance : Instances)
+		if (const AActor* Actor = Cast<AActor>(AgentTrainer->GetAgent(Instance)))
 		{
-			if (const AActor* Actor = Cast<AActor>(AgentTrainer->GetAgent(Instance)))
-			{
-				UE_LEARNING_AGENTS_VLOG_STRING(this, LogLearning, Display,
-					Actor->GetActorLocation(),
-					VisualLogColor.ToFColor(true),
-					TEXT("Agent %i\nCondition: %s\nCompletion: %s"),
-					Instance,
-					ConditionView[Instance] ? TEXT("true") : TEXT("false"),
-					UE::Learning::Completion::CompletionModeString(CompletionView[Instance]));
-			}
+			UE_LEARNING_AGENTS_VLOG_STRING(this, LogLearning, Display,
+				Actor->GetActorLocation(),
+				VisualLogColor.ToFColor(true),
+				TEXT("Agent %i\nThreshold: [%6.3f]\nTime: [%6.3f]\nCompletion: %s"),
+				Instance,
+				ThresholdView[Instance],
+				TimeView[Instance],
+				UE::Learning::Completion::CompletionModeString(CompletionView[Instance]));
 		}
 	}
 }
@@ -111,47 +127,36 @@ void UConditionalCompletion::VisualLog(const UE::Learning::FIndexSet Instances)
 //------------------------------------------------------------------
 
 UPlanarPositionDifferenceCompletion* UPlanarPositionDifferenceCompletion::AddPlanarPositionDifferenceCompletion(
-	ULearningAgentsTrainer* AgentTrainer,
+	ULearningAgentsTrainer* InAgentTrainer,
 	const FName Name,
 	const float Threshold,
 	const ELearningAgentsCompletion InCompletionMode,
 	const FVector Axis0,
 	const FVector Axis1)
 {
-	return UE::Learning::Agents::Private::AddCompletion<UPlanarPositionDifferenceCompletion, UE::Learning::FPlanarPositionDifferenceCompletion>(
-		AgentTrainer,
+	return UE::Learning::Agents::Completions::Private::AddCompletion<UPlanarPositionDifferenceCompletion, UE::Learning::FPlanarPositionDifferenceCompletion>(
+		InAgentTrainer,
 		Name,
 		1,
 		Threshold,
-		InCompletionMode == ELearningAgentsCompletion::Termination ?
-		UE::Learning::ECompletionMode::Terminated :
-		UE::Learning::ECompletionMode::Truncated,
+		UE::Learning::Agents::GetCompletionMode(InCompletionMode),
 		Axis0,
 		Axis1);
 }
 
 void UPlanarPositionDifferenceCompletion::SetPlanarPositionDifferenceCompletion(const int32 AgentId, const FVector Position0, const FVector Position1)
 {
-	const TLearningArrayView<2, FVector> Position0View = CompletionObject->InstanceData->View(CompletionObject->Position0Handle);
-	const TLearningArrayView<2, FVector> Position1View = CompletionObject->InstanceData->View(CompletionObject->Position1Handle);
-
-	if (AgentId == INDEX_NONE)
-	{
-		UE_LOG(LogLearning, Error, TEXT("AgentId is invalid (INDEX_NONE)"));
-		return;
-	}
-
-	if (AgentId < 0 || AgentId >= Position0View.Num<0>())
+	if (!AgentTrainer->HasAgent(AgentId))
 	{
-		UE_LOG(LogLearning, Error, TEXT("AgentId %d is out of index. Valid range [0, %d]."), AgentId, Position0View.Num<0>() - 1);
+		UE_LOG(LogLearning, Error, TEXT("%s: AgentId %d not found in the agents set."), *GetName(), AgentId);
 		return;
 	}
 
-	Position0View[AgentId][0] = Position0;
-	Position1View[AgentId][0] = Position1;
+	CompletionObject->InstanceData->View(CompletionObject->Position0Handle)[AgentId][0] = Position0;
+	CompletionObject->InstanceData->View(CompletionObject->Position1Handle)[AgentId][0] = Position1;
 }
 
-#if ENABLE_VISUAL_LOG
+#if UE_LEARNING_AGENTS_ENABLE_VISUAL_LOG
 void UPlanarPositionDifferenceCompletion::VisualLog(const UE::Learning::FIndexSet Instances) const
 {
 	UE_LEARNING_TRACE_CPUPROFILER_EVENT_SCOPE(UPlanarPositionDifferenceCompletion::VisualLog);
@@ -161,68 +166,157 @@ void UPlanarPositionDifferenceCompletion::VisualLog(const UE::Learning::FIndexSe
 	const TLearningArrayView<1, const float> ThresholdView = CompletionObject->InstanceData->ConstView(CompletionObject->ThresholdHandle);
 	const TLearningArrayView<1, const UE::Learning::ECompletionMode> CompletionView = CompletionObject->InstanceData->ConstView(CompletionObject->CompletionHandle);
 
-	if (const ULearningAgentsTrainer* AgentTrainer = Cast<ULearningAgentsTrainer>(GetOuter()))
+	for (const int32 Instance : Instances)
 	{
-		for (const int32 Instance : Instances)
+		if (const AActor* Actor = Cast<AActor>(AgentTrainer->GetAgent(Instance)))
 		{
-			if (const AActor* Actor = Cast<AActor>(AgentTrainer->GetAgent(Instance)))
-			{
-				const FVector Position0 = Position0View[Instance][0];
-				const FVector Position1 = Position1View[Instance][0];
-
-				const FVector PlanarPosition0 = FVector(CompletionObject->Axis0.Dot(Position0), CompletionObject->Axis1.Dot(Position0), 0.0f);
-				const FVector PlanarPosition1 = FVector(CompletionObject->Axis0.Dot(Position1), CompletionObject->Axis1.Dot(Position1), 0.0f);
-
-				UE_VLOG_LOCATION(this, LogLearning, Display,
-					Position0,
-					10.0f,
-					VisualLogColor.ToFColor(true),
-					TEXT("Position0: [% 6.1f % 6.1f % 6.1f]\nPlanar Position0: [% 6.1f % 6.1f]"),
-					Position0.X, Position0.Y, Position0.Z,
-					PlanarPosition0.X, PlanarPosition0.Y);
-
-				UE_LEARNING_AGENTS_VLOG_PLANE(this, LogLearning, Display,
-					Position0,
-					FQuat::Identity,
-					CompletionObject->Axis0,
-					CompletionObject->Axis1,
-					VisualLogColor.ToFColor(true),
-					TEXT(""));
-
-				UE_VLOG_LOCATION(this, LogLearning, Display,
-					Position1,
-					10.0f,
-					VisualLogColor.ToFColor(true),
-					TEXT("Position1: [% 6.1f % 6.1f % 6.1f]\nPlanar Position1: [% 6.1f % 6.1f]"),
-					Position1.X, Position1.Y, Position1.Z,
-					PlanarPosition1.X, PlanarPosition1.Y);
-
-				UE_LEARNING_AGENTS_VLOG_PLANE(this, LogLearning, Display,
-					Position1,
-					FQuat::Identity,
-					CompletionObject->Axis0,
-					CompletionObject->Axis1,
-					VisualLogColor.ToFColor(true),
-					TEXT(""));
-
-				UE_VLOG_SEGMENT(this, LogLearning, Display,
-					Position0,
-					Position1,
-					VisualLogColor.ToFColor(true),
-					TEXT("Agent %i\nDistance: [% 6.3f]\nPlanar Distance: [% 6.3f]\nThreshold: [% 6.2f]\nCompletion: %s"),
-					Instance,
-					FVector::Distance(Position0, Position1),
-					FVector::Distance(PlanarPosition0, PlanarPosition1),
-					ThresholdView[Instance],
-					UE::Learning::Completion::CompletionModeString(CompletionView[Instance]));
-			}
+			const FVector Position0 = Position0View[Instance][0];
+			const FVector Position1 = Position1View[Instance][0];
+
+			const FVector PlanarPosition0 = FVector(CompletionObject->Axis0.Dot(Position0), CompletionObject->Axis1.Dot(Position0), 0.0f);
+			const FVector PlanarPosition1 = FVector(CompletionObject->Axis0.Dot(Position1), CompletionObject->Axis1.Dot(Position1), 0.0f);
+
+			UE_LEARNING_AGENTS_VLOG_LOCATION(this, LogLearning, Display,
+				Position0,
+				10.0f,
+				VisualLogColor.ToFColor(true),
+				TEXT("Position0: [% 6.1f % 6.1f % 6.1f]\nPlanar Position0: [% 6.1f % 6.1f]"),
+				Position0.X, Position0.Y, Position0.Z,
+				PlanarPosition0.X, PlanarPosition0.Y);
+
+			UE_LEARNING_AGENTS_VLOG_PLANE(this, LogLearning, Display,
+				Position0,
+				FQuat::Identity,
+				CompletionObject->Axis0,
+				CompletionObject->Axis1,
+				VisualLogColor.ToFColor(true),
+				TEXT(""));
+
+			UE_LEARNING_AGENTS_VLOG_LOCATION(this, LogLearning, Display,
+				Position1,
+				10.0f,
+				VisualLogColor.ToFColor(true),
+				TEXT("Position1: [% 6.1f % 6.1f % 6.1f]\nPlanar Position1: [% 6.1f % 6.1f]"),
+				Position1.X, Position1.Y, Position1.Z,
+				PlanarPosition1.X, PlanarPosition1.Y);
+
+			UE_LEARNING_AGENTS_VLOG_PLANE(this, LogLearning, Display,
+				Position1,
+				FQuat::Identity,
+				CompletionObject->Axis0,
+				CompletionObject->Axis1,
+				VisualLogColor.ToFColor(true),
+				TEXT(""));
+
+			UE_LEARNING_AGENTS_VLOG_SEGMENT(this, LogLearning, Display,
+				Position0,
+				Position1,
+				VisualLogColor.ToFColor(true),
+				TEXT("Agent %i\nDistance: [% 6.3f]\nPlanar Distance: [% 6.3f]\nThreshold: [% 6.2f]\nCompletion: %s"),
+				Instance,
+				FVector::Distance(Position0, Position1),
+				FVector::Distance(PlanarPosition0, PlanarPosition1),
+				ThresholdView[Instance],
+				UE::Learning::Completion::CompletionModeString(CompletionView[Instance]));
 		}
 	}
 }
 #endif
 
-//------------------------------------------------------------------
 
-#undef UE_LEARNING_AGENTS_VLOG_STRING
-#undef UE_LEARNING_AGENTS_VLOG_TRANSFORM
-#undef UE_LEARNING_AGENTS_VLOG_PLANE
+UPlanarPositionSimilarityCompletion* UPlanarPositionSimilarityCompletion::AddPlanarPositionSimilarityCompletion(
+	ULearningAgentsTrainer* InAgentTrainer,
+	const FName Name,
+	const float Threshold,
+	const ELearningAgentsCompletion InCompletionMode,
+	const FVector Axis0,
+	const FVector Axis1)
+{
+	return UE::Learning::Agents::Completions::Private::AddCompletion<UPlanarPositionSimilarityCompletion, UE::Learning::FPlanarPositionSimilarityCompletion>(
+		InAgentTrainer,
+		Name,
+		1,
+		Threshold,
+		UE::Learning::Agents::GetCompletionMode(InCompletionMode),
+		Axis0,
+		Axis1);
+}
+
+void UPlanarPositionSimilarityCompletion::SetPlanarPositionSimilarityCompletion(const int32 AgentId, const FVector Position0, const FVector Position1)
+{
+	if (!AgentTrainer->HasAgent(AgentId))
+	{
+		UE_LOG(LogLearning, Error, TEXT("%s: AgentId %d not found in the agents set."), *GetName(), AgentId);
+		return;
+	}
+
+	CompletionObject->InstanceData->View(CompletionObject->Position0Handle)[AgentId][0] = Position0;
+	CompletionObject->InstanceData->View(CompletionObject->Position1Handle)[AgentId][0] = Position1;
+}
+
+#if UE_LEARNING_AGENTS_ENABLE_VISUAL_LOG
+void UPlanarPositionSimilarityCompletion::VisualLog(const UE::Learning::FIndexSet Instances) const
+{
+	UE_LEARNING_TRACE_CPUPROFILER_EVENT_SCOPE(UPlanarPositionSimilarityCompletion::VisualLog);
+
+	const TLearningArrayView<2, const FVector> Position0View = CompletionObject->InstanceData->ConstView(CompletionObject->Position0Handle);
+	const TLearningArrayView<2, const FVector> Position1View = CompletionObject->InstanceData->ConstView(CompletionObject->Position1Handle);
+	const TLearningArrayView<1, const float> ThresholdView = CompletionObject->InstanceData->ConstView(CompletionObject->ThresholdHandle);
+	const TLearningArrayView<1, const UE::Learning::ECompletionMode> CompletionView = CompletionObject->InstanceData->ConstView(CompletionObject->CompletionHandle);
+
+	for (const int32 Instance : Instances)
+	{
+		if (const AActor* Actor = Cast<AActor>(AgentTrainer->GetAgent(Instance)))
+		{
+			const FVector Position0 = Position0View[Instance][0];
+			const FVector Position1 = Position1View[Instance][0];
+
+			const FVector PlanarPosition0 = FVector(CompletionObject->Axis0.Dot(Position0), CompletionObject->Axis1.Dot(Position0), 0.0f);
+			const FVector PlanarPosition1 = FVector(CompletionObject->Axis0.Dot(Position1), CompletionObject->Axis1.Dot(Position1), 0.0f);
+
+			UE_LEARNING_AGENTS_VLOG_LOCATION(this, LogLearning, Display,
+				Position0,
+				10.0f,
+				VisualLogColor.ToFColor(true),
+				TEXT("Position0: [% 6.1f % 6.1f % 6.1f]\nPlanar Position0: [% 6.1f % 6.1f]"),
+				Position0.X, Position0.Y, Position0.Z,
+				PlanarPosition0.X, PlanarPosition0.Y);
+
+			UE_LEARNING_AGENTS_VLOG_PLANE(this, LogLearning, Display,
+				Position0,
+				FQuat::Identity,
+				CompletionObject->Axis0,
+				CompletionObject->Axis1,
+				VisualLogColor.ToFColor(true),
+				TEXT(""));
+
+			UE_LEARNING_AGENTS_VLOG_LOCATION(this, LogLearning, Display,
+				Position1,
+				10.0f,
+				VisualLogColor.ToFColor(true),
+				TEXT("Position1: [% 6.1f % 6.1f % 6.1f]\nPlanar Position1: [% 6.1f % 6.1f]"),
+				Position1.X, Position1.Y, Position1.Z,
+				PlanarPosition1.X, PlanarPosition1.Y);
+
+			UE_LEARNING_AGENTS_VLOG_PLANE(this, LogLearning, Display,
+				Position1,
+				FQuat::Identity,
+				CompletionObject->Axis0,
+				CompletionObject->Axis1,
+				VisualLogColor.ToFColor(true),
+				TEXT(""));
+
+			UE_LEARNING_AGENTS_VLOG_SEGMENT(this, LogLearning, Display,
+				Position0,
+				Position1,
+				VisualLogColor.ToFColor(true),
+				TEXT("Agent %i\nDistance: [% 6.3f]\nPlanar Distance: [% 6.3f]\nThreshold: [% 6.2f]\nCompletion: %s"),
+				Instance,
+				FVector::Distance(Position0, Position1),
+				FVector::Distance(PlanarPosition0, PlanarPosition1),
+				ThresholdView[Instance],
+				UE::Learning::Completion::CompletionModeString(CompletionView[Instance]));
+		}
+	}
+}
+#endif
