LearningAgents: Added some additional observations, actions, rewards, completions and helpers. Cleaned up some of the error messaging to make it more consistent.

diff --git a/Engine/Plugins/Experimental/LearningAgents/Source/LearningAgentsTraining/Private/LearningAgentsDataStorage.cpp b/Engine/Plugins/Experimental/LearningAgents/Source/LearningAgentsTraining/Private/LearningAgentsDataStorage.cpp
--- a/Engine/Plugins/Experimental/LearningAgents/Source/LearningAgentsTraining/Private/LearningAgentsDataStorage.cpp
+++ b/Engine/Plugins/Experimental/LearningAgents/Source/LearningAgentsTraining/Private/LearningAgentsDataStorage.cpp
@@ -74,13 +74,13 @@ ULearningAgentsRecord* ULearningAgentsDataStorage::CreateRecord(const FName Reco
 {
 	if (!AgentType)
 	{
-		UE_LOG(LogLearning, Warning, TEXT("CreateRecord: AgentType is nullptr. You must pass a valid agent type. Skipping creation."));
+		UE_LOG(LogLearning, Error, TEXT("%s: AgentType is nullptr."), *GetName());
 		return nullptr;
 	}
 
 	if (!AgentType->IsSetup())
 	{
-		UE_LOG(LogLearning, Warning, TEXT("CreateRecord: AgentType Setup must be performed before record can be created."));
+		UE_LOG(LogLearning, Error, TEXT("%s: %s's Setup must be run before it can be used."), *GetName(), *AgentType->GetName());
 		return nullptr;
 	}
 
@@ -105,7 +105,7 @@ int32 ULearningAgentsDataStorage::LoadAllRecords(ULearningAgentsType* AgentType,
 {
 	if (!AgentType)
 	{
-		UE_LOG(LogLearning, Warning, TEXT("LoadAllRecords: AgentType is nullptr. You must pass a valid agent type. Skipping loading."));
+		UE_LOG(LogLearning, Error, TEXT("%s: AgentType is nullptr."), *GetName());
 		return 0;
 	}
 
@@ -114,7 +114,7 @@ int32 ULearningAgentsDataStorage::LoadAllRecords(ULearningAgentsType* AgentType,
 
 	if (DataFiles.Num() == 0)
 	{
-		UE_LOG(LogLearning, Warning, TEXT("LoadAllRecords: Found no .%s files in directory %s"), *FileExtension, *Directory.Path);
+		UE_LOG(LogLearning, Warning, TEXT("%s: Found no .%s files in directory %s"), *GetName(), *FileExtension, *Directory.Path);
 		return 0;
 	}
 
@@ -130,13 +130,13 @@ ULearningAgentsRecord* ULearningAgentsDataStorage::LoadRecord(ULearningAgentsTyp
 {
 	if (!AgentType)
 	{
-		UE_LOG(LogLearning, Warning, TEXT("LoadRecord: AgentType is nullptr. You must pass a valid agent type. Skipping loading."));
+		UE_LOG(LogLearning, Error, TEXT("%s: AgentType is nullptr."), *GetName());
 		return nullptr;
 	}
 
 	if (!AgentType->IsSetup())
 	{
-		UE_LOG(LogLearning, Warning, TEXT("LoadRecord: AgentType Setup must be performed before record can be created."));
+		UE_LOG(LogLearning, Error, TEXT("%s: %s's Setup must be run before it can be used."), *GetName(), *AgentType->GetName());
 		return nullptr;
 	}
 
@@ -146,14 +146,14 @@ ULearningAgentsRecord* ULearningAgentsDataStorage::LoadRecord(ULearningAgentsTyp
 
 	if (!FPaths::FileExists(InputFilePath))
 	{
-		UE_LOG(LogLearning, Error, TEXT("LoadRecord: FilePath %s does not exist."), *InputFilePath);
+		UE_LOG(LogLearning, Error, TEXT("%s: FilePath %s does not exist."), *GetName(), *InputFilePath);
 		return nullptr;
 	}
 
 	const bool bSuccess = FFileHelper::LoadFileToArray(Bytes, *InputFilePath);
 	if (!bSuccess)
 	{
-		UE_LOG(LogLearning, Error, TEXT("Failed to load data from file %s"), *InputFilePath);
+		UE_LOG(LogLearning, Error, TEXT("%s: Failed to load data from file %s"), *GetName(), *InputFilePath);
 		return nullptr;
 	}
 
@@ -163,7 +163,7 @@ ULearningAgentsRecord* ULearningAgentsDataStorage::LoadRecord(ULearningAgentsTyp
 	Reader.Serialize(&ChunksNum, sizeof(int32));
 	if (ChunksNum <= 0)
 	{
-		UE_LOG(LogLearning, Error, TEXT("LoadRecord: file ChunksNum header was %d. Must be greater than 0."), ChunksNum);
+		UE_LOG(LogLearning, Error, TEXT("%s: file ChunksNum header was %d. Must be greater than 0."), *GetName(), ChunksNum);
 		return nullptr;
 	}
 
@@ -172,7 +172,7 @@ ULearningAgentsRecord* ULearningAgentsDataStorage::LoadRecord(ULearningAgentsTyp
 	int32 AgentObsNum = AgentType->GetObservationFeature().DimNum();
 	if (AgentObsNum != ObsNum)
 	{
-		UE_LOG(LogLearning, Error, TEXT("LoadRecord: Observation data dimension size %d incompatible with agent's observation dimension size %d"), ObsNum, AgentObsNum);
+		UE_LOG(LogLearning, Error, TEXT("%s: Observation data dimension size %d incompatible with agent's observation dimension size %d"), *GetName(), ObsNum, AgentObsNum);
 		return nullptr;
 	}
 
@@ -181,7 +181,7 @@ ULearningAgentsRecord* ULearningAgentsDataStorage::LoadRecord(ULearningAgentsTyp
 	int32 AgentActionNum = AgentType->GetActionFeature().DimNum();
 	if (AgentActionNum != ActNum)
 	{
-		UE_LOG(LogLearning, Error, TEXT("LoadRecord: Action data dimension size %d incompatible with agent's action dimension size %d"), ObsNum, AgentObsNum);
+		UE_LOG(LogLearning, Error, TEXT("%s: Action data dimension size %d incompatible with agent's action dimension size %d"), *GetName(), ObsNum, AgentObsNum);
 		return nullptr;
 	}
 
@@ -222,7 +222,7 @@ void ULearningAgentsDataStorage::SaveRecord(const FDirectoryPath& Directory, ULe
 {
 	if (!Record)
 	{
-		UE_LOG(LogLearning, Warning, TEXT("SaveRecord: Record is nullptr. You must pass a valid record. Skipping save."));
+		UE_LOG(LogLearning, Warning, TEXT("%s: Record is nullptr."), *GetName());
 		return;
 	}
 
@@ -252,10 +252,11 @@ void ULearningAgentsDataStorage::SaveRecord(const FDirectoryPath& Directory, ULe
 	const FString OutputFilePath = Directory.Path + FGenericPlatformMisc::GetDefaultPathSeparator() +
 		(!Record->MetaData.bWasLoadedFromFile && bPrependUtcTimeStamp ? FDateTime::UtcNow().ToString() + "_" : "") + FileName + "." + FileExtension;
 	const bool bSuccess = FFileHelper::SaveArrayToFile(Bytes, *OutputFilePath, &IFileManager::Get(), EFileWrite::FILEWRITE_Append);
+
 	if (!bSuccess)
 	{
-		UE_LOG(LogLearning, Error, TEXT("Failed to write data to output file for %s"), *FileName);
+		UE_LOG(LogLearning, Error, TEXT("%s: Failed to write data to output file for %s"), *GetName(), *FileName);
 	}
 
-	UE_LOG(LogLearning, Display, TEXT("SaveRecord: Saved data for agent to %s"), *OutputFilePath);
+	UE_LOG(LogLearning, Display, TEXT("%s: Saved data for agent to %s"), *GetName(), *OutputFilePath);
 }
diff --git a/Engine/Plugins/Experimental/LearningAgents/Source/LearningAgentsTraining/Private/LearningAgentsImitationTrainer.cpp b/Engine/Plugins/Experimental/LearningAgents/Source/LearningAgentsTraining/Private/LearningAgentsImitationTrainer.cpp
--- a/Engine/Plugins/Experimental/LearningAgents/Source/LearningAgentsTraining/Private/LearningAgentsImitationTrainer.cpp
+++ b/Engine/Plugins/Experimental/LearningAgents/Source/LearningAgentsTraining/Private/LearningAgentsImitationTrainer.cpp
@@ -38,25 +38,25 @@ void ULearningAgentsImitationTrainer::BeginTraining(
 {
 	if (IsTraining())
 	{
-		UE_LOG(LogLearning, Error, TEXT("Already Training!"));
+		UE_LOG(LogLearning, Error, TEXT("%s: Already Training!"), *GetName());
 		return;
 	}
 
 	if (!InPolicy)
 	{
-		UE_LOG(LogLearning, Error, TEXT("BeginTraining called with nullptr for Policy."));
+		UE_LOG(LogLearning, Error, TEXT("%s: InPolicy is nullptr."), *GetName());
 		return;
 	}
 
 	if (!InPolicy->IsSetup())
 	{
-		UE_LOG(LogLearning, Error, TEXT("Policy Setup not Performed"));
+		UE_LOG(LogLearning, Error, TEXT("%s: %s's Setup must be run before it can be used."), *GetName(), *InPolicy->GetName());
 		return;
 	}
 
 	if (Records.IsEmpty())
 	{
-		UE_LOG(LogLearning, Error, TEXT("Records list is empty! Imitation Trainer needs at least one record to work correctly."));
+		UE_LOG(LogLearning, Error, TEXT("%s: Records list is empty! Imitation Trainer needs at least one record."), *GetName());
 		return;
 	}
 
@@ -69,7 +69,7 @@ void ULearningAgentsImitationTrainer::BeginTraining(
 	{
 		if (!Record)
 		{
-			UE_LOG(LogLearning, Warning, TEXT("BeginTraining: Null record object."));
+			UE_LOG(LogLearning, Warning, TEXT("%s: Record is nullptr."));
 			continue;
 		}
 
@@ -108,7 +108,7 @@ void ULearningAgentsImitationTrainer::BeginTraining(
 			}
 			else
 			{
-				UE_LOG(LogLearning, Warning, TEXT("Record input or output size does not match policy."));
+				UE_LOG(LogLearning, Warning, TEXT("%s: Record input or output size does not match policy."), *GetName());
 			}
 		}
 	}
@@ -117,7 +117,7 @@ void ULearningAgentsImitationTrainer::BeginTraining(
 
 	if (DataIndex == 0)
 	{
-		UE_LOG(LogLearning, Warning, TEXT("BeginTraining: input contains no valid training data."));
+		UE_LOG(LogLearning, Warning, TEXT("%s: input contains no valid training data."), *GetName());
 		RecordedObservations.Empty();
 		RecordedActions.Empty();
 		return;
@@ -130,7 +130,7 @@ void ULearningAgentsImitationTrainer::BeginTraining(
 
 	// Begin Training Properly
 
-	UE_LOG(LogLearning, Display, TEXT("Imitation Training Started"));
+	UE_LOG(LogLearning, Display, TEXT("%s: Imitation Training Started"), *GetName());
 
 	bIsTraining = true;
 	bIsTrainingComplete = false;
@@ -194,19 +194,18 @@ void ULearningAgentsImitationTrainer::BeginTraining(
 
 void ULearningAgentsImitationTrainer::EndTraining()
 {
-	if (IsTraining())
+	if (!IsTraining())
 	{
-		UE_LOG(LogLearning, Display, TEXT("Imitation Training Ended"));
+		UE_LOG(LogLearning, Warning, TEXT("%s: Not Training!"), *GetName());
+		return;
+	}
 
-		bRequestImitationTrainingStop = true;
-		ImitationTrainingTask.Wait(FTimespan(0, 0, 30));
+	UE_LOG(LogLearning, Display, TEXT("%s: Imitation Training Ended."), *GetName());
 
-		bIsTraining = false;
-	}
-	else
-	{
-		UE_LOG(LogLearning, Warning, TEXT("Not Training!"));
-	}
+	bRequestImitationTrainingStop = true;
+	ImitationTrainingTask.Wait(FTimespan(0, 0, 30));
+
+	bIsTraining = false;
 }
 
 bool ULearningAgentsImitationTrainer::IsTraining() const
diff --git a/Engine/Plugins/Experimental/LearningAgents/Source/LearningAgentsTraining/Private/LearningAgentsRewards.cpp b/Engine/Plugins/Experimental/LearningAgents/Source/LearningAgentsTraining/Private/LearningAgentsRewards.cpp
--- a/Engine/Plugins/Experimental/LearningAgents/Source/LearningAgentsTraining/Private/LearningAgentsRewards.cpp
+++ b/Engine/Plugins/Experimental/LearningAgents/Source/LearningAgentsTraining/Private/LearningAgentsRewards.cpp
@@ -10,41 +10,28 @@
 #include "LearningLog.h"
 
 #include "GameFramework/Actor.h"
-#include "VisualLogger/VisualLogger.h"
 
-#define UE_LEARNING_AGENTS_VLOG_STRING(Owner, Category, Verbosity, Location, Color, Format, ...) \
-	UE_VLOG_LOCATION(Owner, Category, Verbosity, Location, 0.0f, Color, Format, ##__VA_ARGS__)
-
-#define UE_LEARNING_AGENTS_VLOG_TRANSFORM(Owner, Category, Verbosity, Location, Rotation, Color, Format, ...) \
-	UE_VLOG_SEGMENT(Owner, Category, Verbosity, Location, Location + 15.0f * Rotation.RotateVector(FVector::ForwardVector), FColor::Red, TEXT("")); \
-	UE_VLOG_SEGMENT(Owner, Category, Verbosity, Location, Location + 15.0f * Rotation.RotateVector(FVector::RightVector), FColor::Green, TEXT("")); \
-	UE_VLOG_SEGMENT(Owner, Category, Verbosity, Location, Location + 15.0f * Rotation.RotateVector(FVector::UpVector), FColor::Blue, TEXT("")); \
-	UE_VLOG_OBOX(Owner, Category, Verbosity, FBox(10.0f * FVector(-1, -1, -1), 10.0f * FVector(1, 1, 1)), FTransform(Rotation, Location, FVector::OneVector).ToMatrixNoScale(), Color, TEXT("")); \
-	UE_LEARNING_AGENTS_VLOG_STRING(Owner, Category, Verbosity, Location + FVector(0.0f, 0.0f, 20.0f), Color, Format, ##__VA_ARGS__)
-
-#define UE_LEARNING_AGENTS_VLOG_PLANE(Owner, Category, Verbosity, Location, Rotation, Axis0, Axis1, Color, Format, ...) \
-	UE_VLOG_OBOX(Owner, Category, Verbosity, FBox(-25.0f * (Axis0 + Axis1), 25.0f * (Axis0 + Axis1)), FTransform(Rotation, Location, FVector::OneVector).ToMatrixNoScale(), Color, Format, ##__VA_ARGS__)
-
-namespace UE::Learning::Agents::Private
+namespace UE::Learning::Agents::Rewards::Private
 {
 	template<typename RewardUObject, typename RewardFObject, typename... InArgTypes>
-	RewardUObject* AddReward(ULearningAgentsTrainer* AgentTrainer, const FName Name, InArgTypes&& ...Args)
+	RewardUObject* AddReward(ULearningAgentsTrainer* InAgentTrainer, const FName Name, InArgTypes&& ...Args)
 	{
-		if (!AgentTrainer)
+		if (!InAgentTrainer)
 		{
-			UE_LOG(LogLearning, Error, TEXT("AgentTrainer is nullptr"));
+			UE_LOG(LogLearning, Error, TEXT("InAgentTrainer is nullptr."));
 			return nullptr;
 		}
 
-		RewardUObject* Reward = NewObject<RewardUObject>(AgentTrainer, Name);
+		RewardUObject* Reward = NewObject<RewardUObject>(InAgentTrainer, Name);
 
+		Reward->AgentTrainer = InAgentTrainer;
 		Reward->RewardObject = MakeShared<RewardFObject>(
 			Reward->GetFName(),
-			AgentTrainer->GetAgentManager()->GetInstanceData().ToSharedRef(),
-			AgentTrainer->GetAgentManager()->GetMaxInstanceNum(),
+			InAgentTrainer->GetAgentManager()->GetInstanceData().ToSharedRef(),
+			InAgentTrainer->GetAgentManager()->GetMaxInstanceNum(),
 			Forward<InArgTypes>(Args)...);
 
-		AgentTrainer->AddReward(Reward, Reward->RewardObject.ToSharedRef());
+		InAgentTrainer->AddReward(Reward, Reward->RewardObject.ToSharedRef());
 
 		return Reward;
 	}
@@ -52,31 +39,23 @@ namespace UE::Learning::Agents::Private
 
 //------------------------------------------------------------------
 
-UFloatReward* UFloatReward::AddFloatReward(ULearningAgentsTrainer* AgentTrainer, const FName Name, const float Weight)
+UFloatReward* UFloatReward::AddFloatReward(ULearningAgentsTrainer* InAgentTrainer, const FName Name, const float Weight)
 {
-	return UE::Learning::Agents::Private::AddReward<UFloatReward, UE::Learning::FFloatReward>(AgentTrainer, Name, Weight);
+	return UE::Learning::Agents::Rewards::Private::AddReward<UFloatReward, UE::Learning::FFloatReward>(InAgentTrainer, Name, Weight);
 }
 
 void UFloatReward::SetFloatReward(const int32 AgentId, const float Reward)
 {
-	const TLearningArrayView<1, float> View = RewardObject->InstanceData->View(RewardObject->ValueHandle);
-
-	if (AgentId == INDEX_NONE)
-	{
-		UE_LOG(LogLearning, Error, TEXT("AgentId is invalid (INDEX_NONE)"));
-		return;
-	}
-
-	if (AgentId < 0 || AgentId >= View.Num<0>())
+	if (!AgentTrainer->HasAgent(AgentId))
 	{
-		UE_LOG(LogLearning, Error, TEXT("AgentId %d is out of index. Valid range [0, %d]."), AgentId, View.Num<0>() - 1);
+		UE_LOG(LogLearning, Error, TEXT("%s: AgentId %d not found in the agents set."), *GetName(), AgentId);
 		return;
 	}
 
-	View[AgentId] = Reward;
+	RewardObject->InstanceData->View(RewardObject->ValueHandle)[AgentId] = Reward;
 }
 
-#if ENABLE_VISUAL_LOG
+#if UE_LEARNING_AGENTS_ENABLE_VISUAL_LOG
 void UFloatReward::VisualLog(const UE::Learning::FIndexSet Instances) const
 {
 	UE_LEARNING_TRACE_CPUPROFILER_EVENT_SCOPE(UFloatReward::VisualLog);
@@ -85,21 +64,18 @@ void UFloatReward::VisualLog(const UE::Learning::FIndexSet Instances) const
 	const TLearningArrayView<1, const float> WeightView = RewardObject->InstanceData->ConstView(RewardObject->WeightHandle);
 	const TLearningArrayView<1, const float> RewardView = RewardObject->InstanceData->ConstView(RewardObject->RewardHandle);
 
-	if (const ULearningAgentsTrainer* AgentTrainer = Cast<ULearningAgentsTrainer>(GetOuter()))
+	for (const int32 Instance : Instances)
 	{
-		for (const int32 Instance : Instances)
+		if (const AActor* Actor = Cast<AActor>(AgentTrainer->GetAgent(Instance)))
 		{
-			if (const AActor* Actor = Cast<AActor>(AgentTrainer->GetAgent(Instance)))
-			{
-				UE_LEARNING_AGENTS_VLOG_STRING(this, LogLearning, Display,
-					Actor->GetActorLocation(),
-					VisualLogColor.ToFColor(true),
-					TEXT("Agent %i\nWeight: [% 6.2f]\nValue: [% 6.2f]\nReward: [% 6.3f]"),
-					Instance,
-					WeightView[Instance],
-					ValueView[Instance],
-					RewardView[Instance]);
-			}
+			UE_LEARNING_AGENTS_VLOG_STRING(this, LogLearning, Display,
+				Actor->GetActorLocation(),
+				VisualLogColor.ToFColor(true),
+				TEXT("Agent %i\nWeight: [% 6.2f]\nValue: [% 6.2f]\nReward: [% 6.3f]"),
+				Instance,
+				WeightView[Instance],
+				ValueView[Instance],
+				RewardView[Instance]);
 		}
 	}
 }
@@ -107,31 +83,23 @@ void UFloatReward::VisualLog(const UE::Learning::FIndexSet Instances) const
 
 //------------------------------------------------------------------
 
-UScalarVelocityReward* UScalarVelocityReward::AddScalarVelocityReward(ULearningAgentsTrainer* AgentTrainer, const FName Name, const float Weight, const float Scale)
+UScalarVelocityReward* UScalarVelocityReward::AddScalarVelocityReward(ULearningAgentsTrainer* InAgentTrainer, const FName Name, const float Weight, const float Scale)
 {
-	return UE::Learning::Agents::Private::AddReward<UScalarVelocityReward, UE::Learning::FScalarVelocityReward>(AgentTrainer, Name, Weight, Scale);
+	return UE::Learning::Agents::Rewards::Private::AddReward<UScalarVelocityReward, UE::Learning::FScalarVelocityReward>(InAgentTrainer, Name, Weight, Scale);
 }
 
 void UScalarVelocityReward::SetScalarVelocityReward(int32 AgentId, float Velocity)
 {
-	const TLearningArrayView<1, float> VelocityView = RewardObject->InstanceData->View(RewardObject->VelocityHandle);
-
-	if (AgentId == INDEX_NONE)
+	if (!AgentTrainer->HasAgent(AgentId))
 	{
-		UE_LOG(LogLearning, Error, TEXT("AgentId is invalid (INDEX_NONE)"));
+		UE_LOG(LogLearning, Error, TEXT("%s: AgentId %d not found in the agents set."), *GetName(), AgentId);
 		return;
 	}
 
-	if (AgentId < 0 || AgentId >= VelocityView.Num<0>())
-	{
-		UE_LOG(LogLearning, Error, TEXT("AgentId %d is out of index. Valid range [0, %d]."), AgentId, VelocityView.Num<0>() - 1);
-		return;
-	}
-
-	VelocityView[AgentId] = Velocity;
+	RewardObject->InstanceData->View(RewardObject->VelocityHandle)[AgentId] = Velocity;
 }
 
-#if ENABLE_VISUAL_LOG
+#if UE_LEARNING_AGENTS_ENABLE_VISUAL_LOG
 void UScalarVelocityReward::VisualLog(const UE::Learning::FIndexSet Instances) const
 {
 	UE_LEARNING_TRACE_CPUPROFILER_EVENT_SCOPE(UScalarVelocityReward::VisualLog);
@@ -141,54 +109,42 @@ void UScalarVelocityReward::VisualLog(const UE::Learning::FIndexSet Instances) c
 	const TLearningArrayView<1, const float> ScaleView = RewardObject->InstanceData->ConstView(RewardObject->ScaleHandle);
 	const TLearningArrayView<1, const float> RewardView = RewardObject->InstanceData->ConstView(RewardObject->RewardHandle);
 
-	if (const ULearningAgentsTrainer* AgentTrainer = Cast<ULearningAgentsTrainer>(GetOuter()))
+	for (const int32 Instance : Instances)
 	{
-		for (const int32 Instance : Instances)
+		if (const AActor* Actor = Cast<const AActor>(AgentTrainer->GetAgent(Instance)))
 		{
-			if (const AActor* Actor = Cast<const AActor>(AgentTrainer->GetAgent(Instance)))
-			{
-				UE_LEARNING_AGENTS_VLOG_STRING(this, LogLearning, Display,
-					Actor->GetActorLocation(),
-					VisualLogColor.ToFColor(true),
-					TEXT("Agent %i\nWeight: [% 6.2f]\nScale: [% 6.2f]\nVelocity: [% 6.2f]\nReward: [% 6.3f]"),
-					Instance,
-					WeightView[Instance],
-					ScaleView[Instance],
-					VelocityView[Instance],
-					RewardView[Instance]);
-			}
+			UE_LEARNING_AGENTS_VLOG_STRING(this, LogLearning, Display,
+				Actor->GetActorLocation(),
+				VisualLogColor.ToFColor(true),
+				TEXT("Agent %i\nWeight: [% 6.2f]\nScale: [% 6.2f]\nVelocity: [% 6.2f]\nReward: [% 6.3f]"),
+				Instance,
+				WeightView[Instance],
+				ScaleView[Instance],
+				VelocityView[Instance],
+				RewardView[Instance]);
 		}
 	}
 }
 #endif
 
-ULocalDirectionalVelocityReward* ULocalDirectionalVelocityReward::AddLocalDirectionalVelocityReward(ULearningAgentsTrainer* AgentTrainer, const FName Name, const float Weight, const float Scale, const FVector Axis)
+ULocalDirectionalVelocityReward* ULocalDirectionalVelocityReward::AddLocalDirectionalVelocityReward(ULearningAgentsTrainer* InAgentTrainer, const FName Name, const float Weight, const float Scale, const FVector Axis)
 {
-	return UE::Learning::Agents::Private::AddReward<ULocalDirectionalVelocityReward, UE::Learning::FLocalDirectionalVelocityReward>(AgentTrainer, Name, Weight, Scale, Axis);
+	return UE::Learning::Agents::Rewards::Private::AddReward<ULocalDirectionalVelocityReward, UE::Learning::FLocalDirectionalVelocityReward>(InAgentTrainer, Name, Weight, Scale, Axis);
 }
 
 void ULocalDirectionalVelocityReward::SetLocalDirectionalVelocityReward(const int32 AgentId, const FVector Velocity, const FRotator RelativeRotation)
 {
-	const TLearningArrayView<1, FVector> VelocityView = RewardObject->InstanceData->View(RewardObject->VelocityHandle);
-	const TLearningArrayView<1, FQuat> RelativeRotationView = RewardObject->InstanceData->View(RewardObject->RelativeRotationHandle);
-
-	if (AgentId == INDEX_NONE)
-	{
-		UE_LOG(LogLearning, Error, TEXT("AgentId is invalid (INDEX_NONE)"));
-		return;
-	}
-
-	if (AgentId < 0 || AgentId >= VelocityView.Num<0>())
+	if (!AgentTrainer->HasAgent(AgentId))
 	{
-		UE_LOG(LogLearning, Error, TEXT("AgentId %d is out of index. Valid range [0, %d]."), AgentId, VelocityView.Num<0>() - 1);
+		UE_LOG(LogLearning, Error, TEXT("%s: AgentId %d not found in the agents set."), *GetName(), AgentId);
 		return;
 	}
 
-	VelocityView[AgentId] = Velocity;
-	RelativeRotationView[AgentId] = FQuat::MakeFromRotator(RelativeRotation);
+	RewardObject->InstanceData->View(RewardObject->VelocityHandle)[AgentId] = Velocity;
+	RewardObject->InstanceData->View(RewardObject->RelativeRotationHandle)[AgentId] = RelativeRotation.Quaternion();
 }
 
-#if ENABLE_VISUAL_LOG
+#if UE_LEARNING_AGENTS_ENABLE_VISUAL_LOG
 void ULocalDirectionalVelocityReward::VisualLog(const UE::Learning::FIndexSet Instances) const
 {
 	UE_LEARNING_TRACE_CPUPROFILER_EVENT_SCOPE(ULocalDirectionalVelocityReward::VisualLog);
@@ -199,54 +155,51 @@ void ULocalDirectionalVelocityReward::VisualLog(const UE::Learning::FIndexSet In
 	const TLearningArrayView<1, const float> ScaleView = RewardObject->InstanceData->ConstView(RewardObject->ScaleHandle);
 	const TLearningArrayView<1, const float> RewardView = RewardObject->InstanceData->ConstView(RewardObject->RewardHandle);
 
-	if (const ULearningAgentsTrainer* AgentTrainer = Cast<ULearningAgentsTrainer>(GetOuter()))
+	for (const int32 Instance : Instances)
 	{
-		for (const int32 Instance : Instances)
+		if (const AActor* Actor = Cast<AActor>(AgentTrainer->GetAgent(Instance)))
 		{
-			if (const AActor* Actor = Cast<AActor>(AgentTrainer->GetAgent(Instance)))
-			{
-				const FVector Velocity = VelocityView[Instance];
-				const FQuat RelativeRotation = RelativeRotationView[Instance];
-				const FVector LocalVelocity = RelativeRotation.UnrotateVector(Velocity);
-				const FVector Direction = RelativeRotation.RotateVector(RewardObject->Axis);
-
-				UE_VLOG_ARROW(this, LogLearning, Display,
-					Actor->GetActorLocation(),
-					Actor->GetActorLocation() + Velocity,
-					VisualLogColor.ToFColor(true),
-					TEXT(""));
-
-				UE_LEARNING_AGENTS_VLOG_STRING(this, LogLearning, Display,
-					Actor->GetActorLocation() + Velocity,
-					VisualLogColor.ToFColor(true),
-					TEXT("Velocity: [% 6.3f % 6.3f % 6.3f]\nLocal Velocity: [% 6.3f % 6.3f % 6.3f]"),
-					Velocity.X, Velocity.Y, Velocity.Z,
-					LocalVelocity.X, LocalVelocity.Y, LocalVelocity.Z);
-
-				UE_VLOG_ARROW(this, LogLearning, Display,
-					Actor->GetActorLocation(),
-					Actor->GetActorLocation() + 100.0f * Direction,
-					VisualLogColor.ToFColor(true),
-					TEXT(""));
-
-				UE_LEARNING_AGENTS_VLOG_STRING(this, LogLearning, Display,
-					Actor->GetActorLocation() + 100.0f * Direction,
-					VisualLogColor.ToFColor(true),
-					TEXT("Direction: [% 6.3f % 6.3f % 6.3f]\nLocal Direction: [% 6.3f % 6.3f % 6.3f]"),
-					Direction.X, Direction.Y, Direction.Z,
-					RewardObject->Axis.X, RewardObject->Axis.Y, RewardObject->Axis.Z);
-
-				UE_LEARNING_AGENTS_VLOG_TRANSFORM(this, LogLearning, Display,
-					Actor->GetActorLocation(),
-					RelativeRotation,
-					VisualLogColor.ToFColor(true),
-					TEXT("Agent %i\nDot Product: [% 6.3f]\nWeight: [% 6.2f]\nScale: [% 6.2f]\nReward: [% 6.3f]"),
-					Instance,
-					LocalVelocity.Dot(RewardObject->Axis),
-					WeightView[Instance],
-					ScaleView[Instance],
-					RewardView[Instance]);
-			}
+			const FVector Velocity = VelocityView[Instance];
+			const FQuat RelativeRotation = RelativeRotationView[Instance];
+			const FVector LocalVelocity = RelativeRotation.UnrotateVector(Velocity);
+			const FVector Direction = RelativeRotation.RotateVector(RewardObject->Axis);
+
+			UE_LEARNING_AGENTS_VLOG_ARROW(this, LogLearning, Display,
+				Actor->GetActorLocation(),
+				Actor->GetActorLocation() + Velocity,
+				VisualLogColor.ToFColor(true),
+				TEXT(""));
+
+			UE_LEARNING_AGENTS_VLOG_STRING(this, LogLearning, Display,
+				Actor->GetActorLocation() + Velocity,
+				VisualLogColor.ToFColor(true),
+				TEXT("Velocity: [% 6.3f % 6.3f % 6.3f]\nLocal Velocity: [% 6.3f % 6.3f % 6.3f]"),
+				Velocity.X, Velocity.Y, Velocity.Z,
+				LocalVelocity.X, LocalVelocity.Y, LocalVelocity.Z);
+
+			UE_LEARNING_AGENTS_VLOG_ARROW(this, LogLearning, Display,
+				Actor->GetActorLocation(),
+				Actor->GetActorLocation() + 100.0f * Direction,
+				VisualLogColor.ToFColor(true),
+				TEXT(""));
+
+			UE_LEARNING_AGENTS_VLOG_STRING(this, LogLearning, Display,
+				Actor->GetActorLocation() + 100.0f * Direction,
+				VisualLogColor.ToFColor(true),
+				TEXT("Direction: [% 6.3f % 6.3f % 6.3f]\nLocal Direction: [% 6.3f % 6.3f % 6.3f]"),
+				Direction.X, Direction.Y, Direction.Z,
+				RewardObject->Axis.X, RewardObject->Axis.Y, RewardObject->Axis.Z);
+
+			UE_LEARNING_AGENTS_VLOG_TRANSFORM(this, LogLearning, Display,
+				Actor->GetActorLocation(),
+				RelativeRotation,
+				VisualLogColor.ToFColor(true),
+				TEXT("Agent %i\nDot Product: [% 6.3f]\nWeight: [% 6.2f]\nScale: [% 6.2f]\nReward: [% 6.3f]"),
+				Instance,
+				LocalVelocity.Dot(RewardObject->Axis),
+				WeightView[Instance],
+				ScaleView[Instance],
+				RewardView[Instance]);
 		}
 	}
 }
@@ -255,7 +208,7 @@ void ULocalDirectionalVelocityReward::VisualLog(const UE::Learning::FIndexSet In
 //------------------------------------------------------------------
 
 UPlanarPositionDifferencePenalty* UPlanarPositionDifferencePenalty::AddPlanarPositionDifferencePenalty(
-	ULearningAgentsTrainer* AgentTrainer,
+	ULearningAgentsTrainer* InAgentTrainer,
 	const FName Name,
 	const float Weight,
 	const float Scale,
@@ -263,31 +216,22 @@ UPlanarPositionDifferencePenalty* UPlanarPositionDifferencePenalty::AddPlanarPos
 	const FVector Axis0,
 	const FVector Axis1)
 {
-	return UE::Learning::Agents::Private::AddReward<UPlanarPositionDifferencePenalty, UE::Learning::FPlanarPositionDifferencePenalty>(AgentTrainer, Name, Weight, Scale, Threshold, Axis0, Axis1);
+	return UE::Learning::Agents::Rewards::Private::AddReward<UPlanarPositionDifferencePenalty, UE::Learning::FPlanarPositionDifferencePenalty>(InAgentTrainer, Name, Weight, Scale, Threshold, Axis0, Axis1);
 }
 
 void UPlanarPositionDifferencePenalty::SetPlanarPositionDifferencePenalty(const int32 AgentId, const FVector Position0, const FVector Position1)
 {
-	const TLearningArrayView<1, FVector> Position0View = RewardObject->InstanceData->View(RewardObject->Position0Handle);
-	const TLearningArrayView<1, FVector> Position1View = RewardObject->InstanceData->View(RewardObject->Position1Handle);
-
-	if (AgentId == INDEX_NONE)
+	if (!AgentTrainer->HasAgent(AgentId))
 	{
-		UE_LOG(LogLearning, Error, TEXT("AgentId is invalid (INDEX_NONE)"));
+		UE_LOG(LogLearning, Error, TEXT("%s: AgentId %d not found in the agents set."), *GetName(), AgentId);
 		return;
 	}
 
-	if (AgentId < 0 || AgentId >= Position0View.Num<0>())
-	{
-		UE_LOG(LogLearning, Error, TEXT("AgentId %d is out of index. Valid range [0, %d]."), AgentId, Position0View.Num<0>() - 1);
-		return;
-	}
-
-	Position0View[AgentId] = Position0;
-	Position1View[AgentId] = Position1;
+	RewardObject->InstanceData->View(RewardObject->Position0Handle)[AgentId] = Position0;
+	RewardObject->InstanceData->View(RewardObject->Position1Handle)[AgentId] = Position1;
 }
 
-#if ENABLE_VISUAL_LOG
+#if UE_LEARNING_AGENTS_ENABLE_VISUAL_LOG
 void UPlanarPositionDifferencePenalty::VisualLog(const UE::Learning::FIndexSet Instances) const
 {
 	UE_LEARNING_TRACE_CPUPROFILER_EVENT_SCOPE(UPlanarPositionDifferencePenalty::VisualLog);
@@ -299,63 +243,60 @@ void UPlanarPositionDifferencePenalty::VisualLog(const UE::Learning::FIndexSet I
 	const TLearningArrayView<1, const float> ThresholdView = RewardObject->InstanceData->ConstView(RewardObject->ThresholdHandle);
 	const TLearningArrayView<1, const float> RewardView = RewardObject->InstanceData->ConstView(RewardObject->RewardHandle);
 
-	if (const ULearningAgentsTrainer* AgentTrainer = Cast<ULearningAgentsTrainer>(GetOuter()))
+	for (const int32 Instance : Instances)
 	{
-		for (const int32 Instance : Instances)
+		if (const AActor* Actor = Cast<AActor>(AgentTrainer->GetAgent(Instance)))
 		{
-			if (const AActor* Actor = Cast<AActor>(AgentTrainer->GetAgent(Instance)))
-			{
-				const FVector Position0 = Position0View[Instance];
-				const FVector Position1 = Position1View[Instance];
-
-				const FVector PlanarPosition0 = FVector(RewardObject->Axis0.Dot(Position0), RewardObject->Axis1.Dot(Position0), 0.0f);
-				const FVector PlanarPosition1 = FVector(RewardObject->Axis0.Dot(Position1), RewardObject->Axis1.Dot(Position1), 0.0f);
-
-				UE_VLOG_LOCATION(this, LogLearning, Display,
-					Position0,
-					10.0f,
-					VisualLogColor.ToFColor(true),
-					TEXT("Position0: [% 6.1f % 6.1f % 6.1f]\nPlanar Position0: [% 6.1f % 6.1f]"),
-					Position0.X, Position0.Y, Position0.Z,
-					PlanarPosition0.X, PlanarPosition0.Y);
-
-				UE_LEARNING_AGENTS_VLOG_PLANE(this, LogLearning, Display,
-					Position0,
-					FQuat::Identity,
-					RewardObject->Axis0,
-					RewardObject->Axis1,
-					VisualLogColor.ToFColor(true),
-					TEXT(""));
-
-				UE_VLOG_LOCATION(this, LogLearning, Display,
-					Position1,
-					10.0f,
-					VisualLogColor.ToFColor(true),
-					TEXT("Position1: [% 6.1f % 6.1f % 6.1f]\nPlanar Position1: [% 6.1f % 6.1f]"),
-					Position1.X, Position1.Y, Position1.Z,
-					PlanarPosition1.X, PlanarPosition1.Y);
-
-				UE_LEARNING_AGENTS_VLOG_PLANE(this, LogLearning, Display,
-					Position1,
-					FQuat::Identity,
-					RewardObject->Axis0,
-					RewardObject->Axis1,
-					VisualLogColor.ToFColor(true),
-					TEXT(""));
-
-				UE_VLOG_SEGMENT(this, LogLearning, Display,
-					Position0,
-					Position1,
-					VisualLogColor.ToFColor(true),
-					TEXT("Agent %i\nDistance: [% 6.3f]\nPlanar Distance: [% 6.3f]\nWeight: [% 6.2f]\nScale: [% 6.2f]\nThreshold: [% 6.2f]\nReward: [% 6.3f]"),
-					Instance,
-					FVector::Distance(Position0, Position1),
-					FVector::Distance(PlanarPosition0, PlanarPosition1),
-					WeightView[Instance],
-					ScaleView[Instance],
-					ThresholdView[Instance],
-					RewardView[Instance]);
-			}
+			const FVector Position0 = Position0View[Instance];
+			const FVector Position1 = Position1View[Instance];
+
+			const FVector PlanarPosition0 = FVector(RewardObject->Axis0.Dot(Position0), RewardObject->Axis1.Dot(Position0), 0.0f);
+			const FVector PlanarPosition1 = FVector(RewardObject->Axis0.Dot(Position1), RewardObject->Axis1.Dot(Position1), 0.0f);
+
+			UE_LEARNING_AGENTS_VLOG_LOCATION(this, LogLearning, Display,
+				Position0,
+				10.0f,
+				VisualLogColor.ToFColor(true),
+				TEXT("Position0: [% 6.1f % 6.1f % 6.1f]\nPlanar Position0: [% 6.1f % 6.1f]"),
+				Position0.X, Position0.Y, Position0.Z,
+				PlanarPosition0.X, PlanarPosition0.Y);
+
+			UE_LEARNING_AGENTS_VLOG_PLANE(this, LogLearning, Display,
+				Position0,
+				FQuat::Identity,
+				RewardObject->Axis0,
+				RewardObject->Axis1,
+				VisualLogColor.ToFColor(true),
+				TEXT(""));
+
+			UE_LEARNING_AGENTS_VLOG_LOCATION(this, LogLearning, Display,
+				Position1,
+				10.0f,
+				VisualLogColor.ToFColor(true),
+				TEXT("Position1: [% 6.1f % 6.1f % 6.1f]\nPlanar Position1: [% 6.1f % 6.1f]"),
+				Position1.X, Position1.Y, Position1.Z,
+				PlanarPosition1.X, PlanarPosition1.Y);
+
+			UE_LEARNING_AGENTS_VLOG_PLANE(this, LogLearning, Display,
+				Position1,
+				FQuat::Identity,
+				RewardObject->Axis0,
+				RewardObject->Axis1,
+				VisualLogColor.ToFColor(true),
+				TEXT(""));
+
+			UE_LEARNING_AGENTS_VLOG_SEGMENT(this, LogLearning, Display,
+				Position0,
+				Position1,
+				VisualLogColor.ToFColor(true),
+				TEXT("Agent %i\nDistance: [% 6.3f]\nPlanar Distance: [% 6.3f]\nWeight: [% 6.2f]\nScale: [% 6.2f]\nThreshold: [% 6.2f]\nReward: [% 6.3f]"),
+				Instance,
+				FVector::Distance(Position0, Position1),
+				FVector::Distance(PlanarPosition0, PlanarPosition1),
+				WeightView[Instance],
+				ScaleView[Instance],
+				ThresholdView[Instance],
+				RewardView[Instance]);
 		}
 	}
 }
@@ -364,13 +305,13 @@ void UPlanarPositionDifferencePenalty::VisualLog(const UE::Learning::FIndexSet I
 //------------------------------------------------------------------
 
 UPositionArraySimilarityReward* UPositionArraySimilarityReward::AddPositionArraySimilarityReward(
-	ULearningAgentsTrainer* AgentTrainer,
+	ULearningAgentsTrainer* InAgentTrainer,
 	const FName Name,
 	const int32 PositionNum,
 	const float Weight,
 	const float Scale)
 {
-	return UE::Learning::Agents::Private::AddReward<UPositionArraySimilarityReward, UE::Learning::FPositionArraySimilarityReward>(AgentTrainer, Name, PositionNum, Weight, Scale);
+	return UE::Learning::Agents::Rewards::Private::AddReward<UPositionArraySimilarityReward, UE::Learning::FPositionArraySimilarityReward>(InAgentTrainer, Name, PositionNum, Weight, Scale);
 }
 
 void UPositionArraySimilarityReward::SetPositionArraySimilarityReward(
@@ -382,6 +323,12 @@ void UPositionArraySimilarityReward::SetPositionArraySimilarityReward(
 	const FRotator RelativeRotation0,
 	const FRotator RelativeRotation1)
 {
+	if (!AgentTrainer->HasAgent(AgentId))
+	{
+		UE_LOG(LogLearning, Error, TEXT("%s: AgentId %d not found in the agents set."), *GetName(), AgentId);
+		return;
+	}
+
 	const TLearningArrayView<2, FVector> Positions0View = RewardObject->InstanceData->View(RewardObject->Positions0Handle);
 	const TLearningArrayView<2, FVector> Positions1View = RewardObject->InstanceData->View(RewardObject->Positions1Handle);
 	const TLearningArrayView<1, FVector> RelativePosition0View = RewardObject->InstanceData->View(RewardObject->RelativePosition0Handle);
@@ -389,35 +336,21 @@ void UPositionArraySimilarityReward::SetPositionArraySimilarityReward(
 	const TLearningArrayView<1, FQuat> RelativeRotation0View = RewardObject->InstanceData->View(RewardObject->RelativeRotation0Handle);
 	const TLearningArrayView<1, FQuat> RelativeRotation1View = RewardObject->InstanceData->View(RewardObject->RelativeRotation1Handle);
 
-	if (AgentId == INDEX_NONE)
+	if (Positions0.Num() != Positions0View.Num<1>() || Positions1.Num() != Positions0View.Num<1>())
 	{
-		UE_LOG(LogLearning, Error, TEXT("AgentId is invalid (INDEX_NONE)"));
-		return;
-	}
-
-	if (AgentId < 0 || AgentId >= Positions0View.Num<0>())
-	{
-		UE_LOG(LogLearning, Error, TEXT("AgentId %d is out of index. Valid range [0, %d]."), AgentId, Positions0View.Num<0>() - 1);
-		return;
-	}
-
-	const int32 PositionNum = Positions0View.Num<1>();
-
-	if (Positions0.Num() != PositionNum || Positions1.Num() != PositionNum)
-	{
-		UE_LOG(LogLearning, Error, TEXT("Incorrect number of positions in array. Got %i and %i, expected %i."), Positions0.Num(), Positions1.Num(), PositionNum);
+		UE_LOG(LogLearning, Error, TEXT("%s: Got wrong number of elements in array. Expected %i, got %i and %i."), *GetName(), Positions0View.Num<1>(), Positions0.Num(), Positions1.Num());
 		return;
 	}
 
 	RelativePosition0View[AgentId] = RelativePosition0;
 	RelativePosition1View[AgentId] = RelativePosition1;
-	RelativeRotation0View[AgentId] = FQuat::MakeFromRotator(RelativeRotation0);
-	RelativeRotation1View[AgentId] = FQuat::MakeFromRotator(RelativeRotation1);
-	UE::Learning::Array::Copy(Positions0View[AgentId], TLearningArrayView<1, const FVector>(Positions0));
-	UE::Learning::Array::Copy(Positions1View[AgentId], TLearningArrayView<1, const FVector>(Positions1));
+	RelativeRotation0View[AgentId] = RelativeRotation0.Quaternion();
+	RelativeRotation1View[AgentId] = RelativeRotation1.Quaternion();
+	UE::Learning::Array::Copy<1, FVector>(Positions0View[AgentId], Positions0);
+	UE::Learning::Array::Copy<1, FVector>(Positions1View[AgentId], Positions1);
 }
 
-#if ENABLE_VISUAL_LOG
+#if UE_LEARNING_AGENTS_ENABLE_VISUAL_LOG
 void UPositionArraySimilarityReward::VisualLog(const UE::Learning::FIndexSet Instances) const
 {
 	UE_LEARNING_TRACE_CPUPROFILER_EVENT_SCOPE(UPositionArraySimilarityReward::VisualLog);
@@ -435,79 +368,70 @@ void UPositionArraySimilarityReward::VisualLog(const UE::Learning::FIndexSet Ins
 
 	const int32 PositionNum = Positions0View.Num<1>();
 		
-	if (const ULearningAgentsTrainer* AgentTrainer = Cast<ULearningAgentsTrainer>(GetOuter()))
+	for (const int32 Instance : Instances)
 	{
-		for (const int32 Instance : Instances)
+		if (const AActor* Actor = Cast<AActor>(AgentTrainer->GetAgent(Instance)))
 		{
-			if (const AActor* Actor = Cast<AActor>(AgentTrainer->GetAgent(Instance)))
+			const FVector RelativePosition0 = RelativePosition0View[Instance];
+			const FVector RelativePosition1 = RelativePosition1View[Instance];
+			const FQuat RelativeRotation0 = RelativeRotation0View[Instance];
+			const FQuat RelativeRotation1 = RelativeRotation1View[Instance];
+
+			for (int32 PositionIdx = 0; PositionIdx < PositionNum; PositionIdx++)
 			{
-				const FVector RelativePosition0 = RelativePosition0View[Instance];
-				const FVector RelativePosition1 = RelativePosition1View[Instance];
-				const FQuat RelativeRotation0 = RelativeRotation0View[Instance];
-				const FQuat RelativeRotation1 = RelativeRotation1View[Instance];
-
-				for (int32 PositionIdx = 0; PositionIdx < PositionNum; PositionIdx++)
-				{
-					const FVector Position0 = Positions0View[Instance][PositionIdx];
-					const FVector Position1 = Positions1View[Instance][PositionIdx];
-
-					const FVector LocalPosition0 = RelativeRotation0.UnrotateVector(Positions0View[Instance][PositionIdx] - RelativePosition0);
-					const FVector LocalPosition1 = RelativeRotation1.UnrotateVector(Positions1View[Instance][PositionIdx] - RelativePosition1);
-
-					UE_VLOG_LOCATION(this, LogLearning, Display,
-						Position0,
-						10.0f,
-						VisualLogColor.ToFColor(true),
-						TEXT("Position0: [% 6.1f % 6.1f % 6.1f]\nLocal Position0: [% 6.1f % 6.1f % 6.1f]"),
-						Position0.X, Position0.Y, Position0.Z,
-						LocalPosition0.X, LocalPosition0.Y, LocalPosition0.Z);
-
-					UE_VLOG_LOCATION(this, LogLearning, Display,
-						Position1,
-						10.0f,
-						VisualLogColor.ToFColor(true),
-						TEXT("Position1: [% 6.1f % 6.1f % 6.1f]\nLocal Position1: [% 6.1f % 6.1f % 6.1f]"),
-						Position1.X, Position1.Y, Position1.Z,
-						LocalPosition1.X, LocalPosition1.Y, LocalPosition1.Z);
-
-					UE_VLOG_SEGMENT(this, LogLearning, Display,
-						Position0,
-						Position1,
-						VisualLogColor.ToFColor(true),
-						TEXT("Distance: [% 6.1f]\nLocal Distance: [% 6.1f]"),
-						FVector::Distance(Position0, Position1),
-						FVector::Distance(LocalPosition0, LocalPosition1));
-				}
-
-				UE_LEARNING_AGENTS_VLOG_TRANSFORM(this, LogLearning, Display,
-					RelativePosition0,
-					RelativeRotation0,
+				const FVector Position0 = Positions0View[Instance][PositionIdx];
+				const FVector Position1 = Positions1View[Instance][PositionIdx];
+
+				const FVector LocalPosition0 = RelativeRotation0.UnrotateVector(Positions0View[Instance][PositionIdx] - RelativePosition0);
+				const FVector LocalPosition1 = RelativeRotation1.UnrotateVector(Positions1View[Instance][PositionIdx] - RelativePosition1);
+
+				UE_LEARNING_AGENTS_VLOG_LOCATION(this, LogLearning, Display,
+					Position0,
+					10.0f,
 					VisualLogColor.ToFColor(true),
-					TEXT("Relative Transform 0"));
+					TEXT("Position0: [% 6.1f % 6.1f % 6.1f]\nLocal Position0: [% 6.1f % 6.1f % 6.1f]"),
+					Position0.X, Position0.Y, Position0.Z,
+					LocalPosition0.X, LocalPosition0.Y, LocalPosition0.Z);
 
-				UE_LEARNING_AGENTS_VLOG_TRANSFORM(this, LogLearning, Display,
-					RelativePosition1,
-					RelativeRotation1,
+				UE_LEARNING_AGENTS_VLOG_LOCATION(this, LogLearning, Display,
+					Position1,
+					10.0f,
 					VisualLogColor.ToFColor(true),
-					TEXT("Relative Transform 1"));
+					TEXT("Position1: [% 6.1f % 6.1f % 6.1f]\nLocal Position1: [% 6.1f % 6.1f % 6.1f]"),
+					Position1.X, Position1.Y, Position1.Z,
+					LocalPosition1.X, LocalPosition1.Y, LocalPosition1.Z);
 
-				UE_LEARNING_AGENTS_VLOG_STRING(this, LogLearning, Display,
-					Actor->GetActorLocation(),
+				UE_LEARNING_AGENTS_VLOG_SEGMENT(this, LogLearning, Display,
+					Position0,
+					Position1,
 					VisualLogColor.ToFColor(true),
-					TEXT("Agent %i\nWeight: [% 6.2f]\nScale: [% 6.2f]\nThreshold: [% 6.2f]\nReward: [% 6.3f]"),
-					Instance,
-					WeightView[Instance],
-					ScaleView[Instance],
-					ThresholdView[Instance],
-					RewardView[Instance]);
+					TEXT("Distance: [% 6.1f]\nLocal Distance: [% 6.1f]"),
+					FVector::Distance(Position0, Position1),
+					FVector::Distance(LocalPosition0, LocalPosition1));
 			}
+
+			UE_LEARNING_AGENTS_VLOG_TRANSFORM(this, LogLearning, Display,
+				RelativePosition0,
+				RelativeRotation0,
+				VisualLogColor.ToFColor(true),
+				TEXT("Relative Transform 0"));
+
+			UE_LEARNING_AGENTS_VLOG_TRANSFORM(this, LogLearning, Display,
+				RelativePosition1,
+				RelativeRotation1,
+				VisualLogColor.ToFColor(true),
+				TEXT("Relative Transform 1"));
+
+			UE_LEARNING_AGENTS_VLOG_STRING(this, LogLearning, Display,
+				Actor->GetActorLocation(),
+				VisualLogColor.ToFColor(true),
+				TEXT("Agent %i\nWeight: [% 6.2f]\nScale: [% 6.2f]\nThreshold: [% 6.2f]\nReward: [% 6.3f]"),
+				Instance,
+				WeightView[Instance],
+				ScaleView[Instance],
+				ThresholdView[Instance],
+				RewardView[Instance]);
 		}
 	}
 }
 #endif
-
-//------------------------------------------------------------------
-
-#undef UE_LEARNING_AGENTS_VLOG_STRING
-#undef UE_LEARNING_AGENTS_VLOG_TRANSFORM
-#undef UE_LEARNING_AGENTS_VLOG_PLANE
diff --git a/Engine/Plugins/Experimental/LearningAgents/Source/LearningAgentsTraining/Private/LearningAgentsTrainer.cpp b/Engine/Plugins/Experimental/LearningAgents/Source/LearningAgentsTraining/Private/LearningAgentsTrainer.cpp
--- a/Engine/Plugins/Experimental/LearningAgents/Source/LearningAgentsTraining/Private/LearningAgentsTrainer.cpp
+++ b/Engine/Plugins/Experimental/LearningAgents/Source/LearningAgentsTraining/Private/LearningAgentsTrainer.cpp
@@ -29,6 +29,29 @@
 #include "Engine/GameViewportClient.h"
 #include "EngineDefines.h"
 
+namespace UE::Learning::Agents
+{
+	ELearningAgentsCompletion GetLearningAgentsCompletion(const ECompletionMode CompletionMode)
+	{
+		switch (CompletionMode)
+		{
+		case ECompletionMode::Running: UE_LOG(LogLearning, Error, TEXT("Cannot convert from ECompletionMode::Running to ELearningAgentsCompletion")); return ELearningAgentsCompletion::Termination;
+		case ECompletionMode::Terminated: return ELearningAgentsCompletion::Termination;
+		case ECompletionMode::Truncated: return ELearningAgentsCompletion::Truncation;
+		default:UE_LOG(LogLearning, Error, TEXT("Unknown Completion Mode.")); return ELearningAgentsCompletion::Termination;
+		}
+	}
+
+	ECompletionMode GetCompletionMode(const ELearningAgentsCompletion Completion)
+	{
+		switch (Completion)
+		{
+		case ELearningAgentsCompletion::Termination: return ECompletionMode::Terminated;
+		case ELearningAgentsCompletion::Truncation: return ECompletionMode::Truncated;
+		default:UE_LOG(LogLearning, Error, TEXT("Unknown Completion.")); return ECompletionMode::Running;
+		}
+	}
+}
 ULearningAgentsTrainer::ULearningAgentsTrainer() : ULearningAgentsManagerComponent() {}
 ULearningAgentsTrainer::ULearningAgentsTrainer(FVTableHelper& Helper) : ULearningAgentsTrainer() {}
 ULearningAgentsTrainer::~ULearningAgentsTrainer() {}
@@ -38,11 +61,11 @@ void ULearningAgentsTrainer::SetupTrainer(
 	ULearningAgentsType* InAgentType,
 	ULearningAgentsPolicy* InPolicy,
 	ULearningAgentsCritic* InCritic,
-	const FLearningAgentsTrainerSettings& Settings)
+	const FLearningAgentsTrainerSettings& TrainerSettings)
 {
 	if (IsSetup())
 	{
-		UE_LOG(LogLearning, Error, TEXT("%s: Setup already performed!"), *GetName());
+		UE_LOG(LogLearning, Error, TEXT("%s: Setup already run!"), *GetName());
 		return;
 	}
 
@@ -54,7 +77,7 @@ void ULearningAgentsTrainer::SetupTrainer(
 
 	if (!InAgentManager->IsManagerSetup())
 	{
-		UE_LOG(LogLearning, Error, TEXT("%s's SetupManager() must be run before %s can be setup."), *InAgentManager->GetName(), *GetName());
+		UE_LOG(LogLearning, Error, TEXT("%s: %s's SetupManager must be run before it can be used."), *GetName(), *InAgentManager->GetName());
 		return;
 	}
 
@@ -122,26 +145,26 @@ void ULearningAgentsTrainer::SetupTrainer(
 	EpisodeBuffer = MakeUnique<UE::Learning::FEpisodeBuffer>();
 	EpisodeBuffer->Resize(
 		AgentManager->GetMaxInstanceNum(),
-		Settings.MaxStepNum,
+		TrainerSettings.MaxStepNum,
 		AgentType->GetObservationFeature().DimNum(),
 		AgentType->GetActionFeature().DimNum());
 
-	MaxStepsCompletion = Settings.MaxStepsCompletion;
+	MaxStepsCompletion = TrainerSettings.MaxStepsCompletion;
 
 	// Create Replay Buffer
 	ReplayBuffer = MakeUnique<UE::Learning::FReplayBuffer>();
 	ReplayBuffer->Resize(
 		AgentType->GetObservationFeature().DimNum(),
 		AgentType->GetActionFeature().DimNum(),
-		Settings.MaximumRecordedEpisodesPerIteration,
-		Settings.MaximumRecordedStepsPerIteration);
+		TrainerSettings.MaximumRecordedEpisodesPerIteration,
+		TrainerSettings.MaximumRecordedStepsPerIteration);
 
 	// Create Reset Buffer
 	ResetBuffer = MakeUnique<UE::Learning::FResetInstanceBuffer>();
 	ResetBuffer->Resize(AgentManager->GetMaxInstanceNum());
 
 	// Record Timeout Setting
-	TrainerTimeout = Settings.TrainerCommunicationTimeout;
+	TrainerTimeout = TrainerSettings.TrainerCommunicationTimeout;
 
 	bIsSetup = true;
 }
@@ -154,7 +177,7 @@ bool ULearningAgentsTrainer::AddAgent(const int32 AgentId)
 	{
 		// Reset the instance and the buffer in case we have stale 
 		// data from another agent that was using this id previously
-		ResetInstance({ AgentId });
+		ResetEpisodes({ AgentId });
 		EpisodeBuffer->Reset(AgentId);
 	}
 
@@ -219,7 +242,13 @@ void ULearningAgentsTrainer::BeginTraining(
 {
 	if (!IsSetup())
 	{
-		UE_LOG(LogLearning, Error, TEXT("Setup must be run before training can start."));
+		UE_LOG(LogLearning, Error, TEXT("%s: Setup not run."), *GetName());
+		return;
+	}
+
+	if (IsTraining())
+	{
+		UE_LOG(LogLearning, Error, TEXT("%s: Already Training!"), *GetName());
 		return;
 	}
 
@@ -260,7 +289,7 @@ void ULearningAgentsTrainer::BeginTraining(
 	}
 	else
 	{
-		UE_LOG(LogLearning, Warning, TEXT("Provided invalid FixedTimeStepFrequency: %0.5f"), TrainerGameSettings.FixedTimeStepFrequency);
+		UE_LOG(LogLearning, Warning, TEXT("%s: Provided invalid FixedTimeStepFrequency: %0.5f"), *GetName(), TrainerGameSettings.FixedTimeStepFrequency);
 	}
 
 	if (TrainerGameSettings.bDisableVSync && GameSettings)
@@ -326,7 +355,7 @@ void ULearningAgentsTrainer::BeginTraining(
 			CriticSettings.LayerNum != Critic->GetCriticNetwork().GetLayerNum() ||
 			UE::Learning::Agents::GetActivationFunction(CriticSettings.ActivationFunction) != Critic->GetCriticNetwork().ActivationFunction)
 		{
-			UE_LOG(LogLearning, Warning, TEXT("StartTraining got different Critic Network Settings to those provided to SetupCritic."));
+			UE_LOG(LogLearning, Warning, TEXT("%s: BeginTraining got different Critic Network Settings to those provided to SetupCritic."), *GetName());
 		}
 
 		PPONetworkSettings.CriticHiddenLayerSize = Critic->GetCriticNetwork().GetHiddenNum();
@@ -362,7 +391,7 @@ void ULearningAgentsTrainer::BeginTraining(
 		PPONetworkSettings,
 		TrainerFlags);
 
-	UE_LOG(LogLearning, Display, TEXT("Receiving initial policy..."));
+	UE_LOG(LogLearning, Display, TEXT("%s: Receiving initial policy..."), *GetName());
 
 	UE::Learning::ETrainerResponse Response = UE::Learning::ETrainerResponse::Success;
 
@@ -377,7 +406,7 @@ void ULearningAgentsTrainer::BeginTraining(
 
 	if (Response != UE::Learning::ETrainerResponse::Success)
 	{
-		UE_LOG(LogLearning, Error, TEXT("Error sending or receiving policy from trainer: %s. Check log for errors."), UE::Learning::Trainer::GetResponseString(Response));
+		UE_LOG(LogLearning, Error, TEXT("%s: Error sending or receiving policy from trainer: %s. Check log for errors."), *GetName(), UE::Learning::Trainer::GetResponseString(Response));
 		Trainer->Terminate();
 		return;
 	}
@@ -399,14 +428,14 @@ void ULearningAgentsTrainer::BeginTraining(
 
 		if (Response != UE::Learning::ETrainerResponse::Success)
 		{
-			UE_LOG(LogLearning, Error, TEXT("Error sending or receiving critic from trainer: %s. Check log for errors."), UE::Learning::Trainer::GetResponseString(Response));
+			UE_LOG(LogLearning, Error, TEXT("%s: Error sending or receiving critic from trainer: %s. Check log for errors."), *GetName(), UE::Learning::Trainer::GetResponseString(Response));
 			Trainer->Terminate();
 			return;
 		}
 	}
 
 	// Reset Agents, Episode Buffer, and Replay Buffer
-	ResetInstance(AddedAgentIds);
+	ResetEpisodes(AddedAgentIds);
 	EpisodeBuffer->Reset(AddedAgentSet);
 	ReplayBuffer->Reset();
 
@@ -453,7 +482,7 @@ void ULearningAgentsTrainer::EndTraining()
 {
 	if (IsTraining())
 	{
-		UE_LOG(LogLearning, Display, TEXT("Stopping training..."));
+		UE_LOG(LogLearning, Display, TEXT("%s: Stopping training..."), *GetName());
 		Trainer->SendStop();
 		DoneTraining();
 	}
@@ -465,7 +494,7 @@ void ULearningAgentsTrainer::EvaluateRewards()
 
 	if (!IsSetup())
 	{
-		UE_LOG(LogLearning, Error, TEXT("Setup must be run before rewards can be evaluated."));
+		UE_LOG(LogLearning, Error, TEXT("%s: Setup not run."), *GetName());
 		return;
 	}
 
@@ -490,7 +519,7 @@ void ULearningAgentsTrainer::EvaluateCompletions()
 
 	if (!IsSetup())
 	{
-		UE_LOG(LogLearning, Error, TEXT("Setup must be run before completions can be evaluated."));
+		UE_LOG(LogLearning, Error, TEXT("%s: Setup not run."), *GetName());
 		return;
 	}
 
@@ -509,17 +538,17 @@ void ULearningAgentsTrainer::EvaluateCompletions()
 #endif
 }
 
-void ULearningAgentsTrainer::IterateTraining()
+void ULearningAgentsTrainer::ProcessExperience()
 {
 	if (!IsSetup())
 	{
-		UE_LOG(LogLearning, Error, TEXT("Setup must be run before training can be performed."));
+		UE_LOG(LogLearning, Error, TEXT("%s: Setup not run."), *GetName());
 		return;
 	}
 
 	if (!IsTraining())
 	{
-		UE_LOG(LogLearning, Error, TEXT("Attempted to iterate training but the trainer is not running. Ensure StartTraining() is called before IterateTraining() and no errors occurred."));
+		UE_LOG(LogLearning, Error, TEXT("%s: Training not running."), *GetName());
 		return;
 	}
 
@@ -574,7 +603,7 @@ void ULearningAgentsTrainer::IterateTraining()
 
 			if (Response != UE::Learning::ETrainerResponse::Success)
 			{
-				UE_LOG(LogLearning, Error, TEXT("Error waiting to push experience to trainer. Check log for errors."));
+				UE_LOG(LogLearning, Error, TEXT("%s: Error waiting to push experience to trainer. Check log for errors."), *GetName());
 				EndTraining();
 				return;
 			}
@@ -586,13 +615,13 @@ void ULearningAgentsTrainer::IterateTraining()
 
 			if (Response == UE::Learning::ETrainerResponse::Completed)
 			{
-				UE_LOG(LogLearning, Display, TEXT("Trainer completed training."));
+				UE_LOG(LogLearning, Display, TEXT("%s: Trainer completed training."), *GetName());
 				DoneTraining();
 				return;
 			}
 			else if (Response != UE::Learning::ETrainerResponse::Success)
 			{
-				UE_LOG(LogLearning, Error, TEXT("Error waiting for policy from trainer. Check log for errors."));
+				UE_LOG(LogLearning, Error, TEXT("%s: Error waiting for policy from trainer. Check log for errors."), *GetName());
 				EndTraining();
 				return;
 			}
@@ -604,7 +633,7 @@ void ULearningAgentsTrainer::IterateTraining()
 
 				if (Response != UE::Learning::ETrainerResponse::Success)
 				{
-					UE_LOG(LogLearning, Error, TEXT("Error waiting for critic from trainer. Check log for errors."));
+					UE_LOG(LogLearning, Error, TEXT("%s: Error waiting for critic from trainer. Check log for errors."), *GetName());
 					EndTraining();
 					return;
 				}
@@ -614,28 +643,90 @@ void ULearningAgentsTrainer::IterateTraining()
 			ResetBuffer->SetResetInstances(AddedAgentSet);
 		}
 
-		ResetInstance(ResetBuffer->GetResetInstances().ToArray());
+		ResetEpisodes(ResetBuffer->GetResetInstances().ToArray());
 
 		EpisodeBuffer->Reset(ResetBuffer->GetResetInstances());
 	}
 }
 
-void ULearningAgentsTrainer::ResetAllInstances()
+void ULearningAgentsTrainer::ResetAllEpisodes()
 {
-	UE_LEARNING_TRACE_CPUPROFILER_EVENT_SCOPE(ULearningAgentsTrainer::ResetAllInstances);
+	UE_LEARNING_TRACE_CPUPROFILER_EVENT_SCOPE(ULearningAgentsTrainer::ResetAllEpisodes);
 
 	if (!IsSetup())
 	{
-		UE_LOG(LogLearning, Error, TEXT("Setup must be run before agents can be reset."));
+		UE_LOG(LogLearning, Error, TEXT("%s: Setup not run."), *GetName());
 		return;
 	}
 
-	ResetInstance(AddedAgentIds);
+	ResetEpisodes(AddedAgentIds);
 
 	EpisodeBuffer->Reset(AddedAgentSet);
 }
 
-void ULearningAgentsTrainer::ResetInstance_Implementation(const TArray<int32>& AgentId)
+void ULearningAgentsTrainer::RunTraining()
+{
+	if (!IsTraining())
+	{
+		BeginTraining();
+	}
+	else
+	{
+		EvaluateCompletions();
+		EvaluateRewards();
+		ProcessExperience();
+	}
+
+	Policy->RunInference();
+}
+
+float ULearningAgentsTrainer::GetReward(const int32 AgentId) const
+{
+	if (!IsSetup())
+	{
+		UE_LOG(LogLearning, Error, TEXT("%s: Setup not run."), *GetName());
+		return 0.0f;
+	}
+
+	if (!HasAgent(AgentId))
+	{
+		UE_LOG(LogLearning, Error, TEXT("%s: AgentId %d not found in the agents set."), *GetName(), AgentId);
+		return 0.0f;
+	}
+
+	return Rewards->RewardBuffer()[AgentId];
+}
+
+bool ULearningAgentsTrainer::IsCompleted(const int32 AgentId, ELearningAgentsCompletion& OutCompletion) const
+{
+	if (!IsSetup())
+	{
+		UE_LOG(LogLearning, Error, TEXT("%s: Setup not run."), *GetName());
+		OutCompletion = ELearningAgentsCompletion::Termination;
+		return false;
+	}
+
+	if (!HasAgent(AgentId))
+	{
+		UE_LOG(LogLearning, Error, TEXT("%s: AgentId %d not found in the agents set."), *GetName(), AgentId);
+		OutCompletion = ELearningAgentsCompletion::Termination;
+		return false;
+	}
+
+	const UE::Learning::ECompletionMode CompletionMode = Completions->CompletionBuffer()[AgentId];
+
+	if (CompletionMode == UE::Learning::ECompletionMode::Running)
+	{
+		OutCompletion = ELearningAgentsCompletion::Termination;
+		return false;
+	}
+	else
+	{
+		OutCompletion = UE::Learning::Agents::GetLearningAgentsCompletion(CompletionMode);
+		return true;
+	}
+}
+void ULearningAgentsTrainer::ResetEpisodes_Implementation(const TArray<int32>& AgentId)
 {
 	// Can be overridden to reset agent without blueprints
 }
diff --git a/Engine/Plugins/Experimental/LearningAgents/Source/LearningAgentsTraining/Public/LearningAgentsCompletions.h b/Engine/Plugins/Experimental/LearningAgents/Source/LearningAgentsTraining/Public/LearningAgentsCompletions.h
--- a/Engine/Plugins/Experimental/LearningAgents/Source/LearningAgentsTraining/Public/LearningAgentsCompletions.h
+++ b/Engine/Plugins/Experimental/LearningAgents/Source/LearningAgentsTraining/Public/LearningAgentsCompletions.h
@@ -4,10 +4,10 @@
 
 #include "LearningArray.h"
 #include "LearningAgentsTrainer.h" // Required for ELearningAgentsCompletion::Termination
+#include "LearningAgentsDebug.h"
 
 #include "Templates/SharedPointer.h"
 #include "UObject/Object.h"
-#include "EngineDefines.h"
 
 #include "LearningAgentsCompletions.generated.h"
 
@@ -15,12 +15,16 @@ namespace UE::Learning
 {
 	struct FCompletionObject;
 	struct FConditionalCompletion;
+	struct FTimeElapsedCompletion;
 	struct FPlanarPositionDifferenceCompletion;
+	struct FPlanarPositionSimilarityCompletion;
 }
 
 // For functions in this file, we are favoring having more verbose names such as "AddConditionalCompletion" vs simply "Add" in 
 // order to keep it easy to find the correct function in blueprints.
 
+//------------------------------------------------------------------
+
 /**
  * The base class for all completions. Completions contain logic that determines if an agent's current episode should
  * end, e.g. because the agent achieved the normal win/loss condition for the game. Additionally, completions can speed
@@ -39,15 +43,20 @@ class LEARNINGAGENTSTRAINING_API ULearningAgentsCompletion : public UObject
 
 public:
 
-#if ENABLE_VISUAL_LOG
+	UPROPERTY(VisibleAnywhere, Transient, Category = "LearningAgents")
+	TObjectPtr<ULearningAgentsTrainer> AgentTrainer;
+
+#if UE_LEARNING_AGENTS_ENABLE_VISUAL_LOG
 	/** Color used to draw this completion in the visual log */
 	FLinearColor VisualLogColor = FColor::Yellow;
 
-	/** Describes this observation to the visual logger for debugging purposes. */
+	/** Describes this completion to the visual logger for debugging purposes. */
 	virtual void VisualLog(const UE::Learning::FIndexSet Instances) const {}
 #endif
 };
 
+//------------------------------------------------------------------
+
 /** A simple boolean completion. Used as a catch-all for situations where a more type-specific completion does not exist yet. */
 UCLASS()
 class LEARNINGAGENTSTRAINING_API UConditionalCompletion : public ULearningAgentsCompletion
@@ -58,30 +67,69 @@ public:
 
 	/**
 	* Adds a new conditional completion to the given trainer. Call during ULearningAgentsTrainer::SetupCompletions event.
-	* @param AgentTrainer The trainer to add this completion to.
+	* @param InAgentTrainer The trainer to add this completion to.
 	* @param Name The name of this new completion. Used for debugging.
 	* @param InCompletionMode The completion mode.
 	* @return The newly created completion.
 	*/
 	UFUNCTION(BlueprintCallable, Category = "LearningAgents")
-	static UConditionalCompletion* AddConditionalCompletion(ULearningAgentsTrainer* AgentTrainer, const FName Name = NAME_None, const ELearningAgentsCompletion InCompletionMode = ELearningAgentsCompletion::Termination);
+	static UConditionalCompletion* AddConditionalCompletion(ULearningAgentsTrainer* InAgentTrainer, const FName Name = NAME_None, const ELearningAgentsCompletion InCompletionMode = ELearningAgentsCompletion::Termination);
 
 	/**
 	* Sets the data for this completion. Call during ULearningAgentsTrainer::SetCompletions event.
 	* @param AgentId The agent id this data corresponds to.
 	* @param bIsCompleted Pass in true if condition is met. Otherwise, false.
 	*/
-	UFUNCTION(BlueprintCallable, Category = "LearningAgents")
+	UFUNCTION(BlueprintCallable, Category = "LearningAgents", meta = (AgentId = "-1"))
 	void SetConditionalCompletion(const int32 AgentId, const bool bIsCompleted);
 
-#if ENABLE_VISUAL_LOG
-	/** Describes this observation to the visual logger for debugging purposes. */
+#if UE_LEARNING_AGENTS_ENABLE_VISUAL_LOG
+	/** Describes this completion to the visual logger for debugging purposes. */
 	virtual void VisualLog(const UE::Learning::FIndexSet Instances) const override;
 #endif
 
 	TSharedPtr<UE::Learning::FConditionalCompletion> CompletionObject;
 };
 
+//------------------------------------------------------------------
+
+/** A completion for if a given amount of time has elapsed. */
+UCLASS()
+class LEARNINGAGENTSTRAINING_API UTimeElapsedCompletion : public ULearningAgentsCompletion
+{
+	GENERATED_BODY()
+
+public:
+
+	/**
+	* Adds a new time elapsed completion to the given trainer. Call during ULearningAgentsTrainer::SetupCompletions event.
+	* @param InAgentTrainer The trainer to add this completion to.
+	* @param Name The name of this new completion. Used for debugging.
+	* @param Threshold How much time should be elapsed for the completion to trigger.
+	* @param InCompletionMode The completion mode.
+	* @return The newly created completion.
+	*/
+	UFUNCTION(BlueprintCallable, Category = "LearningAgents")
+	static UTimeElapsedCompletion* AddTimeElapsedCompletion(ULearningAgentsTrainer* InAgentTrainer, const FName Name = NAME_None, const float Threshold = 10.0f, const ELearningAgentsCompletion InCompletionMode = ELearningAgentsCompletion::Termination);
+
+	/**
+	* Sets the data for this completion. Call during ULearningAgentsTrainer::SetCompletions event.
+	* @param AgentId The agent id this data corresponds to.
+	* @param Time The amount of time that has passed
+	*/
+	UFUNCTION(BlueprintCallable, Category = "LearningAgents", meta = (AgentId = "-1"))
+	void SetTimeElapsedCompletion(const int32 AgentId, const float Time);
+
+#if UE_LEARNING_AGENTS_ENABLE_VISUAL_LOG
+	/** Describes this completion to the visual logger for debugging purposes. */
+	virtual void VisualLog(const UE::Learning::FIndexSet Instances) const override;
+#endif
+
+	TSharedPtr<UE::Learning::FTimeElapsedCompletion> CompletionObject;
+};
+
+//------------------------------------------------------------------
+
 /**
 * A completion for if two positions differ by some threshold in a plane, e.g. if the agent gets too far from a
 * starting position.
@@ -96,7 +144,7 @@ public:
 	/**
 	* Adds a new planar position difference completion to the given trainer. The axis parameters define the plane.
 	* Call during ULearningAgentsTrainer::SetupCompletions event.
-	* @param AgentTrainer The trainer to add this completion to.
+	* @param InAgentTrainer The trainer to add this completion to.
 	* @param Name The name of this new completion. Used for debugging.
 	* @param Threshold If the distance becomes greater than this threshold, then the episode will complete.
 	* @param InCompletionMode The completion mode.
@@ -106,7 +154,7 @@ public:
 	*/
 	UFUNCTION(BlueprintCallable, Category = "LearningAgents")
 	static UPlanarPositionDifferenceCompletion* AddPlanarPositionDifferenceCompletion(
-		ULearningAgentsTrainer* AgentTrainer, 
+		ULearningAgentsTrainer* InAgentTrainer,
 		const FName Name = NAME_None,
 		const float Threshold = 100.0f,
 		const ELearningAgentsCompletion InCompletionMode = ELearningAgentsCompletion::Termination,
@@ -119,13 +167,60 @@ public:
 	* @param Position0 The first position.
 	* @param Position1 The second position.
 	*/
-	UFUNCTION(BlueprintCallable, Category = "LearningAgents")
+	UFUNCTION(BlueprintCallable, Category = "LearningAgents", meta = (AgentId = "-1"))
 	void SetPlanarPositionDifferenceCompletion(const int32 AgentId, const FVector Position0, const FVector Position1);
 
-#if ENABLE_VISUAL_LOG
-	/** Describes this observation to the visual logger for debugging purposes. */
+#if UE_LEARNING_AGENTS_ENABLE_VISUAL_LOG
+	/** Describes this completion to the visual logger for debugging purposes. */
 	virtual void VisualLog(const UE::Learning::FIndexSet Instances) const override;
 #endif
 
 	TSharedPtr<UE::Learning::FPlanarPositionDifferenceCompletion> CompletionObject;
 };
+
+/**
+* A completion for if two positions are near by some threshold in a plane, e.g. if the agent gets close to a position.
+*/
+UCLASS()
+class LEARNINGAGENTSTRAINING_API UPlanarPositionSimilarityCompletion : public ULearningAgentsCompletion
+{
+	GENERATED_BODY()
+
+public:
+
+	/**
+	* Adds a new planar position difference completion to the given trainer. The axis parameters define the plane.
+	* Call during ULearningAgentsTrainer::SetupCompletions event.
+	* @param InAgentTrainer The trainer to add this completion to.
+	* @param Name The name of this new completion. Used for debugging.
+	* @param Threshold If the distance becomes greater than this threshold, then the episode will complete.
+	* @param InCompletionMode The completion mode.
+	* @param Axis0 The forward axis of the plane.
+	* @param Axis1 The right axis of the plane.
+	* @return The newly created completion.
+	*/
+	UFUNCTION(BlueprintCallable, Category = "LearningAgents")
+	static UPlanarPositionSimilarityCompletion* AddPlanarPositionSimilarityCompletion(
+		ULearningAgentsTrainer* InAgentTrainer,
+		const FName Name = NAME_None,
+		const float Threshold = 100.0f,
+		const ELearningAgentsCompletion InCompletionMode = ELearningAgentsCompletion::Termination,
+		const FVector Axis0 = FVector::ForwardVector,
+		const FVector Axis1 = FVector::RightVector);
+
+	/**
+	* Sets the data for this completion. Call during ULearningAgentsTrainer::SetCompletions event.
+	* @param AgentId The agent id this data corresponds to.
+	* @param Position0 The first position.
+	* @param Position1 The second position.
+	*/
+	UFUNCTION(BlueprintCallable, Category = "LearningAgents", meta = (AgentId = "-1"))
+	void SetPlanarPositionSimilarityCompletion(const int32 AgentId, const FVector Position0, const FVector Position1);
+
+#if UE_LEARNING_AGENTS_ENABLE_VISUAL_LOG
+	/** Describes this completion to the visual logger for debugging purposes. */
+	virtual void VisualLog(const UE::Learning::FIndexSet Instances) const override;
+#endif
+
+	TSharedPtr<UE::Learning::FPlanarPositionSimilarityCompletion> CompletionObject;
+};
diff --git a/Engine/Plugins/Experimental/LearningAgents/Source/LearningAgentsTraining/Public/LearningAgentsRewards.h b/Engine/Plugins/Experimental/LearningAgents/Source/LearningAgentsTraining/Public/LearningAgentsRewards.h
--- a/Engine/Plugins/Experimental/LearningAgents/Source/LearningAgentsTraining/Public/LearningAgentsRewards.h
+++ b/Engine/Plugins/Experimental/LearningAgents/Source/LearningAgentsTraining/Public/LearningAgentsRewards.h
@@ -3,10 +3,10 @@
 #pragma once
 
 #include "LearningArray.h"
+#include "LearningAgentsDebug.h"
 
 #include "Templates/SharedPointer.h"
 #include "UObject/Object.h"
-#include "EngineDefines.h"
 
 #include "LearningAgentsRewards.generated.h"
 
@@ -26,6 +26,8 @@ class ULearningAgentsTrainer;
 // For functions in this file, we are favoring having more verbose names such as "AddFloatReward" vs simply "Add" in 
 // order to keep it easy to find the correct function in blueprints.
 
+//------------------------------------------------------------------
+
 /**
  * Base class for all rewards/penalties. Rewards are used during reinforcement learning to encourage/discourage
  * certain behaviors from occurring.
@@ -37,7 +39,10 @@ class LEARNINGAGENTSTRAINING_API ULearningAgentsReward : public UObject
 
 public:
 
-#if ENABLE_VISUAL_LOG
+	UPROPERTY(VisibleAnywhere, Transient, Category = "LearningAgents")
+	TObjectPtr<ULearningAgentsTrainer> AgentTrainer;
+
+#if UE_LEARNING_AGENTS_ENABLE_VISUAL_LOG
 	/** Color used to draw this reward in the visual log */
 	FLinearColor VisualLogColor = FColor::Green;
 
@@ -46,6 +51,8 @@ public:
 #endif
 };
 
+//------------------------------------------------------------------
+
 /** A simple float reward. Used as a catch-all for situations where a more type-specific reward does not exist yet. */
 UCLASS()
 class LEARNINGAGENTSTRAINING_API UFloatReward : public ULearningAgentsReward
@@ -56,14 +63,14 @@ public:
 
 	/**
 	* Adds a new float reward to the given trainer. Call during ULearningAgentsTrainer::SetupRewards event.
-	* @param AgentTrainer The trainer to add this reward to.
+	* @param InAgentTrainer The trainer to add this reward to.
 	* @param Name The name of this new reward. Used for debugging.
 	* @param Weight Multiplier for this reward when being summed up for the total reward.
 	* @return The newly created reward.
 	*/
 	UFUNCTION(BlueprintCallable, Category = "LearningAgents")
 	static UFloatReward* AddFloatReward(
-		ULearningAgentsTrainer* AgentTrainer,
+		ULearningAgentsTrainer* InAgentTrainer,
 		const FName Name = NAME_None,
 		const float Weight = 1.0f);
 
@@ -72,10 +79,10 @@ public:
 	* @param AgentId The agent id this data corresponds to.
 	* @param Reward The value currently being rewarded.
 	*/
-	UFUNCTION(BlueprintCallable, Category = "LearningAgents")
+	UFUNCTION(BlueprintCallable, Category = "LearningAgents", meta = (AgentId = "-1"))
 	void SetFloatReward(const int32 AgentId, const float Reward);
 
-#if ENABLE_VISUAL_LOG
+#if UE_LEARNING_AGENTS_ENABLE_VISUAL_LOG
 	/** Describes this reward to the visual logger for debugging purposes. */
 	virtual void VisualLog(const UE::Learning::FIndexSet Instances) const override;
 #endif
@@ -83,6 +90,8 @@ public:
 	TSharedPtr<UE::Learning::FFloatReward> RewardObject;
 };
 
+//------------------------------------------------------------------
+
 /** A reward for maximizing speed. */
 UCLASS()
 class LEARNINGAGENTSTRAINING_API UScalarVelocityReward : public ULearningAgentsReward
@@ -93,7 +102,7 @@ public:
 
 	/**
 	* Adds a new scalar velocity reward to the given trainer. Call during ULearningAgentsTrainer::SetupRewards event.
-	* @param AgentTrainer The trainer to add this reward to.
+	* @param InAgentTrainer The trainer to add this reward to.
 	* @param Name The name of this new reward. Used for debugging.
 	* @param Weight Multiplier for this reward when being summed up for the total reward.
 	* @param Scale Used to normalize the data for the reward.
@@ -101,7 +110,7 @@ public:
 	*/
 	UFUNCTION(BlueprintCallable, Category = "LearningAgents")
 	static UScalarVelocityReward* AddScalarVelocityReward(
-		ULearningAgentsTrainer* AgentTrainer,
+		ULearningAgentsTrainer* InAgentTrainer,
 		const FName Name = NAME_None,
 		const float Weight = 1.0f,
 		const float Scale = 200.0f);
@@ -111,10 +120,10 @@ public:
 	* @param AgentId The agent id this data corresponds to.
 	* @param Velocity The current scalar velocity.
 	*/
-	UFUNCTION(BlueprintCallable, Category = "LearningAgents")
+	UFUNCTION(BlueprintCallable, Category = "LearningAgents", meta = (AgentId = "-1"))
 	void SetScalarVelocityReward(const int32 AgentId, const float Velocity);
 
-#if ENABLE_VISUAL_LOG
+#if UE_LEARNING_AGENTS_ENABLE_VISUAL_LOG
 	/** Describes this reward to the visual logger for debugging purposes. */
 	virtual void VisualLog(const UE::Learning::FIndexSet Instances) const override;
 #endif
@@ -122,6 +131,8 @@ public:
 	TSharedPtr<UE::Learning::FScalarVelocityReward> RewardObject;
 };
 
+//------------------------------------------------------------------
+
 /** A reward for maximizing velocity along a given local axis. */
 UCLASS()
 class LEARNINGAGENTSTRAINING_API ULocalDirectionalVelocityReward : public ULearningAgentsReward
@@ -132,7 +143,7 @@ public:
 
 	/**
 	* Adds a new directional velocity reward to the given trainer. Call during ULearningAgentsTrainer::SetupRewards event.
-	* @param AgentTrainer The trainer to add this reward to.
+	* @param InAgentTrainer The trainer to add this reward to.
 	* @param Name The name of this new reward. Used for debugging.
 	* @param Weight Multiplier for this reward when being summed up for the total reward.
 	* @param Scale Used to normalize the data for the reward.
@@ -141,7 +152,7 @@ public:
 	*/
 	UFUNCTION(BlueprintCallable, Category = "LearningAgents")
 	static ULocalDirectionalVelocityReward* AddLocalDirectionalVelocityReward(
-		ULearningAgentsTrainer* AgentTrainer,
+		ULearningAgentsTrainer* InAgentTrainer,
 		const FName Name = NAME_None,
 		const float Weight = 1.0f,
 		const float Scale = 200.0f,
@@ -153,13 +164,13 @@ public:
 	* @param Velocity The current velocity.
 	* @param RelativeRotation The frame of reference rotation.
 	*/
-	UFUNCTION(BlueprintCallable, Category = "LearningAgents")
+	UFUNCTION(BlueprintCallable, Category = "LearningAgents", meta = (AgentId = "-1"))
 	void SetLocalDirectionalVelocityReward(
 		const int32 AgentId,
 		const FVector Velocity,
 		const FRotator RelativeRotation = FRotator::ZeroRotator);
 
-#if ENABLE_VISUAL_LOG
+#if UE_LEARNING_AGENTS_ENABLE_VISUAL_LOG
 	/** Describes this reward to the visual logger for debugging purposes. */
 	virtual void VisualLog(const UE::Learning::FIndexSet Instances) const override;
 #endif
@@ -167,6 +178,8 @@ public:
 	TSharedPtr<UE::Learning::FLocalDirectionalVelocityReward> RewardObject;
 };
 
+//------------------------------------------------------------------
+
 /** A penalty for being far from a goal position in a plane. */
 UCLASS()
 class LEARNINGAGENTSTRAINING_API UPlanarPositionDifferencePenalty : public ULearningAgentsReward
@@ -178,7 +191,7 @@ public:
 	/**
 	* Adds a new planar difference penalty to the given trainer. The axis parameters define the plane.
 	* Call during ULearningAgentsTrainer::SetupRewards event.
-	* @param AgentTrainer The trainer to add this penalty to.
+	* @param InAgentTrainer The trainer to add this penalty to.
 	* @param Name The name of this new penalty. Used for debugging.
 	* @param Weight Multiplier for this penalty when being summed up for the total reward.
 	* @param Scale Used to normalize the data for the penalty.
@@ -189,7 +202,7 @@ public:
 	*/
 	UFUNCTION(BlueprintCallable, Category = "LearningAgents")
 	static UPlanarPositionDifferencePenalty* AddPlanarPositionDifferencePenalty(
-		ULearningAgentsTrainer* AgentTrainer, 
+		ULearningAgentsTrainer* InAgentTrainer,
 		const FName Name = NAME_None,
 		const float Weight = 1.0f,
 		const float Scale = 100.0f,
@@ -203,10 +216,10 @@ public:
 	* @param Position0 The current position.
 	* @param Position1 The goal position.
 	*/
-	UFUNCTION(BlueprintCallable, Category = "LearningAgents")
+	UFUNCTION(BlueprintCallable, Category = "LearningAgents", meta = (AgentId = "-1"))
 	void SetPlanarPositionDifferencePenalty(const int32 AgentId, const FVector Position0, const FVector Position1);
 
-#if ENABLE_VISUAL_LOG
+#if UE_LEARNING_AGENTS_ENABLE_VISUAL_LOG
 	/** Describes this reward to the visual logger for debugging purposes. */
 	virtual void VisualLog(const UE::Learning::FIndexSet Instances) const override;
 #endif
@@ -224,7 +237,7 @@ public:
 
 	/**
 	* Adds a new position array similarity reward to the given trainer. Call during ULearningAgentsTrainer::SetupRewards event.
-	* @param AgentTrainer The trainer to add this reward to.
+	* @param InAgentTrainer The trainer to add this reward to.
 	* @param Name The name of this new reward. Used for debugging.
 	* @param PositionNum The number of positions in the array.
 	* @param Scale Used to normalize the data for the reward.
@@ -233,7 +246,7 @@ public:
 	*/
 	UFUNCTION(BlueprintCallable, Category = "LearningAgents")
 	static UPositionArraySimilarityReward* AddPositionArraySimilarityReward(
-		ULearningAgentsTrainer* AgentTrainer,
+		ULearningAgentsTrainer* InAgentTrainer,
 		const FName Name = NAME_None,
 		const int32 PositionNum = 0,
 		const float Scale = 100.0f,
@@ -249,7 +262,7 @@ public:
 	* @param RelativeRotation0 The frame of reference rotation for Positions0.
 	* @param RelativeRotation1 The frame of reference rotation for Positions1.
 	*/
-	UFUNCTION(BlueprintCallable, Category = "LearningAgents")
+	UFUNCTION(BlueprintCallable, Category = "LearningAgents", meta = (AgentId = "-1"))
 	void SetPositionArraySimilarityReward(
 		const int32 AgentId,
 		const TArray<FVector>& Positions0, 
@@ -259,7 +272,7 @@ public:
 		const FRotator RelativeRotation0 = FRotator::ZeroRotator,
 		const FRotator RelativeRotation1 = FRotator::ZeroRotator);
 
-#if ENABLE_VISUAL_LOG
+#if UE_LEARNING_AGENTS_ENABLE_VISUAL_LOG
 	/** Describes this reward to the visual logger for debugging purposes. */
 	virtual void VisualLog(const UE::Learning::FIndexSet Instances) const override;
 #endif
diff --git a/Engine/Plugins/Experimental/LearningAgents/Source/LearningAgentsTraining/Public/LearningAgentsTrainer.h b/Engine/Plugins/Experimental/LearningAgents/Source/LearningAgentsTraining/Public/LearningAgentsTrainer.h
--- a/Engine/Plugins/Experimental/LearningAgents/Source/LearningAgentsTraining/Public/LearningAgentsTrainer.h
+++ b/Engine/Plugins/Experimental/LearningAgents/Source/LearningAgentsTraining/Public/LearningAgentsTrainer.h
@@ -22,6 +22,7 @@ namespace UE::Learning
 	struct FSharedMemoryPPOTrainer;
 	struct FSumReward;
 	struct FCompletionObject;
+	enum class ECompletionMode : uint8;
 }
 
 class ALearningAgentsManager;
@@ -41,6 +42,14 @@ enum class ELearningAgentsCompletion : uint8
 	/** Episode ended early and zero reward was expected for all future steps. */
 	Termination	UMETA(DisplayName = "Termination"),
 };
+namespace UE::Learning::Agents
+{
+	/** Get the learning agents completion from the UE::Learning completion. */
+	LEARNINGAGENTSTRAINING_API ELearningAgentsCompletion GetLearningAgentsCompletion(const ECompletionMode CompletionMode);
+
+	/** Get the UE::Learning completion from the learning agents completion. */
+	LEARNINGAGENTSTRAINING_API ECompletionMode GetCompletionMode(const ELearningAgentsCompletion Completion);
+}
 
 /** The configurable settings for a ULearningAgentsTrainer. */
 USTRUCT(BlueprintType, Category = "LearningAgents")
@@ -208,7 +217,7 @@ public:
 	* @param InAgentType The agent type we are training with.
 	* @param InPolicy The policy to be trained.
 	* @param InCritic Optional - only needs to be provided if we want the critic to be accessible at runtime.
-	* @param Settings The trainer settings to use.
+	* @param TrainerSettings The trainer settings to use.
 	*/
 	UFUNCTION(BlueprintCallable, Category = "LearningAgents")
 	void SetupTrainer(
@@ -216,7 +225,7 @@ public:
 		ULearningAgentsType* InAgentType,
 		ULearningAgentsPolicy* InPolicy,
 		ULearningAgentsCritic* InCritic = nullptr,
-		const FLearningAgentsTrainerSettings& Settings = FLearningAgentsTrainerSettings());
+		const FLearningAgentsTrainerSettings& TrainerSettings = FLearningAgentsTrainerSettings());
 
 public: 
 
@@ -280,12 +289,12 @@ public:
 public:
 
 	/**
-	 * During this event, you will receive the ids of each agent that needs to be reset. Both the agent's actor and its
-	 * training environment should be reset for a new episode to commence.
+	 * During this event, all episodes should be reset for each agent.
 	 * @param AgentIds The ids of the agents that need resetting.
+	 * @see GetAgent to get the agent corresponding to each id.
 	 */
 	UFUNCTION(BlueprintNativeEvent, Category = "LearningAgents")
-	void ResetInstance(const TArray<int32>& AgentIds);
+	void ResetEpisodes(const TArray<int32>& AgentIds);
 
 // ----- Training Process -----
 public:
@@ -332,17 +341,45 @@ public:
 
 	/**
 	* Call this function at the end of each step of your training loop. This takes the current observations/actions/
-	* rewards and moves them into the current episode's experience buffer. Finished episodes will have their agents
-	* reset and their data will be sent to the external training process. Finally, the latest iteration of the
-	* trained policy will be synced back to UE so further experience can be acquired on-policy.
+	* rewards and moves them into the episode experience buffer. All agents with full episode buffers or those which
+	* have been signaled complete will be reset. If enough experience is gathered, it will be sent to the training 
+	* process and an iteration of training will be run and the updated policy will be synced back.
 	*/
 	UFUNCTION(BlueprintCallable, Category = "LearningAgents")
-	void IterateTraining();
+	void ProcessExperience();
 
 	/** Manually reset all agents. Does not record the experience gathered up to this point by each agent. */
 	UFUNCTION(BlueprintCallable, Category = "LearningAgents")
-	void ResetAllInstances();
+	void ResetAllEpisodes();
+
+	/**
+	* Convenience function that runs a basic training loop. If training has not been started, it will start it, and 
+	* then call RunInference. On each following call to this function, it will call EvaluateRewards, 
+	* EvaluateCompletions, and ProcessExperience, followed by RunInference.
+	*/
+	UFUNCTION(BlueprintCallable, Category = "LearningAgents")
+	void RunTraining();
 
+	/**
+	* Gets the current reward for an agent according to the critic. Should be called only after EvaluateRewards.
+	*
+	* @param AgentId	The AgentId to look-up the reward for
+	* @returns			The reward
+	*/
+	UFUNCTION(BlueprintPure, Category = "LearningAgents", meta = (AgentId = "-1"))
+	float GetReward(const int32 AgentId) const;
+
+	/**
+	* Gets if the agent will complete the episode or not according to the given set of completions. Should be called 
+	* only after EvaluateCompletions.
+	*
+	* @param AgentId		The AgentId to look-up the completion for
+	* @param OutCompletion	The completion type if the agent will complete the episode
+	* @returns				If the agent will complete the episode
+	*/
+	UFUNCTION(BlueprintPure, Category = "LearningAgents", meta = (AgentId = "-1"))
+	bool IsCompleted(const int32 AgentId, ELearningAgentsCompletion& OutCompletion) const;
+	
 // ----- Private Data ----- 
 private:
 
