Learning Agents: Rename Agent Type to Agent Interactor & add settings to RunTraining

diff --git a/Engine/Plugins/Experimental/LearningAgents/Source/LearningAgents/Private/LearningAgentsActions.cpp b/Engine/Plugins/Experimental/LearningAgents/Source/LearningAgents/Private/LearningAgentsActions.cpp
--- a/Engine/Plugins/Experimental/LearningAgents/Source/LearningAgents/Private/LearningAgentsActions.cpp
+++ b/Engine/Plugins/Experimental/LearningAgents/Source/LearningAgents/Private/LearningAgentsActions.cpp
@@ -3,7 +3,7 @@
 #include "LearningAgentsActions.h"
 
 #include "LearningAgentsManager.h"
-#include "LearningAgentsType.h"
+#include "LearningAgentsInteractor.h"
 #include "LearningArray.h"
 #include "LearningArrayMap.h"
 #include "LearningFeatureObject.h"
@@ -14,24 +14,24 @@
 namespace UE::Learning::Agents::Actions::Private
 {
 	template<typename ActionUObject, typename ActionFObject, typename... InArgTypes>
-	ActionUObject* AddAction(ULearningAgentsType* InAgentType, const FName Name, InArgTypes&& ...Args)
+	ActionUObject* AddAction(ULearningAgentsInteractor* InInteractor, const FName Name, InArgTypes&& ...Args)
 	{
-		if (!InAgentType)
+		if (!InInteractor)
 		{
-			UE_LOG(LogLearning, Error, TEXT("InAgentType is nullptr."));
+			UE_LOG(LogLearning, Error, TEXT("InInteractor is nullptr."));
 			return nullptr;
 		}
 
-		ActionUObject* Action = NewObject<ActionUObject>(InAgentType, Name);
+		ActionUObject* Action = NewObject<ActionUObject>(InInteractor, Name);
 
-		Action->AgentType = InAgentType;
+		Action->Interactor = InInteractor;
 		Action->FeatureObject = MakeShared<ActionFObject>(
 			Action->GetFName(),
-			InAgentType->GetAgentManager()->GetInstanceData().ToSharedRef(),
-			InAgentType->GetAgentManager()->GetMaxInstanceNum(),
+			InInteractor->GetAgentManager()->GetInstanceData().ToSharedRef(),
+			InInteractor->GetAgentManager()->GetMaxInstanceNum(),
 			Forward<InArgTypes>(Args)...);
 
-		InAgentType->AddAction(Action, Action->FeatureObject.ToSharedRef());
+		InInteractor->AddAction(Action, Action->FeatureObject.ToSharedRef());
 
 		return Action;
 	}
@@ -39,14 +39,14 @@ namespace UE::Learning::Agents::Actions::Private
 
 //------------------------------------------------------------------
 
-UFloatAction* UFloatAction::AddFloatAction(ULearningAgentsType* InAgentType, const FName Name, const float Scale)
+UFloatAction* UFloatAction::AddFloatAction(ULearningAgentsInteractor* InInteractor, const FName Name, const float Scale)
 {
-	return UE::Learning::Agents::Actions::Private::AddAction<UFloatAction, UE::Learning::FFloatFeature>(InAgentType, Name, 1, Scale);
+	return UE::Learning::Agents::Actions::Private::AddAction<UFloatAction, UE::Learning::FFloatFeature>(InInteractor, Name, 1, Scale);
 }
 
 float UFloatAction::GetFloatAction(const int32 AgentId) const
 {
-	if (!AgentType->HasAgent(AgentId))
+	if (!Interactor->HasAgent(AgentId))
 	{
 		UE_LOG(LogLearning, Error, TEXT("%s: AgentId %d not found in the agents set."), *GetName(), AgentId);
 		return 0.0f;
@@ -57,7 +57,7 @@ float UFloatAction::GetFloatAction(const int32 AgentId) const
 
 void UFloatAction::SetFloatAction(const int32 AgentId, const float Value)
 {
-	if (!AgentType->HasAgent(AgentId))
+	if (!Interactor->HasAgent(AgentId))
 	{
 		UE_LOG(LogLearning, Error, TEXT("%s: AgentId %d not found in the agents set."), *GetName(), AgentId);
 		return;
@@ -76,7 +76,7 @@ void UFloatAction::VisualLog(const UE::Learning::FIndexSet Instances) const
 
 	for (const int32 Instance : Instances)
 	{
-		if (const AActor* Actor = Cast<AActor>(AgentType->GetAgent(Instance)))
+		if (const AActor* Actor = Cast<AActor>(Interactor->GetAgent(Instance)))
 		{
 			UE_LEARNING_AGENTS_VLOG_STRING(this, LogLearning, Display,
 				Actor->GetActorLocation(),
@@ -91,7 +91,7 @@ void UFloatAction::VisualLog(const UE::Learning::FIndexSet Instances) const
 }
 #endif
 
-UFloatArrayAction* UFloatArrayAction::AddFloatArrayAction(ULearningAgentsType* InAgentType, const FName Name, const int32 Num, const float Scale)
+UFloatArrayAction* UFloatArrayAction::AddFloatArrayAction(ULearningAgentsInteractor* InInteractor, const FName Name, const int32 Num, const float Scale)
 {
 	if (Num < 1)
 	{
@@ -99,12 +99,12 @@ UFloatArrayAction* UFloatArrayAction::AddFloatArrayAction(ULearningAgentsType* I
 		return nullptr;
 	}
 
-	return UE::Learning::Agents::Actions::Private::AddAction<UFloatArrayAction, UE::Learning::FFloatFeature>(InAgentType, Name, Num, Scale);
+	return UE::Learning::Agents::Actions::Private::AddAction<UFloatArrayAction, UE::Learning::FFloatFeature>(InInteractor, Name, Num, Scale);
 }
 
 void UFloatArrayAction::GetFloatArrayAction(const int32 AgentId, TArray<float>& OutValues) const
 {
-	if (!AgentType->HasAgent(AgentId))
+	if (!Interactor->HasAgent(AgentId))
 	{
 		UE_LOG(LogLearning, Error, TEXT("%s: AgentId %d not found in the agents set."), *GetName(), AgentId);
 		OutValues.Empty();
@@ -119,7 +119,7 @@ void UFloatArrayAction::GetFloatArrayAction(const int32 AgentId, TArray<float>&
 
 void UFloatArrayAction::SetFloatArrayAction(const int32 AgentId, const TArray<float>& Values)
 {
-	if (!AgentType->HasAgent(AgentId))
+	if (!Interactor->HasAgent(AgentId))
 	{
 		UE_LOG(LogLearning, Error, TEXT("%s: AgentId %d not found in the agents set."), *GetName(), AgentId);
 		return;
@@ -146,7 +146,7 @@ void UFloatArrayAction::VisualLog(const UE::Learning::FIndexSet Instances) const
 
 	for (const int32 Instance : Instances)
 	{
-		if (const AActor* Actor = Cast<AActor>(AgentType->GetAgent(Instance)))
+		if (const AActor* Actor = Cast<AActor>(Interactor->GetAgent(Instance)))
 		{
 			UE_LEARNING_AGENTS_VLOG_STRING(this, LogLearning, Display,
 				Actor->GetActorLocation(),
@@ -163,14 +163,14 @@ void UFloatArrayAction::VisualLog(const UE::Learning::FIndexSet Instances) const
 
 //------------------------------------------------------------------
 
-UVectorAction* UVectorAction::AddVectorAction(ULearningAgentsType* InAgentType, const FName Name, const float Scale)
+UVectorAction* UVectorAction::AddVectorAction(ULearningAgentsInteractor* InInteractor, const FName Name, const float Scale)
 {
-	return UE::Learning::Agents::Actions::Private::AddAction<UVectorAction, UE::Learning::FFloatFeature>(InAgentType, Name, 3, Scale);
+	return UE::Learning::Agents::Actions::Private::AddAction<UVectorAction, UE::Learning::FFloatFeature>(InInteractor, Name, 3, Scale);
 }
 
 FVector UVectorAction::GetVectorAction(const int32 AgentId) const
 {
-	if (!AgentType->HasAgent(AgentId))
+	if (!Interactor->HasAgent(AgentId))
 	{
 		UE_LOG(LogLearning, Error, TEXT("%s: AgentId %d not found in the agents set."), *GetName(), AgentId);
 		return FVector::ZeroVector;
@@ -183,7 +183,7 @@ FVector UVectorAction::GetVectorAction(const int32 AgentId) const
 
 void UVectorAction::SetVectorAction(const int32 AgentId, const FVector InAction)
 {
-	if (!AgentType->HasAgent(AgentId))
+	if (!Interactor->HasAgent(AgentId))
 	{
 		UE_LOG(LogLearning, Error, TEXT("%s: AgentId %d not found in the agents set."), *GetName(), AgentId);
 		return;
@@ -206,7 +206,7 @@ void UVectorAction::VisualLog(const UE::Learning::FIndexSet Instances) const
 
 	for (const int32 Instance : Instances)
 	{
-		if (const AActor* Actor = Cast<AActor>(AgentType->GetAgent(Instance)))
+		if (const AActor* Actor = Cast<AActor>(Interactor->GetAgent(Instance)))
 		{
 			const FVector Vector(ValueView[Instance][0], ValueView[Instance][1], ValueView[Instance][2]);
 
@@ -234,7 +234,7 @@ void UVectorAction::VisualLog(const UE::Learning::FIndexSet Instances) const
 }
 #endif
 
-UVectorArrayAction* UVectorArrayAction::AddVectorArrayAction(ULearningAgentsType* InAgentType, const FName Name, const int32 Num, const float Scale)
+UVectorArrayAction* UVectorArrayAction::AddVectorArrayAction(ULearningAgentsInteractor* InInteractor, const FName Name, const int32 Num, const float Scale)
 {
 	if (Num < 1)
 	{
@@ -242,12 +242,12 @@ UVectorArrayAction* UVectorArrayAction::AddVectorArrayAction(ULearningAgentsType
 		return nullptr;
 	}
 
-	return UE::Learning::Agents::Actions::Private::AddAction<UVectorArrayAction, UE::Learning::FFloatFeature>(InAgentType, Name, Num * 3, Scale);
+	return UE::Learning::Agents::Actions::Private::AddAction<UVectorArrayAction, UE::Learning::FFloatFeature>(InInteractor, Name, Num * 3, Scale);
 }
 
 void UVectorArrayAction::GetVectorArrayAction(const int32 AgentId, TArray<FVector>& OutVectors) const
 {
-	if (!AgentType->HasAgent(AgentId))
+	if (!Interactor->HasAgent(AgentId))
 	{
 		UE_LOG(LogLearning, Error, TEXT("%s: AgentId %d not found in the agents set."), *GetName(), AgentId);
 		OutVectors.Empty();
@@ -269,7 +269,7 @@ void UVectorArrayAction::GetVectorArrayAction(const int32 AgentId, TArray<FVecto
 
 void UVectorArrayAction::SetVectorArrayAction(const int32 AgentId, const TArray<FVector>& Vectors)
 {
-	if (!AgentType->HasAgent(AgentId))
+	if (!Interactor->HasAgent(AgentId))
 	{
 		UE_LOG(LogLearning, Error, TEXT("%s: AgentId %d not found in the agents set."), *GetName(), AgentId);
 		return;
@@ -303,7 +303,7 @@ void UVectorArrayAction::VisualLog(const UE::Learning::FIndexSet Instances) cons
 
 	for (const int32 Instance : Instances)
 	{
-		if (const AActor* Actor = Cast<AActor>(AgentType->GetAgent(Instance)))
+		if (const AActor* Actor = Cast<AActor>(Interactor->GetAgent(Instance)))
 		{
 			for (int32 VectorIdx = 0; VectorIdx < VectorNum; VectorIdx++)
 			{
@@ -338,14 +338,14 @@ void UVectorArrayAction::VisualLog(const UE::Learning::FIndexSet Instances) cons
 
 //------------------------------------------------------------------
 
-URotationAction* URotationAction::AddRotationAction(ULearningAgentsType* InAgentType, const FName Name, const float Scale)
+URotationAction* URotationAction::AddRotationAction(ULearningAgentsInteractor* InInteractor, const FName Name, const float Scale)
 {
-	return UE::Learning::Agents::Actions::Private::AddAction<URotationAction, UE::Learning::FRotationVectorFeature>(InAgentType, Name, 1, Scale);
+	return UE::Learning::Agents::Actions::Private::AddAction<URotationAction, UE::Learning::FRotationVectorFeature>(InInteractor, Name, 1, Scale);
 }
 
 FRotator URotationAction::GetRotationAction(const int32 AgentId) const
 {
-	if (!AgentType->HasAgent(AgentId))
+	if (!Interactor->HasAgent(AgentId))
 	{
 		UE_LOG(LogLearning, Error, TEXT("%s: AgentId %d not found in the agents set."), *GetName(), AgentId);
 		return FRotator::ZeroRotator;
@@ -358,7 +358,7 @@ FRotator URotationAction::GetRotationAction(const int32 AgentId) const
 
 FVector URotationAction::GetRotationActionAsRotationVector(const int32 AgentId) const
 {
-	if (!AgentType->HasAgent(AgentId))
+	if (!Interactor->HasAgent(AgentId))
 	{
 		UE_LOG(LogLearning, Error, TEXT("%s: AgentId %d not found in the agents set."), *GetName(), AgentId);
 		return FVector::ZeroVector;
@@ -369,7 +369,7 @@ FVector URotationAction::GetRotationActionAsRotationVector(const int32 AgentId)
 
 FQuat URotationAction::GetRotationActionAsQuat(const int32 AgentId) const
 {
-	if (!AgentType->HasAgent(AgentId))
+	if (!Interactor->HasAgent(AgentId))
 	{
 		UE_LOG(LogLearning, Error, TEXT("%s: AgentId %d not found in the agents set."), *GetName(), AgentId);
 		return FQuat::Identity;
@@ -392,7 +392,7 @@ void URotationAction::VisualLog(const UE::Learning::FIndexSet Instances) const
 
 	for (const int32 Instance : Instances)
 	{
-		if (const AActor* Actor = Cast<AActor>(AgentType->GetAgent(Instance)))
+		if (const AActor* Actor = Cast<AActor>(Interactor->GetAgent(Instance)))
 		{
 			const FQuat Rotation = FQuat::MakeFromRotationVector((UE_TWO_PI / 180.0f) * ValueView[Instance][0]);
 			const FRotator Rotator = Rotation.Rotator();
@@ -412,7 +412,7 @@ void URotationAction::VisualLog(const UE::Learning::FIndexSet Instances) const
 #endif
 
 
-URotationArrayAction* URotationArrayAction::AddRotationArrayAction(ULearningAgentsType* InAgentType, const FName Name, const int32 RotationNum, const float Scale)
+URotationArrayAction* URotationArrayAction::AddRotationArrayAction(ULearningAgentsInteractor* InInteractor, const FName Name, const int32 RotationNum, const float Scale)
 {
 	if (RotationNum < 1)
 	{
@@ -420,12 +420,12 @@ URotationArrayAction* URotationArrayAction::AddRotationArrayAction(ULearningAgen
 		return nullptr;
 	}
 
-	return UE::Learning::Agents::Actions::Private::AddAction<URotationArrayAction, UE::Learning::FRotationVectorFeature>(InAgentType, Name, RotationNum, Scale);
+	return UE::Learning::Agents::Actions::Private::AddAction<URotationArrayAction, UE::Learning::FRotationVectorFeature>(InInteractor, Name, RotationNum, Scale);
 }
 
 void URotationArrayAction::GetRotationArrayAction(const int32 AgentId, TArray<FRotator>& OutRotations) const
 {
-	if (!AgentType->HasAgent(AgentId))
+	if (!Interactor->HasAgent(AgentId))
 	{
 		UE_LOG(LogLearning, Error, TEXT("%s: AgentId %d not found in the agents set."), *GetName(), AgentId);
 		return;
@@ -442,7 +442,7 @@ void URotationArrayAction::GetRotationArrayAction(const int32 AgentId, TArray<FR
 
 void URotationArrayAction::GetRotationArrayActionAsRotationVectors(const int32 AgentId, TArray<FVector>& OutRotationVectors) const
 {
-	if (!AgentType->HasAgent(AgentId))
+	if (!Interactor->HasAgent(AgentId))
 	{
 		UE_LOG(LogLearning, Error, TEXT("%s: AgentId %d not found in the agents set."), *GetName(), AgentId);
 		return;
@@ -459,7 +459,7 @@ void URotationArrayAction::GetRotationArrayActionAsRotationVectors(const int32 A
 
 void URotationArrayAction::GetRotationArrayActionAsQuats(const int32 AgentId, TArray<FQuat>& OutRotations) const
 {
-	if (!AgentType->HasAgent(AgentId))
+	if (!Interactor->HasAgent(AgentId))
 	{
 		UE_LOG(LogLearning, Error, TEXT("%s: AgentId %d not found in the agents set."), *GetName(), AgentId);
 		return;
@@ -486,7 +486,7 @@ void URotationArrayAction::VisualLog(const UE::Learning::FIndexSet Instances) co
 
 	for (const int32 Instance : Instances)
 	{
-		if (const AActor* Actor = Cast<AActor>(AgentType->GetAgent(Instance)))
+		if (const AActor* Actor = Cast<AActor>(Interactor->GetAgent(Instance)))
 		{
 			for (int32 RotationIdx = 0; RotationIdx < RotationNum; RotationIdx++)
 			{
diff --git a/Engine/Plugins/Experimental/LearningAgents/Source/LearningAgents/Private/LearningAgentsController.cpp b/Engine/Plugins/Experimental/LearningAgents/Source/LearningAgents/Private/LearningAgentsController.cpp
--- a/Engine/Plugins/Experimental/LearningAgents/Source/LearningAgents/Private/LearningAgentsController.cpp
+++ b/Engine/Plugins/Experimental/LearningAgents/Source/LearningAgents/Private/LearningAgentsController.cpp
@@ -3,7 +3,7 @@
 #include "LearningAgentsController.h"
 
 #include "LearningAgentsManager.h"
-#include "LearningAgentsType.h"
+#include "LearningAgentsInteractor.h"
 #include "LearningAgentsActions.h"
 #include "LearningFeatureObject.h"
 #include "LearningLog.h"
@@ -19,7 +19,7 @@ void ULearningAgentsController::SetActions_Implementation(const TArray<int32>& A
 	// Can be overridden to get actions without blueprints
 }
 
-void ULearningAgentsController::SetupController(ALearningAgentsManager* InAgentManager, ULearningAgentsType* InAgentType)
+void ULearningAgentsController::SetupController(ALearningAgentsManager* InAgentManager, ULearningAgentsInteractor* InInteractor)
 {
 	if (IsSetup())
 	{
@@ -42,19 +42,19 @@ void ULearningAgentsController::SetupController(ALearningAgentsManager* InAgentM
 	// This manager is not referenced in this class but we need it in blueprints to call GetAgent()
 	AgentManager = InAgentManager;
 
-	if (!InAgentType)
+	if (!InInteractor)
 	{
-		UE_LOG(LogLearning, Error, TEXT("%s: InAgentType is nullptr."), *GetName());
+		UE_LOG(LogLearning, Error, TEXT("%s: InInteractor is nullptr."), *GetName());
 		return;
 	}
 
-	if (!InAgentType->IsSetup())
+	if (!InInteractor->IsSetup())
 	{
-		UE_LOG(LogLearning, Error, TEXT("%s: %s's Setup must be run before it can be used."), *GetName(), *InAgentType->GetName());
+		UE_LOG(LogLearning, Error, TEXT("%s: %s's Setup must be run before it can be used."), *GetName(), *InInteractor->GetName());
 		return;
 	}
 
-	AgentType = InAgentType;
+	Interactor = InInteractor;
 
 	bIsSetup = true;
 }
@@ -71,10 +71,10 @@ void ULearningAgentsController::EncodeActions()
 
 	SetActions(AddedAgentIds);
 
-	AgentType->GetActionFeature().Encode(AddedAgentSet);
+	Interactor->GetActionFeature().Encode(AddedAgentSet);
 
 #if ENABLE_VISUAL_LOG
-	for (const ULearningAgentsAction* ActionObject : AgentType->GetActionObjects())
+	for (const ULearningAgentsAction* ActionObject : Interactor->GetActionObjects())
 	{
 		if (ActionObject)
 		{
@@ -84,15 +84,15 @@ void ULearningAgentsController::EncodeActions()
 #endif
 }
 
-ULearningAgentsType* ULearningAgentsController::GetAgentType(const TSubclassOf<ULearningAgentsType> AgentTypeClass) const
+ULearningAgentsInteractor* ULearningAgentsController::GetInteractor(const TSubclassOf<ULearningAgentsInteractor> InteractorClass) const
 {
-	if (!AgentType)
+	if (!Interactor)
 	{
-		UE_LOG(LogLearning, Error, TEXT("%s: AgentType is nullptr. Did we forget to call Setup on this component?"), *GetName());
+		UE_LOG(LogLearning, Error, TEXT("%s: Interactor is nullptr. Did we forget to call Setup on this component?"), *GetName());
 		return nullptr;
 	}
 
-	return AgentType;
+	return Interactor;
 }
 
 void ULearningAgentsController::RunController()
@@ -105,7 +105,7 @@ void ULearningAgentsController::RunController()
 		return;
 	}
 
-	AgentType->EncodeObservations();
+	Interactor->EncodeObservations();
 	EncodeActions();
-	AgentType->DecodeActions();
+	Interactor->DecodeActions();
 }
diff --git a/Engine/Plugins/Experimental/LearningAgents/Source/LearningAgents/Private/LearningAgentsCritic.cpp b/Engine/Plugins/Experimental/LearningAgents/Source/LearningAgents/Private/LearningAgentsCritic.cpp
--- a/Engine/Plugins/Experimental/LearningAgents/Source/LearningAgents/Private/LearningAgentsCritic.cpp
+++ b/Engine/Plugins/Experimental/LearningAgents/Source/LearningAgents/Private/LearningAgentsCritic.cpp
@@ -3,7 +3,7 @@
 #include "LearningAgentsCritic.h"
 
 #include "LearningAgentsManager.h"
-#include "LearningAgentsType.h"
+#include "LearningAgentsInteractor.h"
 #include "LearningFeatureObject.h"
 #include "LearningNeuralNetwork.h"
 #include "LearningNeuralNetworkObject.h"
@@ -18,7 +18,7 @@ ULearningAgentsCritic::ULearningAgentsCritic() : ULearningAgentsManagerComponent
 ULearningAgentsCritic::ULearningAgentsCritic(FVTableHelper& Helper) : ULearningAgentsCritic() {}
 ULearningAgentsCritic::~ULearningAgentsCritic() {}
 
-void ULearningAgentsCritic::SetupCritic(ALearningAgentsManager* InAgentManager, ULearningAgentsType* InAgentType, const FLearningAgentsCriticSettings& CriticSettings)
+void ULearningAgentsCritic::SetupCritic(ALearningAgentsManager* InAgentManager, ULearningAgentsInteractor* InInteractor, const FLearningAgentsCriticSettings& CriticSettings)
 {
 	if (IsSetup())
 	{
@@ -40,25 +40,25 @@ void ULearningAgentsCritic::SetupCritic(ALearningAgentsManager* InAgentManager,
 
 	AgentManager = InAgentManager;
 
-	if (!InAgentType)
+	if (!InInteractor)
 	{
-		UE_LOG(LogLearning, Error, TEXT("%s: InAgentType is nullptr."), *GetName());
+		UE_LOG(LogLearning, Error, TEXT("%s: InInteractor is nullptr."), *GetName());
 		return;
 	}
 
-	if (!InAgentType->IsSetup())
+	if (!InInteractor->IsSetup())
 	{
-		UE_LOG(LogLearning, Error, TEXT("%s: %s's Setup must be run before it can be used."), *GetName(), *InAgentType->GetName());
+		UE_LOG(LogLearning, Error, TEXT("%s: %s's Setup must be run before it can be used."), *GetName(), *InInteractor->GetName());
 		return;
 	}
 
-	AgentType = InAgentType;
+	Interactor = InInteractor;
 
 	// Setup Neural Network
 	Network = NewObject<ULearningAgentsNeuralNetwork>(this, TEXT("CriticNetwork"));
 	Network->NeuralNetwork = MakeShared<UE::Learning::FNeuralNetwork>();
 	Network->NeuralNetwork->Resize(
-		AgentType->GetObservationFeature().DimNum(),
+		Interactor->GetObservationFeature().DimNum(),
 		1,
 		CriticSettings.HiddenLayerSize,
 		CriticSettings.LayerNum);
@@ -71,7 +71,7 @@ void ULearningAgentsCritic::SetupCritic(ALearningAgentsManager* InAgentManager,
 		AgentManager->GetMaxInstanceNum(),
 		Network->NeuralNetwork.ToSharedRef());
 
-	AgentManager->GetInstanceData()->Link(AgentType->GetObservationFeature().FeatureHandle, CriticObject->InputHandle);
+	AgentManager->GetInstanceData()->Link(Interactor->GetObservationFeature().FeatureHandle, CriticObject->InputHandle);
 
 	bIsSetup = true;
 }
@@ -249,7 +249,7 @@ void ULearningAgentsCritic::VisualLog(const UE::Learning::FIndexSet Instances) c
 
 	for (const int32 Instance : Instances)
 	{
-		if (const AActor* Actor = Cast<AActor>(AgentType->GetAgent(Instance)))
+		if (const AActor* Actor = Cast<AActor>(Interactor->GetAgent(Instance)))
 		{
 			UE_LEARNING_AGENTS_VLOG_STRING(this, LogLearning, Display,
 				Actor->GetActorLocation(),
diff --git a/Engine/Plugins/Experimental/LearningAgents/Source/LearningAgents/Private/LearningAgentsHelpers.cpp b/Engine/Plugins/Experimental/LearningAgents/Source/LearningAgents/Private/LearningAgentsHelpers.cpp
--- a/Engine/Plugins/Experimental/LearningAgents/Source/LearningAgents/Private/LearningAgentsHelpers.cpp
+++ b/Engine/Plugins/Experimental/LearningAgents/Source/LearningAgents/Private/LearningAgentsHelpers.cpp
@@ -448,7 +448,7 @@ URayCastHelper* URayCastHelper::AddRayCastHelper(ULearningAgentsManagerComponent
 {
 	if (!InManagerComponent)
 	{
-		UE_LOG(LogLearning, Error, TEXT("InAgentType is nullptr."));
+		UE_LOG(LogLearning, Error, TEXT("InInteractor is nullptr."));
 		return nullptr;
 	}
 
diff --git a/Engine/Plugins/Experimental/LearningAgents/Source/LearningAgents/Private/LearningAgentsInteractor.cpp b/Engine/Plugins/Experimental/LearningAgents/Source/LearningAgents/Private/LearningAgentsInteractor.cpp
new file mode 100644
--- /dev/null
+++ b/Engine/Plugins/Experimental/LearningAgents/Source/LearningAgents/Private/LearningAgentsInteractor.cpp
@@ -0,0 +1,203 @@
+// Copyright Epic Games, Inc. All Rights Reserved.
+
+#include "LearningAgentsInteractor.h"
+
+#include "LearningAgentsManager.h"
+#include "LearningAgentsObservations.h"
+#include "LearningAgentsActions.h"
+
+#include "LearningArray.h"
+#include "LearningArrayMap.h"
+#include "LearningFeatureObject.h"
+#include "LearningLog.h"
+#include "EngineDefines.h"
+
+ULearningAgentsInteractor::ULearningAgentsInteractor() : ULearningAgentsManagerComponent() {}
+ULearningAgentsInteractor::ULearningAgentsInteractor(FVTableHelper& Helper) : ULearningAgentsInteractor() {}
+ULearningAgentsInteractor::~ULearningAgentsInteractor() {}
+
+void ULearningAgentsInteractor::SetupInteractor(ALearningAgentsManager* InAgentManager)
+{
+	if (IsSetup())
+	{
+		UE_LOG(LogLearning, Error, TEXT("%s: Setup already run!"), *GetName());
+		return;
+	}
+
+	if (!InAgentManager)
+	{
+		UE_LOG(LogLearning, Error, TEXT("%s: InAgentManager is nullptr."), *GetName());
+		return;
+	}
+
+	if (!InAgentManager->IsManagerSetup())
+	{
+		UE_LOG(LogLearning, Error, TEXT("%s: %s's SetupManager must be run before it can be used."), *GetName(), *InAgentManager->GetName());
+		return;
+	}
+
+	AgentManager = InAgentManager;
+
+	// Setup Observations
+	ObservationObjects.Empty();
+	ObservationFeatures.Empty();
+	SetupObservations(this);
+	Observations = MakeShared<UE::Learning::FConcatenateFeature>(TEXT("Observations"),
+		TLearningArrayView<1, const TSharedRef<UE::Learning::FFeatureObject>>(ObservationFeatures),
+		InAgentManager->GetInstanceData().ToSharedRef(),
+		InAgentManager->GetMaxInstanceNum());
+
+	// Setup Actions
+	ActionObjects.Empty();
+	ActionFeatures.Empty();
+	SetupActions(this);
+	Actions = MakeShared<UE::Learning::FConcatenateFeature>(TEXT("Actions"),
+		TLearningArrayView<1, const TSharedRef<UE::Learning::FFeatureObject>>(ActionFeatures),
+		InAgentManager->GetInstanceData().ToSharedRef(),
+		InAgentManager->GetMaxInstanceNum());
+
+	bIsSetup = true;
+}
+
+UE::Learning::FFeatureObject& ULearningAgentsInteractor::GetObservationFeature() const
+{
+	return *Observations;
+}
+
+UE::Learning::FFeatureObject& ULearningAgentsInteractor::GetActionFeature() const
+{
+	return *Actions;
+}
+
+TConstArrayView<ULearningAgentsObservation*> ULearningAgentsInteractor::GetObservationObjects() const
+{
+	return ObservationObjects;
+}
+
+TConstArrayView<ULearningAgentsAction*> ULearningAgentsInteractor::GetActionObjects() const
+{
+	return ActionObjects;
+}
+
+void ULearningAgentsInteractor::SetupObservations_Implementation(ULearningAgentsInteractor* Interactor)
+{
+	// Can be overridden to setup observations without blueprints
+}
+
+void ULearningAgentsInteractor::SetObservations_Implementation(const TArray<int32>& AgentIds)
+{
+	// Can be overridden to set observations without blueprints
+}
+
+void ULearningAgentsInteractor::AddObservation(TObjectPtr<ULearningAgentsObservation> Object, const TSharedRef<UE::Learning::FFeatureObject>& Feature)
+{
+	UE_LEARNING_CHECK(!IsSetup());
+	ObservationObjects.Add(Object);
+	ObservationFeatures.Add(Feature);
+}
+
+void ULearningAgentsInteractor::SetupActions_Implementation(ULearningAgentsInteractor* Interactor)
+{
+	// Can be overridden to setup actions without blueprints
+}
+
+void ULearningAgentsInteractor::GetActions_Implementation(const TArray<int32>& AgentIds)
+{
+	// Can be overridden to get actions without blueprints
+}
+
+void ULearningAgentsInteractor::AddAction(TObjectPtr<ULearningAgentsAction> Object, const TSharedRef<UE::Learning::FFeatureObject>& Feature)
+{
+	UE_LEARNING_CHECK(!IsSetup());
+	ActionObjects.Add(Object);
+	ActionFeatures.Add(Feature);
+}
+
+void ULearningAgentsInteractor::EncodeObservations()
+{
+	UE_LEARNING_TRACE_CPUPROFILER_EVENT_SCOPE(ULearningAgentsInteractor::EncodeObservations);
+
+	if (!IsSetup())
+	{
+		UE_LOG(LogLearning, Error, TEXT("%s: Setup not run."), *GetName());
+		return;
+	}
+
+	SetObservations(AddedAgentIds);
+
+	Observations->Encode(AddedAgentSet);
+
+#if ENABLE_VISUAL_LOG
+	for (const ULearningAgentsObservation* ObservationObject : ObservationObjects)
+	{
+		if (ObservationObject)
+		{
+			ObservationObject->VisualLog(AddedAgentSet);
+		}
+	}
+#endif
+}
+
+void ULearningAgentsInteractor::DecodeActions()
+{
+	UE_LEARNING_TRACE_CPUPROFILER_EVENT_SCOPE(ULearningAgentsInteractor::DecodeActions);
+
+	if (!IsSetup())
+	{
+		UE_LOG(LogLearning, Error, TEXT("%s: Setup not run."), *GetName());
+		return;
+	}
+
+	Actions->Decode(AddedAgentSet);
+
+	GetActions(AddedAgentIds);
+
+#if ENABLE_VISUAL_LOG
+	for (const ULearningAgentsAction* ActionObject : ActionObjects)
+	{
+		if (ActionObject)
+		{
+			ActionObject->VisualLog(AddedAgentSet);
+		}
+	}
+#endif
+}
+void ULearningAgentsInteractor::GetObservationVector(const int32 AgentId, TArray<float>& OutObservationVector) const
+{
+	if (!IsSetup())
+	{
+		UE_LOG(LogLearning, Error, TEXT("%s: Setup not run."), *GetName());
+		OutObservationVector.Empty();
+		return;
+	}
+
+	if (!HasAgent(AgentId))
+	{
+		UE_LOG(LogLearning, Error, TEXT("%s: AgentId %d not found in the agents set."), *GetName(), AgentId);
+		OutObservationVector.Empty();
+		return;
+	}
+
+	OutObservationVector.SetNumUninitialized(Observations->DimNum());
+	UE::Learning::Array::Copy<1, float>(OutObservationVector, Observations->FeatureBuffer()[AgentId]);
+}
+
+void ULearningAgentsInteractor::GetActionVector(const int32 AgentId, TArray<float>& OutActionVector) const
+{
+	if (!IsSetup())
+	{
+		UE_LOG(LogLearning, Error, TEXT("%s: Setup not run."), *GetName());
+		OutActionVector.Empty();
+		return;
+	}
+
+	if (!HasAgent(AgentId))
+	{
+		UE_LOG(LogLearning, Error, TEXT("%s: AgentId %d not found in the agents set."), *GetName(), AgentId);
+		OutActionVector.Empty();
+		return;
+	}
+
+	OutActionVector.SetNumUninitialized(Actions->DimNum());
+	UE::Learning::Array::Copy<1, float>(OutActionVector, Actions->FeatureBuffer()[AgentId]);
+}
\ No newline at end of file
diff --git a/Engine/Plugins/Experimental/LearningAgents/Source/LearningAgents/Private/LearningAgentsObservations.cpp b/Engine/Plugins/Experimental/LearningAgents/Source/LearningAgents/Private/LearningAgentsObservations.cpp
--- a/Engine/Plugins/Experimental/LearningAgents/Source/LearningAgents/Private/LearningAgentsObservations.cpp
+++ b/Engine/Plugins/Experimental/LearningAgents/Source/LearningAgents/Private/LearningAgentsObservations.cpp
@@ -3,7 +3,7 @@
 #include "LearningAgentsObservations.h"
 
 #include "LearningAgentsManager.h"
-#include "LearningAgentsType.h"
+#include "LearningAgentsInteractor.h"
 #include "LearningArray.h"
 #include "LearningArrayMap.h"
 #include "LearningFeatureObject.h"
@@ -14,24 +14,24 @@
 namespace UE::Learning::Agents::Observations::Private
 {
 	template<typename ObservationUObject, typename ObservationFObject, typename... InArgTypes>
-	ObservationUObject* AddObservation(ULearningAgentsType* InAgentType, const FName Name, InArgTypes&& ...Args)
+	ObservationUObject* AddObservation(ULearningAgentsInteractor* InInteractor, const FName Name, InArgTypes&& ...Args)
 	{
-		if (!InAgentType)
+		if (!InInteractor)
 		{
-			UE_LOG(LogLearning, Error, TEXT("InAgentType is nullptr."));
+			UE_LOG(LogLearning, Error, TEXT("InInteractor is nullptr."));
 			return nullptr;
 		}
 
-		ObservationUObject* Observation = NewObject<ObservationUObject>(InAgentType, Name);
+		ObservationUObject* Observation = NewObject<ObservationUObject>(InInteractor, Name);
 
-		Observation->AgentType = InAgentType;
+		Observation->Interactor = InInteractor;
 		Observation->FeatureObject = MakeShared<ObservationFObject>(
 			Observation->GetFName(),
-			InAgentType->GetAgentManager()->GetInstanceData().ToSharedRef(),
-			InAgentType->GetAgentManager()->GetMaxInstanceNum(),
+			InInteractor->GetAgentManager()->GetInstanceData().ToSharedRef(),
+			InInteractor->GetAgentManager()->GetMaxInstanceNum(),
 			Forward<InArgTypes>(Args)...);
 
-		InAgentType->AddObservation(Observation, Observation->FeatureObject.ToSharedRef());
+		InInteractor->AddObservation(Observation, Observation->FeatureObject.ToSharedRef());
 
 		return Observation;
 	}
@@ -39,14 +39,14 @@ namespace UE::Learning::Agents::Observations::Private
 
 //------------------------------------------------------------------
 
-UFloatObservation* UFloatObservation::AddFloatObservation(ULearningAgentsType* InAgentType, const FName Name, const float Scale)
+UFloatObservation* UFloatObservation::AddFloatObservation(ULearningAgentsInteractor* InInteractor, const FName Name, const float Scale)
 {
-	return UE::Learning::Agents::Observations::Private::AddObservation<UFloatObservation, UE::Learning::FFloatFeature>(InAgentType, Name, 1, Scale);
+	return UE::Learning::Agents::Observations::Private::AddObservation<UFloatObservation, UE::Learning::FFloatFeature>(InInteractor, Name, 1, Scale);
 }
 
 void UFloatObservation::SetFloatObservation(const int32 AgentId, const float Value)
 {
-	if (!AgentType->HasAgent(AgentId))
+	if (!Interactor->HasAgent(AgentId))
 	{
 		UE_LOG(LogLearning, Error, TEXT("%s: AgentId %d not found in the agents set."), *GetName(), AgentId);
 		return;
@@ -65,7 +65,7 @@ void UFloatObservation::VisualLog(const UE::Learning::FIndexSet Instances) const
 
 	for (const int32 Instance : Instances)
 	{
-		if (const AActor* Actor = Cast<AActor>(AgentType->GetAgent(Instance)))
+		if (const AActor* Actor = Cast<AActor>(Interactor->GetAgent(Instance)))
 		{
 			UE_LEARNING_AGENTS_VLOG_STRING(this, LogLearning, Display,
 				Actor->GetActorLocation(),
@@ -82,7 +82,7 @@ void UFloatObservation::VisualLog(const UE::Learning::FIndexSet Instances) const
 
 //------------------------------------------------------------------
 
-UFloatArrayObservation* UFloatArrayObservation::AddFloatArrayObservation(ULearningAgentsType* InAgentType, const FName Name, const int32 Num, const float Scale)
+UFloatArrayObservation* UFloatArrayObservation::AddFloatArrayObservation(ULearningAgentsInteractor* InInteractor, const FName Name, const int32 Num, const float Scale)
 {
 	if (Num < 1)
 	{
@@ -90,12 +90,12 @@ UFloatArrayObservation* UFloatArrayObservation::AddFloatArrayObservation(ULearni
 		return nullptr;
 	}
 
-	return UE::Learning::Agents::Observations::Private::AddObservation<UFloatArrayObservation, UE::Learning::FFloatFeature>(InAgentType, Name, Num, Scale);
+	return UE::Learning::Agents::Observations::Private::AddObservation<UFloatArrayObservation, UE::Learning::FFloatFeature>(InInteractor, Name, Num, Scale);
 }
 
 void UFloatArrayObservation::SetFloatArrayObservation(const int32 AgentId, const TArray<float>& Values)
 {
-	if (!AgentType->HasAgent(AgentId))
+	if (!Interactor->HasAgent(AgentId))
 	{
 		UE_LOG(LogLearning, Error, TEXT("%s: AgentId %d not found in the agents set."), *GetName(), AgentId);
 		return;
@@ -122,7 +122,7 @@ void UFloatArrayObservation::VisualLog(const UE::Learning::FIndexSet Instances)
 
 	for (const int32 Instance : Instances)
 	{
-		if (const AActor* Actor = Cast<AActor>(AgentType->GetAgent(Instance)))
+		if (const AActor* Actor = Cast<AActor>(Interactor->GetAgent(Instance)))
 		{
 			UE_LEARNING_AGENTS_VLOG_STRING(this, LogLearning, Display,
 				Actor->GetActorLocation(),
@@ -139,14 +139,14 @@ void UFloatArrayObservation::VisualLog(const UE::Learning::FIndexSet Instances)
 
 //------------------------------------------------------------------
 
-UVectorObservation* UVectorObservation::AddVectorObservation(ULearningAgentsType* InAgentType, const FName Name, const float Scale)
+UVectorObservation* UVectorObservation::AddVectorObservation(ULearningAgentsInteractor* InInteractor, const FName Name, const float Scale)
 {
-	return UE::Learning::Agents::Observations::Private::AddObservation<UVectorObservation, UE::Learning::FFloatFeature>(InAgentType, Name, 3, Scale);
+	return UE::Learning::Agents::Observations::Private::AddObservation<UVectorObservation, UE::Learning::FFloatFeature>(InInteractor, Name, 3, Scale);
 }
 
 void UVectorObservation::SetVectorObservation(const int32 AgentId, const FVector Vector)
 {
-	if (!AgentType->HasAgent(AgentId))
+	if (!Interactor->HasAgent(AgentId))
 	{
 		UE_LOG(LogLearning, Error, TEXT("%s: AgentId %d not found in the agents set."), *GetName(), AgentId);
 		return;
@@ -169,7 +169,7 @@ void UVectorObservation::VisualLog(const UE::Learning::FIndexSet Instances) cons
 
 	for (const int32 Instance : Instances)
 	{
-		if (const AActor* Actor = Cast<AActor>(AgentType->GetAgent(Instance)))
+		if (const AActor* Actor = Cast<AActor>(Interactor->GetAgent(Instance)))
 		{
 			const FVector Vector(ValueView[Instance][0], ValueView[Instance][1], ValueView[Instance][2]);
 
@@ -197,7 +197,7 @@ void UVectorObservation::VisualLog(const UE::Learning::FIndexSet Instances) cons
 }
 #endif
 
-UVectorArrayObservation* UVectorArrayObservation::AddVectorArrayObservation(ULearningAgentsType* InAgentType, const FName Name, const int32 Num, const float Scale)
+UVectorArrayObservation* UVectorArrayObservation::AddVectorArrayObservation(ULearningAgentsInteractor* InInteractor, const FName Name, const int32 Num, const float Scale)
 {
 	if (Num < 1)
 	{
@@ -205,12 +205,12 @@ UVectorArrayObservation* UVectorArrayObservation::AddVectorArrayObservation(ULea
 		return nullptr;
 	}
 
-	return UE::Learning::Agents::Observations::Private::AddObservation<UVectorArrayObservation, UE::Learning::FFloatFeature>(InAgentType, Name, Num * 3, Scale);
+	return UE::Learning::Agents::Observations::Private::AddObservation<UVectorArrayObservation, UE::Learning::FFloatFeature>(InInteractor, Name, Num * 3, Scale);
 }
 
 void UVectorArrayObservation::SetVectorArrayObservation(const int32 AgentId, const TArray<FVector>& Vectors)
 {
-	if (!AgentType->HasAgent(AgentId))
+	if (!Interactor->HasAgent(AgentId))
 	{
 		UE_LOG(LogLearning, Error, TEXT("%s: AgentId %d not found in the agents set."), *GetName(), AgentId);
 		return;
@@ -244,7 +244,7 @@ void UVectorArrayObservation::VisualLog(const UE::Learning::FIndexSet Instances)
 
 	for (const int32 Instance : Instances)
 	{
-		if (const AActor* Actor = Cast<AActor>(AgentType->GetAgent(Instance)))
+		if (const AActor* Actor = Cast<AActor>(Interactor->GetAgent(Instance)))
 		{
 			for (int32 VectorIdx = 0; VectorIdx < VectorNum; VectorIdx++)
 			{
@@ -279,7 +279,7 @@ void UVectorArrayObservation::VisualLog(const UE::Learning::FIndexSet Instances)
 
 //------------------------------------------------------------------
 
-UEnumObservation* UEnumObservation::AddEnumObservation(ULearningAgentsType* InAgentType, const UEnum* EnumType, const FName Name)
+UEnumObservation* UEnumObservation::AddEnumObservation(ULearningAgentsInteractor* InInteractor, const UEnum* EnumType, const FName Name)
 {
 	if (!EnumType)
 	{
@@ -293,7 +293,7 @@ UEnumObservation* UEnumObservation::AddEnumObservation(ULearningAgentsType* InAg
 		return nullptr;
 	}
 
-	UEnumObservation* Observation = UE::Learning::Agents::Observations::Private::AddObservation<UEnumObservation, UE::Learning::FFloatFeature>(InAgentType, Name, EnumType->NumEnums());
+	UEnumObservation* Observation = UE::Learning::Agents::Observations::Private::AddObservation<UEnumObservation, UE::Learning::FFloatFeature>(InInteractor, Name, EnumType->NumEnums());
 	if (Observation) { Observation->Enum = EnumType; }
 	
 	return Observation;
@@ -301,7 +301,7 @@ UEnumObservation* UEnumObservation::AddEnumObservation(ULearningAgentsType* InAg
 
 void UEnumObservation::SetEnumObservation(const int32 AgentId, const uint8 EnumValue)
 {
-	if (!AgentType->HasAgent(AgentId))
+	if (!Interactor->HasAgent(AgentId))
 	{
 		UE_LOG(LogLearning, Error, TEXT("%s: AgentId %d not found in the agents set."), *GetName(), AgentId);
 		return;
@@ -332,7 +332,7 @@ void UEnumObservation::VisualLog(const UE::Learning::FIndexSet Instances) const
 		
 	for (const int32 Instance : Instances)
 	{
-		if (const AActor* Actor = Cast<AActor>(AgentType->GetAgent(Instance)))
+		if (const AActor* Actor = Cast<AActor>(Interactor->GetAgent(Instance)))
 		{
 			int32 EnumEntryIdx = INDEX_NONE;
 
@@ -366,7 +366,7 @@ void UEnumObservation::VisualLog(const UE::Learning::FIndexSet Instances) const
 }
 #endif
 
-UEnumArrayObservation* UEnumArrayObservation::AddEnumArrayObservation(ULearningAgentsType* InAgentType, const UEnum* EnumType, const FName Name, const int32 EnumNum)
+UEnumArrayObservation* UEnumArrayObservation::AddEnumArrayObservation(ULearningAgentsInteractor* InInteractor, const UEnum* EnumType, const FName Name, const int32 EnumNum)
 {
 	if (!EnumType)
 	{
@@ -386,7 +386,7 @@ UEnumArrayObservation* UEnumArrayObservation::AddEnumArrayObservation(ULearningA
 		return nullptr;
 	}
 
-	UEnumArrayObservation* Observation = UE::Learning::Agents::Observations::Private::AddObservation<UEnumArrayObservation, UE::Learning::FFloatFeature>(InAgentType, Name, EnumNum * EnumType->NumEnums());
+	UEnumArrayObservation* Observation = UE::Learning::Agents::Observations::Private::AddObservation<UEnumArrayObservation, UE::Learning::FFloatFeature>(InInteractor, Name, EnumNum * EnumType->NumEnums());
 	if (Observation) { Observation->Enum = EnumType; }
 
 	return Observation;
@@ -394,7 +394,7 @@ UEnumArrayObservation* UEnumArrayObservation::AddEnumArrayObservation(ULearningA
 
 void UEnumArrayObservation::SetEnumArrayObservation(const int32 AgentId, const TArray<uint8>& EnumValues)
 {
-	if (!AgentType->HasAgent(AgentId))
+	if (!Interactor->HasAgent(AgentId))
 	{
 		UE_LOG(LogLearning, Error, TEXT("%s: AgentId %d not found in the agents set."), *GetName(), AgentId);
 		return;
@@ -437,7 +437,7 @@ void UEnumArrayObservation::VisualLog(const UE::Learning::FIndexSet Instances) c
 
 	for (const int32 Instance : Instances)
 	{
-		if (const AActor* Actor = Cast<AActor>(AgentType->GetAgent(Instance)))
+		if (const AActor* Actor = Cast<AActor>(Interactor->GetAgent(Instance)))
 		{
 			UE_LEARNING_AGENTS_VLOG_STRING(this, LogLearning, Display,
 				Actor->GetActorLocation(),
@@ -453,14 +453,14 @@ void UEnumArrayObservation::VisualLog(const UE::Learning::FIndexSet Instances) c
 
 //------------------------------------------------------------------
 
-UTimeObservation* UTimeObservation::AddTimeObservation(ULearningAgentsType* InAgentType, const FName Name, const float Scale)
+UTimeObservation* UTimeObservation::AddTimeObservation(ULearningAgentsInteractor* InInteractor, const FName Name, const float Scale)
 {
-	return UE::Learning::Agents::Observations::Private::AddObservation<UTimeObservation, UE::Learning::FTimeFeature>(InAgentType, Name, 1, Scale);
+	return UE::Learning::Agents::Observations::Private::AddObservation<UTimeObservation, UE::Learning::FTimeFeature>(InInteractor, Name, 1, Scale);
 }
 
 void UTimeObservation::SetTimeObservation(const int32 AgentId, const float Time, const float RelativeTime)
 {
-	if (!AgentType->HasAgent(AgentId))
+	if (!Interactor->HasAgent(AgentId))
 	{
 		UE_LOG(LogLearning, Error, TEXT("%s: AgentId %d not found in the agents set."), *GetName(), AgentId);
 		return;
@@ -481,7 +481,7 @@ void UTimeObservation::VisualLog(const UE::Learning::FIndexSet Instances) const
 
 	for (const int32 Instance : Instances)
 	{
-		if (const AActor* Actor = Cast<AActor>(AgentType->GetAgent(Instance)))
+		if (const AActor* Actor = Cast<AActor>(Interactor->GetAgent(Instance)))
 		{
 			const float Time = TimeView[Instance][0];
 			const float RelativeTime = RelativeTimeView[Instance];
@@ -500,7 +500,7 @@ void UTimeObservation::VisualLog(const UE::Learning::FIndexSet Instances) const
 }
 #endif
 
-UTimeArrayObservation* UTimeArrayObservation::AddTimeArrayObservation(ULearningAgentsType* InAgentType, const FName Name, const int32 TimeNum, const float Scale)
+UTimeArrayObservation* UTimeArrayObservation::AddTimeArrayObservation(ULearningAgentsInteractor* InInteractor, const FName Name, const int32 TimeNum, const float Scale)
 {
 	if (TimeNum < 1)
 	{
@@ -508,12 +508,12 @@ UTimeArrayObservation* UTimeArrayObservation::AddTimeArrayObservation(ULearningA
 		return nullptr;
 	}
 
-	return UE::Learning::Agents::Observations::Private::AddObservation<UTimeArrayObservation, UE::Learning::FTimeFeature>(InAgentType, Name, TimeNum, Scale);
+	return UE::Learning::Agents::Observations::Private::AddObservation<UTimeArrayObservation, UE::Learning::FTimeFeature>(InInteractor, Name, TimeNum, Scale);
 }
 
 void UTimeArrayObservation::SetTimeArrayObservation(const int32 AgentId, const TArray<float>& Times, const float RelativeTime)
 {
-	if (!AgentType->HasAgent(AgentId))
+	if (!Interactor->HasAgent(AgentId))
 	{
 		UE_LOG(LogLearning, Error, TEXT("%s: AgentId %d not found in the agents set."), *GetName(), AgentId);
 		return;
@@ -543,7 +543,7 @@ void UTimeArrayObservation::VisualLog(const UE::Learning::FIndexSet Instances) c
 
 	for (const int32 Instance : Instances)
 	{
-		if (const AActor* Actor = Cast<AActor>(AgentType->GetAgent(Instance)))
+		if (const AActor* Actor = Cast<AActor>(Interactor->GetAgent(Instance)))
 		{
 			UE_LEARNING_AGENTS_VLOG_STRING(this, LogLearning, Display,
 				Actor->GetActorLocation(),
@@ -561,14 +561,14 @@ void UTimeArrayObservation::VisualLog(const UE::Learning::FIndexSet Instances) c
 
 //------------------------------------------------------------------
 
-UAngleObservation* UAngleObservation::AddAngleObservation(ULearningAgentsType* InAgentType, const FName Name, const float Scale)
+UAngleObservation* UAngleObservation::AddAngleObservation(ULearningAgentsInteractor* InInteractor, const FName Name, const float Scale)
 {
-	return UE::Learning::Agents::Observations::Private::AddObservation<UAngleObservation, UE::Learning::FAngleFeature>(InAgentType, Name, 1, Scale);
+	return UE::Learning::Agents::Observations::Private::AddObservation<UAngleObservation, UE::Learning::FAngleFeature>(InInteractor, Name, 1, Scale);
 }
 
 void UAngleObservation::SetAngleObservation(const int32 AgentId, const float Angle, const float RelativeAngle)
 {
-	if (!AgentType->HasAgent(AgentId))
+	if (!Interactor->HasAgent(AgentId))
 	{
 		UE_LOG(LogLearning, Error, TEXT("%s: AgentId %d not found in the agents set."), *GetName(), AgentId);
 		return;
@@ -589,7 +589,7 @@ void UAngleObservation::VisualLog(const UE::Learning::FIndexSet Instances) const
 
 	for (const int32 Instance : Instances)
 	{
-		if (const AActor* Actor = Cast<AActor>(AgentType->GetAgent(Instance)))
+		if (const AActor* Actor = Cast<AActor>(Interactor->GetAgent(Instance)))
 		{
 			const float Angle = AngleView[Instance][0];
 			const float RelativeAngle = RelativeAngleView[Instance];
@@ -611,7 +611,7 @@ void UAngleObservation::VisualLog(const UE::Learning::FIndexSet Instances) const
 }
 #endif
 
-UAngleArrayObservation* UAngleArrayObservation::AddAngleArrayObservation(ULearningAgentsType* InAgentType, const FName Name, const int32 AngleNum, const float Scale)
+UAngleArrayObservation* UAngleArrayObservation::AddAngleArrayObservation(ULearningAgentsInteractor* InInteractor, const FName Name, const int32 AngleNum, const float Scale)
 {
 	if (AngleNum < 1)
 	{
@@ -619,12 +619,12 @@ UAngleArrayObservation* UAngleArrayObservation::AddAngleArrayObservation(ULearni
 		return nullptr;
 	}
 
-	return UE::Learning::Agents::Observations::Private::AddObservation<UAngleArrayObservation, UE::Learning::FAngleFeature>(InAgentType, Name, AngleNum, Scale);
+	return UE::Learning::Agents::Observations::Private::AddObservation<UAngleArrayObservation, UE::Learning::FAngleFeature>(InInteractor, Name, AngleNum, Scale);
 }
 
 void UAngleArrayObservation::SetAngleArrayObservation(const int32 AgentId, const TArray<float>& Angles, const float RelativeAngle)
 {
-	if (!AgentType->HasAgent(AgentId))
+	if (!Interactor->HasAgent(AgentId))
 	{
 		UE_LOG(LogLearning, Error, TEXT("%s: AgentId %d not found in the agents set."), *GetName(), AgentId);
 		return;
@@ -660,7 +660,7 @@ void UAngleArrayObservation::VisualLog(const UE::Learning::FIndexSet Instances)
 
 	for (const int32 Instance : Instances)
 	{
-		if (const AActor* Actor = Cast<AActor>(AgentType->GetAgent(Instance)))
+		if (const AActor* Actor = Cast<AActor>(Interactor->GetAgent(Instance)))
 		{
 			const float RelativeAngle = RelativeAngleView[Instance];
 
@@ -695,9 +695,9 @@ void UAngleArrayObservation::VisualLog(const UE::Learning::FIndexSet Instances)
 
 //------------------------------------------------------------------
 
-URotationObservation* URotationObservation::AddRotationObservation(ULearningAgentsType* InAgentType, const FName Name, const float Scale)
+URotationObservation* URotationObservation::AddRotationObservation(ULearningAgentsInteractor* InInteractor, const FName Name, const float Scale)
 {
-	return UE::Learning::Agents::Observations::Private::AddObservation<URotationObservation, UE::Learning::FRotationFeature>(InAgentType, Name, 1, Scale);
+	return UE::Learning::Agents::Observations::Private::AddObservation<URotationObservation, UE::Learning::FRotationFeature>(InInteractor, Name, 1, Scale);
 }
 
 void URotationObservation::SetRotationObservation(const int32 AgentId, const FRotator Rotation, const FRotator RelativeRotation)
@@ -707,7 +707,7 @@ void URotationObservation::SetRotationObservation(const int32 AgentId, const FRo
 
 void URotationObservation::SetRotationObservationFromQuat(const int32 AgentId, const FQuat Rotation, const FQuat RelativeRotation)
 {
-	if (!AgentType->HasAgent(AgentId))
+	if (!Interactor->HasAgent(AgentId))
 	{
 		UE_LOG(LogLearning, Error, TEXT("%s: AgentId %d not found in the agents set."), *GetName(), AgentId);
 		return;
@@ -728,7 +728,7 @@ void URotationObservation::VisualLog(const UE::Learning::FIndexSet Instances) co
 
 	for (const int32 Instance : Instances)
 	{
-		if (const AActor* Actor = Cast<AActor>(AgentType->GetAgent(Instance)))
+		if (const AActor* Actor = Cast<AActor>(Interactor->GetAgent(Instance)))
 		{
 			const FQuat Rotation = RotationView[Instance][0];
 			const FRotator Rotator = Rotation.Rotator();
@@ -756,7 +756,7 @@ void URotationObservation::VisualLog(const UE::Learning::FIndexSet Instances) co
 }
 #endif
 
-URotationArrayObservation* URotationArrayObservation::AddRotationArrayObservation(ULearningAgentsType* InAgentType, const FName Name, const int32 RotationNum, const float Scale)
+URotationArrayObservation* URotationArrayObservation::AddRotationArrayObservation(ULearningAgentsInteractor* InInteractor, const FName Name, const int32 RotationNum, const float Scale)
 {
 	if (RotationNum < 1)
 	{
@@ -764,12 +764,12 @@ URotationArrayObservation* URotationArrayObservation::AddRotationArrayObservatio
 		return nullptr;
 	}
 
-	return UE::Learning::Agents::Observations::Private::AddObservation<URotationArrayObservation, UE::Learning::FRotationFeature>(InAgentType, Name, RotationNum, Scale);
+	return UE::Learning::Agents::Observations::Private::AddObservation<URotationArrayObservation, UE::Learning::FRotationFeature>(InInteractor, Name, RotationNum, Scale);
 }
 
 void URotationArrayObservation::SetRotationArrayObservation(const int32 AgentId, const TArray<FRotator>& Rotations, const FRotator RelativeRotation)
 {
-	if (!AgentType->HasAgent(AgentId))
+	if (!Interactor->HasAgent(AgentId))
 	{
 		UE_LOG(LogLearning, Error, TEXT("%s: AgentId %d not found in the agents set."), *GetName(), AgentId);
 		return;
@@ -794,7 +794,7 @@ void URotationArrayObservation::SetRotationArrayObservation(const int32 AgentId,
 
 void URotationArrayObservation::SetRotationArrayObservationFromQuats(const int32 AgentId, const TArray<FQuat>& Rotations, const FQuat RelativeRotation)
 {
-	if (!AgentType->HasAgent(AgentId))
+	if (!Interactor->HasAgent(AgentId))
 	{
 		UE_LOG(LogLearning, Error, TEXT("%s: AgentId %d not found in the agents set."), *GetName(), AgentId);
 		return;
@@ -826,7 +826,7 @@ void URotationArrayObservation::VisualLog(const UE::Learning::FIndexSet Instance
 
 	for (const int32 Instance : Instances)
 	{
-		if (const AActor* Actor = Cast<AActor>(AgentType->GetAgent(Instance)))
+		if (const AActor* Actor = Cast<AActor>(Interactor->GetAgent(Instance)))
 		{
 			const FQuat RelativeRotation = RelativeRotationView[Instance];
 			const FRotator RelativeRotator = RelativeRotation.Rotator();
@@ -863,14 +863,14 @@ void URotationArrayObservation::VisualLog(const UE::Learning::FIndexSet Instance
 
 //------------------------------------------------------------------
 
-UDirectionObservation* UDirectionObservation::AddDirectionObservation(ULearningAgentsType* InAgentType, const FName Name, const float Scale)
+UDirectionObservation* UDirectionObservation::AddDirectionObservation(ULearningAgentsInteractor* InInteractor, const FName Name, const float Scale)
 {
-	return UE::Learning::Agents::Observations::Private::AddObservation<UDirectionObservation, UE::Learning::FDirectionFeature>(InAgentType, Name, 1, Scale);
+	return UE::Learning::Agents::Observations::Private::AddObservation<UDirectionObservation, UE::Learning::FDirectionFeature>(InInteractor, Name, 1, Scale);
 }
 
 void UDirectionObservation::SetDirectionObservation(const int32 AgentId, const FVector Direction, const FRotator RelativeRotation)
 {
-	if (!AgentType->HasAgent(AgentId))
+	if (!Interactor->HasAgent(AgentId))
 	{
 		UE_LOG(LogLearning, Error, TEXT("%s: AgentId %d not found in the agents set."), *GetName(), AgentId);
 		return;
@@ -891,7 +891,7 @@ void UDirectionObservation::VisualLog(const UE::Learning::FIndexSet Instances) c
 
 	for (const int32 Instance : Instances)
 	{
-		if (const AActor* Actor = Cast<AActor>(AgentType->GetAgent(Instance)))
+		if (const AActor* Actor = Cast<AActor>(Interactor->GetAgent(Instance)))
 		{
 			const FVector Direction = DirectionView[Instance][0];
 			const FQuat RelativeRotation = RelativeRotationView[Instance];
@@ -923,7 +923,7 @@ void UDirectionObservation::VisualLog(const UE::Learning::FIndexSet Instances) c
 }
 #endif
 
-UDirectionArrayObservation* UDirectionArrayObservation::AddDirectionArrayObservation(ULearningAgentsType* InAgentType, const FName Name, const int32 DirectionNum, const float Scale)
+UDirectionArrayObservation* UDirectionArrayObservation::AddDirectionArrayObservation(ULearningAgentsInteractor* InInteractor, const FName Name, const int32 DirectionNum, const float Scale)
 {
 	if (DirectionNum < 1)
 	{
@@ -931,12 +931,12 @@ UDirectionArrayObservation* UDirectionArrayObservation::AddDirectionArrayObserva
 		return nullptr;
 	}
 
-	return UE::Learning::Agents::Observations::Private::AddObservation<UDirectionArrayObservation, UE::Learning::FDirectionFeature>(InAgentType, Name, DirectionNum, Scale);
+	return UE::Learning::Agents::Observations::Private::AddObservation<UDirectionArrayObservation, UE::Learning::FDirectionFeature>(InInteractor, Name, DirectionNum, Scale);
 }
 
 void UDirectionArrayObservation::SetDirectionArrayObservation(const int32 AgentId, const TArray<FVector>& Directions, const FRotator RelativeRotation)
 {
-	if (!AgentType->HasAgent(AgentId))
+	if (!Interactor->HasAgent(AgentId))
 	{
 		UE_LOG(LogLearning, Error, TEXT("%s: AgentId %d not found in the agents set."), *GetName(), AgentId);
 		return;
@@ -968,7 +968,7 @@ void UDirectionArrayObservation::VisualLog(const UE::Learning::FIndexSet Instanc
 
 	for (const int32 Instance : Instances)
 	{
-		if (const AActor* Actor = Cast<AActor>(AgentType->GetAgent(Instance)))
+		if (const AActor* Actor = Cast<AActor>(Interactor->GetAgent(Instance)))
 		{
 			const FQuat RelativeRotation = RelativeRotationView[Instance];
 
@@ -1006,14 +1006,14 @@ void UDirectionArrayObservation::VisualLog(const UE::Learning::FIndexSet Instanc
 }
 #endif
 
-UPlanarDirectionObservation* UPlanarDirectionObservation::AddPlanarDirectionObservation(ULearningAgentsType* InAgentType, const FName Name, const float Scale, const FVector Axis0, const FVector Axis1)
+UPlanarDirectionObservation* UPlanarDirectionObservation::AddPlanarDirectionObservation(ULearningAgentsInteractor* InInteractor, const FName Name, const float Scale, const FVector Axis0, const FVector Axis1)
 {
-	return UE::Learning::Agents::Observations::Private::AddObservation<UPlanarDirectionObservation, UE::Learning::FPlanarDirectionFeature>(InAgentType, Name, 1, Scale, Axis0, Axis1);
+	return UE::Learning::Agents::Observations::Private::AddObservation<UPlanarDirectionObservation, UE::Learning::FPlanarDirectionFeature>(InInteractor, Name, 1, Scale, Axis0, Axis1);
 }
 
 void UPlanarDirectionObservation::SetPlanarDirectionObservation(const int32 AgentId, const FVector Direction, const FRotator RelativeRotation)
 {
-	if (!AgentType->HasAgent(AgentId))
+	if (!Interactor->HasAgent(AgentId))
 	{
 		UE_LOG(LogLearning, Error, TEXT("%s: AgentId %d not found in the agents set."), *GetName(), AgentId);
 		return;
@@ -1034,7 +1034,7 @@ void UPlanarDirectionObservation::VisualLog(const UE::Learning::FIndexSet Instan
 
 	for (const int32 Instance : Instances)
 	{
-		if (const AActor* Actor = Cast<AActor>(AgentType->GetAgent(Instance)))
+		if (const AActor* Actor = Cast<AActor>(Interactor->GetAgent(Instance)))
 		{
 			const FVector Direction = DirectionView[Instance][0];
 			const FQuat RelativeRotation = RelativeRotationView[Instance];
@@ -1074,14 +1074,14 @@ void UPlanarDirectionObservation::VisualLog(const UE::Learning::FIndexSet Instan
 }
 #endif
 
-UPlanarDirectionArrayObservation* UPlanarDirectionArrayObservation::AddPlanarDirectionArrayObservation(ULearningAgentsType* InAgentType, const FName Name, const int32 DirectionNum, const float Scale, const FVector Axis0, const FVector Axis1)
+UPlanarDirectionArrayObservation* UPlanarDirectionArrayObservation::AddPlanarDirectionArrayObservation(ULearningAgentsInteractor* InInteractor, const FName Name, const int32 DirectionNum, const float Scale, const FVector Axis0, const FVector Axis1)
 {
-	return UE::Learning::Agents::Observations::Private::AddObservation<UPlanarDirectionArrayObservation, UE::Learning::FPlanarDirectionFeature>(InAgentType, Name, DirectionNum, Scale, Axis0, Axis1);
+	return UE::Learning::Agents::Observations::Private::AddObservation<UPlanarDirectionArrayObservation, UE::Learning::FPlanarDirectionFeature>(InInteractor, Name, DirectionNum, Scale, Axis0, Axis1);
 }
 
 void UPlanarDirectionArrayObservation::SetPlanarDirectionArrayObservation(const int32 AgentId, const TArray<FVector>& Directions, const FRotator RelativeRotation)
 {
-	if (!AgentType->HasAgent(AgentId))
+	if (!Interactor->HasAgent(AgentId))
 	{
 		UE_LOG(LogLearning, Error, TEXT("%s: AgentId %d not found in the agents set."), *GetName(), AgentId);
 		return;
@@ -1113,7 +1113,7 @@ void UPlanarDirectionArrayObservation::VisualLog(const UE::Learning::FIndexSet I
 	
 	for (const int32 Instance : Instances)
 	{
-		if (const AActor* Actor = Cast<AActor>(AgentType->GetAgent(Instance)))
+		if (const AActor* Actor = Cast<AActor>(Interactor->GetAgent(Instance)))
 		{
 			const FQuat RelativeRotation = RelativeRotationView[Instance];
 
@@ -1163,14 +1163,14 @@ void UPlanarDirectionArrayObservation::VisualLog(const UE::Learning::FIndexSet I
 
 //------------------------------------------------------------------
 
-UPositionObservation* UPositionObservation::AddPositionObservation(ULearningAgentsType* InAgentType, const FName Name, const float Scale)
+UPositionObservation* UPositionObservation::AddPositionObservation(ULearningAgentsInteractor* InInteractor, const FName Name, const float Scale)
 {
-	return UE::Learning::Agents::Observations::Private::AddObservation<UPositionObservation, UE::Learning::FPositionFeature>(InAgentType, Name, 1, Scale);
+	return UE::Learning::Agents::Observations::Private::AddObservation<UPositionObservation, UE::Learning::FPositionFeature>(InInteractor, Name, 1, Scale);
 }
 
 void UPositionObservation::SetPositionObservation(const int32 AgentId, const FVector Position, const FVector RelativePosition, const FRotator RelativeRotation)
 {
-	if (!AgentType->HasAgent(AgentId))
+	if (!Interactor->HasAgent(AgentId))
 	{
 		UE_LOG(LogLearning, Error, TEXT("%s: AgentId %d not found in the agents set."), *GetName(), AgentId);
 		return;
@@ -1224,7 +1224,7 @@ void UPositionObservation::VisualLog(const UE::Learning::FIndexSet Instances) co
 }
 #endif
 
-UPositionArrayObservation* UPositionArrayObservation::AddPositionArrayObservation(ULearningAgentsType* InAgentType, const FName Name, const int32 PositionNum, const float Scale)
+UPositionArrayObservation* UPositionArrayObservation::AddPositionArrayObservation(ULearningAgentsInteractor* InInteractor, const FName Name, const int32 PositionNum, const float Scale)
 {
 	if (PositionNum < 1)
 	{
@@ -1232,12 +1232,12 @@ UPositionArrayObservation* UPositionArrayObservation::AddPositionArrayObservatio
 		return nullptr;
 	}
 
-	return UE::Learning::Agents::Observations::Private::AddObservation<UPositionArrayObservation, UE::Learning::FPositionFeature>(InAgentType, Name, PositionNum, Scale);
+	return UE::Learning::Agents::Observations::Private::AddObservation<UPositionArrayObservation, UE::Learning::FPositionFeature>(InInteractor, Name, PositionNum, Scale);
 }
 
 void UPositionArrayObservation::SetPositionArrayObservation(const int32 AgentId, const TArray<FVector>& Positions, const FVector RelativePosition, const FRotator RelativeRotation)
 {
-	if (!AgentType->HasAgent(AgentId))
+	if (!Interactor->HasAgent(AgentId))
 	{
 		UE_LOG(LogLearning, Error, TEXT("%s: AgentId %d not found in the agents set."), *GetName(), AgentId);
 		return;
@@ -1308,14 +1308,14 @@ void UPositionArrayObservation::VisualLog(const UE::Learning::FIndexSet Instance
 }
 #endif
 
-UScalarPositionObservation* UScalarPositionObservation::AddScalarPositionObservation(ULearningAgentsType* InAgentType, const FName Name, const float Scale)
+UScalarPositionObservation* UScalarPositionObservation::AddScalarPositionObservation(ULearningAgentsInteractor* InInteractor, const FName Name, const float Scale)
 {
-	return UE::Learning::Agents::Observations::Private::AddObservation<UScalarPositionObservation, UE::Learning::FScalarPositionFeature>(InAgentType, Name, 1, Scale);
+	return UE::Learning::Agents::Observations::Private::AddObservation<UScalarPositionObservation, UE::Learning::FScalarPositionFeature>(InInteractor, Name, 1, Scale);
 }
 
 void UScalarPositionObservation::SetScalarPositionObservation(const int32 AgentId, const float Position, const float RelativePosition)
 {
-	if (!AgentType->HasAgent(AgentId))
+	if (!Interactor->HasAgent(AgentId))
 	{
 		UE_LOG(LogLearning, Error, TEXT("%s: AgentId %d not found in the agents set."), *GetName(), AgentId);
 		return;
@@ -1341,7 +1341,7 @@ void UScalarPositionObservation::VisualLog(const UE::Learning::FIndexSet Instanc
 
 	for (const int32 Instance : Instances)
 	{
-		if (const AActor* Actor = Cast<AActor>(AgentType->GetAgent(Instance)))
+		if (const AActor* Actor = Cast<AActor>(Interactor->GetAgent(Instance)))
 		{
 			UE_LEARNING_AGENTS_VLOG_STRING(this, LogLearning, Display,
 				Actor->GetActorLocation(),
@@ -1358,7 +1358,7 @@ void UScalarPositionObservation::VisualLog(const UE::Learning::FIndexSet Instanc
 }
 #endif
 
-UScalarPositionArrayObservation* UScalarPositionArrayObservation::AddScalarPositionArrayObservation(ULearningAgentsType* InAgentType, const FName Name, const int32 PositionNum, const float Scale)
+UScalarPositionArrayObservation* UScalarPositionArrayObservation::AddScalarPositionArrayObservation(ULearningAgentsInteractor* InInteractor, const FName Name, const int32 PositionNum, const float Scale)
 {
 	if (PositionNum < 1)
 	{
@@ -1366,12 +1366,12 @@ UScalarPositionArrayObservation* UScalarPositionArrayObservation::AddScalarPosit
 		return nullptr;
 	}
 
-	return UE::Learning::Agents::Observations::Private::AddObservation<UScalarPositionArrayObservation, UE::Learning::FScalarPositionFeature>(InAgentType, Name, PositionNum, Scale);
+	return UE::Learning::Agents::Observations::Private::AddObservation<UScalarPositionArrayObservation, UE::Learning::FScalarPositionFeature>(InInteractor, Name, PositionNum, Scale);
 }
 
 void UScalarPositionArrayObservation::SetScalarPositionArrayObservation(const int32 AgentId, const TArray<float>& Positions, const float RelativePosition)
 {
-	if (!AgentType->HasAgent(AgentId))
+	if (!Interactor->HasAgent(AgentId))
 	{
 		UE_LOG(LogLearning, Error, TEXT("%s: AgentId %d not found in the agents set."), *GetName(), AgentId);
 		return;
@@ -1392,7 +1392,7 @@ void UScalarPositionArrayObservation::SetScalarPositionArrayObservation(const in
 
 void UScalarPositionArrayObservation::SetScalarPositionArrayObservationWithAxis(const int32 AgentId, const TArray<FVector>& Positions, const FVector RelativePosition, const FVector Axis)
 {
-	if (!AgentType->HasAgent(AgentId))
+	if (!Interactor->HasAgent(AgentId))
 	{
 		UE_LOG(LogLearning, Error, TEXT("%s: AgentId %d not found in the agents set."), *GetName(), AgentId);
 		return;
@@ -1427,7 +1427,7 @@ void UScalarPositionArrayObservation::VisualLog(const UE::Learning::FIndexSet In
 
 	for (const int32 Instance : Instances)
 	{
-		if (const AActor* Actor = Cast<AActor>(AgentType->GetAgent(Instance)))
+		if (const AActor* Actor = Cast<AActor>(Interactor->GetAgent(Instance)))
 		{
 			UE_LEARNING_AGENTS_VLOG_STRING(this, LogLearning, Display,
 				Actor->GetActorLocation(),
@@ -1443,14 +1443,14 @@ void UScalarPositionArrayObservation::VisualLog(const UE::Learning::FIndexSet In
 }
 #endif
 
-UPlanarPositionObservation* UPlanarPositionObservation::AddPlanarPositionObservation(ULearningAgentsType* InAgentType, const FName Name, const float Scale, const FVector Axis0, const FVector Axis1)
+UPlanarPositionObservation* UPlanarPositionObservation::AddPlanarPositionObservation(ULearningAgentsInteractor* InInteractor, const FName Name, const float Scale, const FVector Axis0, const FVector Axis1)
 {
-	return UE::Learning::Agents::Observations::Private::AddObservation<UPlanarPositionObservation, UE::Learning::FPlanarPositionFeature>(InAgentType, Name, 1, Scale, Axis0, Axis1);
+	return UE::Learning::Agents::Observations::Private::AddObservation<UPlanarPositionObservation, UE::Learning::FPlanarPositionFeature>(InInteractor, Name, 1, Scale, Axis0, Axis1);
 }
 
 void UPlanarPositionObservation::SetPlanarPositionObservation(const int32 AgentId, const FVector Position, const FVector RelativePosition, const FRotator RelativeRotation)
 {
-	if (!AgentType->HasAgent(AgentId))
+	if (!Interactor->HasAgent(AgentId))
 	{
 		UE_LOG(LogLearning, Error, TEXT("%s: AgentId %d not found in the agents set."), *GetName(), AgentId);
 		return;
@@ -1512,7 +1512,7 @@ void UPlanarPositionObservation::VisualLog(const UE::Learning::FIndexSet Instanc
 }
 #endif
 
-UPlanarPositionArrayObservation* UPlanarPositionArrayObservation::AddPlanarPositionArrayObservation(ULearningAgentsType* InAgentType, const FName Name, const int32 PositionNum, const float Scale, const FVector Axis0, const FVector Axis1)
+UPlanarPositionArrayObservation* UPlanarPositionArrayObservation::AddPlanarPositionArrayObservation(ULearningAgentsInteractor* InInteractor, const FName Name, const int32 PositionNum, const float Scale, const FVector Axis0, const FVector Axis1)
 {
 	if (PositionNum < 1)
 	{
@@ -1520,12 +1520,12 @@ UPlanarPositionArrayObservation* UPlanarPositionArrayObservation::AddPlanarPosit
 		return nullptr;
 	}
 
-	return UE::Learning::Agents::Observations::Private::AddObservation<UPlanarPositionArrayObservation, UE::Learning::FPlanarPositionFeature>(InAgentType, Name, PositionNum, Scale, Axis0, Axis1);
+	return UE::Learning::Agents::Observations::Private::AddObservation<UPlanarPositionArrayObservation, UE::Learning::FPlanarPositionFeature>(InInteractor, Name, PositionNum, Scale, Axis0, Axis1);
 }
 
 void UPlanarPositionArrayObservation::SetPlanarPositionArrayObservation(const int32 AgentId, const TArray<FVector>& Positions, const FVector RelativePosition, const FRotator RelativeRotation)
 {
-	if (!AgentType->HasAgent(AgentId))
+	if (!Interactor->HasAgent(AgentId))
 	{
 		UE_LOG(LogLearning, Error, TEXT("%s: AgentId %d not found in the agents set."), *GetName(), AgentId);
 		return;
@@ -1607,14 +1607,14 @@ void UPlanarPositionArrayObservation::VisualLog(const UE::Learning::FIndexSet In
 
 //------------------------------------------------------------------
 
-UVelocityObservation* UVelocityObservation::AddVelocityObservation(ULearningAgentsType* InAgentType, const FName Name, const float Scale)
+UVelocityObservation* UVelocityObservation::AddVelocityObservation(ULearningAgentsInteractor* InInteractor, const FName Name, const float Scale)
 {
-	return UE::Learning::Agents::Observations::Private::AddObservation<UVelocityObservation, UE::Learning::FVelocityFeature>(InAgentType, Name, 1, Scale);
+	return UE::Learning::Agents::Observations::Private::AddObservation<UVelocityObservation, UE::Learning::FVelocityFeature>(InInteractor, Name, 1, Scale);
 }
 
 void UVelocityObservation::SetVelocityObservation(const int32 AgentId, const FVector Velocity, const FRotator RelativeRotation)
 {
-	if (!AgentType->HasAgent(AgentId))
+	if (!Interactor->HasAgent(AgentId))
 	{
 		UE_LOG(LogLearning, Error, TEXT("%s: AgentId %d not found in the agents set."), *GetName(), AgentId);
 		return;
@@ -1635,7 +1635,7 @@ void UVelocityObservation::VisualLog(const UE::Learning::FIndexSet Instances) co
 
 	for (const int32 Instance : Instances)
 	{
-		if (const AActor* Actor = Cast<AActor>(AgentType->GetAgent(Instance)))
+		if (const AActor* Actor = Cast<AActor>(Interactor->GetAgent(Instance)))
 		{
 			const FVector Velocity = VelocityView[Instance][0];
 			const FQuat RelativeRotation = RelativeRotationView[Instance];
@@ -1667,7 +1667,7 @@ void UVelocityObservation::VisualLog(const UE::Learning::FIndexSet Instances) co
 }
 #endif
 
-UVelocityArrayObservation* UVelocityArrayObservation::AddVelocityArrayObservation(ULearningAgentsType* InAgentType, const FName Name, const int32 VelocityNum, const float Scale)
+UVelocityArrayObservation* UVelocityArrayObservation::AddVelocityArrayObservation(ULearningAgentsInteractor* InInteractor, const FName Name, const int32 VelocityNum, const float Scale)
 {
 	if (VelocityNum < 1)
 	{
@@ -1675,12 +1675,12 @@ UVelocityArrayObservation* UVelocityArrayObservation::AddVelocityArrayObservatio
 		return nullptr;
 	}
 
-	return UE::Learning::Agents::Observations::Private::AddObservation<UVelocityArrayObservation, UE::Learning::FVelocityFeature>(InAgentType, Name, VelocityNum, Scale);
+	return UE::Learning::Agents::Observations::Private::AddObservation<UVelocityArrayObservation, UE::Learning::FVelocityFeature>(InInteractor, Name, VelocityNum, Scale);
 }
 
 void UVelocityArrayObservation::SetVelocityArrayObservation(const int32 AgentId, const TArray<FVector>& Velocities, const FRotator RelativeRotation)
 {
-	if (!AgentType->HasAgent(AgentId))
+	if (!Interactor->HasAgent(AgentId))
 	{
 		UE_LOG(LogLearning, Error, TEXT("%s: AgentId %d not found in the agents set."), *GetName(), AgentId);
 		return;
@@ -1712,7 +1712,7 @@ void UVelocityArrayObservation::VisualLog(const UE::Learning::FIndexSet Instance
 
 	for (const int32 Instance : Instances)
 	{
-		if (const AActor* Actor = Cast<AActor>(AgentType->GetAgent(Instance)))
+		if (const AActor* Actor = Cast<AActor>(Interactor->GetAgent(Instance)))
 		{
 			const FQuat RelativeRotation = RelativeRotationView[Instance];
 				
@@ -1750,14 +1750,14 @@ void UVelocityArrayObservation::VisualLog(const UE::Learning::FIndexSet Instance
 }
 #endif
 
-UScalarVelocityObservation* UScalarVelocityObservation::AddScalarVelocityObservation(ULearningAgentsType* InAgentType, const FName Name, const float Scale)
+UScalarVelocityObservation* UScalarVelocityObservation::AddScalarVelocityObservation(ULearningAgentsInteractor* InInteractor, const FName Name, const float Scale)
 {
-	return UE::Learning::Agents::Observations::Private::AddObservation<UScalarVelocityObservation, UE::Learning::FScalarVelocityFeature>(InAgentType, Name, 1, Scale);
+	return UE::Learning::Agents::Observations::Private::AddObservation<UScalarVelocityObservation, UE::Learning::FScalarVelocityFeature>(InInteractor, Name, 1, Scale);
 }
 
 void UScalarVelocityObservation::SetScalarVelocityObservation(const int32 AgentId, const float Velocity)
 {
-	if (!AgentType->HasAgent(AgentId))
+	if (!Interactor->HasAgent(AgentId))
 	{
 		UE_LOG(LogLearning, Error, TEXT("%s: AgentId %d not found in the agents set."), *GetName(), AgentId);
 		return;
@@ -1781,7 +1781,7 @@ void UScalarVelocityObservation::VisualLog(const UE::Learning::FIndexSet Instanc
 
 	for (const int32 Instance : Instances)
 	{
-		if (const AActor* Actor = Cast<AActor>(AgentType->GetAgent(Instance)))
+		if (const AActor* Actor = Cast<AActor>(Interactor->GetAgent(Instance)))
 		{
 			UE_LEARNING_AGENTS_VLOG_STRING(this, LogLearning, Display,
 				Actor->GetActorLocation(),
@@ -1797,7 +1797,7 @@ void UScalarVelocityObservation::VisualLog(const UE::Learning::FIndexSet Instanc
 #endif
 
 
-UScalarVelocityArrayObservation* UScalarVelocityArrayObservation::AddScalarVelocityArrayObservation(ULearningAgentsType* InAgentType, const FName Name, const int32 VelocityNum, const float Scale)
+UScalarVelocityArrayObservation* UScalarVelocityArrayObservation::AddScalarVelocityArrayObservation(ULearningAgentsInteractor* InInteractor, const FName Name, const int32 VelocityNum, const float Scale)
 {
 	if (VelocityNum < 1)
 	{
@@ -1805,12 +1805,12 @@ UScalarVelocityArrayObservation* UScalarVelocityArrayObservation::AddScalarVeloc
 		return nullptr;
 	}
 
-	return UE::Learning::Agents::Observations::Private::AddObservation<UScalarVelocityArrayObservation, UE::Learning::FScalarVelocityFeature>(InAgentType, Name, VelocityNum, Scale);
+	return UE::Learning::Agents::Observations::Private::AddObservation<UScalarVelocityArrayObservation, UE::Learning::FScalarVelocityFeature>(InInteractor, Name, VelocityNum, Scale);
 }
 
 void UScalarVelocityArrayObservation::SetScalarVelocityArrayObservation(const int32 AgentId, const TArray<float>& Velocities)
 {
-	if (!AgentType->HasAgent(AgentId))
+	if (!Interactor->HasAgent(AgentId))
 	{
 		UE_LOG(LogLearning, Error, TEXT("%s: AgentId %d not found in the agents set."), *GetName(), AgentId);
 		return;
@@ -1829,7 +1829,7 @@ void UScalarVelocityArrayObservation::SetScalarVelocityArrayObservation(const in
 
 void UScalarVelocityArrayObservation::SetScalarVelocityArrayObservationWithAxis(const int32 AgentId, const TArray<FVector>& Velocities, const FVector Axis)
 {
-	if (!AgentType->HasAgent(AgentId))
+	if (!Interactor->HasAgent(AgentId))
 	{
 		UE_LOG(LogLearning, Error, TEXT("%s: AgentId %d not found in the agents set."), *GetName(), AgentId);
 		return;
@@ -1861,7 +1861,7 @@ void UScalarVelocityArrayObservation::VisualLog(const UE::Learning::FIndexSet In
 
 	for (const int32 Instance : Instances)
 	{
-		if (const AActor* Actor = Cast<AActor>(AgentType->GetAgent(Instance)))
+		if (const AActor* Actor = Cast<AActor>(Interactor->GetAgent(Instance)))
 		{
 			UE_LEARNING_AGENTS_VLOG_STRING(this, LogLearning, Display,
 				Actor->GetActorLocation(),
@@ -1877,14 +1877,14 @@ void UScalarVelocityArrayObservation::VisualLog(const UE::Learning::FIndexSet In
 #endif
 
 
-UPlanarVelocityObservation* UPlanarVelocityObservation::AddPlanarVelocityObservation(ULearningAgentsType* InAgentType, const FName Name, const float Scale, const FVector Axis0, const FVector Axis1)
+UPlanarVelocityObservation* UPlanarVelocityObservation::AddPlanarVelocityObservation(ULearningAgentsInteractor* InInteractor, const FName Name, const float Scale, const FVector Axis0, const FVector Axis1)
 {
-	return UE::Learning::Agents::Observations::Private::AddObservation<UPlanarVelocityObservation, UE::Learning::FPlanarVelocityFeature>(InAgentType, Name, 1, Scale, Axis0, Axis1);
+	return UE::Learning::Agents::Observations::Private::AddObservation<UPlanarVelocityObservation, UE::Learning::FPlanarVelocityFeature>(InInteractor, Name, 1, Scale, Axis0, Axis1);
 }
 
 void UPlanarVelocityObservation::SetPlanarVelocityObservation(const int32 AgentId, const FVector Velocity, const FRotator RelativeRotation)
 {
-	if (!AgentType->HasAgent(AgentId))
+	if (!Interactor->HasAgent(AgentId))
 	{
 		UE_LOG(LogLearning, Error, TEXT("%s: AgentId %d not found in the agents set."), *GetName(), AgentId);
 		return;
@@ -1905,7 +1905,7 @@ void UPlanarVelocityObservation::VisualLog(const UE::Learning::FIndexSet Instanc
 
 	for (const int32 Instance : Instances)
 	{
-		if (const AActor* Actor = Cast<AActor>(AgentType->GetAgent(Instance)))
+		if (const AActor* Actor = Cast<AActor>(Interactor->GetAgent(Instance)))
 		{
 			const FVector Velocity = VelocityView[Instance][0];
 			const FQuat RelativeRotation = RelativeRotationView[Instance];
@@ -1946,7 +1946,7 @@ void UPlanarVelocityObservation::VisualLog(const UE::Learning::FIndexSet Instanc
 #endif
 
 
-UPlanarVelocityArrayObservation* UPlanarVelocityArrayObservation::AddPlanarVelocityArrayObservation(ULearningAgentsType* InAgentType, const FName Name, const int32 VelocityNum, const float Scale, const FVector Axis0, const FVector Axis1)
+UPlanarVelocityArrayObservation* UPlanarVelocityArrayObservation::AddPlanarVelocityArrayObservation(ULearningAgentsInteractor* InInteractor, const FName Name, const int32 VelocityNum, const float Scale, const FVector Axis0, const FVector Axis1)
 {
 	if (VelocityNum < 1)
 	{
@@ -1954,12 +1954,12 @@ UPlanarVelocityArrayObservation* UPlanarVelocityArrayObservation::AddPlanarVeloc
 		return nullptr;
 	}
 
-	return UE::Learning::Agents::Observations::Private::AddObservation<UPlanarVelocityArrayObservation, UE::Learning::FPlanarVelocityFeature>(InAgentType, Name, VelocityNum, Scale, Axis0, Axis1);
+	return UE::Learning::Agents::Observations::Private::AddObservation<UPlanarVelocityArrayObservation, UE::Learning::FPlanarVelocityFeature>(InInteractor, Name, VelocityNum, Scale, Axis0, Axis1);
 }
 
 void UPlanarVelocityArrayObservation::SetPlanarVelocityArrayObservation(const int32 AgentId, const TArray<FVector>& Velocities, const FRotator RelativeRotation)
 {
-	if (!AgentType->HasAgent(AgentId))
+	if (!Interactor->HasAgent(AgentId))
 	{
 		UE_LOG(LogLearning, Error, TEXT("%s: AgentId %d not found in the agents set."), *GetName(), AgentId);
 		return;
@@ -1991,7 +1991,7 @@ void UPlanarVelocityArrayObservation::VisualLog(const UE::Learning::FIndexSet In
 
 	for (const int32 Instance : Instances)
 	{
-		if (const AActor* Actor = Cast<AActor>(AgentType->GetAgent(Instance)))
+		if (const AActor* Actor = Cast<AActor>(Interactor->GetAgent(Instance)))
 		{
 			const FQuat RelativeRotation = RelativeRotationView[Instance];
 
@@ -2039,14 +2039,14 @@ void UPlanarVelocityArrayObservation::VisualLog(const UE::Learning::FIndexSet In
 
 //------------------------------------------------------------------
 
-UAngularVelocityObservation* UAngularVelocityObservation::AddAngularVelocityObservation(ULearningAgentsType* InAgentType, const FName Name, const float Scale)
+UAngularVelocityObservation* UAngularVelocityObservation::AddAngularVelocityObservation(ULearningAgentsInteractor* InInteractor, const FName Name, const float Scale)
 {
-	return UE::Learning::Agents::Observations::Private::AddObservation<UAngularVelocityObservation, UE::Learning::FAngularVelocityFeature>(InAgentType, Name, 1, Scale);
+	return UE::Learning::Agents::Observations::Private::AddObservation<UAngularVelocityObservation, UE::Learning::FAngularVelocityFeature>(InInteractor, Name, 1, Scale);
 }
 
 void UAngularVelocityObservation::SetAngularVelocityObservation(const int32 AgentId, const FVector AngularVelocity, const FRotator RelativeRotation)
 {
-	if (!AgentType->HasAgent(AgentId))
+	if (!Interactor->HasAgent(AgentId))
 	{
 		UE_LOG(LogLearning, Error, TEXT("%s: AgentId %d not found in the agents set."), *GetName(), AgentId);
 		return;
@@ -2067,7 +2067,7 @@ void UAngularVelocityObservation::VisualLog(const UE::Learning::FIndexSet Instan
 
 	for (const int32 Instance : Instances)
 	{
-		if (const AActor* Actor = Cast<AActor>(AgentType->GetAgent(Instance)))
+		if (const AActor* Actor = Cast<AActor>(Interactor->GetAgent(Instance)))
 		{
 			const FVector AngularVelocity = AngularVelocityView[Instance][0];
 			const FQuat RelativeRotation = RelativeRotationView[Instance];
@@ -2099,7 +2099,7 @@ void UAngularVelocityObservation::VisualLog(const UE::Learning::FIndexSet Instan
 }
 #endif
 
-UAngularVelocityArrayObservation* UAngularVelocityArrayObservation::AddAngularVelocityArrayObservation(ULearningAgentsType* InAgentType, const FName Name, const int32 AngularVelocityNum, const float Scale)
+UAngularVelocityArrayObservation* UAngularVelocityArrayObservation::AddAngularVelocityArrayObservation(ULearningAgentsInteractor* InInteractor, const FName Name, const int32 AngularVelocityNum, const float Scale)
 {
 	if (AngularVelocityNum < 1)
 	{
@@ -2107,12 +2107,12 @@ UAngularVelocityArrayObservation* UAngularVelocityArrayObservation::AddAngularVe
 		return nullptr;
 	}
 
-	return UE::Learning::Agents::Observations::Private::AddObservation<UAngularVelocityArrayObservation, UE::Learning::FAngularVelocityFeature>(InAgentType, Name, AngularVelocityNum, Scale);
+	return UE::Learning::Agents::Observations::Private::AddObservation<UAngularVelocityArrayObservation, UE::Learning::FAngularVelocityFeature>(InInteractor, Name, AngularVelocityNum, Scale);
 }
 
 void UAngularVelocityArrayObservation::SetAngularVelocityArrayObservation(const int32 AgentId, const TArray<FVector>& AngularVelocities, const FRotator RelativeRotation)
 {
-	if (!AgentType->HasAgent(AgentId))
+	if (!Interactor->HasAgent(AgentId))
 	{
 		UE_LOG(LogLearning, Error, TEXT("%s: AgentId %d not found in the agents set."), *GetName(), AgentId);
 		return;
@@ -2144,7 +2144,7 @@ void UAngularVelocityArrayObservation::VisualLog(const UE::Learning::FIndexSet I
 
 	for (const int32 Instance : Instances)
 	{
-		if (const AActor* Actor = Cast<AActor>(AgentType->GetAgent(Instance)))
+		if (const AActor* Actor = Cast<AActor>(Interactor->GetAgent(Instance)))
 		{
 			const FQuat RelativeRotation = RelativeRotationView[Instance];
 
@@ -2182,14 +2182,14 @@ void UAngularVelocityArrayObservation::VisualLog(const UE::Learning::FIndexSet I
 }
 #endif
 
-UScalarAngularVelocityObservation* UScalarAngularVelocityObservation::AddScalarAngularVelocityObservation(ULearningAgentsType* InAgentType, const FName Name, const float Scale)
+UScalarAngularVelocityObservation* UScalarAngularVelocityObservation::AddScalarAngularVelocityObservation(ULearningAgentsInteractor* InInteractor, const FName Name, const float Scale)
 {
-	return UE::Learning::Agents::Observations::Private::AddObservation<UScalarAngularVelocityObservation, UE::Learning::FScalarAngularVelocityFeature>(InAgentType, Name, 1, Scale);
+	return UE::Learning::Agents::Observations::Private::AddObservation<UScalarAngularVelocityObservation, UE::Learning::FScalarAngularVelocityFeature>(InInteractor, Name, 1, Scale);
 }
 
 void UScalarAngularVelocityObservation::SetScalarAngularVelocityObservation(const int32 AgentId, const float AngularVelocity)
 {
-	if (!AgentType->HasAgent(AgentId))
+	if (!Interactor->HasAgent(AgentId))
 	{
 		UE_LOG(LogLearning, Error, TEXT("%s: AgentId %d not found in the agents set."), *GetName(), AgentId);
 		return;
@@ -2213,7 +2213,7 @@ void UScalarAngularVelocityObservation::VisualLog(const UE::Learning::FIndexSet
 
 	for (const int32 Instance : Instances)
 	{
-		if (const AActor* Actor = Cast<AActor>(AgentType->GetAgent(Instance)))
+		if (const AActor* Actor = Cast<AActor>(Interactor->GetAgent(Instance)))
 		{
 			UE_LEARNING_AGENTS_VLOG_STRING(this, LogLearning, Display,
 				Actor->GetActorLocation(),
@@ -2229,7 +2229,7 @@ void UScalarAngularVelocityObservation::VisualLog(const UE::Learning::FIndexSet
 #endif
 
 
-UScalarAngularVelocityArrayObservation* UScalarAngularVelocityArrayObservation::AddScalarAngularVelocityArrayObservation(ULearningAgentsType* InAgentType, const FName Name, const int32 AngularVelocityNum, const float Scale)
+UScalarAngularVelocityArrayObservation* UScalarAngularVelocityArrayObservation::AddScalarAngularVelocityArrayObservation(ULearningAgentsInteractor* InInteractor, const FName Name, const int32 AngularVelocityNum, const float Scale)
 {
 	if (AngularVelocityNum < 1)
 	{
@@ -2237,12 +2237,12 @@ UScalarAngularVelocityArrayObservation* UScalarAngularVelocityArrayObservation::
 		return nullptr;
 	}
 
-	return UE::Learning::Agents::Observations::Private::AddObservation<UScalarAngularVelocityArrayObservation, UE::Learning::FScalarAngularVelocityFeature>(InAgentType, Name, AngularVelocityNum, Scale);
+	return UE::Learning::Agents::Observations::Private::AddObservation<UScalarAngularVelocityArrayObservation, UE::Learning::FScalarAngularVelocityFeature>(InInteractor, Name, AngularVelocityNum, Scale);
 }
 
 void UScalarAngularVelocityArrayObservation::SetScalarAngularVelocityArrayObservation(const int32 AgentId, const TArray<float>& AngularVelocities)
 {
-	if (!AgentType->HasAgent(AgentId))
+	if (!Interactor->HasAgent(AgentId))
 	{
 		UE_LOG(LogLearning, Error, TEXT("%s: AgentId %d not found in the agents set."), *GetName(), AgentId);
 		return;
@@ -2261,7 +2261,7 @@ void UScalarAngularVelocityArrayObservation::SetScalarAngularVelocityArrayObserv
 
 void UScalarAngularVelocityArrayObservation::SetScalarAngularVelocityArrayObservationWithAxis(const int32 AgentId, const TArray<FVector>& AngularVelocities, const FVector Axis)
 {
-	if (!AgentType->HasAgent(AgentId))
+	if (!Interactor->HasAgent(AgentId))
 	{
 		UE_LOG(LogLearning, Error, TEXT("%s: AgentId %d not found in the agents set."), *GetName(), AgentId);
 		return;
@@ -2291,7 +2291,7 @@ void UScalarAngularVelocityArrayObservation::VisualLog(const UE::Learning::FInde
 
 	for (const int32 Instance : Instances)
 	{
-		if (const AActor* Actor = Cast<AActor>(AgentType->GetAgent(Instance)))
+		if (const AActor* Actor = Cast<AActor>(Interactor->GetAgent(Instance)))
 		{
 			UE_LEARNING_AGENTS_VLOG_STRING(this, LogLearning, Display,
 				Actor->GetActorLocation(),
diff --git a/Engine/Plugins/Experimental/LearningAgents/Source/LearningAgents/Private/LearningAgentsPolicy.cpp b/Engine/Plugins/Experimental/LearningAgents/Source/LearningAgents/Private/LearningAgentsPolicy.cpp
--- a/Engine/Plugins/Experimental/LearningAgents/Source/LearningAgents/Private/LearningAgentsPolicy.cpp
+++ b/Engine/Plugins/Experimental/LearningAgents/Source/LearningAgents/Private/LearningAgentsPolicy.cpp
@@ -3,7 +3,7 @@
 #include "LearningAgentsPolicy.h"
 
 #include "LearningAgentsManager.h"
-#include "LearningAgentsType.h"
+#include "LearningAgentsInteractor.h"
 #include "LearningFeatureObject.h"
 #include "LearningNeuralNetwork.h"
 #include "LearningNeuralNetworkObject.h"
@@ -18,7 +18,7 @@ ULearningAgentsPolicy::ULearningAgentsPolicy() : ULearningAgentsManagerComponent
 ULearningAgentsPolicy::ULearningAgentsPolicy(FVTableHelper& Helper) : ULearningAgentsPolicy() {}
 ULearningAgentsPolicy::~ULearningAgentsPolicy() {}
 
-void ULearningAgentsPolicy::SetupPolicy(ALearningAgentsManager* InAgentManager, ULearningAgentsType* InAgentType, const FLearningAgentsPolicySettings& PolicySettings)
+void ULearningAgentsPolicy::SetupPolicy(ALearningAgentsManager* InAgentManager, ULearningAgentsInteractor* InInteractor, const FLearningAgentsPolicySettings& PolicySettings)
 {
 	if (IsSetup())
 	{
@@ -40,26 +40,26 @@ void ULearningAgentsPolicy::SetupPolicy(ALearningAgentsManager* InAgentManager,
 
 	AgentManager = InAgentManager;
 
-	if (!InAgentType)
+	if (!InInteractor)
 	{
-		UE_LOG(LogLearning, Error, TEXT("%s: InAgentType is nullptr."), *GetName());
+		UE_LOG(LogLearning, Error, TEXT("%s: InInteractor is nullptr."), *GetName());
 		return;
 	}
 
-	if (!InAgentType->IsSetup())
+	if (!InInteractor->IsSetup())
 	{
-		UE_LOG(LogLearning, Error, TEXT("%s: %s's Setup must be run before it can be used."), *GetName(), *InAgentType->GetName());
+		UE_LOG(LogLearning, Error, TEXT("%s: %s's Setup must be run before it can be used."), *GetName(), *InInteractor->GetName());
 		return;
 	}
 
-	AgentType = InAgentType;
+	Interactor = InInteractor;
 
 	// Setup Neural Network
 	Network = NewObject<ULearningAgentsNeuralNetwork>(this, TEXT("PolicyNetwork"));
 	Network->NeuralNetwork = MakeShared<UE::Learning::FNeuralNetwork>();
 	Network->NeuralNetwork->Resize(
-		AgentType->GetObservationFeature().DimNum(),
-		2 * AgentType->GetActionFeature().DimNum(),
+		Interactor->GetObservationFeature().DimNum(),
+		2 * Interactor->GetActionFeature().DimNum(),
 		PolicySettings.HiddenLayerSize,
 		PolicySettings.LayerNum);
 	Network->NeuralNetwork->ActivationFunction = UE::Learning::Agents::GetActivationFunction(PolicySettings.ActivationFunction);
@@ -78,7 +78,7 @@ void ULearningAgentsPolicy::SetupPolicy(ALearningAgentsManager* InAgentManager,
 		PolicySettings.ActionNoiseSeed,
 		PolicyFunctionSettings);
 
-	AgentManager->GetInstanceData()->Link(AgentType->GetObservationFeature().FeatureHandle, PolicyObject->InputHandle);
+	AgentManager->GetInstanceData()->Link(Interactor->GetObservationFeature().FeatureHandle, PolicyObject->InputHandle);
 
 	InitialActionNoiseScale = PolicySettings.InitialActionNoiseScale;
 	
@@ -244,7 +244,7 @@ void ULearningAgentsPolicy::EvaluatePolicy()
 	// there may be multiple writers to the action feature vector handle and so therefore
 	// the handles cannot be linked and we need to do the copy manually
 	UE::Learning::Array::Copy(
-		PolicyObject->InstanceData->View(AgentType->GetActionFeature().FeatureHandle),
+		PolicyObject->InstanceData->View(Interactor->GetActionFeature().FeatureHandle),
 		PolicyObject->InstanceData->ConstView(PolicyObject->OutputHandle),
 		AddedAgentIds);
 
@@ -263,9 +263,9 @@ void ULearningAgentsPolicy::RunInference()
 		return;
 	}
 
-	AgentType->EncodeObservations();
+	Interactor->EncodeObservations();
 	EvaluatePolicy();
-	AgentType->DecodeActions();
+	Interactor->DecodeActions();
 }
 float ULearningAgentsPolicy::GetAgentActionNoiseScale(const int32 AgentId) const
 {
@@ -328,7 +328,7 @@ void ULearningAgentsPolicy::VisualLog(const UE::Learning::FIndexSet Instances) c
 
 	for (const int32 Instance : Instances)
 	{
-		if (const AActor* Actor = Cast<AActor>(AgentType->GetAgent(Instance)))
+		if (const AActor* Actor = Cast<AActor>(Interactor->GetAgent(Instance)))
 		{
 
 			UE_LEARNING_AGENTS_VLOG_STRING(this, LogLearning, Display,
diff --git a/Engine/Plugins/Experimental/LearningAgents/Source/LearningAgents/Private/LearningAgentsType.cpp b/Engine/Plugins/Experimental/LearningAgents/Source/LearningAgents/Private/LearningAgentsType.cpp
deleted file mode 100644
diff --git a/Engine/Plugins/Experimental/LearningAgents/Source/LearningAgents/Public/LearningAgentsActions.h b/Engine/Plugins/Experimental/LearningAgents/Source/LearningAgents/Public/LearningAgentsActions.h
--- a/Engine/Plugins/Experimental/LearningAgents/Source/LearningAgents/Public/LearningAgentsActions.h
+++ b/Engine/Plugins/Experimental/LearningAgents/Source/LearningAgents/Public/LearningAgentsActions.h
@@ -17,7 +17,7 @@ namespace UE::Learning
 	struct FRotationVectorFeature;
 }
 
-class ULearningAgentsType;
+class ULearningAgentsInteractor;
 
 // For functions in this file, we are favoring having more verbose names such as "AddFloatAction" vs simply "Add" in 
 // order to keep it easy to find the correct function in blueprints.
@@ -37,7 +37,7 @@ class LEARNINGAGENTS_API ULearningAgentsAction : public UObject
 public:
 
 	UPROPERTY(VisibleAnywhere, Transient, Category = "LearningAgents")
-	TObjectPtr<ULearningAgentsType> AgentType;
+	TObjectPtr<ULearningAgentsInteractor> Interactor;
 
 #if UE_LEARNING_AGENTS_ENABLE_VISUAL_LOG
 	/** Color used to draw this action in the visual log */
@@ -59,17 +59,17 @@ class LEARNINGAGENTS_API UFloatAction : public ULearningAgentsAction
 public:
 
 	/**
-	* Adds a new float action to the given agent type. Call during ULearningAgentsType::SetupActions event.
-	* @param InAgentType The agent type to add this action to.
+	* Adds a new float action to the given agent interactor. Call during ULearningAgentsInteractor::SetupActions event.
+	* @param InInteractor The agent interactor to add this action to.
 	* @param Name The name of this new action. Used for debugging.
 	* @param Scale Used to normalize the data for the action.
 	* @return The newly created action.
 	*/
 	UFUNCTION(BlueprintCallable, Category = "LearningAgents")
-	static UFloatAction* AddFloatAction(ULearningAgentsType* InAgentType, const FName Name = NAME_None, const float Scale = 1.0f);
+	static UFloatAction* AddFloatAction(ULearningAgentsInteractor* InInteractor, const FName Name = NAME_None, const float Scale = 1.0f);
 
 	/**
-	* Gets the data for this action. Call during ULearningAgentsType::GetActions event.
+	* Gets the data for this action. Call during ULearningAgentsInteractor::GetActions event.
 	* @param AgentId The agent id to get data for.
 	* @return The current action value.
 	*/
@@ -101,18 +101,18 @@ class LEARNINGAGENTS_API UFloatArrayAction : public ULearningAgentsAction
 public:
 
 	/**
-	* Adds a new float array action to the given agent type. Call during ULearningAgentsType::SetupActions event.
-	* @param InAgentType The agent type to add this action to.
+	* Adds a new float array action to the given agent interactor. Call during ULearningAgentsInteractor::SetupActions event.
+	* @param InInteractor The agent interactor to add this action to.
 	* @param Name The name of this new action. Used for debugging.
 	* @param Num The number of floats in the array
 	* @param Scale Used to normalize the data for the action.
 	* @return The newly created action.
 	*/
 	UFUNCTION(BlueprintCallable, Category = "LearningAgents")
-	static UFloatArrayAction* AddFloatArrayAction(ULearningAgentsType* InAgentType, const FName Name = NAME_None, const int32 Num = 1, const float Scale = 1.0f);
+	static UFloatArrayAction* AddFloatArrayAction(ULearningAgentsInteractor* InInteractor, const FName Name = NAME_None, const int32 Num = 1, const float Scale = 1.0f);
 
 	/**
-	* Gets the data for this action. Call during ULearningAgentsType::GetActions event.
+	* Gets the data for this action. Call during ULearningAgentsInteractor::GetActions event.
 	* @param AgentId The agent id to get data for.
 	* @param OutValues The output array of floats
 	*/
@@ -146,17 +146,17 @@ class LEARNINGAGENTS_API UVectorAction : public ULearningAgentsAction
 public:
 
 	/**
-	* Adds a new vector action to the given agent type. Call during ULearningAgentsType::SetupActions event.
-	* @param InAgentType The agent type to add this action to.
+	* Adds a new vector action to the given agent interactor. Call during ULearningAgentsInteractor::SetupActions event.
+	* @param InInteractor The agent interactor to add this action to.
 	* @param Name The name of this new action. Used for debugging.
 	* @param Scale Used to normalize the data for the action.
 	* @return The newly created action.
 	*/
 	UFUNCTION(BlueprintCallable, Category = "LearningAgents")
-	static UVectorAction* AddVectorAction(ULearningAgentsType* InAgentType, const FName Name = NAME_None, const float Scale = 1.0f);
+	static UVectorAction* AddVectorAction(ULearningAgentsInteractor* InInteractor, const FName Name = NAME_None, const float Scale = 1.0f);
 
 	/**
-	* Gets the data for this action. Call during ULearningAgentsType::GetActions event.
+	* Gets the data for this action. Call during ULearningAgentsInteractor::GetActions event.
 	* @param AgentId The agent id to get data for.
 	* @return The current action values.
 	*/
@@ -188,18 +188,18 @@ class LEARNINGAGENTS_API UVectorArrayAction : public ULearningAgentsAction
 public:
 
 	/**
-	* Adds a new vector action to the given agent type. Call during ULearningAgentsType::SetupActions event.
-	* @param InAgentType The agent type to add this action to.
+	* Adds a new vector action to the given agent interactor. Call during ULearningAgentsInteractor::SetupActions event.
+	* @param InInteractor The agent interactor to add this action to.
 	* @param Name The name of this new action. Used for debugging.
 	* @param Num The number of vectors in the array
 	* @param Scale Used to normalize the data for the action.
 	* @return The newly created action.
 	*/
 	UFUNCTION(BlueprintCallable, Category = "LearningAgents")
-	static UVectorArrayAction* AddVectorArrayAction(ULearningAgentsType* InAgentType, const FName Name = NAME_None, const int32 Num = 1, const float Scale = 1.0f);
+	static UVectorArrayAction* AddVectorArrayAction(ULearningAgentsInteractor* InInteractor, const FName Name = NAME_None, const int32 Num = 1, const float Scale = 1.0f);
 
 	/**
-	* Gets the data for this action. Call during ULearningAgentsType::GetActions event.
+	* Gets the data for this action. Call during ULearningAgentsInteractor::GetActions event.
 	* @param AgentId The agent id to get data for.
 	* @param OutVectors The current action values.
 	*/
@@ -233,17 +233,17 @@ class LEARNINGAGENTS_API URotationAction : public ULearningAgentsAction
 public:
 
 	/**
-	* Adds a new rotation action to the given agent type. Call during ULearningAgentsType::SetupActions event.
-	* @param InAgentType The agent type to add this action to.
+	* Adds a new rotation action to the given agent interactor. Call during ULearningAgentsInteractor::SetupActions event.
+	* @param InInteractor The agent interactor to add this action to.
 	* @param Name The name of this new action. Used for debugging.
 	* @param Scale Used to normalize the data for the action.
 	* @return The newly created action.
 	*/
 	UFUNCTION(BlueprintCallable, Category = "LearningAgents")
-	static URotationAction* AddRotationAction(ULearningAgentsType* InAgentType, const FName Name = NAME_None, const float Scale = 180.0f);
+	static URotationAction* AddRotationAction(ULearningAgentsInteractor* InInteractor, const FName Name = NAME_None, const float Scale = 180.0f);
 
 	/**
-	* Gets the data for this action as a rotator. Call during ULearningAgentsType::GetActions event.
+	* Gets the data for this action as a rotator. Call during ULearningAgentsInteractor::GetActions event.
 	* @param AgentId The agent id to get data for.
 	* @return The current action value.
 	*/
@@ -251,7 +251,7 @@ public:
 	FRotator GetRotationAction(const int32 AgentId) const;
 
 	/**
-	* Gets the data for this action as a rotation vector. Call during ULearningAgentsType::GetActions event.
+	* Gets the data for this action as a rotation vector. Call during ULearningAgentsInteractor::GetActions event.
 	* @param AgentId The agent id to get data for.
 	* @return The current action value.
 	*/
@@ -259,7 +259,7 @@ public:
 	FVector GetRotationActionAsRotationVector(const int32 AgentId) const;
 
 	/**
-	* Gets the data for this action as a quaternion. Call during ULearningAgentsType::GetActions event.
+	* Gets the data for this action as a quaternion. Call during ULearningAgentsInteractor::GetActions event.
 	* @param AgentId The agent id to get data for.
 	* @return The current action value.
 	*/
@@ -283,18 +283,18 @@ class LEARNINGAGENTS_API URotationArrayAction : public ULearningAgentsAction
 public:
 
 	/**
-	* Adds a new rotation array action to the given agent type. Call during ULearningAgentsType::SetupActions event.
-	* @param InAgentType The agent type to add this action to.
+	* Adds a new rotation array action to the given agent interactor. Call during ULearningAgentsInteractor::SetupActions event.
+	* @param InInteractor The agent interactor to add this action to.
 	* @param Name The name of this new action. Used for debugging.
 	* @param RotationNum The number of rotations in the array.
 	* @param Scale Used to normalize the data for the action.
 	* @return The newly created action.
 	*/
 	UFUNCTION(BlueprintCallable, Category = "LearningAgents")
-	static URotationArrayAction* AddRotationArrayAction(ULearningAgentsType* InAgentType, const FName Name = NAME_None, const int32 RotationNum = 1, const float Scale = 180.0f);
+	static URotationArrayAction* AddRotationArrayAction(ULearningAgentsInteractor* InInteractor, const FName Name = NAME_None, const int32 RotationNum = 1, const float Scale = 180.0f);
 
 	/**
-	* Gets the data for this action as rotators. Call during ULearningAgentsType::GetActions event.
+	* Gets the data for this action as rotators. Call during ULearningAgentsInteractor::GetActions event.
 	* @param AgentId The agent id to get data for.
 	* @param OutRotations The current action values.
 	*/
@@ -302,7 +302,7 @@ public:
 	void GetRotationArrayAction(const int32 AgentId, TArray<FRotator>& OutRotations) const;
 
 	/**
-	* Gets the data for this action as rotation vectors. Call during ULearningAgentsType::GetActions event.
+	* Gets the data for this action as rotation vectors. Call during ULearningAgentsInteractor::GetActions event.
 	* @param AgentId The agent id to get data for.
 	* @param OutRotationVectors The current action values.
 	*/
@@ -310,7 +310,7 @@ public:
 	void GetRotationArrayActionAsRotationVectors(const int32 AgentId, TArray<FVector>& OutRotationVectors) const;
 
 	/**
-	* Gets the data for this action as quaternions. Call during ULearningAgentsType::GetActions event.
+	* Gets the data for this action as quaternions. Call during ULearningAgentsInteractor::GetActions event.
 	* @param AgentId The agent id to get data for.
 	* @param OutRotations The current action values.
 	*/
diff --git a/Engine/Plugins/Experimental/LearningAgents/Source/LearningAgents/Public/LearningAgentsController.h b/Engine/Plugins/Experimental/LearningAgents/Source/LearningAgents/Public/LearningAgentsController.h
--- a/Engine/Plugins/Experimental/LearningAgents/Source/LearningAgents/Public/LearningAgentsController.h
+++ b/Engine/Plugins/Experimental/LearningAgents/Source/LearningAgents/Public/LearningAgentsController.h
@@ -23,22 +23,22 @@ public:
 	ULearningAgentsController(FVTableHelper& Helper);
 	virtual ~ULearningAgentsController();
 
-	/** Initializes this object to be used with the given manager and agent type. */
+	/** Initializes this object to be used with the given manager and agent interactor. */
 	UFUNCTION(BlueprintCallable, Category = "LearningAgents")
-	void SetupController(ALearningAgentsManager* InAgentManager, ULearningAgentsType* InAgentType);
+	void SetupController(ALearningAgentsManager* InAgentManager, ULearningAgentsInteractor* InInteractor);
 
 	/**
 	* During this event, you should set the actions of your agents.
 	* @param AgentIds The list of agent ids to set actions for.
 	* @see LearningAgentsActions.h for the list of available actions.
-	* @see ULearningAgentsType::GetAgent to get the agent corresponding to each id.
+	* @see ULearningAgentsInteractor::GetAgent to get the agent corresponding to each id.
 	*/
 	UFUNCTION(BlueprintNativeEvent, Category = "LearningAgents")
 	void SetActions(const TArray<int32>& AgentIds);
 
 	/**
 	* Call this function when it is time to gather all the actions for your agents. This should be called roughly 
-	* whenever you are calling ULearningAgentsType::EncodeObservations. This will call this controller's SetActions event.
+	* whenever you are calling ULearningAgentsInteractor::EncodeObservations. This will call this controller's SetActions event.
 	*/
 	UFUNCTION(BlueprintCallable, Category = "LearningAgents")
 	void EncodeActions();
@@ -52,15 +52,15 @@ public:
 protected:
 
 	/**
-	* Gets the agent type associated with this component.
-	* @param AgentClass The class to cast the agent type to (in blueprint).
+	* Gets the agent interactor associated with this component.
+	* @param AgentClass The class to cast the agent interactor to (in blueprint).
 	*/
-	UFUNCTION(BlueprintPure, Category = "LearningAgents", meta = (DeterminesOutputType = "AgentTypeClass"))
-	ULearningAgentsType* GetAgentType(const TSubclassOf<ULearningAgentsType> AgentTypeClass) const;
+	UFUNCTION(BlueprintPure, Category = "LearningAgents", meta = (DeterminesOutputType = "InteractorClass"))
+	ULearningAgentsInteractor* GetInteractor(const TSubclassOf<ULearningAgentsInteractor> InteractorClass) const;
 
 private:
 
-	/** The agent type this controller is associated with. */
+	/** The agent interactor this controller is associated with. */
 	UPROPERTY(VisibleAnywhere, Transient, Category = "LearningAgents")
-	TObjectPtr<ULearningAgentsType> AgentType;
+	TObjectPtr<ULearningAgentsInteractor> Interactor;
 };
diff --git a/Engine/Plugins/Experimental/LearningAgents/Source/LearningAgents/Public/LearningAgentsCritic.h b/Engine/Plugins/Experimental/LearningAgents/Source/LearningAgents/Public/LearningAgentsCritic.h
--- a/Engine/Plugins/Experimental/LearningAgents/Source/LearningAgents/Public/LearningAgentsCritic.h
+++ b/Engine/Plugins/Experimental/LearningAgents/Source/LearningAgents/Public/LearningAgentsCritic.h
@@ -54,11 +54,11 @@ public:
 	ULearningAgentsCritic(FVTableHelper& Helper);
 	virtual ~ULearningAgentsCritic();
 
-	/** Initializes this object to be used with the given agent type and critic settings. */
+	/** Initializes this object to be used with the given agent interactor and critic settings. */
 	UFUNCTION(BlueprintCallable, Category = "LearningAgents")
 	void SetupCritic(
 		ALearningAgentsManager* InAgentManager,
-		ULearningAgentsType* InAgentType,
+		ULearningAgentsInteractor* InInteractor,
 		const FLearningAgentsCriticSettings& CriticSettings = FLearningAgentsCriticSettings());
 
 // ----- Load / Save -----
@@ -99,7 +99,7 @@ public:
 
 	/**
 	* Calling this function will run the underlying neural network on the previously buffered observations to populate
-	* the output value buffer. This should be called after the corresponding agent type's EncodeObservations.
+	* the output value buffer. This should be called after the corresponding agent interactor's EncodeObservations.
 	*/
 	UFUNCTION(BlueprintCallable, Category = "LearningAgents")
 	void EvaluateCritic();
@@ -128,9 +128,9 @@ public:
 // ----- Private Data -----
 private:
 
-	/** The agent type this critic is associated with. */
+	/** The agent interactor this critic is associated with. */
 	UPROPERTY(VisibleAnywhere, Transient, Category = "LearningAgents")
-	TObjectPtr<ULearningAgentsType> AgentType;
+	TObjectPtr<ULearningAgentsInteractor> Interactor;
 
 	/** The underlying neural network. */
 	UPROPERTY(VisibleAnywhere, Transient, Category = "LearningAgents")
diff --git a/Engine/Plugins/Experimental/LearningAgents/Source/LearningAgents/Public/LearningAgentsInteractor.h b/Engine/Plugins/Experimental/LearningAgents/Source/LearningAgents/Public/LearningAgentsInteractor.h
new file mode 100644
--- /dev/null
+++ b/Engine/Plugins/Experimental/LearningAgents/Source/LearningAgents/Public/LearningAgentsInteractor.h
@@ -0,0 +1,175 @@
+// Copyright Epic Games, Inc. All Rights Reserved.
+
+#pragma once
+
+#include "LearningAgentsManagerComponent.h"
+
+#include "LearningArray.h"
+#include "Containers/Array.h"
+#include "Containers/ArrayView.h"
+#include "Containers/ContainerAllocationPolicies.h"
+#include "Templates/SharedPointer.h"
+#include "UObject/Object.h"
+#include "UObject/ObjectPtr.h"
+
+#include "LearningAgentsInteractor.generated.h"
+
+namespace UE::Learning
+{
+	struct FArrayMap;
+	struct FConcatenateFeature;
+	struct FFeatureObject;
+}
+
+class ALearningAgentsManager;
+class ULearningAgentsAction;
+class ULearningAgentsObservation;
+
+/**
+* ULearningAgentsInteractor defines how agents interact with the environment through their observations and actions.
+* Additionally, it provides methods to be called during the inference process of those agents.
+*
+* To use this class, you need to implement the SetupObservations and SetupActions functions (as well as their
+* corresponding SetObservations and SetActions functions), which will define the size of inputs and outputs to your
+* policy. Before you can do inference, you need to call SetupInteractor, which will initialize the underlying data
+* structure, and you need to call AddAgent for each object you want controlled by this agent interactor.
+*/
+UCLASS(Abstract, BlueprintType, Blueprintable)
+class LEARNINGAGENTS_API ULearningAgentsInteractor : public ULearningAgentsManagerComponent
+{
+	GENERATED_BODY()
+
+public:
+
+	// These constructors/destructors are needed to make forward declarations happy
+	ULearningAgentsInteractor();
+	ULearningAgentsInteractor(FVTableHelper& Helper);
+	virtual ~ULearningAgentsInteractor();
+
+	/** Initializes this object and runs the setup events for observations and actions. */
+	UFUNCTION(BlueprintCallable, Category = "LearningAgents")
+	void SetupInteractor(ALearningAgentsManager* InAgentManager);
+
+// ----- Observations -----
+public:
+
+	/**
+	 * During this event, all observations should be added to this agent interactor.
+	 * @param Interactor This agent interactor. This is a convenience for blueprints so you don't have to find the Self pin.
+	 * @see LearningAgentsObservations.h for the list of available observations.
+	 */
+	UFUNCTION(BlueprintNativeEvent, Category = "LearningAgents")
+	void SetupObservations(ULearningAgentsInteractor* Interactor);
+
+	/**
+	* During this event, all observations should be set for each agent.
+	* @param AgentIds The list of agent ids to set observations for.
+	* @see LearningAgentsObservations.h for the list of available observations.
+	* @see GetAgent to get the agent corresponding to each id.
+	*/
+	UFUNCTION(BlueprintNativeEvent, Category = "LearningAgents")
+	void SetObservations(const TArray<int32>& AgentIds);
+
+	/**
+	* Used by objects derived from ULearningAgentsObservation to add themselves to this agent interactor during their
+	* creation. You shouldn't need to call this directly.
+	*/
+	void AddObservation(TObjectPtr<ULearningAgentsObservation> Object, const TSharedRef<UE::Learning::FFeatureObject>& Feature);
+
+// ----- Actions -----
+public:
+
+	/**
+	 * During this event, all actions should be added to this agent interactor.
+	 * @param Interactor This Interactor. This is a convenience for blueprints so you don't have to find the Self pin.
+	 * @see LearningAgentsActions.h for the list of available actions.
+	 */
+	UFUNCTION(BlueprintNativeEvent, Category = "LearningAgents")
+	void SetupActions(ULearningAgentsInteractor* Interactor);
+
+	/**
+	* During this event, you should retrieve the actions and apply them to your agents.
+	* @param AgentIds The list of agent ids to get actions for.
+	* @see LearningAgentsActions.h for the list of available actions.
+	* @see GetAgent to get the agent corresponding to each id.
+	*/
+	UFUNCTION(BlueprintNativeEvent, Category = "LearningAgents")
+	void GetActions(const TArray<int32>& AgentIds);
+
+	/**
+	* Used by objects derived from ULearningAgentsAction to add themselves to this agent interactor during their creation.
+	* You shouldn't need to call this directly.
+	*/
+	void AddAction(TObjectPtr<ULearningAgentsAction> Object, const TSharedRef<UE::Learning::FFeatureObject>& Feature);
+
+// ----- Encoding / Decoding -----
+public:
+
+	/**
+	* Call this function when it is time to gather all the observations for your agents. This can be done each frame or
+	* you can consider wiring it up to some kind of meaningful event, e.g. a hypothetical "OnAITurnStarted" if you have
+	* a turn-based game. This will call this agent interactor's SetObservations event.
+	*/
+	UFUNCTION(BlueprintCallable, Category = "LearningAgents")
+	void EncodeObservations();
+
+	/**
+	* Call this function when it is time for your agents to take their actions. You most likely want to call this after
+	* your policy's EvaluatePolicy function to ensure that you are receiving the latest actions. This will call this
+	* agent interactor's GetActions event.
+	*/
+	UFUNCTION(BlueprintCallable, Category = "LearningAgents")
+	void DecodeActions();
+
+	/**
+	* Gets the observation vector used by a given agent. Should be called only after EncodeObservations.
+	*
+	* @param AgentId				The AgentId to look-up the observation vector for
+	* @param OutObservationVector	The observation vector for the given agent
+	*/
+	UFUNCTION(BlueprintPure, Category = "LearningAgents", meta = (AgentId = "-1"))
+	void GetObservationVector(const int32 AgentId, TArray<float>& OutObservationVector) const;
+
+	/**
+	* Gets the action vector used by a given agent. Should be called only after EncodeActions or EvaluatePolicy.
+	*
+	* @param AgentId				The AgentId to look-up the action vector for
+	* @param OutActionVector		The action vector for the given agent
+	*/
+	UFUNCTION(BlueprintPure, Category = "LearningAgents", meta = (AgentId = "-1"))
+	void GetActionVector(const int32 AgentId, TArray<float>& OutActionVector) const;
+
+// ----- Non-blueprint public interface -----
+public:
+
+	/** Get a reference to this agent interactor's observation feature. */
+	UE::Learning::FFeatureObject& GetObservationFeature() const;
+
+	/** Get a reference to this agent interactor's action feature. */
+	UE::Learning::FFeatureObject& GetActionFeature() const;
+
+	/** Get a const array view of this agent interactor's observation objects. */
+	TConstArrayView<ULearningAgentsObservation*> GetObservationObjects() const;
+
+	/** Get a const array view of this agent interactor's action objects. */
+	TConstArrayView<ULearningAgentsAction*> GetActionObjects() const;
+
+private:
+
+	/** The list of current observation objects. */
+	UPROPERTY(VisibleAnywhere, Transient, Category = "LearningAgents")
+	TArray<TObjectPtr<ULearningAgentsObservation>> ObservationObjects;
+
+	/** The list of current action objects. */
+	UPROPERTY(VisibleAnywhere, Transient, Category = "LearningAgents")
+	TArray<TObjectPtr<ULearningAgentsAction>> ActionObjects;
+
+// ----- Private Data -----
+private:
+
+	TArray<TSharedRef<UE::Learning::FFeatureObject>, TInlineAllocator<16>> ObservationFeatures;
+	TArray<TSharedRef<UE::Learning::FFeatureObject>, TInlineAllocator<16>> ActionFeatures;
+
+	TSharedPtr<UE::Learning::FConcatenateFeature> Observations;
+	TSharedPtr<UE::Learning::FConcatenateFeature> Actions;
+};
diff --git a/Engine/Plugins/Experimental/LearningAgents/Source/LearningAgents/Public/LearningAgentsObservations.h b/Engine/Plugins/Experimental/LearningAgents/Source/LearningAgents/Public/LearningAgentsObservations.h
--- a/Engine/Plugins/Experimental/LearningAgents/Source/LearningAgents/Public/LearningAgentsObservations.h
+++ b/Engine/Plugins/Experimental/LearningAgents/Source/LearningAgents/Public/LearningAgentsObservations.h
@@ -29,7 +29,7 @@ namespace UE::Learning
 	struct FScalarAngularVelocityFeature;
 }
 
-class ULearningAgentsType;
+class ULearningAgentsInteractor;
 
 // For functions in this file, we are favoring having more verbose names such as "AddFloatObservation" vs simply "Add" in 
 // order to keep it easy to find the correct function in blueprints.
@@ -45,7 +45,7 @@ class LEARNINGAGENTS_API ULearningAgentsObservation : public UObject
 public:
 
 	UPROPERTY(VisibleAnywhere, Transient, Category = "LearningAgents")
-	TObjectPtr<ULearningAgentsType> AgentType;
+	TObjectPtr<ULearningAgentsInteractor> Interactor;
 
 #if UE_LEARNING_AGENTS_ENABLE_VISUAL_LOG
 	/** Color used to draw this observation in the visual log */
@@ -67,17 +67,17 @@ class LEARNINGAGENTS_API UFloatObservation : public ULearningAgentsObservation
 public:
 
 	/**
-	* Adds a new float observation to the given agent type. Call during ULearningAgentsType::SetupObservations event.
-	* @param InAgentType The agent type to add this observation to.
+	* Adds a new float observation to the given agent interactor. Call during ULearningAgentsInteractor::SetupObservations event.
+	* @param InInteractor The agent interactor to add this observation to.
 	* @param Name The name of this new observation. Used for debugging.
 	* @param Scale Used to normalize the data for the observation.
 	* @return The newly created observation.
 	*/
 	UFUNCTION(BlueprintCallable, Category = "LearningAgents")
-	static UFloatObservation* AddFloatObservation(ULearningAgentsType* InAgentType, const FName Name = NAME_None, const float Scale = 1.0f);
+	static UFloatObservation* AddFloatObservation(ULearningAgentsInteractor* InInteractor, const FName Name = NAME_None, const float Scale = 1.0f);
 
 	/**
-	* Sets the data for this observation. Call during ULearningAgentsType::SetObservations event.
+	* Sets the data for this observation. Call during ULearningAgentsInteractor::SetObservations event.
 	* @param AgentId The agent id this data corresponds to.
 	* @param Observation The value currently being observed.
 	*/
@@ -101,18 +101,18 @@ class LEARNINGAGENTS_API UFloatArrayObservation : public ULearningAgentsObservat
 public:
 
 	/**
-	* Adds a new float array observation to the given agent type. Call during ULearningAgentsType::SetupObservations event.
-	* @param InAgentType The agent type to add this observation to.
+	* Adds a new float array observation to the given agent interactor. Call during ULearningAgentsInteractor::SetupObservations event.
+	* @param InInteractor The agent interactor to add this observation to.
 	* @param Name The name of this new observation. Used for debugging.
 	* @param Num The number of floats in the array
 	* @param Scale Used to normalize the data for the observation.
 	* @return The newly created observation.
 	*/
 	UFUNCTION(BlueprintCallable, Category = "LearningAgents")
-	static UFloatArrayObservation* AddFloatArrayObservation(ULearningAgentsType* InAgentType, const FName Name = NAME_None, const int32 Num = 1, const float Scale = 1.0f);
+	static UFloatArrayObservation* AddFloatArrayObservation(ULearningAgentsInteractor* InInteractor, const FName Name = NAME_None, const int32 Num = 1, const float Scale = 1.0f);
 
 	/**
-	* Sets the data for this observation. Call during ULearningAgentsType::SetObservations event.
+	* Sets the data for this observation. Call during ULearningAgentsInteractor::SetObservations event.
 	* @param AgentId The agent id this data corresponds to.
 	* @param Observation The value currently being observed.
 	*/
@@ -138,17 +138,17 @@ class LEARNINGAGENTS_API UVectorObservation : public ULearningAgentsObservation
 public:
 
 	/**
-	* Adds a new vector observation to the given agent type. Call during ULearningAgentsType::SetupObservations event.
-	* @param InAgentType The agent type to add this observation to.
+	* Adds a new vector observation to the given agent interactor. Call during ULearningAgentsInteractor::SetupObservations event.
+	* @param InInteractor The agent interactor to add this observation to.
 	* @param Name The name of this new observation. Used for debugging.
 	* @param Scale Used to normalize the data for the observation.
 	* @return The newly created observation.
 	*/
 	UFUNCTION(BlueprintCallable, Category = "LearningAgents")
-	static UVectorObservation* AddVectorObservation(ULearningAgentsType* InAgentType, const FName Name = NAME_None, const float Scale = 1.0f);
+	static UVectorObservation* AddVectorObservation(ULearningAgentsInteractor* InInteractor, const FName Name = NAME_None, const float Scale = 1.0f);
 
 	/**
-	* Sets the data for this observation. Call during ULearningAgentsType::SetObservations event.
+	* Sets the data for this observation. Call during ULearningAgentsInteractor::SetObservations event.
 	* @param AgentId The agent id this data corresponds to.
 	* @param Observation The values currently being observed.
 	*/
@@ -172,18 +172,18 @@ class LEARNINGAGENTS_API UVectorArrayObservation : public ULearningAgentsObserva
 public:
 
 	/**
-	* Adds a new vector array observation to the given agent type. Call during ULearningAgentsType::SetupObservations event.
-	* @param InAgentType The agent type to add this observation to.
+	* Adds a new vector array observation to the given agent interactor. Call during ULearningAgentsInteractor::SetupObservations event.
+	* @param InInteractor The agent interactor to add this observation to.
 	* @param Name The name of this new observation. Used for debugging.
 	* @param Num The number of vectors in the array
 	* @param Scale Used to normalize the data for the observation.
 	* @return The newly created observation.
 	*/
 	UFUNCTION(BlueprintCallable, Category = "LearningAgents")
-	static UVectorArrayObservation* AddVectorArrayObservation(ULearningAgentsType* InAgentType, const FName Name = NAME_None, const int32 Num = 1, const float Scale = 1.0f);
+	static UVectorArrayObservation* AddVectorArrayObservation(ULearningAgentsInteractor* InInteractor, const FName Name = NAME_None, const int32 Num = 1, const float Scale = 1.0f);
 
 	/**
-	* Sets the data for this observation. Call during ULearningAgentsType::SetObservations event.
+	* Sets the data for this observation. Call during ULearningAgentsInteractor::SetObservations event.
 	* @param AgentId The agent id this data corresponds to.
 	* @param Observation The values currently being observed.
 	*/
@@ -209,17 +209,17 @@ class LEARNINGAGENTS_API UEnumObservation : public ULearningAgentsObservation
 public:
 
 	/**
-	* Adds a new enum observation to the given agent type. Call during ULearningAgentsType::SetupObservations event.
-	* @param InAgentType The agent type to add this observation to.
+	* Adds a new enum observation to the given agent interactor. Call during ULearningAgentsInteractor::SetupObservations event.
+	* @param InInteractor The agent interactor to add this observation to.
 	* @param EnumType The type of enum to use
 	* @param Name The name of this new observation. Used for debugging.
 	* @return The newly created observation.
 	*/
 	UFUNCTION(BlueprintCallable, Category = "LearningAgents")
-	static UEnumObservation* AddEnumObservation(ULearningAgentsType* InAgentType, const UEnum* EnumType, const FName Name = NAME_None);
+	static UEnumObservation* AddEnumObservation(ULearningAgentsInteractor* InInteractor, const UEnum* EnumType, const FName Name = NAME_None);
 
 	/**
-	* Sets the data for this observation. Call during ULearningAgentsType::SetObservations event.
+	* Sets the data for this observation. Call during ULearningAgentsInteractor::SetObservations event.
 	* @param AgentId The agent id this data corresponds to.
 	* @param Value The enum value currently being observed.
 	*/
@@ -244,18 +244,18 @@ class LEARNINGAGENTS_API UEnumArrayObservation : public ULearningAgentsObservati
 public:
 
 	/**
-	* Adds a new enum array observation to the given agent type. Call during ULearningAgentsType::SetupObservations event.
-	* @param InAgentType The agent type to add this observation to.
+	* Adds a new enum array observation to the given agent interactor. Call during ULearningAgentsInteractor::SetupObservations event.
+	* @param InInteractor The agent interactor to add this observation to.
 	* @param EnumType The type of enum to use
 	* @param Name The name of this new observation. Used for debugging.
 	* @param EnumNum The number of enum observations in the array
 	* @return The newly created observation.
 	*/
 	UFUNCTION(BlueprintCallable, Category = "LearningAgents")
-	static UEnumArrayObservation* AddEnumArrayObservation(ULearningAgentsType* InAgentType, const UEnum* EnumType, const FName Name = NAME_None, const int32 EnumNum = 1);
+	static UEnumArrayObservation* AddEnumArrayObservation(ULearningAgentsInteractor* InInteractor, const UEnum* EnumType, const FName Name = NAME_None, const int32 EnumNum = 1);
 
 	/**
-	* Sets the data for this observation. Call during ULearningAgentsType::SetObservations event.
+	* Sets the data for this observation. Call during ULearningAgentsInteractor::SetObservations event.
 	* @param AgentId The agent id this data corresponds to.
 	* @param Values The enum values currently being observed.
 	*/
@@ -282,17 +282,17 @@ class LEARNINGAGENTS_API UTimeObservation : public ULearningAgentsObservation
 public:
 
 	/**
-	* Adds a new time observation to the given agent type. Call during ULearningAgentsType::SetupObservations event.
-	* @param InAgentType The agent type to add this observation to.
+	* Adds a new time observation to the given agent interactor. Call during ULearningAgentsInteractor::SetupObservations event.
+	* @param InInteractor The agent interactor to add this observation to.
 	* @param Name The name of this new observation. Used for debugging.
 	* @param Scale Used to normalize the data for the observation.
 	* @return The newly created observation.
 	*/
 	UFUNCTION(BlueprintCallable, Category = "LearningAgents")
-	static UTimeObservation* AddTimeObservation(ULearningAgentsType* InAgentType, const FName Name = NAME_None, const float Scale = 1.0f);
+	static UTimeObservation* AddTimeObservation(ULearningAgentsInteractor* InInteractor, const FName Name = NAME_None, const float Scale = 1.0f);
 
 	/**
-	* Sets the data for this observation. Call during ULearningAgentsType::SetObservations event.
+	* Sets the data for this observation. Call during ULearningAgentsInteractor::SetObservations event.
 	* @param AgentId The agent id this data corresponds to.
 	* @param Time The time currently being observed.
 	* @param RelativeTime The time the provided time should be encoded relative to.
@@ -317,18 +317,18 @@ class LEARNINGAGENTS_API UTimeArrayObservation : public ULearningAgentsObservati
 public:
 
 	/**
-	* Adds a new angle observation to the given agent type. Call during ULearningAgentsType::SetupObservations event.
-	* @param InAgentType The agent type to add this observation to.
+	* Adds a new angle observation to the given agent interactor. Call during ULearningAgentsInteractor::SetupObservations event.
+	* @param InInteractor The agent interactor to add this observation to.
 	* @param Name The name of this new observation. Used for debugging.
 	* @param TimeNum The number of times in the array
 	* @param Scale Used to normalize the data for the observation.
 	* @return The newly created observation.
 	*/
 	UFUNCTION(BlueprintCallable, Category = "LearningAgents")
-	static UTimeArrayObservation* AddTimeArrayObservation(ULearningAgentsType* InAgentType, const FName Name = NAME_None, const int32 TimeNum = 1, const float Scale = 1.0f);
+	static UTimeArrayObservation* AddTimeArrayObservation(ULearningAgentsInteractor* InInteractor, const FName Name = NAME_None, const int32 TimeNum = 1, const float Scale = 1.0f);
 
 	/**
-	* Sets the data for this observation. Call during ULearningAgentsType::SetObservations event.
+	* Sets the data for this observation. Call during ULearningAgentsInteractor::SetObservations event.
 	* @param AgentId The agent id this data corresponds to.
 	* @param Time The times currently being observed.
 	* @param RelativeTime The time the provided time should be encoded relative to.
@@ -355,17 +355,17 @@ class LEARNINGAGENTS_API UAngleObservation : public ULearningAgentsObservation
 public:
 
 	/**
-	* Adds a new angle observation to the given agent type. Call during ULearningAgentsType::SetupObservations event.
-	* @param InAgentType The agent type to add this observation to.
+	* Adds a new angle observation to the given agent interactor. Call during ULearningAgentsInteractor::SetupObservations event.
+	* @param InInteractor The agent interactor to add this observation to.
 	* @param Name The name of this new observation. Used for debugging.
 	* @param Scale Used to normalize the data for the observation. Angle observations are encoded as directions. 
 	* @return The newly created observation.
 	*/
 	UFUNCTION(BlueprintCallable, Category = "LearningAgents")
-	static UAngleObservation* AddAngleObservation(ULearningAgentsType* InAgentType, const FName Name = NAME_None, const float Scale = 1.0f);
+	static UAngleObservation* AddAngleObservation(ULearningAgentsInteractor* InInteractor, const FName Name = NAME_None, const float Scale = 1.0f);
 
 	/**
-	* Sets the data for this observation. Call during ULearningAgentsType::SetObservations event.
+	* Sets the data for this observation. Call during ULearningAgentsInteractor::SetObservations event.
 	* @param AgentId The agent id this data corresponds to.
 	* @param Angle The angle currently being observed.
 	* @param RelativeAngle The frame of reference angle.
@@ -390,18 +390,18 @@ class LEARNINGAGENTS_API UAngleArrayObservation : public ULearningAgentsObservat
 public:
 
 	/**
-	* Adds a new angle observation to the given agent type. Call during ULearningAgentsType::SetupObservations event.
-	* @param InAgentType The agent type to add this observation to.
+	* Adds a new angle observation to the given agent interactor. Call during ULearningAgentsInteractor::SetupObservations event.
+	* @param InInteractor The agent interactor to add this observation to.
 	* @param Name The name of this new observation. Used for debugging.
 	* @param AngleNum The number of angles in the array
 	* @param Scale Used to normalize the data for the observation. Angle observations are encoded as directions. 
 	* @return The newly created observation.
 	*/
 	UFUNCTION(BlueprintCallable, Category = "LearningAgents")
-	static UAngleArrayObservation* AddAngleArrayObservation(ULearningAgentsType* InAgentType, const FName Name = NAME_None, const int32 AngleNum = 1, const float Scale = 1.0f);
+	static UAngleArrayObservation* AddAngleArrayObservation(ULearningAgentsInteractor* InInteractor, const FName Name = NAME_None, const int32 AngleNum = 1, const float Scale = 1.0f);
 
 	/**
-	* Sets the data for this observation. Call during ULearningAgentsType::SetObservations event.
+	* Sets the data for this observation. Call during ULearningAgentsInteractor::SetObservations event.
 	* @param AgentId The agent id this data corresponds to.
 	* @param Angles The angles currently being observed.
 	* @param RelativeAngle The frame of reference angle.
@@ -426,17 +426,17 @@ class LEARNINGAGENTS_API URotationObservation : public ULearningAgentsObservatio
 public:
 
 	/**
-	* Adds a new rotation observation to the given agent type. Call during ULearningAgentsType::SetupObservations event.
-	* @param InAgentType The agent type to add this observation to.
+	* Adds a new rotation observation to the given agent interactor. Call during ULearningAgentsInteractor::SetupObservations event.
+	* @param InInteractor The agent interactor to add this observation to.
 	* @param Name The name of this new observation. Used for debugging.
 	* @param Scale Used to normalize the data for the observation. Rotation observations are encoded as directions. 
 	* @return The newly created observation.
 	*/
 	UFUNCTION(BlueprintCallable, Category = "LearningAgents")
-	static URotationObservation* AddRotationObservation(ULearningAgentsType* InAgentType, const FName Name = NAME_None, const float Scale = 1.0f);
+	static URotationObservation* AddRotationObservation(ULearningAgentsInteractor* InInteractor, const FName Name = NAME_None, const float Scale = 1.0f);
 
 	/**
-	* Sets the data for this observation. Call during ULearningAgentsType::SetObservations event.
+	* Sets the data for this observation. Call during ULearningAgentsInteractor::SetObservations event.
 	* @param AgentId The agent id this data corresponds to.
 	* @param Rotation The rotation currently being observed.
 	* @param RelativeRotation The frame of reference rotation.
@@ -445,7 +445,7 @@ public:
 	void SetRotationObservation(const int32 AgentId, const FRotator Rotation, const FRotator RelativeRotation = FRotator::ZeroRotator);
 
 	/**
-	* Sets the data for this observation. Call during ULearningAgentsType::SetObservations event.
+	* Sets the data for this observation. Call during ULearningAgentsInteractor::SetObservations event.
 	* @param AgentId The agent id this data corresponds to.
 	* @param Rotation The rotation currently being observed.
 	* @param RelativeRotation The frame of reference rotation.
@@ -470,18 +470,18 @@ class LEARNINGAGENTS_API URotationArrayObservation : public ULearningAgentsObser
 public:
 
 	/**
-	* Adds a new rotation observation to the given agent type. Call during ULearningAgentsType::SetupObservations event.
-	* @param InAgentType The agent type to add this observation to.
+	* Adds a new rotation observation to the given agent interactor. Call during ULearningAgentsInteractor::SetupObservations event.
+	* @param InInteractor The agent interactor to add this observation to.
 	* @param Name The name of this new observation. Used for debugging.
 	* @param RotationNum The number of rotations in the array
 	* @param Scale Used to normalize the data for the observation. Rotation observations are encoded as directions. 
 	* @return The newly created observation.
 	*/
 	UFUNCTION(BlueprintCallable, Category = "LearningAgents")
-	static URotationArrayObservation* AddRotationArrayObservation(ULearningAgentsType* InAgentType, const FName Name = NAME_None, const int32 RotationNum = 1, const float Scale = 1.0f);
+	static URotationArrayObservation* AddRotationArrayObservation(ULearningAgentsInteractor* InInteractor, const FName Name = NAME_None, const int32 RotationNum = 1, const float Scale = 1.0f);
 
 	/**
-	* Sets the data for this observation. Call during ULearningAgentsType::SetObservations event.
+	* Sets the data for this observation. Call during ULearningAgentsInteractor::SetObservations event.
 	* @param AgentId The agent id this data corresponds to.
 	* @param Rotations The rotations currently being observed.
 	* @param RelativeRotation The frame of reference rotation.
@@ -490,7 +490,7 @@ public:
 	void SetRotationArrayObservation(const int32 AgentId, const TArray<FRotator>& Rotations, const FRotator RelativeRotation = FRotator::ZeroRotator);
 
 	/**
-	* Sets the data for this observation. Call during ULearningAgentsType::SetObservations event.
+	* Sets the data for this observation. Call during ULearningAgentsInteractor::SetObservations event.
 	* @param AgentId The agent id this data corresponds to.
 	* @param Rotations The rotations currently being observed.
 	* @param RelativeRotation The frame of reference rotation.
@@ -517,18 +517,18 @@ class LEARNINGAGENTS_API UDirectionObservation : public ULearningAgentsObservati
 public:
 
 	/**
-	* Adds a new direction observation to the given agent type. Call during ULearningAgentsType::SetupObservations event.
-	* @param InAgentType The agent type to add this observation to.
+	* Adds a new direction observation to the given agent interactor. Call during ULearningAgentsInteractor::SetupObservations event.
+	* @param InInteractor The agent interactor to add this observation to.
 	* @param Name The name of this new observation. Used for debugging.
 	* @param Scale Used to normalize the data for the observation.
 	* @return The newly created observation.
 	*/
 	UFUNCTION(BlueprintCallable, Category = "LearningAgents")
-	static UDirectionObservation* AddDirectionObservation(ULearningAgentsType* InAgentType, const FName Name = NAME_None, const float Scale = 1.0f);
+	static UDirectionObservation* AddDirectionObservation(ULearningAgentsInteractor* InInteractor, const FName Name = NAME_None, const float Scale = 1.0f);
 
 	/**
 	* Sets the data for this observation. The relative rotation can be used to make this observation relative to the
-	* agent's perspective, e.g. by passing the agent's forward rotation. Call during ULearningAgentsType::SetObservations event.
+	* agent's perspective, e.g. by passing the agent's forward rotation. Call during ULearningAgentsInteractor::SetObservations event.
 	* @param AgentId The agent id this data corresponds to.
 	* @param Direction The direction currently being observed.
 	* @param RelativeRotation The frame of reference rotation.
@@ -553,19 +553,19 @@ class LEARNINGAGENTS_API UDirectionArrayObservation : public ULearningAgentsObse
 public:
 
 	/**
-	* Adds a new direction array observation to the given agent type. Call during ULearningAgentsType::SetupObservations event.
-	* @param InAgentType The agent type to add this observation to.
+	* Adds a new direction array observation to the given agent interactor. Call during ULearningAgentsInteractor::SetupObservations event.
+	* @param InInteractor The agent interactor to add this observation to.
 	* @param Name The name of this new observation. Used for debugging.
 	* @param DirectionNum The number of directions in the array
 	* @param Scale Used to normalize the data for the observation.
 	* @return The newly created observation.
 	*/
 	UFUNCTION(BlueprintCallable, Category = "LearningAgents")
-	static UDirectionArrayObservation* AddDirectionArrayObservation(ULearningAgentsType* InAgentType, const FName Name = NAME_None, const int32 DirectionNum = 1, const float Scale = 1.0f);
+	static UDirectionArrayObservation* AddDirectionArrayObservation(ULearningAgentsInteractor* InInteractor, const FName Name = NAME_None, const int32 DirectionNum = 1, const float Scale = 1.0f);
 
 	/**
 	* Sets the data for this observation. The relative rotation can be used to make this observation relative to the
-	* agent's perspective, e.g. by passing the agent's forward rotation. Call during ULearningAgentsType::SetObservations event.
+	* agent's perspective, e.g. by passing the agent's forward rotation. Call during ULearningAgentsInteractor::SetObservations event.
 	* @param AgentId The agent id this data corresponds to.
 	* @param Directions The directions currently being observed.
 	* @param RelativeRotation The frame of reference rotation.
@@ -590,9 +590,9 @@ class LEARNINGAGENTS_API UPlanarDirectionObservation : public ULearningAgentsObs
 public:
 
 	/**
-	* Adds a new planar direction observation to the given agent type. The axis parameters define the plane. Call
-	* during ULearningAgentsType::SetupObservations event.
-	* @param InAgentType The agent type to add this observation to.
+	* Adds a new planar direction observation to the given agent interactor. The axis parameters define the plane. Call
+	* during ULearningAgentsInteractor::SetupObservations event.
+	* @param InInteractor The agent interactor to add this observation to.
 	* @param Name The name of this new observation. Used for debugging.
 	* @param Scale Used to normalize the data for the observation.
 	* @param Axis0 The forward axis of the plane.
@@ -600,11 +600,11 @@ public:
 	* @return The newly created observation.
 	*/
 	UFUNCTION(BlueprintCallable, Category = "LearningAgents")
-	static UPlanarDirectionObservation* AddPlanarDirectionObservation(ULearningAgentsType* InAgentType, const FName Name = NAME_None, const float Scale = 1.0f, const FVector Axis0 = FVector::ForwardVector, const FVector Axis1 = FVector::RightVector);
+	static UPlanarDirectionObservation* AddPlanarDirectionObservation(ULearningAgentsInteractor* InInteractor, const FName Name = NAME_None, const float Scale = 1.0f, const FVector Axis0 = FVector::ForwardVector, const FVector Axis1 = FVector::RightVector);
 
 	/**
 	* Sets the data for this observation. The relative rotation can be used to make this observation relative to the
-	* agent's perspective, e.g. by passing the agent's forward rotation. Call during ULearningAgentsType::SetObservations event.
+	* agent's perspective, e.g. by passing the agent's forward rotation. Call during ULearningAgentsInteractor::SetObservations event.
 	* @param AgentId The agent id this data corresponds to.
 	* @param Direction The direction currently being observed.
 	* @param RelativeRotation The frame of reference rotation.
@@ -629,9 +629,9 @@ class LEARNINGAGENTS_API UPlanarDirectionArrayObservation : public ULearningAgen
 public:
 
 	/**
-	* Adds a new planar direction observation to the given agent type. The axis parameters define the plane. Call
-	* during ULearningAgentsType::SetupObservations event.
-	* @param InAgentType The agent type to add this observation to.
+	* Adds a new planar direction observation to the given agent interactor. The axis parameters define the plane. Call
+	* during ULearningAgentsInteractor::SetupObservations event.
+	* @param InInteractor The agent interactor to add this observation to.
 	* @param Name The name of this new observation. Used for debugging.
 	* @param DirectionNum The number of directions in the array
 	* @param Scale Used to normalize the data for the observation.
@@ -640,11 +640,11 @@ public:
 	* @return The newly created observation.
 	*/
 	UFUNCTION(BlueprintCallable, Category = "LearningAgents")
-	static UPlanarDirectionArrayObservation* AddPlanarDirectionArrayObservation(ULearningAgentsType* InAgentType, const FName Name = NAME_None, const int32 DirectionNum = 1, const float Scale = 1.0f, const FVector Axis0 = FVector::ForwardVector, const FVector Axis1 = FVector::RightVector);
+	static UPlanarDirectionArrayObservation* AddPlanarDirectionArrayObservation(ULearningAgentsInteractor* InInteractor, const FName Name = NAME_None, const int32 DirectionNum = 1, const float Scale = 1.0f, const FVector Axis0 = FVector::ForwardVector, const FVector Axis1 = FVector::RightVector);
 
 	/**
 	* Sets the data for this observation. The relative rotation can be used to make this observation relative to the
-	* agent's perspective, e.g. by passing the agent's forward rotation. Call during ULearningAgentsType::SetObservations event.
+	* agent's perspective, e.g. by passing the agent's forward rotation. Call during ULearningAgentsInteractor::SetObservations event.
 	* @param AgentId The agent id this data corresponds to.
 	* @param Directions The directions currently being observed.
 	* @param RelativeRotation The frame of reference rotation.
@@ -671,19 +671,19 @@ class LEARNINGAGENTS_API UPositionObservation : public ULearningAgentsObservatio
 public:
 
 	/**
-	* Adds a new position observation to the given agent type. Call during ULearningAgentsType::SetupObservations event.
-	* @param InAgentType The agent type to add this observation to.
+	* Adds a new position observation to the given agent interactor. Call during ULearningAgentsInteractor::SetupObservations event.
+	* @param InInteractor The agent interactor to add this observation to.
 	* @param Name The name of this new observation. Used for debugging.
 	* @param Scale Used to normalize the data for the observation.
 	* @return The newly created observation.
 	*/
 	UFUNCTION(BlueprintCallable, Category = "LearningAgents")
-	static UPositionObservation* AddPositionObservation(ULearningAgentsType* InAgentType, const FName Name = NAME_None, const float Scale = 100.0f);
+	static UPositionObservation* AddPositionObservation(ULearningAgentsInteractor* InInteractor, const FName Name = NAME_None, const float Scale = 100.0f);
 
 	/**
 	* Sets the data for this observation. The relative position & rotation can be used to make this observation
 	* relative to the agent's perspective, e.g. by passing the agent's position & forward rotation. Call during
-	* ULearningAgentsType::SetObservations event.
+	* ULearningAgentsInteractor::SetObservations event.
 	* @param AgentId The agent id this data corresponds to.
 	* @param Position The position currently being observed.
 	* @param RelativePosition The vector Position will be offset from.
@@ -709,20 +709,20 @@ class LEARNINGAGENTS_API UPositionArrayObservation : public ULearningAgentsObser
 public:
 
 	/**
-	* Adds a new position array observation to the given agent type. Call during ULearningAgentsType::SetupObservations event.
-	* @param InAgentType The agent type to add this observation to.
+	* Adds a new position array observation to the given agent interactor. Call during ULearningAgentsInteractor::SetupObservations event.
+	* @param InInteractor The agent interactor to add this observation to.
 	* @param Name The name of this new observation. Used for debugging.
 	* @param PositionNum The number of positions in the array.
 	* @param Scale Used to normalize the data for the observation.
 	* @return The newly created observation.
 	*/
 	UFUNCTION(BlueprintCallable, Category = "LearningAgents")
-	static UPositionArrayObservation* AddPositionArrayObservation(ULearningAgentsType* InAgentType, const FName Name = NAME_None, const int32 PositionNum = 1, const float Scale = 100.0f);
+	static UPositionArrayObservation* AddPositionArrayObservation(ULearningAgentsInteractor* InInteractor, const FName Name = NAME_None, const int32 PositionNum = 1, const float Scale = 100.0f);
 
 	/**
 	* Sets the data for this observation. The relative position & rotation can be used to make this observation
 	* relative to the agent's perspective, e.g. by passing the agent's position & forward rotation. Call during
-	* ULearningAgentsType::SetObservations event.
+	* ULearningAgentsInteractor::SetObservations event.
 	* @param AgentId The agent id this data corresponds to.
 	* @param Positions The positions currently being observed.
 	* @param RelativePosition The vector Positions will be offset from.
@@ -748,20 +748,20 @@ class LEARNINGAGENTS_API UScalarPositionObservation : public ULearningAgentsObse
 public:
 
 	/**
-	* Adds a new scalar position observation to the given agent type. Call during 
-	* ULearningAgentsType::SetupObservations event.
-	* @param InAgentType The agent type to add this observation to.
+	* Adds a new scalar position observation to the given agent interactor. Call during 
+	* ULearningAgentsInteractor::SetupObservations event.
+	* @param InInteractor The agent interactor to add this observation to.
 	* @param Name The name of this new observation. Used for debugging.
 	* @param Scale Used to normalize the data for the observation.
 	* @return The newly created observation.
 	*/
 	UFUNCTION(BlueprintCallable, Category = "LearningAgents")
-	static UScalarPositionObservation* AddScalarPositionObservation(ULearningAgentsType* InAgentType, const FName Name = NAME_None, const float Scale = 100.0f);
+	static UScalarPositionObservation* AddScalarPositionObservation(ULearningAgentsInteractor* InInteractor, const FName Name = NAME_None, const float Scale = 100.0f);
 
 	/**
 	* Sets the data for this observation. The relative position can be used to make this observation
 	* relative to the agent's perspective, e.g. by passing the agent's position. Call during 
-	* ULearningAgentsType::SetObservations event.
+	* ULearningAgentsInteractor::SetObservations event.
 	* @param AgentId The agent id this data corresponds to.
 	* @param Position The position currently being observed.
 	* @param RelativePosition The vector Position will be offset from.
@@ -772,7 +772,7 @@ public:
 	/**
 	* Sets the data for this observation. The relative position can be used to make this observation
 	* relative to the agent's perspective, e.g. by passing the agent's position. Call during
-	* ULearningAgentsType::SetObservations event.
+	* ULearningAgentsInteractor::SetObservations event.
 	* @param AgentId The agent id this data corresponds to.
 	* @param Position The position currently being observed.
 	* @param RelativePosition The vector Position will be offset from.
@@ -798,21 +798,21 @@ class LEARNINGAGENTS_API UScalarPositionArrayObservation : public ULearningAgent
 public:
 
 	/**
-	* Adds a new scalar position array observation to the given agent type. Call during 
-	* ULearningAgentsType::SetupObservations event.
-	* @param InAgentType The agent type to add this observation to.
+	* Adds a new scalar position array observation to the given agent interactor. Call during 
+	* ULearningAgentsInteractor::SetupObservations event.
+	* @param InInteractor The agent interactor to add this observation to.
 	* @param Name The name of this new observation. Used for debugging.
 	* @param PositionNum The number of positions in the array.
 	* @param Scale Used to normalize the data for the observation.
 	* @return The newly created observation.
 	*/
 	UFUNCTION(BlueprintCallable, Category = "LearningAgents")
-	static UScalarPositionArrayObservation* AddScalarPositionArrayObservation(ULearningAgentsType* InAgentType, const FName Name = NAME_None, const int32 PositionNum = 1, const float Scale = 100.0f);
+	static UScalarPositionArrayObservation* AddScalarPositionArrayObservation(ULearningAgentsInteractor* InInteractor, const FName Name = NAME_None, const int32 PositionNum = 1, const float Scale = 100.0f);
 
 	/**
 	* Sets the data for this observation. The relative position can be used to make this observation
 	* relative to the agent's perspective, e.g. by passing the agent's position. Call during
-	* ULearningAgentsType::SetObservations event.
+	* ULearningAgentsInteractor::SetObservations event.
 	* @param AgentId The agent id this data corresponds to.
 	* @param Positions The positions currently being observed.
 	* @param RelativePosition The vector Positions will be offset from.
@@ -823,7 +823,7 @@ public:
 	/**
 	* Sets the data for this observation. The relative position can be used to make this observation
 	* relative to the agent's perspective, e.g. by passing the agent's position. Call during
-	* ULearningAgentsType::SetObservations event.
+	* ULearningAgentsInteractor::SetObservations event.
 	* @param AgentId The agent id this data corresponds to.
 	* @param Positions The positions currently being observed.
 	* @param RelativePosition The vector Positions will be offset from.
@@ -849,9 +849,9 @@ class LEARNINGAGENTS_API UPlanarPositionObservation : public ULearningAgentsObse
 public:
 
 	/**
-	* Adds a new planar position observation to the given agent type. The axis parameters define the plane. Call
-	* during ULearningAgentsType::SetupObservations event.
-	* @param InAgentType The agent type to add this observation to.
+	* Adds a new planar position observation to the given agent interactor. The axis parameters define the plane. Call
+	* during ULearningAgentsInteractor::SetupObservations event.
+	* @param InInteractor The agent interactor to add this observation to.
 	* @param Name The name of this new observation. Used for debugging.
 	* @param Scale Used to normalize the data for the observation.
 	* @param Axis0 The forward axis of the plane.
@@ -859,12 +859,12 @@ public:
 	* @return The newly created observation.
 	*/
 	UFUNCTION(BlueprintCallable, Category = "LearningAgents")
-	static UPlanarPositionObservation* AddPlanarPositionObservation(ULearningAgentsType* InAgentType, const FName Name = NAME_None, const float Scale = 100.0f, const FVector Axis0 = FVector::ForwardVector, const FVector Axis1 = FVector::RightVector);
+	static UPlanarPositionObservation* AddPlanarPositionObservation(ULearningAgentsInteractor* InInteractor, const FName Name = NAME_None, const float Scale = 100.0f, const FVector Axis0 = FVector::ForwardVector, const FVector Axis1 = FVector::RightVector);
 
 	/**
 	* Sets the data for this observation. The relative position & rotation can be used to make this observation
 	* relative to the agent's perspective, e.g. by passing the agent's position & forward rotation. Call during
-	* ULearningAgentsType::SetObservations event.
+	* ULearningAgentsInteractor::SetObservations event.
 	* @param AgentId The agent id this data corresponds to.
 	* @param Position The position currently being observed.
 	* @param RelativePosition The vector Position will be offset from.
@@ -890,9 +890,9 @@ class LEARNINGAGENTS_API UPlanarPositionArrayObservation : public ULearningAgent
 public:
 
 	/**
-	* Adds a new planar position array observation to the given agent type. The axis parameters define the plane.
-	* Call during ULearningAgentsType::SetupObservations event.
-	* @param InAgentType The agent type to add this observation to.
+	* Adds a new planar position array observation to the given agent interactor. The axis parameters define the plane.
+	* Call during ULearningAgentsInteractor::SetupObservations event.
+	* @param InInteractor The agent interactor to add this observation to.
 	* @param Name The name of this new observation. Used for debugging.
 	* @param PositionNum The number of positions in the array.
 	* @param Scale Used to normalize the data for the observation.
@@ -901,12 +901,12 @@ public:
 	* @return The newly created observation.
 	*/
 	UFUNCTION(BlueprintCallable, Category = "LearningAgents")
-	static UPlanarPositionArrayObservation* AddPlanarPositionArrayObservation(ULearningAgentsType* InAgentType, const FName Name = NAME_None, const int32 PositionNum = 1, const float Scale = 100.0f, const FVector Axis0 = FVector::ForwardVector, const FVector Axis1 = FVector::RightVector);
+	static UPlanarPositionArrayObservation* AddPlanarPositionArrayObservation(ULearningAgentsInteractor* InInteractor, const FName Name = NAME_None, const int32 PositionNum = 1, const float Scale = 100.0f, const FVector Axis0 = FVector::ForwardVector, const FVector Axis1 = FVector::RightVector);
 
 	/**
 	* Sets the data for this observation. The relative position & rotation can be used to make this observation
 	* relative to the agent's perspective, e.g. by passing the agent's position & forward rotation. Call during
-	* ULearningAgentsType::SetObservations event.
+	* ULearningAgentsInteractor::SetObservations event.
 	* @param AgentId The agent id this data corresponds to.
 	* @param Positions The positions currently being observed.
 	* @param RelativePosition The vector Positions will be offset from.
@@ -934,18 +934,18 @@ class LEARNINGAGENTS_API UVelocityObservation : public ULearningAgentsObservatio
 public:
 
 	/**
-	* Adds a new velocity observation to the given agent type. Call during ULearningAgentsType::SetupObservations event.
-	* @param InAgentType The agent type to add this observation to.
+	* Adds a new velocity observation to the given agent interactor. Call during ULearningAgentsInteractor::SetupObservations event.
+	* @param InInteractor The agent interactor to add this observation to.
 	* @param Name The name of this new observation. Used for debugging.
 	* @param Scale Used to normalize the data for the observation.
 	* @return The newly created observation.
 	*/
 	UFUNCTION(BlueprintCallable, Category = "LearningAgents")
-	static UVelocityObservation* AddVelocityObservation(ULearningAgentsType* InAgentType, const FName Name = NAME_None, const float Scale = 200.0f);
+	static UVelocityObservation* AddVelocityObservation(ULearningAgentsInteractor* InInteractor, const FName Name = NAME_None, const float Scale = 200.0f);
 
 	/**
 	* Sets the data for this observation. The relative rotation can be used to make this observation relative to the
-	* agent's perspective, e.g. by passing the agent's forward rotation. Call during ULearningAgentsType::SetObservations event.
+	* agent's perspective, e.g. by passing the agent's forward rotation. Call during ULearningAgentsInteractor::SetObservations event.
 	* @param AgentId The agent id this data corresponds to.
 	* @param Velocity The velocity currently being observed.
 	* @param RelativeRotation The frame of reference rotation.
@@ -970,19 +970,19 @@ class LEARNINGAGENTS_API UVelocityArrayObservation : public ULearningAgentsObser
 public:
 
 	/**
-	* Adds a new velocity observation to the given agent type. Call during ULearningAgentsType::SetupObservations event.
-	* @param InAgentType The agent type to add this observation to.
+	* Adds a new velocity observation to the given agent interactor. Call during ULearningAgentsInteractor::SetupObservations event.
+	* @param InInteractor The agent interactor to add this observation to.
 	* @param Name The name of this new observation. Used for debugging.
 	* @param VelocityNum The number of velocities in the array.
 	* @param Scale Used to normalize the data for the observation.
 	* @return The newly created observation.
 	*/
 	UFUNCTION(BlueprintCallable, Category = "LearningAgents")
-	static UVelocityArrayObservation* AddVelocityArrayObservation(ULearningAgentsType* InAgentType, const FName Name = NAME_None, const int32 VelocityNum = 1, const float Scale = 200.0f);
+	static UVelocityArrayObservation* AddVelocityArrayObservation(ULearningAgentsInteractor* InInteractor, const FName Name = NAME_None, const int32 VelocityNum = 1, const float Scale = 200.0f);
 
 	/**
 	* Sets the data for this observation. The relative rotation can be used to make this observation relative to the
-	* agent's perspective, e.g. by passing the agent's forward rotation. Call during ULearningAgentsType::SetObservations event.
+	* agent's perspective, e.g. by passing the agent's forward rotation. Call during ULearningAgentsInteractor::SetObservations event.
 	* @param AgentId The agent id this data corresponds to.
 	* @param Velocities The velocities currently being observed.
 	* @param RelativeRotation The frame of reference rotation.
@@ -1007,18 +1007,18 @@ class LEARNINGAGENTS_API UScalarVelocityObservation : public ULearningAgentsObse
 public:
 
 	/**
-	* Adds a new scalar velocity observation to the given agent type. 
-	* Call during ULearningAgentsType::SetupObservations event.
-	* @param InAgentType The agent type to add this observation to.
+	* Adds a new scalar velocity observation to the given agent interactor. 
+	* Call during ULearningAgentsInteractor::SetupObservations event.
+	* @param InInteractor The agent interactor to add this observation to.
 	* @param Name The name of this new observation. Used for debugging.
 	* @param Scale Used to normalize the data for the observation.
 	* @return The newly created observation.
 	*/
 	UFUNCTION(BlueprintCallable, Category = "LearningAgents")
-	static UScalarVelocityObservation* AddScalarVelocityObservation(ULearningAgentsType* InAgentType, const FName Name = NAME_None, const float Scale = 200.0f);
+	static UScalarVelocityObservation* AddScalarVelocityObservation(ULearningAgentsInteractor* InInteractor, const FName Name = NAME_None, const float Scale = 200.0f);
 
 	/**
-	* Sets the data for this observation. Call during ULearningAgentsType::SetObservations event.
+	* Sets the data for this observation. Call during ULearningAgentsInteractor::SetObservations event.
 	* @param AgentId The agent id this data corresponds to.
 	* @param Velocity The velocity currently being observed.
 	*/
@@ -1026,7 +1026,7 @@ public:
 	void SetScalarVelocityObservation(const int32 AgentId, const float Velocity);
 
 	/**
-	* Sets the data for this observation. Call during ULearningAgentsType::SetObservations event.
+	* Sets the data for this observation. Call during ULearningAgentsInteractor::SetObservations event.
 	* @param AgentId The agent id this data corresponds to.
 	* @param Velocity The velocity currently being observed.
 	* @param Axis The axis to encode the velocity along
@@ -1051,19 +1051,19 @@ class LEARNINGAGENTS_API UScalarVelocityArrayObservation : public ULearningAgent
 public:
 
 	/**
-	* Adds a new scalar velocity observation to the given agent type.
-	* Call during ULearningAgentsType::SetupObservations event.
-	* @param InAgentType The agent type to add this observation to.
+	* Adds a new scalar velocity observation to the given agent interactor.
+	* Call during ULearningAgentsInteractor::SetupObservations event.
+	* @param InInteractor The agent interactor to add this observation to.
 	* @param Name The name of this new observation. Used for debugging.
 	* @param VelocityNum The number of velocities in the array.
 	* @param Scale Used to normalize the data for the observation.
 	* @return The newly created observation.
 	*/
 	UFUNCTION(BlueprintCallable, Category = "LearningAgents")
-	static UScalarVelocityArrayObservation* AddScalarVelocityArrayObservation(ULearningAgentsType* InAgentType, const FName Name = NAME_None, const int32 VelocityNum = 1, const float Scale = 200.0f);
+	static UScalarVelocityArrayObservation* AddScalarVelocityArrayObservation(ULearningAgentsInteractor* InInteractor, const FName Name = NAME_None, const int32 VelocityNum = 1, const float Scale = 200.0f);
 
 	/**
-	* Sets the data for this observation. Call during ULearningAgentsType::SetObservations event.
+	* Sets the data for this observation. Call during ULearningAgentsInteractor::SetObservations event.
 	* @param AgentId The agent id this data corresponds to.
 	* @param Velocities The velocities currently being observed.
 	*/
@@ -1071,7 +1071,7 @@ public:
 	void SetScalarVelocityArrayObservation(const int32 AgentId, const TArray<float>& Velocities);
 
 	/**
-	* Sets the data for this observation. Call during ULearningAgentsType::SetObservations event.
+	* Sets the data for this observation. Call during ULearningAgentsInteractor::SetObservations event.
 	* @param AgentId The agent id this data corresponds to.
 	* @param Velocities The velocities currently being observed.
 	* @param Axis The axis to encode the velocity along
@@ -1096,9 +1096,9 @@ class LEARNINGAGENTS_API UPlanarVelocityObservation : public ULearningAgentsObse
 public:
 
 	/**
-	* Adds a new planar velocity observation to the given agent type. The axis parameters define the plane.
-	* Call during ULearningAgentsType::SetupObservations event.
-	* @param InAgentType The agent type to add this observation to.
+	* Adds a new planar velocity observation to the given agent interactor. The axis parameters define the plane.
+	* Call during ULearningAgentsInteractor::SetupObservations event.
+	* @param InInteractor The agent interactor to add this observation to.
 	* @param Name The name of this new observation. Used for debugging.
 	* @param Scale Used to normalize the data for the observation.
 	* @param Axis0 The forward axis of the plane.
@@ -1106,11 +1106,11 @@ public:
 	* @return The newly created observation.
 	*/
 	UFUNCTION(BlueprintCallable, Category = "LearningAgents")
-	static UPlanarVelocityObservation* AddPlanarVelocityObservation(ULearningAgentsType* InAgentType, const FName Name = NAME_None, const float Scale = 200.0f, const FVector Axis0 = FVector::ForwardVector, const FVector Axis1 = FVector::RightVector);
+	static UPlanarVelocityObservation* AddPlanarVelocityObservation(ULearningAgentsInteractor* InInteractor, const FName Name = NAME_None, const float Scale = 200.0f, const FVector Axis0 = FVector::ForwardVector, const FVector Axis1 = FVector::RightVector);
 
 	/**
 	* Sets the data for this observation. The relative rotation can be used to make this observation relative to the
-	* agent's perspective, e.g. by passing the agent's forward rotation. Call during ULearningAgentsType::SetObservations event.
+	* agent's perspective, e.g. by passing the agent's forward rotation. Call during ULearningAgentsInteractor::SetObservations event.
 	* @param AgentId The agent id this data corresponds to.
 	* @param Velocity The velocity currently being observed.
 	* @param RelativeRotation The frame of reference rotation.
@@ -1135,9 +1135,9 @@ class LEARNINGAGENTS_API UPlanarVelocityArrayObservation : public ULearningAgent
 public:
 
 	/**
-	* Adds a new planar velocity observation to the given agent type. The axis parameters define the plane.
-	* Call during ULearningAgentsType::SetupObservations event.
-	* @param InAgentType The agent type to add this observation to.
+	* Adds a new planar velocity observation to the given agent interactor. The axis parameters define the plane.
+	* Call during ULearningAgentsInteractor::SetupObservations event.
+	* @param InInteractor The agent interactor to add this observation to.
 	* @param Name The name of this new observation. Used for debugging.
 	* @param VelocityNum The number of velocities in the array.
 	* @param Scale Used to normalize the data for the observation.
@@ -1146,11 +1146,11 @@ public:
 	* @return The newly created observation.
 	*/
 	UFUNCTION(BlueprintCallable, Category = "LearningAgents")
-	static UPlanarVelocityArrayObservation* AddPlanarVelocityArrayObservation(ULearningAgentsType* InAgentType, const FName Name = NAME_None, const int32 VelocityNum = 1, const float Scale = 200.0f, const FVector Axis0 = FVector::ForwardVector, const FVector Axis1 = FVector::RightVector);
+	static UPlanarVelocityArrayObservation* AddPlanarVelocityArrayObservation(ULearningAgentsInteractor* InInteractor, const FName Name = NAME_None, const int32 VelocityNum = 1, const float Scale = 200.0f, const FVector Axis0 = FVector::ForwardVector, const FVector Axis1 = FVector::RightVector);
 
 	/**
 	* Sets the data for this observation. The relative rotation can be used to make this observation relative to the
-	* agent's perspective, e.g. by passing the agent's forward rotation. Call during ULearningAgentsType::SetObservations event.
+	* agent's perspective, e.g. by passing the agent's forward rotation. Call during ULearningAgentsInteractor::SetObservations event.
 	* @param AgentId The agent id this data corresponds to.
 	* @param Velocities The velocities currently being observed.
 	* @param RelativeRotation The frame of reference rotation.
@@ -1177,18 +1177,18 @@ class LEARNINGAGENTS_API UAngularVelocityObservation : public ULearningAgentsObs
 public:
 
 	/**
-	* Adds a new angular velocity observation to the given agent type. Call during ULearningAgentsType::SetupObservations event.
-	* @param InAgentType The agent type to add this observation to.
+	* Adds a new angular velocity observation to the given agent interactor. Call during ULearningAgentsInteractor::SetupObservations event.
+	* @param InInteractor The agent interactor to add this observation to.
 	* @param Name The name of this new observation. Used for debugging.
 	* @param Scale Used to normalize the data for the observation.
 	* @return The newly created observation.
 	*/
 	UFUNCTION(BlueprintCallable, Category = "LearningAgents")
-	static UAngularVelocityObservation* AddAngularVelocityObservation(ULearningAgentsType* InAgentType, const FName Name = NAME_None, const float Scale = 180.0f);
+	static UAngularVelocityObservation* AddAngularVelocityObservation(ULearningAgentsInteractor* InInteractor, const FName Name = NAME_None, const float Scale = 180.0f);
 
 	/**
 	* Sets the data for this observation. The relative rotation can be used to make this observation relative to the
-	* agent's perspective, e.g. by passing the agent's forward rotation. Call during ULearningAgentsType::SetObservations event.
+	* agent's perspective, e.g. by passing the agent's forward rotation. Call during ULearningAgentsInteractor::SetObservations event.
 	* @param AgentId The agent id this data corresponds to.
 	* @param AngularVelocity The angular velocity currently being observed.
 	* @param RelativeRotation The frame of reference rotation.
@@ -1213,19 +1213,19 @@ class LEARNINGAGENTS_API UAngularVelocityArrayObservation : public ULearningAgen
 public:
 
 	/**
-	* Adds a new angular velocity array observation to the given agent type. Call during ULearningAgentsType::SetupObservations event.
-	* @param InAgentType The agent type to add this observation to.
+	* Adds a new angular velocity array observation to the given agent interactor. Call during ULearningAgentsInteractor::SetupObservations event.
+	* @param InInteractor The agent interactor to add this observation to.
 	* @param Name The name of this new observation. Used for debugging.
 	* @param AngularVelocityNum The number of angular velocities in the array.
 	* @param Scale Used to normalize the data for the observation.
 	* @return The newly created observation.
 	*/
 	UFUNCTION(BlueprintCallable, Category = "LearningAgents")
-	static UAngularVelocityArrayObservation* AddAngularVelocityArrayObservation(ULearningAgentsType* InAgentType, const FName Name = NAME_None, const int32 AngularVelocityNum = 1, const float Scale = 180.0f);
+	static UAngularVelocityArrayObservation* AddAngularVelocityArrayObservation(ULearningAgentsInteractor* InInteractor, const FName Name = NAME_None, const int32 AngularVelocityNum = 1, const float Scale = 180.0f);
 
 	/**
 	* Sets the data for this observation. The relative rotation can be used to make this observation relative to the
-	* agent's perspective, e.g. by passing the agent's forward rotation. Call during ULearningAgentsType::SetObservations event.
+	* agent's perspective, e.g. by passing the agent's forward rotation. Call during ULearningAgentsInteractor::SetObservations event.
 	* @param AgentId The agent id this data corresponds to.
 	* @param Velocities The angular velocities currently being observed.
 	* @param RelativeRotation The frame of reference rotation.
@@ -1250,18 +1250,18 @@ class LEARNINGAGENTS_API UScalarAngularVelocityObservation : public ULearningAge
 public:
 
 	/**
-	* Adds a new scalar angular velocity observation to the given agent type.
-	* Call during ULearningAgentsType::SetupObservations event.
-	* @param InAgentType The agent type to add this observation to.
+	* Adds a new scalar angular velocity observation to the given agent interactor.
+	* Call during ULearningAgentsInteractor::SetupObservations event.
+	* @param InInteractor The agent interactor to add this observation to.
 	* @param Name The name of this new observation. Used for debugging.
 	* @param Scale Used to normalize the data for the observation.
 	* @return The newly created observation.
 	*/
 	UFUNCTION(BlueprintCallable, Category = "LearningAgents")
-	static UScalarAngularVelocityObservation* AddScalarAngularVelocityObservation(ULearningAgentsType* InAgentType, const FName Name = NAME_None, const float Scale = 180.0f);
+	static UScalarAngularVelocityObservation* AddScalarAngularVelocityObservation(ULearningAgentsInteractor* InInteractor, const FName Name = NAME_None, const float Scale = 180.0f);
 
 	/**
-	* Sets the data for this observation. Call during ULearningAgentsType::SetObservations event.
+	* Sets the data for this observation. Call during ULearningAgentsInteractor::SetObservations event.
 	* @param AgentId The agent id this data corresponds to.
 	* @param AngularVelocity The angular velocity currently being observed.
 	*/
@@ -1269,7 +1269,7 @@ public:
 	void SetScalarAngularVelocityObservation(const int32 AgentId, const float AngularVelocity);
 
 	/**
-	* Sets the data for this observation. Call during ULearningAgentsType::SetObservations event.
+	* Sets the data for this observation. Call during ULearningAgentsInteractor::SetObservations event.
 	* @param AgentId The agent id this data corresponds to.
 	* @param AngularVelocity The angular velocity currently being observed.
 	* @param Axis The axis to encode the angular velocity around
@@ -1294,19 +1294,19 @@ class LEARNINGAGENTS_API UScalarAngularVelocityArrayObservation : public ULearni
 public:
 
 	/**
-	* Adds a new scalar angular velocity array observation to the given agent type.
-	* Call during ULearningAgentsType::SetupObservations event.
-	* @param InAgentType The agent type to add this observation to.
+	* Adds a new scalar angular velocity array observation to the given agent interactor.
+	* Call during ULearningAgentsInteractor::SetupObservations event.
+	* @param InInteractor The agent interactor to add this observation to.
 	* @param Name The name of this new observation. Used for debugging.
 	* @param AngularVelocityNum The number of angular velocities in the array.
 	* @param Scale Used to normalize the data for the observation.
 	* @return The newly created observation.
 	*/
 	UFUNCTION(BlueprintCallable, Category = "LearningAgents")
-	static UScalarAngularVelocityArrayObservation* AddScalarAngularVelocityArrayObservation(ULearningAgentsType* InAgentType, const FName Name = NAME_None, const int32 AngularVelocityNum = 1, const float Scale = 180.0f);
+	static UScalarAngularVelocityArrayObservation* AddScalarAngularVelocityArrayObservation(ULearningAgentsInteractor* InInteractor, const FName Name = NAME_None, const int32 AngularVelocityNum = 1, const float Scale = 180.0f);
 
 	/**
-	* Sets the data for this observation. Call during ULearningAgentsType::SetObservations event.
+	* Sets the data for this observation. Call during ULearningAgentsInteractor::SetObservations event.
 	* @param AgentId The agent id this data corresponds to.
 	* @param AngularVelocities The angular velocities currently being observed.
 	*/
@@ -1314,7 +1314,7 @@ public:
 	void SetScalarAngularVelocityArrayObservation(const int32 AgentId, const TArray<float>& AngularVelocities);
 
 	/**
-	* Sets the data for this observation. Call during ULearningAgentsType::SetObservations event.
+	* Sets the data for this observation. Call during ULearningAgentsInteractor::SetObservations event.
 	* @param AgentId The agent id this data corresponds to.
 	* @param AngularVelocities The angular velocities currently being observed.
 	* @param Axis The axis to encode the angular velocity around
diff --git a/Engine/Plugins/Experimental/LearningAgents/Source/LearningAgents/Public/LearningAgentsPolicy.h b/Engine/Plugins/Experimental/LearningAgents/Source/LearningAgents/Public/LearningAgentsPolicy.h
--- a/Engine/Plugins/Experimental/LearningAgents/Source/LearningAgents/Public/LearningAgentsPolicy.h
+++ b/Engine/Plugins/Experimental/LearningAgents/Source/LearningAgents/Public/LearningAgentsPolicy.h
@@ -70,9 +70,9 @@ public:
 	ULearningAgentsPolicy(FVTableHelper& Helper);
 	virtual ~ULearningAgentsPolicy();
 
-	/** Initializes this object to be used with the given agent type and policy settings. */
+	/** Initializes this object to be used with the given agent interactor and policy settings. */
 	UFUNCTION(BlueprintCallable, Category = "LearningAgents")
-	void SetupPolicy(ALearningAgentsManager* InAgentManager, ULearningAgentsType* InAgentType, const FLearningAgentsPolicySettings& PolicySettings = FLearningAgentsPolicySettings());
+	void SetupPolicy(ALearningAgentsManager* InAgentManager, ULearningAgentsInteractor* InInteractor, const FLearningAgentsPolicySettings& PolicySettings = FLearningAgentsPolicySettings());
 
 public:
 
@@ -118,7 +118,7 @@ public:
 
 	/**
 	* Calling this function will run the underlying neural network on the previously buffered observations to populate
-	* the output action buffer. This should be called after the associated agent type's EncodeObservations and
+	* the output action buffer. This should be called after the associated agent interactor's EncodeObservations and
 	* before its DecodeActions.
 	*/
 	UFUNCTION(BlueprintCallable, Category = "LearningAgents")
@@ -171,9 +171,9 @@ public:
 // ----- Private Data -----
 private:
 
-	/** The agent type this policy is associated with. */
+	/** The agent interactor this policy is associated with. */
 	UPROPERTY(VisibleAnywhere, Transient, Category = "LearningAgents")
-	TObjectPtr<ULearningAgentsType> AgentType;
+	TObjectPtr<ULearningAgentsInteractor> Interactor;
 
 	/** The underlying neural network. */
 	UPROPERTY(VisibleAnywhere, Transient, Category = "LearningAgents")
diff --git a/Engine/Plugins/Experimental/LearningAgents/Source/LearningAgents/Public/LearningAgentsType.h b/Engine/Plugins/Experimental/LearningAgents/Source/LearningAgents/Public/LearningAgentsType.h
deleted file mode 100644
diff --git a/Engine/Plugins/Experimental/LearningAgents/Source/LearningAgentsEditor/Private/LearningAgentsEditor.cpp b/Engine/Plugins/Experimental/LearningAgents/Source/LearningAgentsEditor/Private/LearningAgentsEditor.cpp
--- a/Engine/Plugins/Experimental/LearningAgents/Source/LearningAgentsEditor/Private/LearningAgentsEditor.cpp
+++ b/Engine/Plugins/Experimental/LearningAgents/Source/LearningAgentsEditor/Private/LearningAgentsEditor.cpp
@@ -3,17 +3,17 @@
 #include "LearningAgentsEditor.h"
 
 #include "Kismet2/KismetEditorUtilities.h"
-#include "LearningAgentsType.h"
+#include "LearningAgentsInteractor.h"
 #include "LearningAgentsTrainer.h"
 #include "LearningAgentsController.h"
 
 void FLearningAgentsEditorModule::StartupModule()
 {
-	// ULearningAgentsType Default Events
-	FKismetEditorUtilities::RegisterAutoGeneratedDefaultEvent(this, ULearningAgentsType::StaticClass(), GET_FUNCTION_NAME_CHECKED(ULearningAgentsType, SetupObservations));
-	FKismetEditorUtilities::RegisterAutoGeneratedDefaultEvent(this, ULearningAgentsType::StaticClass(), GET_FUNCTION_NAME_CHECKED(ULearningAgentsType, SetObservations));
-	FKismetEditorUtilities::RegisterAutoGeneratedDefaultEvent(this, ULearningAgentsType::StaticClass(), GET_FUNCTION_NAME_CHECKED(ULearningAgentsType, SetupActions));
-	FKismetEditorUtilities::RegisterAutoGeneratedDefaultEvent(this, ULearningAgentsType::StaticClass(), GET_FUNCTION_NAME_CHECKED(ULearningAgentsType, GetActions));
+	// ULearningAgentsInteractor Default Events
+	FKismetEditorUtilities::RegisterAutoGeneratedDefaultEvent(this, ULearningAgentsInteractor::StaticClass(), GET_FUNCTION_NAME_CHECKED(ULearningAgentsInteractor, SetupObservations));
+	FKismetEditorUtilities::RegisterAutoGeneratedDefaultEvent(this, ULearningAgentsInteractor::StaticClass(), GET_FUNCTION_NAME_CHECKED(ULearningAgentsInteractor, SetObservations));
+	FKismetEditorUtilities::RegisterAutoGeneratedDefaultEvent(this, ULearningAgentsInteractor::StaticClass(), GET_FUNCTION_NAME_CHECKED(ULearningAgentsInteractor, SetupActions));
+	FKismetEditorUtilities::RegisterAutoGeneratedDefaultEvent(this, ULearningAgentsInteractor::StaticClass(), GET_FUNCTION_NAME_CHECKED(ULearningAgentsInteractor, GetActions));
 	
 	// ULearningAgentsTrainer Default Events
 	FKismetEditorUtilities::RegisterAutoGeneratedDefaultEvent(this, ULearningAgentsTrainer::StaticClass(), GET_FUNCTION_NAME_CHECKED(ULearningAgentsTrainer, SetupRewards));
diff --git a/Engine/Plugins/Experimental/LearningAgents/Source/LearningAgentsTraining/Private/LearningAgentsDataStorage.cpp b/Engine/Plugins/Experimental/LearningAgents/Source/LearningAgentsTraining/Private/LearningAgentsDataStorage.cpp
--- a/Engine/Plugins/Experimental/LearningAgents/Source/LearningAgentsTraining/Private/LearningAgentsDataStorage.cpp
+++ b/Engine/Plugins/Experimental/LearningAgents/Source/LearningAgentsTraining/Private/LearningAgentsDataStorage.cpp
@@ -4,7 +4,7 @@
 
 #include "Engine/EngineTypes.h"
 #include "HAL/FileManager.h"
-#include "LearningAgentsType.h"
+#include "LearningAgentsInteractor.h"
 #include "LearningFeatureObject.h"
 #include "LearningLog.h"
 #include "Misc/FileHelper.h"
@@ -70,27 +70,27 @@ ULearningAgentsDataStorage::ULearningAgentsDataStorage() {}
 ULearningAgentsDataStorage::ULearningAgentsDataStorage(FVTableHelper& Helper) {}
 ULearningAgentsDataStorage::~ULearningAgentsDataStorage() {}
 
-ULearningAgentsRecord* ULearningAgentsDataStorage::CreateRecord(const FName RecordName, ULearningAgentsType* AgentType)
+ULearningAgentsRecord* ULearningAgentsDataStorage::CreateRecord(const FName RecordName, ULearningAgentsInteractor* Interactor)
 {
-	if (!AgentType)
+	if (!Interactor)
 	{
-		UE_LOG(LogLearning, Error, TEXT("%s: AgentType is nullptr."), *GetName());
+		UE_LOG(LogLearning, Error, TEXT("%s: Interactor is nullptr."), *GetName());
 		return nullptr;
 	}
 
-	if (!AgentType->IsSetup())
+	if (!Interactor->IsSetup())
 	{
-		UE_LOG(LogLearning, Error, TEXT("%s: %s's Setup must be run before it can be used."), *GetName(), *AgentType->GetName());
+		UE_LOG(LogLearning, Error, TEXT("%s: %s's Setup must be run before it can be used."), *GetName(), *Interactor->GetName());
 		return nullptr;
 	}
 
 	ULearningAgentsRecord* NewRecord = NewObject<ULearningAgentsRecord>(this, RecordName);
-	NewRecord->MetaData.AgentType = AgentType;
+	NewRecord->MetaData.Interactor = Interactor;
 	NewRecord->MetaData.CreatedOn = FDateTime::Now();
 	NewRecord->MetaData.CreatedOnUtc = FDateTime::UtcNow();
 	NewRecord->MetaData.RecordName = RecordName;
 
-	NewRecord->Init(AgentType->GetObservationFeature().DimNum(), AgentType->GetActionFeature().DimNum());
+	NewRecord->Init(Interactor->GetObservationFeature().DimNum(), Interactor->GetActionFeature().DimNum());
 	Records.Add(NewRecord);
 
 	return NewRecord;
@@ -101,11 +101,11 @@ const TArray<ULearningAgentsRecord*>& ULearningAgentsDataStorage::GetAllRecords(
 	return Records;
 }
 
-int32 ULearningAgentsDataStorage::LoadAllRecords(ULearningAgentsType* AgentType, const FDirectoryPath& Directory)
+int32 ULearningAgentsDataStorage::LoadAllRecords(ULearningAgentsInteractor* Interactor, const FDirectoryPath& Directory)
 {
-	if (!AgentType)
+	if (!Interactor)
 	{
-		UE_LOG(LogLearning, Error, TEXT("%s: AgentType is nullptr."), *GetName());
+		UE_LOG(LogLearning, Error, TEXT("%s: Interactor is nullptr."), *GetName());
 		return 0;
 	}
 
@@ -120,23 +120,23 @@ int32 ULearningAgentsDataStorage::LoadAllRecords(ULearningAgentsType* AgentType,
 
 	for (const FString& DataFile : DataFiles)
 	{
-		LoadRecord(AgentType, Directory, DataFile);
+		LoadRecord(Interactor, Directory, DataFile);
 	}
 
 	return DataFiles.Num();
 }
 
-ULearningAgentsRecord* ULearningAgentsDataStorage::LoadRecord(ULearningAgentsType* AgentType, const FDirectoryPath& Directory, const FString& Filename)
+ULearningAgentsRecord* ULearningAgentsDataStorage::LoadRecord(ULearningAgentsInteractor* Interactor, const FDirectoryPath& Directory, const FString& Filename)
 {
-	if (!AgentType)
+	if (!Interactor)
 	{
-		UE_LOG(LogLearning, Error, TEXT("%s: AgentType is nullptr."), *GetName());
+		UE_LOG(LogLearning, Error, TEXT("%s: Interactor is nullptr."), *GetName());
 		return nullptr;
 	}
 
-	if (!AgentType->IsSetup())
+	if (!Interactor->IsSetup())
 	{
-		UE_LOG(LogLearning, Error, TEXT("%s: %s's Setup must be run before it can be used."), *GetName(), *AgentType->GetName());
+		UE_LOG(LogLearning, Error, TEXT("%s: %s's Setup must be run before it can be used."), *GetName(), *Interactor->GetName());
 		return nullptr;
 	}
 
@@ -169,7 +169,7 @@ ULearningAgentsRecord* ULearningAgentsDataStorage::LoadRecord(ULearningAgentsTyp
 
 	int32 ObsNum = -1;
 	Reader.Serialize(&ObsNum, sizeof(int32));
-	int32 AgentObsNum = AgentType->GetObservationFeature().DimNum();
+	int32 AgentObsNum = Interactor->GetObservationFeature().DimNum();
 	if (AgentObsNum != ObsNum)
 	{
 		UE_LOG(LogLearning, Error, TEXT("%s: Observation data dimension size %d incompatible with agent's observation dimension size %d"), *GetName(), ObsNum, AgentObsNum);
@@ -178,7 +178,7 @@ ULearningAgentsRecord* ULearningAgentsDataStorage::LoadRecord(ULearningAgentsTyp
 
 	int32 ActNum = -1;
 	Reader.Serialize(&ActNum, sizeof(int32));
-	int32 AgentActionNum = AgentType->GetActionFeature().DimNum();
+	int32 AgentActionNum = Interactor->GetActionFeature().DimNum();
 	if (AgentActionNum != ActNum)
 	{
 		UE_LOG(LogLearning, Error, TEXT("%s: Action data dimension size %d incompatible with agent's action dimension size %d"), *GetName(), ObsNum, AgentObsNum);
@@ -186,7 +186,7 @@ ULearningAgentsRecord* ULearningAgentsDataStorage::LoadRecord(ULearningAgentsTyp
 	}
 
 	// Load all the data into the first chunk
-	ULearningAgentsRecord* Record = CreateRecord(FName(Filename), AgentType);
+	ULearningAgentsRecord* Record = CreateRecord(FName(Filename), Interactor);
 	Reader.Serialize(&Record->MetaData.CreatedOn, sizeof(FDateTime));
 	Reader.Serialize(&Record->MetaData.CreatedOnUtc, sizeof(FDateTime));
 	Record->MetaData.bWasLoadedFromFile = true;
@@ -248,7 +248,7 @@ void ULearningAgentsDataStorage::SaveRecord(const FDirectoryPath& Directory, ULe
 		UE::Learning::Array::Serialize(Writer, Record->Actions[ChunkIndex]);
 	}
 
-	const FString FileName = Record->MetaData.RecordName.IsNone() ? Record->MetaData.AgentType->GetName() : Record->MetaData.RecordName.ToString();
+	const FString FileName = Record->MetaData.RecordName.IsNone() ? Record->MetaData.Interactor->GetName() : Record->MetaData.RecordName.ToString();
 	const FString OutputFilePath = Directory.Path + FGenericPlatformMisc::GetDefaultPathSeparator() +
 		(!Record->MetaData.bWasLoadedFromFile && bPrependUtcTimeStamp ? FDateTime::UtcNow().ToString() + "_" : "") + FileName + "." + FileExtension;
 	const bool bSuccess = FFileHelper::SaveArrayToFile(Bytes, *OutputFilePath, &IFileManager::Get(), EFileWrite::FILEWRITE_Append);
diff --git a/Engine/Plugins/Experimental/LearningAgents/Source/LearningAgentsTraining/Private/LearningAgentsImitationTrainer.cpp b/Engine/Plugins/Experimental/LearningAgents/Source/LearningAgentsTraining/Private/LearningAgentsImitationTrainer.cpp
--- a/Engine/Plugins/Experimental/LearningAgents/Source/LearningAgentsTraining/Private/LearningAgentsImitationTrainer.cpp
+++ b/Engine/Plugins/Experimental/LearningAgents/Source/LearningAgentsTraining/Private/LearningAgentsImitationTrainer.cpp
@@ -4,7 +4,7 @@
 
 #include "GenericPlatform/GenericPlatformMisc.h"
 #include "HAL/FileManager.h"
-#include "LearningAgentsType.h"
+#include "LearningAgentsInteractor.h"
 #include "LearningArrayMap.h"
 #include "LearningExperience.h"
 #include "LearningFeatureObject.h"
diff --git a/Engine/Plugins/Experimental/LearningAgents/Source/LearningAgentsTraining/Private/LearningAgentsRecorder.cpp b/Engine/Plugins/Experimental/LearningAgents/Source/LearningAgentsTraining/Private/LearningAgentsRecorder.cpp
--- a/Engine/Plugins/Experimental/LearningAgents/Source/LearningAgentsTraining/Private/LearningAgentsRecorder.cpp
+++ b/Engine/Plugins/Experimental/LearningAgents/Source/LearningAgentsTraining/Private/LearningAgentsRecorder.cpp
@@ -3,7 +3,7 @@
 #include "LearningAgentsRecorder.h"
 
 #include "LearningAgentsManager.h"
-#include "LearningAgentsType.h"
+#include "LearningAgentsInteractor.h"
 #include "LearningAgentsDataStorage.h"
 #include "LearningFeatureObject.h"
 #include "LearningLog.h"
@@ -25,7 +25,7 @@ void ULearningAgentsRecorder::EndPlay(const EEndPlayReason::Type EndPlayReason)
 	Super::EndPlay(EndPlayReason);
 }
 
-void ULearningAgentsRecorder::SetupRecorder(ALearningAgentsManager* InAgentManager, ULearningAgentsType* InAgentType)
+void ULearningAgentsRecorder::SetupRecorder(ALearningAgentsManager* InAgentManager, ULearningAgentsInteractor* InInteractor)
 {
 	if (IsSetup())
 	{
@@ -47,19 +47,19 @@ void ULearningAgentsRecorder::SetupRecorder(ALearningAgentsManager* InAgentManag
 
 	AgentManager = InAgentManager;
 
-	if (!InAgentType)
+	if (!InInteractor)
 	{
-		UE_LOG(LogLearning, Error, TEXT("%s: InAgentType is nullptr."), *GetName());
+		UE_LOG(LogLearning, Error, TEXT("%s: InInteractor is nullptr."), *GetName());
 		return;
 }
 
-	if (!InAgentType->IsSetup())
+	if (!InInteractor->IsSetup())
 	{
-		UE_LOG(LogLearning, Error, TEXT("%s: %s's Setup must be run before it can be used."), *GetName(), *InAgentType->GetName());
+		UE_LOG(LogLearning, Error, TEXT("%s: %s's Setup must be run before it can be used."), *GetName(), *InInteractor->GetName());
 		return;
 	}
 
-	AgentType = InAgentType;
+	Interactor = InInteractor;
 
 	DataStorage = NewObject<ULearningAgentsDataStorage>(this, TEXT("DataStorage"));
 #if WITH_EDITOR
@@ -77,7 +77,7 @@ bool ULearningAgentsRecorder::AddAgent(const int32 AgentId)
 
 	if (bSuccess && IsRecording() && !CurrentRecords.Contains(AgentId))
 	{
-		CurrentRecords.Add(AgentId, DataStorage->CreateRecord(FName(AgentType->GetName() + "_id" + FString::FromInt(AgentId)), AgentType));
+		CurrentRecords.Add(AgentId, DataStorage->CreateRecord(FName(Interactor->GetName() + "_id" + FString::FromInt(AgentId)), Interactor));
 	}
 
 	return bSuccess;
@@ -107,8 +107,8 @@ void ULearningAgentsRecorder::AddExperience()
 	for (const int32 AgentId : AddedAgentSet)
 	{
 		CurrentRecords[AgentId]->AddExperience(
-			AgentType->GetObservationFeature().FeatureBuffer()[AgentId], 
-			AgentType->GetActionFeature().FeatureBuffer()[AgentId]);
+			Interactor->GetObservationFeature().FeatureBuffer()[AgentId], 
+			Interactor->GetActionFeature().FeatureBuffer()[AgentId]);
 	}
 }
 
@@ -156,7 +156,7 @@ void ULearningAgentsRecorder::BeginRecording()
 
 	for (const int32 AgentId : AddedAgentSet)
 	{
-		CurrentRecords.Add(AgentId, DataStorage->CreateRecord(FName(AgentType->GetName() + "_id" + FString::FromInt(AgentId)), AgentType));
+		CurrentRecords.Add(AgentId, DataStorage->CreateRecord(FName(Interactor->GetName() + "_id" + FString::FromInt(AgentId)), Interactor));
 	}
 
 	bIsRecording = true;
diff --git a/Engine/Plugins/Experimental/LearningAgents/Source/LearningAgentsTraining/Private/LearningAgentsTrainer.cpp b/Engine/Plugins/Experimental/LearningAgents/Source/LearningAgentsTraining/Private/LearningAgentsTrainer.cpp
--- a/Engine/Plugins/Experimental/LearningAgents/Source/LearningAgentsTraining/Private/LearningAgentsTrainer.cpp
+++ b/Engine/Plugins/Experimental/LearningAgents/Source/LearningAgentsTraining/Private/LearningAgentsTrainer.cpp
@@ -4,7 +4,7 @@
 
 #include "GenericPlatform/GenericPlatformMisc.h"
 #include "HAL/FileManager.h"
-#include "LearningAgentsType.h"
+#include "LearningAgentsInteractor.h"
 #include "LearningAgentsManager.h"
 #include "LearningAgentsRewards.h"
 #include "LearningAgentsCompletions.h"
@@ -58,7 +58,7 @@ ULearningAgentsTrainer::~ULearningAgentsTrainer() {}
 
 void ULearningAgentsTrainer::SetupTrainer(
 	ALearningAgentsManager* InAgentManager,
-	ULearningAgentsType* InAgentType,
+	ULearningAgentsInteractor* InInteractor,
 	ULearningAgentsPolicy* InPolicy,
 	ULearningAgentsCritic* InCritic,
 	const FLearningAgentsTrainerSettings& TrainerSettings)
@@ -83,19 +83,19 @@ void ULearningAgentsTrainer::SetupTrainer(
 
 	AgentManager = InAgentManager;
 
-	if (!InAgentType)
+	if (!InInteractor)
 	{
-		UE_LOG(LogLearning, Error, TEXT("%s: InAgentType is nullptr."), *GetName());
+		UE_LOG(LogLearning, Error, TEXT("%s: InInteractor is nullptr."), *GetName());
 		return;
 	}
 
-	if (!InAgentType->IsSetup())
+	if (!InInteractor->IsSetup())
 	{
-		UE_LOG(LogLearning, Error, TEXT("%s: %s's Setup must be run before it can be used."), *GetName(), *InAgentType->GetName());
+		UE_LOG(LogLearning, Error, TEXT("%s: %s's Setup must be run before it can be used."), *GetName(), *InInteractor->GetName());
 		return;
 	}
 
-	AgentType = InAgentType;
+	Interactor = InInteractor;
 
 	if (!InPolicy)
 	{
@@ -146,16 +146,16 @@ void ULearningAgentsTrainer::SetupTrainer(
 	EpisodeBuffer->Resize(
 		AgentManager->GetMaxInstanceNum(),
 		TrainerSettings.MaxStepNum,
-		AgentType->GetObservationFeature().DimNum(),
-		AgentType->GetActionFeature().DimNum());
+		Interactor->GetObservationFeature().DimNum(),
+		Interactor->GetActionFeature().DimNum());
 
 	MaxStepsCompletion = TrainerSettings.MaxStepsCompletion;
 
 	// Create Replay Buffer
 	ReplayBuffer = MakeUnique<UE::Learning::FReplayBuffer>();
 	ReplayBuffer->Resize(
-		AgentType->GetObservationFeature().DimNum(),
-		AgentType->GetActionFeature().DimNum(),
+		Interactor->GetObservationFeature().DimNum(),
+		Interactor->GetActionFeature().DimNum(),
 		TrainerSettings.MaximumRecordedEpisodesPerIteration,
 		TrainerSettings.MaximumRecordedStepsPerIteration);
 
@@ -369,7 +369,7 @@ void ULearningAgentsTrainer::BeginTraining(
 		PPONetworkSettings.CriticActivationFunction = UE::Learning::Agents::GetActivationFunction(CriticSettings.ActivationFunction);
 	}
 
-	// We assume that if the critic has been setup on the agent type, then
+	// We assume that if the critic has been setup on the agent interactor, then
 	// the user wants the critic network to be synced during training.
 	UE::Learning::EPPOTrainerFlags TrainerFlags = 
 		Critic ?
@@ -552,8 +552,8 @@ void ULearningAgentsTrainer::ProcessExperience()
 		return;
 	}
 
-	UE::Learning::FFeatureObject& Observations = AgentType->GetObservationFeature();
-	UE::Learning::FFeatureObject& Actions = AgentType->GetActionFeature();
+	UE::Learning::FFeatureObject& Observations = Interactor->GetObservationFeature();
+	UE::Learning::FFeatureObject& Actions = Interactor->GetActionFeature();
 
 	// Add Experience to Episode Buffer
 	EpisodeBuffer->Push(
@@ -578,11 +578,11 @@ void ULearningAgentsTrainer::ProcessExperience()
 	if (ResetBuffer->GetResetInstanceNum() > 0)
 	{
 		// Encode Observations for completed Instances
-		AgentType->SetObservations(ResetBuffer->GetResetInstances().ToArray());
+		Interactor->SetObservations(ResetBuffer->GetResetInstances().ToArray());
 		Observations.Encode(ResetBuffer->GetResetInstances());
 		
 #if ENABLE_VISUAL_LOG
-		for (const ULearningAgentsObservation* ObservationObject : AgentType->GetObservationObjects())
+		for (const ULearningAgentsObservation* ObservationObject : Interactor->GetObservationObjects())
 		{
 			if (ObservationObject)
 			{
@@ -664,19 +664,37 @@ void ULearningAgentsTrainer::ResetAllEpisodes()
 	EpisodeBuffer->Reset(AddedAgentSet);
 }
 
-void ULearningAgentsTrainer::RunTraining()
+void ULearningAgentsTrainer::RunTraining(
+	const FLearningAgentsTrainerTrainingSettings& TrainingSettings,
+	const FLearningAgentsTrainerGameSettings& TrainerGameSettings,
+	const FLearningAgentsCriticSettings& CriticSettings,
+	const bool bReinitializePolicyNetwork,
+	const bool bReinitializeCriticNetwork)
 {
-	if (!IsTraining())
+	if (!IsSetup())
 	{
-		BeginTraining();
+		UE_LOG(LogLearning, Error, TEXT("%s: Setup not run."), *GetName());
+		return;
 	}
-	else
+
+	// If we aren't training yet, then start training and do the first inference step.
+	if (!IsTraining())
 	{
-		EvaluateCompletions();
-		EvaluateRewards();
-		ProcessExperience();
+		BeginTraining(TrainingSettings, TrainerGameSettings, CriticSettings, bReinitializePolicyNetwork, bReinitializeCriticNetwork);
+
+		if (!IsTraining())
+		{
+			// If IsTraining is false, then BeginTraining must have failed and we can't continue.
+			return;
+		}
+
+		Policy->RunInference();
 	}
 
+	// Otherwise, do the regular training process.
+	EvaluateCompletions();
+	EvaluateRewards();
+	ProcessExperience();
 	Policy->RunInference();
 }
 
diff --git a/Engine/Plugins/Experimental/LearningAgents/Source/LearningAgentsTraining/Public/LearningAgentsDataStorage.h b/Engine/Plugins/Experimental/LearningAgents/Source/LearningAgentsTraining/Public/LearningAgentsDataStorage.h
--- a/Engine/Plugins/Experimental/LearningAgents/Source/LearningAgentsTraining/Public/LearningAgentsDataStorage.h
+++ b/Engine/Plugins/Experimental/LearningAgents/Source/LearningAgentsTraining/Public/LearningAgentsDataStorage.h
@@ -13,7 +13,7 @@
 #include "LearningAgentsDataStorage.generated.h"
 
 struct FDirectoryPath;
-class ULearningAgentsType;
+class ULearningAgentsInteractor;
 
 /** Metadata for an agent record. */
 USTRUCT(BlueprintType)
@@ -21,9 +21,9 @@ struct FRecordMetaData
 {
 	GENERATED_BODY()
 
-	/** The agent type the data was recorded from. Determine the shape of observations and actions. */
+	/** The agent interactor the data was recorded from. Determine the shape of observations and actions. */
 	UPROPERTY(BlueprintReadOnly, VisibleAnywhere, Category = "LearningAgents")
-	TObjectPtr<ULearningAgentsType> AgentType;
+	TObjectPtr<ULearningAgentsInteractor> Interactor;
 
 	/** The time the data was recorded (local). */
 	UPROPERTY(BlueprintReadOnly, VisibleAnywhere, Category = "LearningAgents")
@@ -104,21 +104,21 @@ public:
 	ULearningAgentsDataStorage(FVTableHelper& Helper);
 	virtual ~ULearningAgentsDataStorage();
 
-	/** Create a new record with the given name for the agent type. */
+	/** Create a new record with the given name for the agent interactor. */
 	UFUNCTION(BlueprintCallable, Category = "LearningAgents")
-	ULearningAgentsRecord* CreateRecord(const FName RecordName, ULearningAgentsType* AgentType);
+	ULearningAgentsRecord* CreateRecord(const FName RecordName, ULearningAgentsInteractor* Interactor);
 
 	/** Get all currently loaded records. */
 	UFUNCTION(BlueprintPure, Category = "LearningAgents")
 	const TArray<ULearningAgentsRecord*>& GetAllRecords() const;
 
-	/** Load all records from a given directory, ensuring they are valid for the given agent type. */
+	/** Load all records from a given directory, ensuring they are valid for the given agent interactor. */
 	UFUNCTION(BlueprintCallable, Category = "LearningAgents")
-	int32 LoadAllRecords(ULearningAgentsType* AgentType, const FDirectoryPath& Directory);
+	int32 LoadAllRecords(ULearningAgentsInteractor* Interactor, const FDirectoryPath& Directory);
 
-	/** Load a record from a given directory and filename, ensuring the data is valid for the given agent type. */
+	/** Load a record from a given directory and filename, ensuring the data is valid for the given agent interactor. */
 	UFUNCTION(BlueprintCallable, Category = "LearningAgents")
-	ULearningAgentsRecord* LoadRecord(ULearningAgentsType* AgentType, const FDirectoryPath& Directory, const FString& Filename);
+	ULearningAgentsRecord* LoadRecord(ULearningAgentsInteractor* Interactor, const FDirectoryPath& Directory, const FString& Filename);
 
 	/** Save all records which have unsaved changes to the given directory. */
 	UFUNCTION(BlueprintCallable, Category = "LearningAgents")
diff --git a/Engine/Plugins/Experimental/LearningAgents/Source/LearningAgentsTraining/Public/LearningAgentsImitationTrainer.h b/Engine/Plugins/Experimental/LearningAgents/Source/LearningAgentsTraining/Public/LearningAgentsImitationTrainer.h
--- a/Engine/Plugins/Experimental/LearningAgents/Source/LearningAgentsTraining/Public/LearningAgentsImitationTrainer.h
+++ b/Engine/Plugins/Experimental/LearningAgents/Source/LearningAgentsTraining/Public/LearningAgentsImitationTrainer.h
@@ -16,7 +16,7 @@ namespace UE::Learning
 	struct FSharedMemoryImitationTrainer;
 }
 
-class ULearningAgentsType;
+class ULearningAgentsInteractor;
 class ULearningAgentsPolicy;
 
 /** The configurable settings for the training process. */
@@ -71,7 +71,7 @@ public:
 * Imitation training is typically much faster than reinforcement learning, but requires gathering large amounts of
 * data in order to generalize. This can be used to initialize a reinforcement learning policy to speed up initial
 * exploration.
-* @see ULearningAgentsType to understand how observations and actions work.
+* @see ULearningAgentsInteractor to understand how observations and actions work.
 * @see ULearningAgentsRecorder to understand how to make new recordings.
 * @see ULearningAgentsDataStorage to understand how to retrieve previous recordings.
 */
diff --git a/Engine/Plugins/Experimental/LearningAgents/Source/LearningAgentsTraining/Public/LearningAgentsRecorder.h b/Engine/Plugins/Experimental/LearningAgents/Source/LearningAgentsTraining/Public/LearningAgentsRecorder.h
--- a/Engine/Plugins/Experimental/LearningAgents/Source/LearningAgentsTraining/Public/LearningAgentsRecorder.h
+++ b/Engine/Plugins/Experimental/LearningAgents/Source/LearningAgentsTraining/Public/LearningAgentsRecorder.h
@@ -32,10 +32,10 @@ public:
 
 	/**
 	* Initializes this object and runs the setup functions for the underlying data storage.
-	* @param InAgentType The agent type we are recording with.
+	* @param InInteractor The agent interactor we are recording with.
 	*/
 	UFUNCTION(BlueprintCallable, Category = "LearningAgents")
-	void SetupRecorder(ALearningAgentsManager* InAgentManager, ULearningAgentsType* InAgentType);
+	void SetupRecorder(ALearningAgentsManager* InAgentManager, ULearningAgentsInteractor* InInteractor);
 
 public:
 
@@ -49,9 +49,9 @@ public:
 public:
 
 	/**
-	* Adds experience to the added agents' recordings. Call this after ULearningAgentsType::EncodeObservations and
+	* Adds experience to the added agents' recordings. Call this after ULearningAgentsInteractor::EncodeObservations and
 	* either ULearningAgentsController::EncodeActions (if recording a human/AI demonstration) or
-	* ULearningAgentsType::DecodeActions (if recording another policy).
+	* ULearningAgentsInteractor::DecodeActions (if recording another policy).
 	*/
 	UFUNCTION(BlueprintCallable, Category = "LearningAgents")
 	void AddExperience();
@@ -71,9 +71,9 @@ public:
 // ----- Private Data ----- 
 private:
 
-	/** The agent type this recorder is associated with. */
+	/** The agent interactor this recorder is associated with. */
 	UPROPERTY(VisibleAnywhere, Transient, Category = "LearningAgents")
-	TObjectPtr<ULearningAgentsType> AgentType;
+	TObjectPtr<ULearningAgentsInteractor> Interactor;
 
 // ----- Recorder Configuration -----
 private:
diff --git a/Engine/Plugins/Experimental/LearningAgents/Source/LearningAgentsTraining/Public/LearningAgentsTrainer.h b/Engine/Plugins/Experimental/LearningAgents/Source/LearningAgentsTraining/Public/LearningAgentsTrainer.h
--- a/Engine/Plugins/Experimental/LearningAgents/Source/LearningAgentsTraining/Public/LearningAgentsTrainer.h
+++ b/Engine/Plugins/Experimental/LearningAgents/Source/LearningAgentsTraining/Public/LearningAgentsTrainer.h
@@ -26,7 +26,7 @@ namespace UE::Learning
 }
 
 class ALearningAgentsManager;
-class ULearningAgentsType;
+class ULearningAgentsInteractor;
 class ULearningAgentsCompletion;
 class ULearningAgentsReward;
 class ULearningAgentsPolicy;
@@ -193,7 +193,7 @@ public:
 * which will initialize the underlying data structures, and you need to call AddAgent for each agent you want to gather
 * training data from.
 *
-* @see ULearningAgentsType to understand how observations and actions work.
+* @see ULearningAgentsInteractor to understand how observations and actions work.
 */
 UCLASS(Abstract, BlueprintType, Blueprintable)
 class LEARNINGAGENTSTRAINING_API ULearningAgentsTrainer : public ULearningAgentsManagerComponent
@@ -214,7 +214,7 @@ public:
 	/**
 	* Initializes this object and runs the setup functions for rewards and completions.
 	* @param InAgentManager The agent manager we are associated with.
-	* @param InAgentType The agent type we are training with.
+	* @param InInteractor The agent interactor we are training with.
 	* @param InPolicy The policy to be trained.
 	* @param InCritic Optional - only needs to be provided if we want the critic to be accessible at runtime.
 	* @param TrainerSettings The trainer settings to use.
@@ -222,7 +222,7 @@ public:
 	UFUNCTION(BlueprintCallable, Category = "LearningAgents")
 	void SetupTrainer(
 		ALearningAgentsManager* InAgentManager,
-		ULearningAgentsType* InAgentType,
+		ULearningAgentsInteractor* InInteractor,
 		ULearningAgentsPolicy* InPolicy,
 		ULearningAgentsCritic* InCritic = nullptr,
 		const FLearningAgentsTrainerSettings& TrainerSettings = FLearningAgentsTrainerSettings());
@@ -356,9 +356,19 @@ public:
 	* Convenience function that runs a basic training loop. If training has not been started, it will start it, and 
 	* then call RunInference. On each following call to this function, it will call EvaluateRewards, 
 	* EvaluateCompletions, and ProcessExperience, followed by RunInference.
+	* @param TrainingSettings The settings for this training run.
+	* @param GameSettings The settings that will affect the game's simulation.
+	* @param CriticSettings The settings for the critic (if we are using one).
+	* @param bReinitializePolicyNetwork If true, reinitialize the policy. Set this to false if your policy is pre-trained, e.g. with imitation learning.
+	* @param bReinitializeCriticNetwork If true, reinitialize the critic. Set this to false if your critic is pre-trained.
 	*/
 	UFUNCTION(BlueprintCallable, Category = "LearningAgents")
-	void RunTraining();
+	void RunTraining(
+		const FLearningAgentsTrainerTrainingSettings& TrainingSettings = FLearningAgentsTrainerTrainingSettings(),
+		const FLearningAgentsTrainerGameSettings& GameSettings = FLearningAgentsTrainerGameSettings(),
+		const FLearningAgentsCriticSettings& CriticSettings = FLearningAgentsCriticSettings(),
+		const bool bReinitializePolicyNetwork = true,
+		const bool bReinitializeCriticNetwork = true);
 
 	/**
 	* Gets the current reward for an agent according to the critic. Should be called only after EvaluateRewards.
@@ -383,9 +393,9 @@ public:
 // ----- Private Data ----- 
 private:
 
-	/** The agent type associated with this component. */
+	/** The agent interactor associated with this component. */
 	UPROPERTY(VisibleAnywhere, Transient, Category = "LearningAgents")
-	TObjectPtr<ULearningAgentsType> AgentType;
+	TObjectPtr<ULearningAgentsInteractor> Interactor;
 
 	/** The current policy for experience gathering. */
 	UPROPERTY(VisibleAnywhere, Transient, Category = "LearningAgents")
