LearningAgents: Refactoring of how Recordings are stored and used. Updated serialization and added magic number and version number to network and recording files. Updated how paths are handled for the training and recording processes. Moved all generated artefacts from the plugin intermediate directory to the project intermediate directory.

diff --git a/Engine/Plugins/Experimental/LearningAgents/Source/Learning/Private/LearningNeuralNetwork.cpp b/Engine/Plugins/Experimental/LearningAgents/Source/Learning/Private/LearningNeuralNetwork.cpp
--- a/Engine/Plugins/Experimental/LearningAgents/Source/Learning/Private/LearningNeuralNetwork.cpp
+++ b/Engine/Plugins/Experimental/LearningAgents/Source/Learning/Private/LearningNeuralNetwork.cpp
@@ -66,58 +66,39 @@ namespace UE::Learning
 		return Weights[0].Num<1>();
 	}
 
-	void FNeuralNetwork::DeserializeFromBytes(const TLearningArrayView<1, const uint8> RawBytes)
+	void FNeuralNetwork::DeserializeFromBytes(int32& InOutOffset, const TLearningArrayView<1, const uint8> RawBytes)
 	{
 		UE_LEARNING_TRACE_CPUPROFILER_EVENT_SCOPE(Learning::FNeuralNetwork::DeserializeFromBytes);
 
-		const int32 TotalByteNum = GetSerializationByteNum(GetInputNum(), GetOutputNum(), GetHiddenNum(), GetLayerNum());
+		int32 ActivationFunctionInt;
+		UE::Learning::DeserializeFromBytes(InOutOffset, RawBytes, ActivationFunctionInt);
+		ActivationFunction = (EActivationFunction)ActivationFunctionInt;
 
-		UE_LEARNING_CHECK(RawBytes.Num() == TotalByteNum);
+		int32 LayerNum;
+		UE::Learning::DeserializeFromBytes(InOutOffset, RawBytes, LayerNum);
 
-		const int32 LayerNum = Weights.Num();
-
-		int32 Offset = 0;
+		Weights.SetNum(LayerNum);
+		Biases.SetNum(LayerNum);
 
 		for (int32 LayerIdx = 0; LayerIdx < LayerNum; LayerIdx++)
 		{
-			const int32 WeightByteNum = Weights[LayerIdx].Num() * sizeof(float);
-			FMemory::Memcpy((uint8*)Weights[LayerIdx].GetData(), &RawBytes[Offset], WeightByteNum);
-			Array::Check(Weights[LayerIdx]);
-			Offset += WeightByteNum;
-
-			const int32 BiasByteNum = Biases[LayerIdx].Num() * sizeof(float);
-			FMemory::Memcpy((uint8*)Biases[LayerIdx].GetData(), &RawBytes[Offset], BiasByteNum);
-			Array::Check(Biases[LayerIdx]);
-			Offset += BiasByteNum;
+			UE::Learning::Array::DeserializeFromBytes(InOutOffset, RawBytes, Weights[LayerIdx]);
+			UE::Learning::Array::DeserializeFromBytes(InOutOffset, RawBytes, Biases[LayerIdx]);
 		}
-
-		UE_LEARNING_CHECK(Offset == TotalByteNum);
 	}
 
-	void FNeuralNetwork::SerializeToBytes(TLearningArrayView<1, uint8> OutRawBytes) const
+	void FNeuralNetwork::SerializeToBytes(int32& InOutOffset, TLearningArrayView<1, uint8> OutRawBytes) const
 	{
 		UE_LEARNING_TRACE_CPUPROFILER_EVENT_SCOPE(Learning::FNeuralNetwork::SerializeToBytes);
 
-		const int32 TotalByteNum = GetSerializationByteNum(GetInputNum(), GetOutputNum(), GetHiddenNum(), GetLayerNum());
-
-		UE_LEARNING_CHECK(OutRawBytes.Num() == TotalByteNum);
-
-		const int32 LayerNum = Weights.Num();
+		UE::Learning::SerializeToBytes(InOutOffset, OutRawBytes, (int32)ActivationFunction);
+		UE::Learning::SerializeToBytes(InOutOffset, OutRawBytes, GetLayerNum());
 
-		int32 Offset = 0;
-
-		for (int32 LayerIdx = 0; LayerIdx < LayerNum; LayerIdx++)
+		for (int32 LayerIdx = 0; LayerIdx < GetLayerNum(); LayerIdx++)
 		{
-			const int32 WeightByteNum = Weights[LayerIdx].Num() * sizeof(float);
-			FMemory::Memcpy(&OutRawBytes[Offset], (uint8*)Weights[LayerIdx].GetData(), WeightByteNum);
-			Offset += WeightByteNum;
-
-			const int32 BiasByteNum = Biases[LayerIdx].Num() * sizeof(float);
-			FMemory::Memcpy(&OutRawBytes[Offset], (uint8*)Biases[LayerIdx].GetData(), BiasByteNum);
-			Offset += BiasByteNum;
+			UE::Learning::Array::SerializeToBytes(InOutOffset, OutRawBytes, Weights[LayerIdx]);
+			UE::Learning::Array::SerializeToBytes(InOutOffset, OutRawBytes, Biases[LayerIdx]);
 		}
-
-		UE_LEARNING_CHECK(Offset == TotalByteNum);
 	}
 
 	int32 FNeuralNetwork::GetSerializationByteNum(
@@ -127,18 +108,20 @@ namespace UE::Learning
 		const int32 LayerNum)
 	{
 		int32 Total = 0;
-
-		Total += InputNum * HiddenNum * sizeof(float);
-		Total += HiddenNum * sizeof(float);
+			
+		Total += sizeof(int32);  // Activation
+		Total += sizeof(int32);  // Layer Num
+		Total += UE::Learning::Array::SerializationByteNum<2, float>({ InputNum, HiddenNum });
+		Total += UE::Learning::Array::SerializationByteNum<1, float>({ HiddenNum });
 
 		for (int32 LayerIdx = 0; LayerIdx < LayerNum - 2; LayerIdx++)
 		{
-			Total += HiddenNum * HiddenNum * sizeof(float);
-			Total += HiddenNum * sizeof(float);
+			Total += UE::Learning::Array::SerializationByteNum<2, float>({ HiddenNum, HiddenNum });
+			Total += UE::Learning::Array::SerializationByteNum<1, float>({ HiddenNum });
 		}
 
-		Total += HiddenNum * OutputNum * sizeof(float);
-		Total += OutputNum * sizeof(float);
+		Total += UE::Learning::Array::SerializationByteNum<2, float>({ HiddenNum, OutputNum });
+		Total += UE::Learning::Array::SerializationByteNum<1, float>({ OutputNum });
 
 		return Total;
 	}
diff --git a/Engine/Plugins/Experimental/LearningAgents/Source/Learning/Public/LearningArray.h b/Engine/Plugins/Experimental/LearningAgents/Source/Learning/Public/LearningArray.h
--- a/Engine/Plugins/Experimental/LearningAgents/Source/Learning/Public/LearningArray.h
+++ b/Engine/Plugins/Experimental/LearningAgents/Source/Learning/Public/LearningArray.h
@@ -291,6 +291,26 @@ namespace UE::Learning
 	*/
 	LEARNING_API void SlicedParallelFor(const int32 Num, const int32 MinSliceElementNum, const TFunctionRef<void(const int32 Start, const int32 Length)> Body);
 
+	/**
+	* Serialize an integer to bytes
+	*/
+	static inline void SerializeToBytes(int32& InOutOffset, TLearningArrayView<1, uint8> Bytes, const int32 InValue)
+	{
+		UE_LEARNING_CHECK(InOutOffset + sizeof(int32) <= Bytes.Num());
+		FMemory::Memcpy(&Bytes[InOutOffset], (uint8*)&InValue, sizeof(int32));
+		InOutOffset += sizeof(int32);
+	}
+
+	/**
+	* Serialize an integer from bytes
+	*/
+	static inline void DeserializeFromBytes(int32& InOutOffset, TLearningArrayView<1, const uint8> Bytes, int32& OutValue)
+	{
+		UE_LEARNING_CHECK(InOutOffset + sizeof(int32) <= Bytes.Num());
+		FMemory::Memcpy((uint8*)&OutValue, &Bytes[InOutOffset], sizeof(int32));
+		InOutOffset += sizeof(int32);
+	}
+
 	/**
 	* Some additional functions that act on arrays such as copy, set, zero, etc.
 	* 
@@ -945,6 +965,55 @@ namespace UE::Learning
 			}
 		}
 
+		/**
+		* Serialization byte num
+		*/
+		template<uint8 InDimNum, typename InElementType>
+		inline int32 SerializationByteNum(const TLearningArrayShape<InDimNum> Shape)
+		{
+			return sizeof(int32) + sizeof(int32) * InDimNum + sizeof(InElementType) * Shape.Total();
+		}
+
+		/**
+		* Serialize an array to bytes
+		*/
+		template<uint8 InDimNum, typename InElementType, typename Allocator>
+		inline void SerializeToBytes(int32& InOutOffset, TLearningArrayView<1, uint8> Bytes, const TLearningArray<InDimNum, InElementType, Allocator>& InArray)
+		{
+			Learning::SerializeToBytes(InOutOffset, Bytes, (int32)InDimNum);
+			for (uint8 ShapeIdx = 0; ShapeIdx < InDimNum; ShapeIdx++)
+			{
+				Learning::SerializeToBytes(InOutOffset, Bytes, InArray.Shape()[ShapeIdx]);
+			}
+
+			UE_LEARNING_CHECK(InOutOffset + sizeof(InElementType) * InArray.Num() <= Bytes.Num());
+			FMemory::Memcpy(&Bytes[InOutOffset], (uint8*)InArray.GetData(), sizeof(InElementType) * InArray.Num());
+			InOutOffset += sizeof(InElementType) * InArray.Num();
+		}
+
+		/**
+		* Serialize an array from bytes
+		*/
+		template<uint8 InDimNum, typename InElementType, typename Allocator>
+		inline void DeserializeFromBytes(int32& InOutOffset, TLearningArrayView<1, const uint8> Bytes, TLearningArray<InDimNum, InElementType, Allocator>& OutArray)
+		{
+			int32 DimNum = INDEX_NONE;
+			Learning::DeserializeFromBytes(InOutOffset, Bytes, DimNum);
+			UE_LEARNING_CHECK(DimNum == InDimNum);
+
+			TLearningArrayShape<InDimNum> Shape;
+			for (uint8 ShapeIdx = 0; ShapeIdx < InDimNum; ShapeIdx++)
+			{
+				Learning::DeserializeFromBytes(InOutOffset, Bytes, Shape[ShapeIdx]);
+			}
+
+			OutArray.SetNumUninitialized(Shape);
+
+			UE_LEARNING_CHECK(InOutOffset + sizeof(InElementType) * OutArray.Num() <= Bytes.Num());
+			FMemory::Memcpy((uint8*)OutArray.GetData(), &Bytes[InOutOffset],  sizeof(InElementType) * OutArray.Num());
+			InOutOffset += sizeof(InElementType) * OutArray.Num();
+		}
+
 		template<typename InElementType>
 		inline int32 IndexOf(const TLearningArrayView<1, const InElementType> Array, const InElementType& Element)
 		{
diff --git a/Engine/Plugins/Experimental/LearningAgents/Source/Learning/Public/LearningNeuralNetwork.h b/Engine/Plugins/Experimental/LearningAgents/Source/Learning/Public/LearningNeuralNetwork.h
--- a/Engine/Plugins/Experimental/LearningAgents/Source/Learning/Public/LearningNeuralNetwork.h
+++ b/Engine/Plugins/Experimental/LearningAgents/Source/Learning/Public/LearningNeuralNetwork.h
@@ -28,8 +28,8 @@ namespace UE::Learning
 	*/
 	struct LEARNING_API FNeuralNetwork
 	{
-		void DeserializeFromBytes(const TLearningArrayView<1, const uint8> RawBytes);
-		void SerializeToBytes(TLearningArrayView<1, uint8> OutRawBytes) const;
+		void DeserializeFromBytes(int32& InOutOffset, const TLearningArrayView<1, const uint8> RawBytes);
+		void SerializeToBytes(int32& InOutOffset, TLearningArrayView<1, uint8> OutRawBytes) const;
 
 		static int32 GetSerializationByteNum(
 			const int32 InputNum,
diff --git a/Engine/Plugins/Experimental/LearningAgents/Source/LearningAgents/Private/LearningAgentsActions.cpp b/Engine/Plugins/Experimental/LearningAgents/Source/LearningAgents/Private/LearningAgentsActions.cpp
--- a/Engine/Plugins/Experimental/LearningAgents/Source/LearningAgents/Private/LearningAgentsActions.cpp
+++ b/Engine/Plugins/Experimental/LearningAgents/Source/LearningAgents/Private/LearningAgentsActions.cpp
@@ -14,11 +14,11 @@
 namespace UE::Learning::Agents::Actions::Private
 {
 	template<typename ActionUObject, typename ActionFObject, typename... InArgTypes>
-	ActionUObject* AddAction(ULearningAgentsInteractor* InInteractor, const FName Name, InArgTypes&& ...Args)
+	ActionUObject* AddAction(ULearningAgentsInteractor* InInteractor, const FName Name, const TCHAR* FunctionName, InArgTypes&& ...Args)
 	{
 		if (!InInteractor)
 		{
-			UE_LOG(LogLearning, Error, TEXT("InInteractor is nullptr."));
+			UE_LOG(LogLearning, Error, TEXT("%s: InInteractor is nullptr."), FunctionName);
 			return nullptr;
 		}
 
@@ -31,6 +31,9 @@ namespace UE::Learning::Agents::Actions::Private
 			InInteractor->GetAgentManager()->GetMaxInstanceNum(),
 			Forward<InArgTypes>(Args)...);
 
+		// We assume all supported action feature objects can be encoded and decoded
+		UE_LEARNING_CHECK(Action->FeatureObject->IsEncodable() && Action->FeatureObject->IsDecodable());
+
 		InInteractor->AddAction(Action, Action->FeatureObject.ToSharedRef());
 
 		return Action;
@@ -41,7 +44,7 @@ namespace UE::Learning::Agents::Actions::Private
 
 UFloatAction* UFloatAction::AddFloatAction(ULearningAgentsInteractor* InInteractor, const FName Name, const float Scale)
 {
-	return UE::Learning::Agents::Actions::Private::AddAction<UFloatAction, UE::Learning::FFloatFeature>(InInteractor, Name, 1, Scale);
+	return UE::Learning::Agents::Actions::Private::AddAction<UFloatAction,UE::Learning::FFloatFeature>(InInteractor, Name, TEXT("AddFloatAction"), 1, Scale);
 }
 
 float UFloatAction::GetFloatAction(const int32 AgentId) const
@@ -95,11 +98,11 @@ UFloatArrayAction* UFloatArrayAction::AddFloatArrayAction(ULearningAgentsInterac
 {
 	if (Num < 1)
 	{
-		UE_LOG(LogLearning, Error, TEXT("Number of elements in array must be at least 1, got %i."), Num);
+		UE_LOG(LogLearning, Error, TEXT("AddFloatArrayAction: Number of elements in array must be at least 1, got %i."), Num);
 		return nullptr;
 	}
 
-	return UE::Learning::Agents::Actions::Private::AddAction<UFloatArrayAction, UE::Learning::FFloatFeature>(InInteractor, Name, Num, Scale);
+	return UE::Learning::Agents::Actions::Private::AddAction<UFloatArrayAction, UE::Learning::FFloatFeature>(InInteractor, Name, TEXT("AddFloatArrayAction"), Num, Scale);
 }
 
 void UFloatArrayAction::GetFloatArrayAction(const int32 AgentId, TArray<float>& OutValues) const
@@ -165,7 +168,7 @@ void UFloatArrayAction::VisualLog(const UE::Learning::FIndexSet Instances) const
 
 UVectorAction* UVectorAction::AddVectorAction(ULearningAgentsInteractor* InInteractor, const FName Name, const float Scale)
 {
-	return UE::Learning::Agents::Actions::Private::AddAction<UVectorAction, UE::Learning::FFloatFeature>(InInteractor, Name, 3, Scale);
+	return UE::Learning::Agents::Actions::Private::AddAction<UVectorAction, UE::Learning::FFloatFeature>(InInteractor, Name, TEXT("AddVectorAction"), 3, Scale);
 }
 
 FVector UVectorAction::GetVectorAction(const int32 AgentId) const
@@ -238,11 +241,11 @@ UVectorArrayAction* UVectorArrayAction::AddVectorArrayAction(ULearningAgentsInte
 {
 	if (Num < 1)
 	{
-		UE_LOG(LogLearning, Error, TEXT("Number of elements in array must be at least 1, got %i."), Num);
+		UE_LOG(LogLearning, Error, TEXT("AddVectorArrayAction: Number of elements in array must be at least 1, got %i."), Num);
 		return nullptr;
 	}
 
-	return UE::Learning::Agents::Actions::Private::AddAction<UVectorArrayAction, UE::Learning::FFloatFeature>(InInteractor, Name, Num * 3, Scale);
+	return UE::Learning::Agents::Actions::Private::AddAction<UVectorArrayAction, UE::Learning::FFloatFeature>(InInteractor, Name, TEXT("AddVectorArrayAction"), Num * 3, Scale);
 }
 
 void UVectorArrayAction::GetVectorArrayAction(const int32 AgentId, TArray<FVector>& OutVectors) const
@@ -340,7 +343,7 @@ void UVectorArrayAction::VisualLog(const UE::Learning::FIndexSet Instances) cons
 
 URotationAction* URotationAction::AddRotationAction(ULearningAgentsInteractor* InInteractor, const FName Name, const float Scale)
 {
-	return UE::Learning::Agents::Actions::Private::AddAction<URotationAction, UE::Learning::FRotationVectorFeature>(InInteractor, Name, 1, Scale);
+	return UE::Learning::Agents::Actions::Private::AddAction<URotationAction, UE::Learning::FRotationVectorFeature>(InInteractor, Name, TEXT("AddRotationAction"), 1, Scale);
 }
 
 FRotator URotationAction::GetRotationAction(const int32 AgentId) const
@@ -416,11 +419,11 @@ URotationArrayAction* URotationArrayAction::AddRotationArrayAction(ULearningAgen
 {
 	if (RotationNum < 1)
 	{
-		UE_LOG(LogLearning, Error, TEXT("Number of elements in array must be at least 1, got %i."), RotationNum);
+		UE_LOG(LogLearning, Error, TEXT("AddRotationArrayAction: Number of elements in array must be at least 1, got %i."), RotationNum);
 		return nullptr;
 	}
 
-	return UE::Learning::Agents::Actions::Private::AddAction<URotationArrayAction, UE::Learning::FRotationVectorFeature>(InInteractor, Name, RotationNum, Scale);
+	return UE::Learning::Agents::Actions::Private::AddAction<URotationArrayAction, UE::Learning::FRotationVectorFeature>(InInteractor, Name, TEXT("AddRotationArrayAction"), RotationNum, Scale);
 }
 
 void URotationArrayAction::GetRotationArrayAction(const int32 AgentId, TArray<FRotator>& OutRotations) const
diff --git a/Engine/Plugins/Experimental/LearningAgents/Source/LearningAgents/Private/LearningAgentsController.cpp b/Engine/Plugins/Experimental/LearningAgents/Source/LearningAgents/Private/LearningAgentsController.cpp
--- a/Engine/Plugins/Experimental/LearningAgents/Source/LearningAgents/Private/LearningAgentsController.cpp
+++ b/Engine/Plugins/Experimental/LearningAgents/Source/LearningAgents/Private/LearningAgentsController.cpp
@@ -65,7 +65,7 @@ void ULearningAgentsController::EncodeActions()
 
 	if (!IsSetup())
 	{
-		UE_LOG(LogLearning, Error, TEXT("%s: Setup not run."), *GetName());
+		UE_LOG(LogLearning, Error, TEXT("%s: Setup not complete."), *GetName());
 		return;
 	}
 
@@ -101,7 +101,7 @@ void ULearningAgentsController::RunController()
 
 	if (!IsSetup())
 	{
-		UE_LOG(LogLearning, Error, TEXT("%s: Setup not run."), *GetName());
+		UE_LOG(LogLearning, Error, TEXT("%s: Setup not complete."), *GetName());
 		return;
 	}
 
diff --git a/Engine/Plugins/Experimental/LearningAgents/Source/LearningAgents/Private/LearningAgentsCritic.cpp b/Engine/Plugins/Experimental/LearningAgents/Source/LearningAgents/Private/LearningAgentsCritic.cpp
--- a/Engine/Plugins/Experimental/LearningAgents/Source/LearningAgents/Private/LearningAgentsCritic.cpp
+++ b/Engine/Plugins/Experimental/LearningAgents/Source/LearningAgents/Private/LearningAgentsCritic.cpp
@@ -86,67 +86,33 @@ UE::Learning::FNeuralNetworkCriticFunction& ULearningAgentsCritic::GetCriticObje
 	return *CriticObject;
 }
 
-void ULearningAgentsCritic::LoadCriticFromSnapshot(const FDirectoryPath& Directory, const FString Filename)
+void ULearningAgentsCritic::LoadCriticFromSnapshot(const FFilePath& File)
 {
 	if (!IsSetup())
 	{
-		UE_LOG(LogLearning, Error, TEXT("%s: Setup not run."), *GetName());
+		UE_LOG(LogLearning, Error, TEXT("%s: Setup not complete."), *GetName());
 		return;
 	}
 
-	TArray64<uint8> NetworkData;
-	const FString FilePath = Directory.Path + FGenericPlatformMisc::GetDefaultPathSeparator() + Filename;
-	if (FFileHelper::LoadFileToArray(NetworkData, *FilePath))
-	{
-		const int32 TotalByteNum = UE::Learning::FNeuralNetwork::GetSerializationByteNum(
-			Network->NeuralNetwork->GetInputNum(),
-			Network->NeuralNetwork->GetOutputNum(),
-			Network->NeuralNetwork->GetHiddenNum(),
-			Network->NeuralNetwork->GetLayerNum());
-
-		if (NetworkData.Num() != TotalByteNum)
-		{
-			UE_LOG(LogLearning, Error, TEXT("%s: Failed to load network from file %s. File size incorrect."), *GetName(), *FilePath);
-			return;
-		}
-
-		Network->NeuralNetwork->DeserializeFromBytes(NetworkData);
-	}
-	else
-	{
-		UE_LOG(LogLearning, Error, TEXT("%s: Failed to load network. File not found: %s"), *GetName(), *FilePath);
-	}
+	Network->LoadNetworkFromSnapshot(File);
 }
 
-void ULearningAgentsCritic::SaveCriticToSnapshot(const FDirectoryPath& Directory, const FString Filename) const
+void ULearningAgentsCritic::SaveCriticToSnapshot(const FFilePath& File) const
 {
 	if (!IsSetup())
 	{
-		UE_LOG(LogLearning, Error, TEXT("%s: Setup not run."), *GetName());
+		UE_LOG(LogLearning, Error, TEXT("%s: Setup not complete."), *GetName());
 		return;
 	}
 
-	TArray64<uint8> NetworkData;
-	NetworkData.SetNumUninitialized(UE::Learning::FNeuralNetwork::GetSerializationByteNum(
-		Network->NeuralNetwork->GetInputNum(),
-		Network->NeuralNetwork->GetOutputNum(),
-		Network->NeuralNetwork->GetHiddenNum(),
-		Network->NeuralNetwork->GetLayerNum()));
-
-	Network->NeuralNetwork->SerializeToBytes(NetworkData);
-
-	const FString FilePath = Directory.Path + FGenericPlatformMisc::GetDefaultPathSeparator() + Filename;
-	if (!FFileHelper::SaveArrayToFile(NetworkData, *FilePath))
-	{
-		UE_LOG(LogLearning, Error, TEXT("%s: Failed to save network to file: %s"), *GetName(), *FilePath);
-	}
+	Network->SaveNetworkToSnapshot(File);
 }
 
 void ULearningAgentsCritic::LoadCriticFromAsset(const ULearningAgentsNeuralNetwork* NeuralNetworkAsset)
 {
 	if (!IsSetup())
 	{
-		UE_LOG(LogLearning, Error, TEXT("%s: Setup not run."), *GetName());
+		UE_LOG(LogLearning, Error, TEXT("%s: Setup not complete."), *GetName());
 		return;
 	}
 
@@ -170,7 +136,7 @@ void ULearningAgentsCritic::SaveCriticToAsset(ULearningAgentsNeuralNetwork* Neur
 {
 	if (!IsSetup())
 	{
-		UE_LOG(LogLearning, Error, TEXT("%s: Setup not run."), *GetName());
+		UE_LOG(LogLearning, Error, TEXT("%s: Setup not complete."), *GetName());
 		return;
 	}
 
@@ -209,7 +175,7 @@ void ULearningAgentsCritic::EvaluateCritic()
 
 	if (!IsSetup())
 	{
-		UE_LOG(LogLearning, Error, TEXT("%s: Setup not run."), *GetName());
+		UE_LOG(LogLearning, Error, TEXT("%s: Setup not complete."), *GetName());
 		return;
 	}
 
@@ -224,7 +190,7 @@ float ULearningAgentsCritic::GetEstimatedDiscountedReturn(const int32 AgentId) c
 {
 	if (!IsSetup())
 	{
-		UE_LOG(LogLearning, Error, TEXT("%s: Setup not run."), *GetName());
+		UE_LOG(LogLearning, Error, TEXT("%s: Setup not complete."), *GetName());
 		return 0.0f;
 	}
 
diff --git a/Engine/Plugins/Experimental/LearningAgents/Source/LearningAgents/Private/LearningAgentsHelpers.cpp b/Engine/Plugins/Experimental/LearningAgents/Source/LearningAgents/Private/LearningAgentsHelpers.cpp
--- a/Engine/Plugins/Experimental/LearningAgents/Source/LearningAgents/Private/LearningAgentsHelpers.cpp
+++ b/Engine/Plugins/Experimental/LearningAgents/Source/LearningAgents/Private/LearningAgentsHelpers.cpp
@@ -16,7 +16,7 @@ USplineComponentHelper* USplineComponentHelper::AddSplineComponentHelper(ULearni
 {
 	if (!InManagerComponent)
 	{
-		UE_LOG(LogLearning, Error, TEXT("InManagerComponent is nullptr."));
+		UE_LOG(LogLearning, Error, TEXT("AddSplineComponentHelper: InManagerComponent is nullptr."));
 		return nullptr;
 	}
 
@@ -26,6 +26,162 @@ USplineComponentHelper* USplineComponentHelper::AddSplineComponentHelper(ULearni
 	return Helper;
 }
 
+FVector USplineComponentHelper::GetNearestPositionOnSpline(const int32 AgentId, const USplineComponent* SplineComponent, const FVector Position, const ESplineCoordinateSpace::Type CoordinateSpace) const
+{
+	if (!ManagerComponent->HasAgent(AgentId))
+	{
+		UE_LOG(LogLearning, Error, TEXT("%s: AgentId %d not found in the agents set."), *GetName(), AgentId);
+		return FVector::ZeroVector;
+	}
+
+	if (!SplineComponent)
+	{
+		UE_LOG(LogLearning, Error, TEXT("%s: SplineComponent is nullptr."), *GetName());
+		return FVector::ZeroVector;
+	}
+
+	const FVector SplinePosition = SplineComponent->FindLocationClosestToWorldLocation(Position, CoordinateSpace);
+
+#if UE_LEARNING_AGENTS_ENABLE_VISUAL_LOG
+	{
+		UE_LEARNING_AGENTS_VLOG_LOCATION(this, LogLearning, Display,
+			SplinePosition,
+			10.0f,
+			VisualLogColor.ToFColor(true),
+			TEXT(""));
+
+		UE_LEARNING_AGENTS_VLOG_SEGMENT(this, LogLearning, Display,
+			Position,
+			SplinePosition,
+			VisualLogColor.ToFColor(true),
+			TEXT("GetNearestPositionOnSpline\nAgent %i\nPosition: [% 6.1f % 6.1f % 6.1f]\nSpline Position: [% 6.1f % 6.1f % 6.1f]"),
+			AgentId,
+			Position.X, Position.Y, Position.Z,
+			SplinePosition.X, SplinePosition.Y, SplinePosition.Z);
+
+		UE_LEARNING_AGENTS_VLOG_LOCATION(this, LogLearning, Display,
+			Position,
+			10.0f,
+			VisualLogColor.ToFColor(true),
+			TEXT(""));
+	}
+#endif
+
+	return SplinePosition;
+}
+
+float USplineComponentHelper::GetDistanceAlongSplineAtPosition(const int32 AgentId, const USplineComponent* SplineComponent, const FVector Position, const ESplineCoordinateSpace::Type CoordinateSpace) const
+{
+	if (!ManagerComponent->HasAgent(AgentId))
+	{
+		UE_LOG(LogLearning, Error, TEXT("%s: AgentId %d not found in the agents set."), *GetName(), AgentId);
+		return 0.0f;
+	}
+
+	if (!SplineComponent)
+	{
+		UE_LOG(LogLearning, Error, TEXT("%s: SplineComponent is nullptr."), *GetName());
+		return 0.0f;
+	}
+
+	const float DistanceAlongSpline = SplineComponent->GetDistanceAlongSplineAtLocation(Position, CoordinateSpace);
+
+#if UE_LEARNING_AGENTS_ENABLE_VISUAL_LOG
+	{
+		const FVector SplinePosition = SplineComponent->GetLocationAtDistanceAlongSpline(DistanceAlongSpline, ESplineCoordinateSpace::World);
+
+		UE_LEARNING_AGENTS_VLOG_LOCATION(this, LogLearning, Display,
+			SplinePosition,
+			10.0f,
+			VisualLogColor.ToFColor(true),
+			TEXT(""));
+
+		UE_LEARNING_AGENTS_VLOG_SEGMENT(this, LogLearning, Display,
+			Position,
+			SplinePosition,
+			VisualLogColor.ToFColor(true),
+			TEXT("GetDistanceAlongSplineAtPosition\nAgent %i\nPosition: [% 6.1f % 6.1f % 6.1f]\nDistance Along Spline: [% 6.1f]"),
+			AgentId,
+			Position.X, Position.Y, Position.Z,
+			DistanceAlongSpline);
+
+		UE_LEARNING_AGENTS_VLOG_LOCATION(this, LogLearning, Display,
+			Position,
+			10.0f,
+			VisualLogColor.ToFColor(true),
+			TEXT(""));
+	}
+#endif
+
+	return DistanceAlongSpline;
+}
+
+FVector USplineComponentHelper::GetPositionAtDistanceAlongSpline(const int32 AgentId, const USplineComponent* SplineComponent, const float DistanceAlongSpline, const ESplineCoordinateSpace::Type CoordinateSpace) const
+{
+	if (!ManagerComponent->HasAgent(AgentId))
+	{
+		UE_LOG(LogLearning, Error, TEXT("%s: AgentId %d not found in the agents set."), *GetName(), AgentId);
+		return FVector::ZeroVector;
+	}
+
+	if (!SplineComponent)
+	{
+		UE_LOG(LogLearning, Error, TEXT("%s: SplineComponent is nullptr."), *GetName());
+		return FVector::ZeroVector;
+	}
+
+	const FVector Position = SplineComponent->GetLocationAtDistanceAlongSpline(DistanceAlongSpline, CoordinateSpace);
+
+#if UE_LEARNING_AGENTS_ENABLE_VISUAL_LOG
+	{
+		UE_LEARNING_AGENTS_VLOG_LOCATION(this, LogLearning, Display,
+			Position,
+			10.0f,
+			VisualLogColor.ToFColor(true),
+			TEXT("GetPositionAtDistanceAlongSpline\nAgent %i\nDistance Along Spline: [% 6.1f]\nPosition: [% 6.1f % 6.1f % 6.1f]"),
+			AgentId,
+			DistanceAlongSpline,
+			Position.X, Position.Y, Position.Z);
+	}
+#endif
+
+	return Position;
+}
+
+FVector USplineComponentHelper::GetDirectionAtDistanceAlongSpline(const int32 AgentId, const USplineComponent* SplineComponent, const float DistanceAlongSpline, const ESplineCoordinateSpace::Type CoordinateSpace) const
+{
+	if (!ManagerComponent->HasAgent(AgentId))
+	{
+		UE_LOG(LogLearning, Error, TEXT("%s: AgentId %d not found in the agents set."), *GetName(), AgentId);
+		return FVector::ZeroVector;
+	}
+
+	if (!SplineComponent)
+	{
+		UE_LOG(LogLearning, Error, TEXT("%s: SplineComponent is nullptr."), *GetName());
+		return FVector::ZeroVector;
+	}
+
+	const FVector Direction = SplineComponent->GetDirectionAtDistanceAlongSpline(DistanceAlongSpline, CoordinateSpace);
+
+#if UE_LEARNING_AGENTS_ENABLE_VISUAL_LOG
+	{
+		const FVector Position = SplineComponent->GetLocationAtDistanceAlongSpline(DistanceAlongSpline, ESplineCoordinateSpace::World);
+
+		UE_LEARNING_AGENTS_VLOG_ARROW(this, LogLearning, Display,
+			Position,
+			Position + Direction * 100.0f,
+			VisualLogColor.ToFColor(true),
+			TEXT("GetDirectionAtDistanceAlongSpline\nAgent %i\nDistance Along Spline: [% 6.1f]\nDirection: [% 6.3f % 6.3f % 6.3f]"),
+			AgentId,
+			DistanceAlongSpline,
+			Direction.X, Direction.Y, Direction.Z);
+	}
+#endif
+
+	return Direction;
+}
+
 float USplineComponentHelper::GetProportionAlongSpline(const int32 AgentId, const USplineComponent* SplineComponent, const float DistanceAlongSpline) const
 {
 	if (!ManagerComponent->HasAgent(AgentId))
@@ -266,7 +422,7 @@ UProjectionHelper* UProjectionHelper::AddProjectionHelper(ULearningAgentsManager
 {
 	if (!InManagerComponent)
 	{
-		UE_LOG(LogLearning, Error, TEXT("InManagerComponent is nullptr."));
+		UE_LOG(LogLearning, Error, TEXT("AddProjectionHelper: InManagerComponent is nullptr."));
 		return nullptr;
 	}
 
@@ -375,7 +531,7 @@ UMeshComponentHelper* UMeshComponentHelper::AddMeshComponentHelper(ULearningAgen
 {
 	if (!InManagerComponent)
 	{
-		UE_LOG(LogLearning, Error, TEXT("InManagerComponent is nullptr."));
+		UE_LOG(LogLearning, Error, TEXT("AddMeshComponentHelper: InManagerComponent is nullptr."));
 		return nullptr;
 	}
 
@@ -448,7 +604,7 @@ URayCastHelper* URayCastHelper::AddRayCastHelper(ULearningAgentsManagerComponent
 {
 	if (!InManagerComponent)
 	{
-		UE_LOG(LogLearning, Error, TEXT("InInteractor is nullptr."));
+		UE_LOG(LogLearning, Error, TEXT("AddRayCastHelper: InManagerComponent is nullptr."));
 		return nullptr;
 	}
 
diff --git a/Engine/Plugins/Experimental/LearningAgents/Source/LearningAgents/Private/LearningAgentsInteractor.cpp b/Engine/Plugins/Experimental/LearningAgents/Source/LearningAgents/Private/LearningAgentsInteractor.cpp
--- a/Engine/Plugins/Experimental/LearningAgents/Source/LearningAgents/Private/LearningAgentsInteractor.cpp
+++ b/Engine/Plugins/Experimental/LearningAgents/Source/LearningAgents/Private/LearningAgentsInteractor.cpp
@@ -47,6 +47,18 @@ void ULearningAgentsInteractor::SetupInteractor(ALearningAgentsManager* InAgentM
 		InAgentManager->GetInstanceData().ToSharedRef(),
 		InAgentManager->GetMaxInstanceNum());
 
+	if (ObservationObjects.Num() == 0)
+	{
+		UE_LOG(LogLearning, Error, TEXT("%s: No observations added to Interactor during SetupObservations."), *GetName());
+		return;
+	}
+
+	if (Observations->DimNum() == 0)
+	{
+		UE_LOG(LogLearning, Error, TEXT("%s: Observation vector is zero-sized - all added observations have no size."), *GetName());
+		return;
+	}
+
 	// Setup Actions
 	ActionObjects.Empty();
 	ActionFeatures.Empty();
@@ -56,6 +68,18 @@ void ULearningAgentsInteractor::SetupInteractor(ALearningAgentsManager* InAgentM
 		InAgentManager->GetInstanceData().ToSharedRef(),
 		InAgentManager->GetMaxInstanceNum());
 
+	if (ActionObjects.Num() == 0)
+	{
+		UE_LOG(LogLearning, Error, TEXT("%s: No actions added to Interactor during SetupActions."), *GetName());
+		return;
+	}
+
+	if (Actions->DimNum() == 0)
+	{
+		UE_LOG(LogLearning, Error, TEXT("%s: Action vector is zero-sized -  all added actions have no size."), *GetName());
+		return;
+	}
+
 	bIsSetup = true;
 }
 
@@ -119,7 +143,7 @@ void ULearningAgentsInteractor::EncodeObservations()
 
 	if (!IsSetup())
 	{
-		UE_LOG(LogLearning, Error, TEXT("%s: Setup not run."), *GetName());
+		UE_LOG(LogLearning, Error, TEXT("%s: Setup not complete."), *GetName());
 		return;
 	}
 
@@ -144,7 +168,7 @@ void ULearningAgentsInteractor::DecodeActions()
 
 	if (!IsSetup())
 	{
-		UE_LOG(LogLearning, Error, TEXT("%s: Setup not run."), *GetName());
+		UE_LOG(LogLearning, Error, TEXT("%s: Setup not complete."), *GetName());
 		return;
 	}
 
@@ -166,7 +190,7 @@ void ULearningAgentsInteractor::GetObservationVector(const int32 AgentId, TArray
 {
 	if (!IsSetup())
 	{
-		UE_LOG(LogLearning, Error, TEXT("%s: Setup not run."), *GetName());
+		UE_LOG(LogLearning, Error, TEXT("%s: Setup not complete."), *GetName());
 		OutObservationVector.Empty();
 		return;
 	}
@@ -186,7 +210,7 @@ void ULearningAgentsInteractor::GetActionVector(const int32 AgentId, TArray<floa
 {
 	if (!IsSetup())
 	{
-		UE_LOG(LogLearning, Error, TEXT("%s: Setup not run."), *GetName());
+		UE_LOG(LogLearning, Error, TEXT("%s: Setup not complete."), *GetName());
 		OutActionVector.Empty();
 		return;
 	}
diff --git a/Engine/Plugins/Experimental/LearningAgents/Source/LearningAgents/Private/LearningAgentsManager.cpp b/Engine/Plugins/Experimental/LearningAgents/Source/LearningAgents/Private/LearningAgentsManager.cpp
--- a/Engine/Plugins/Experimental/LearningAgents/Source/LearningAgents/Private/LearningAgentsManager.cpp
+++ b/Engine/Plugins/Experimental/LearningAgents/Source/LearningAgents/Private/LearningAgentsManager.cpp
@@ -185,17 +185,39 @@ bool ALearningAgentsManager::HasAgentById(const int32 AgentId) const
 	return OccupiedAgentSet.Contains(AgentId);
 }
 
-UObject* ALearningAgentsManager::GetAgent(const int32 AgentId, const TSubclassOf<UObject> AgentClass)
+UObject* ALearningAgentsManager::GetAgent(const int32 AgentId, const TSubclassOf<UObject> AgentClass) const
 {
 	if (!OccupiedAgentSet.Contains(AgentId))
 	{
-		UE_LOG(LogLearning, Error, TEXT("%s: AgentId %d not found. Be sure to only use AgentIds returned by AddAgent() and check that the agent has not be removed."), *GetName(), AgentId);
+		UE_LOG(LogLearning, Error, TEXT("%s: AgentId %d not found. Be sure to only use AgentIds returned by AddAgent and check that the agent has not be removed."), *GetName(), AgentId);
 		return nullptr;
 	}
 
 	return Agents[AgentId];
 }
 
+void ALearningAgentsManager::GetAgents(const TArray<int32>& AgentIds, const TSubclassOf<UObject> AgentClass, TArray<UObject*>& OutAgents) const
+{
+	OutAgents.SetNumUninitialized(AgentIds.Num());
+
+	for (int32 AgentIdIdx = 0; AgentIdIdx < AgentIds.Num(); AgentIdIdx++)
+	{
+		if (!OccupiedAgentSet.Contains(AgentIds[AgentIdIdx]))
+		{
+			UE_LOG(LogLearning, Error, TEXT("%s: AgentId %d not found. Be sure to only use AgentIds returned by AddAgent and check that the agent has not be removed."), *GetName(), AgentIds[AgentIdIdx]);
+			OutAgents.Empty();
+			return;
+		}
+
+		OutAgents[AgentIdIdx] = Agents[AgentIds[AgentIdIdx]];
+	}
+}
+
+void ALearningAgentsManager::GetAddedAgents(const TSubclassOf<UObject> AgentClass, TArray<UObject*>& OutAgents) const
+{
+	GetAgents(OccupiedAgentIds, AgentClass, OutAgents);
+}
+
 const UObject* ALearningAgentsManager::GetAgent(const int32 AgentId) const
 {
 	UE_LEARNING_CHECK(OccupiedAgentSet.Contains(AgentId));
diff --git a/Engine/Plugins/Experimental/LearningAgents/Source/LearningAgents/Private/LearningAgentsNeuralNetwork.cpp b/Engine/Plugins/Experimental/LearningAgents/Source/LearningAgents/Private/LearningAgentsNeuralNetwork.cpp
--- a/Engine/Plugins/Experimental/LearningAgents/Source/LearningAgents/Private/LearningAgentsNeuralNetwork.cpp
+++ b/Engine/Plugins/Experimental/LearningAgents/Source/LearningAgents/Private/LearningAgentsNeuralNetwork.cpp
@@ -5,6 +5,9 @@
 #include "LearningNeuralNetwork.h"
 #include "LearningLog.h"
 
+#include "UObject/Package.h"
+#include "Misc/FileHelper.h"
+
 ULearningAgentsNeuralNetwork::ULearningAgentsNeuralNetwork() {}
 ULearningAgentsNeuralNetwork::ULearningAgentsNeuralNetwork(FVTableHelper& Helper) : ULearningAgentsNeuralNetwork() {}
 ULearningAgentsNeuralNetwork::~ULearningAgentsNeuralNetwork() {}
@@ -34,6 +37,12 @@ namespace UE::Learning::Agents
 	}
 }
 
+namespace UE::Learning::Agents::NeuralNetwork::Private
+{
+	static constexpr int32 MagicNumber = 0x0cd353cf;
+	static constexpr int32 VersionNumber = 1;
+}
+
 void ULearningAgentsNeuralNetwork::Serialize(FArchive& Ar)
 {
 	Super::Serialize(Ar);
@@ -44,22 +53,20 @@ void ULearningAgentsNeuralNetwork::Serialize(FArchive& Ar)
 		Ar << bValid;
 		if (bValid)
 		{
-			int32 InputNum;
-			int32 OutputNum;
-			int32 HiddenNum;
-			int32 LayerNum;
-			UE::Learning::EActivationFunction Activation;
-			Ar << InputNum;
-			Ar << OutputNum;
-			Ar << HiddenNum;
-			Ar << LayerNum;
-			Ar << Activation;
+			int32 VersionNumber;
+			Ar << VersionNumber;
+
+			if (VersionNumber != UE::Learning::Agents::NeuralNetwork::Private::VersionNumber)
+			{
+				UE_LOG(LogLearning, Warning, TEXT("%s: Unsupported Version Number %i."), *GetName(), VersionNumber);
+			}
+
 			NeuralNetwork = MakeShared<UE::Learning::FNeuralNetwork>();
-			NeuralNetwork->Resize(InputNum, OutputNum, HiddenNum, LayerNum);
-			NeuralNetwork->ActivationFunction = Activation;
 			TArray<uint8> Bytes;
 			Ar << Bytes;
-			NeuralNetwork->DeserializeFromBytes(Bytes);
+			int32 Offset = 0;
+			NeuralNetwork->DeserializeFromBytes(Offset, Bytes);
+			UE_LEARNING_CHECK(Offset == Bytes.Num());
 		}
 		else
 		{
@@ -72,20 +79,102 @@ void ULearningAgentsNeuralNetwork::Serialize(FArchive& Ar)
 		Ar << bValid;
 		if (bValid)
 		{
-			int32 InputNum = NeuralNetwork->GetInputNum();
-			int32 OutputNum = NeuralNetwork->GetOutputNum();
-			int32 HiddenNum = NeuralNetwork->GetHiddenNum();
-			int32 LayerNum = NeuralNetwork->GetLayerNum();
-			UE::Learning::EActivationFunction Activation = NeuralNetwork->ActivationFunction;
-			Ar << InputNum;
-			Ar << OutputNum;
-			Ar << HiddenNum;
-			Ar << LayerNum;
-			Ar << Activation;
+			int32 VersionNumber = UE::Learning::Agents::NeuralNetwork::Private::VersionNumber;
+			Ar << VersionNumber;
 			TArray<uint8> Bytes;
-			Bytes.SetNumUninitialized(UE::Learning::FNeuralNetwork::GetSerializationByteNum(InputNum, OutputNum, HiddenNum, LayerNum));
-			NeuralNetwork->SerializeToBytes(Bytes);
+			Bytes.SetNumUninitialized(UE::Learning::FNeuralNetwork::GetSerializationByteNum(
+				NeuralNetwork->GetInputNum(),
+				NeuralNetwork->GetOutputNum(), 
+				NeuralNetwork->GetHiddenNum(),
+				NeuralNetwork->GetLayerNum()));
+			int32 Offset = 0;
+			NeuralNetwork->SerializeToBytes(Offset, Bytes);
+			UE_LEARNING_CHECK(Offset == Bytes.Num());
 			Ar << Bytes;
 		}
 	}
+}
+
+void ULearningAgentsNeuralNetwork::LoadNetworkFromSnapshot(const FFilePath& File)
+{
+	TArray<uint8> RecordingData;
+
+	if (FFileHelper::LoadFileToArray(RecordingData, *File.FilePath))
+	{
+		if (RecordingData.Num() < sizeof(int32) * 2)
+		{
+			UE_LOG(LogLearning, Error, TEXT("%s: Failed to load network. Incorrect Format."), *GetName());
+			return;
+		}
+
+		int32 Offset = 0;
+
+		int32 MagicNumber;
+		UE::Learning::DeserializeFromBytes(Offset, RecordingData, MagicNumber);
+
+		if (MagicNumber != UE::Learning::Agents::NeuralNetwork::Private::MagicNumber)
+		{
+			UE_LOG(LogLearning, Error, TEXT("%s: Failed to load network. Incorrect Magic Number."), *GetName());
+			return;
+		}
+
+		int32 VersionNumber;
+		UE::Learning::DeserializeFromBytes(Offset, RecordingData, VersionNumber);
+
+		if (VersionNumber != UE::Learning::Agents::NeuralNetwork::Private::VersionNumber)
+		{
+			UE_LOG(LogLearning, Error, TEXT("%s: Failed to load network. Unsupported Version Number %i."), *GetName(), VersionNumber);
+			return;
+		}
+
+		NeuralNetwork->DeserializeFromBytes(Offset, RecordingData);
+
+		UE_LEARNING_CHECK(Offset == RecordingData.Num());
+	}
+	else
+	{
+		UE_LOG(LogLearning, Error, TEXT("%s: Failed to load network. File not found: \"%s\""), *GetName(), *File.FilePath);
+	}
+
+	// Manually mark the package as dirty since just using `Modify` prevents 
+	// marking packages as dirty during PIE which is most likely when this
+	// is being used.
+	if (UPackage* Package = GetPackage())
+	{
+		const bool bIsDirty = Package->IsDirty();
+
+		if (!bIsDirty)
+		{
+			Package->SetDirtyFlag(true);
+		}
+
+		Package->PackageMarkedDirtyEvent.Broadcast(Package, bIsDirty);
+	}
+}
+
+void ULearningAgentsNeuralNetwork::SaveNetworkToSnapshot(const FFilePath& File) const
+{
+	TArray<uint8> RecordingData;
+
+	int32 TotalByteNum =
+		sizeof(int32) + // Magic Num
+		sizeof(int32) + // Version Num
+		NeuralNetwork->GetSerializationByteNum(
+			NeuralNetwork->GetInputNum(), 
+			NeuralNetwork->GetOutputNum(), 
+			NeuralNetwork->GetHiddenNum(), 
+			NeuralNetwork->GetLayerNum());
+
+	RecordingData.SetNumUninitialized(TotalByteNum);
+
+	int32 Offset = 0;
+	UE::Learning::SerializeToBytes(Offset, RecordingData, UE::Learning::Agents::NeuralNetwork::Private::MagicNumber);
+	UE::Learning::SerializeToBytes(Offset, RecordingData, UE::Learning::Agents::NeuralNetwork::Private::VersionNumber);
+	NeuralNetwork->SerializeToBytes(Offset, RecordingData);
+	UE_LEARNING_CHECK(Offset == RecordingData.Num());
+
+	if (!FFileHelper::SaveArrayToFile(RecordingData, *File.FilePath))
+	{
+		UE_LOG(LogLearning, Error, TEXT("%s: Failed to save network to file: \"%s\""), *GetName(), *File.FilePath);
+	}
 }
\ No newline at end of file
diff --git a/Engine/Plugins/Experimental/LearningAgents/Source/LearningAgents/Private/LearningAgentsObservations.cpp b/Engine/Plugins/Experimental/LearningAgents/Source/LearningAgents/Private/LearningAgentsObservations.cpp
--- a/Engine/Plugins/Experimental/LearningAgents/Source/LearningAgents/Private/LearningAgentsObservations.cpp
+++ b/Engine/Plugins/Experimental/LearningAgents/Source/LearningAgents/Private/LearningAgentsObservations.cpp
@@ -14,11 +14,11 @@
 namespace UE::Learning::Agents::Observations::Private
 {
 	template<typename ObservationUObject, typename ObservationFObject, typename... InArgTypes>
-	ObservationUObject* AddObservation(ULearningAgentsInteractor* InInteractor, const FName Name, InArgTypes&& ...Args)
+	ObservationUObject* AddObservation(ULearningAgentsInteractor* InInteractor, const FName Name, const TCHAR* FunctionName, InArgTypes&& ...Args)
 	{
 		if (!InInteractor)
 		{
-			UE_LOG(LogLearning, Error, TEXT("InInteractor is nullptr."));
+			UE_LOG(LogLearning, Error, TEXT("%s: InInteractor is nullptr."), FunctionName);
 			return nullptr;
 		}
 
@@ -31,6 +31,9 @@ namespace UE::Learning::Agents::Observations::Private
 			InInteractor->GetAgentManager()->GetMaxInstanceNum(),
 			Forward<InArgTypes>(Args)...);
 
+		// We assume all supported observation feature objects can be encoded
+		UE_LEARNING_CHECK(Observation->FeatureObject->IsEncodable());
+
 		InInteractor->AddObservation(Observation, Observation->FeatureObject.ToSharedRef());
 
 		return Observation;
@@ -41,7 +44,7 @@ namespace UE::Learning::Agents::Observations::Private
 
 UFloatObservation* UFloatObservation::AddFloatObservation(ULearningAgentsInteractor* InInteractor, const FName Name, const float Scale)
 {
-	return UE::Learning::Agents::Observations::Private::AddObservation<UFloatObservation, UE::Learning::FFloatFeature>(InInteractor, Name, 1, Scale);
+	return UE::Learning::Agents::Observations::Private::AddObservation<UFloatObservation, UE::Learning::FFloatFeature>(InInteractor, Name, TEXT("AddFloatObservation"), 1, Scale);
 }
 
 void UFloatObservation::SetFloatObservation(const int32 AgentId, const float Value)
@@ -86,11 +89,11 @@ UFloatArrayObservation* UFloatArrayObservation::AddFloatArrayObservation(ULearni
 {
 	if (Num < 1)
 	{
-		UE_LOG(LogLearning, Error, TEXT("Number of elements in array must be at least 1, got %i."), Num);
+		UE_LOG(LogLearning, Error, TEXT("AddFloatArrayObservation: Number of elements in array must be at least 1, got %i."), Num);
 		return nullptr;
 	}
 
-	return UE::Learning::Agents::Observations::Private::AddObservation<UFloatArrayObservation, UE::Learning::FFloatFeature>(InInteractor, Name, Num, Scale);
+	return UE::Learning::Agents::Observations::Private::AddObservation<UFloatArrayObservation, UE::Learning::FFloatFeature>(InInteractor, Name, TEXT("AddFloatArrayObservation"), Num, Scale);
 }
 
 void UFloatArrayObservation::SetFloatArrayObservation(const int32 AgentId, const TArray<float>& Values)
@@ -141,7 +144,7 @@ void UFloatArrayObservation::VisualLog(const UE::Learning::FIndexSet Instances)
 
 UVectorObservation* UVectorObservation::AddVectorObservation(ULearningAgentsInteractor* InInteractor, const FName Name, const float Scale)
 {
-	return UE::Learning::Agents::Observations::Private::AddObservation<UVectorObservation, UE::Learning::FFloatFeature>(InInteractor, Name, 3, Scale);
+	return UE::Learning::Agents::Observations::Private::AddObservation<UVectorObservation, UE::Learning::FFloatFeature>(InInteractor, Name, TEXT("AddVectorObservation"), 3, Scale);
 }
 
 void UVectorObservation::SetVectorObservation(const int32 AgentId, const FVector Vector)
@@ -201,11 +204,11 @@ UVectorArrayObservation* UVectorArrayObservation::AddVectorArrayObservation(ULea
 {
 	if (Num < 1)
 	{
-		UE_LOG(LogLearning, Error, TEXT("Number of elements in array must be at least 1, got %i."), Num);
+		UE_LOG(LogLearning, Error, TEXT("AddVectorArrayObservation: Number of elements in array must be at least 1, got %i."), Num);
 		return nullptr;
 	}
 
-	return UE::Learning::Agents::Observations::Private::AddObservation<UVectorArrayObservation, UE::Learning::FFloatFeature>(InInteractor, Name, Num * 3, Scale);
+	return UE::Learning::Agents::Observations::Private::AddObservation<UVectorArrayObservation, UE::Learning::FFloatFeature>(InInteractor, Name, TEXT("AddVectorArrayObservation"), Num * 3, Scale);
 }
 
 void UVectorArrayObservation::SetVectorArrayObservation(const int32 AgentId, const TArray<FVector>& Vectors)
@@ -283,17 +286,17 @@ UEnumObservation* UEnumObservation::AddEnumObservation(ULearningAgentsInteractor
 {
 	if (!EnumType)
 	{
-		UE_LOG(LogLearning, Error, TEXT("Invalid Enum."));
+		UE_LOG(LogLearning, Error, TEXT("AddEnumObservation: Invalid Enum."));
 		return nullptr;
 	}
 
 	if (EnumType->NumEnums() < 1)
 	{
-		UE_LOG(LogLearning, Error, TEXT("Enum requires at least one entry to be used as an observation."));
+		UE_LOG(LogLearning, Error, TEXT("AddEnumObservation: Enum requires at least one entry to be used as an observation."));
 		return nullptr;
 	}
 
-	UEnumObservation* Observation = UE::Learning::Agents::Observations::Private::AddObservation<UEnumObservation, UE::Learning::FFloatFeature>(InInteractor, Name, EnumType->NumEnums());
+	UEnumObservation* Observation = UE::Learning::Agents::Observations::Private::AddObservation<UEnumObservation, UE::Learning::FFloatFeature>(InInteractor, Name, TEXT("AddEnumObservation"), EnumType->NumEnums());
 	if (Observation) { Observation->Enum = EnumType; }
 	
 	return Observation;
@@ -370,23 +373,23 @@ UEnumArrayObservation* UEnumArrayObservation::AddEnumArrayObservation(ULearningA
 {
 	if (!EnumType)
 	{
-		UE_LOG(LogLearning, Error, TEXT("Invalid Enum."));
+		UE_LOG(LogLearning, Error, TEXT("AddEnumArrayObservation: Invalid Enum."));
 		return nullptr;
 	}
 
 	if (EnumType->NumEnums() < 1)
 	{
-		UE_LOG(LogLearning, Error, TEXT("Enum requires at least one entry to be used as an observation."));
+		UE_LOG(LogLearning, Error, TEXT("AddEnumArrayObservation: Enum requires at least one entry to be used as an observation."));
 		return nullptr;
 	}
 
 	if (EnumNum < 1)
 	{
-		UE_LOG(LogLearning, Error, TEXT("Number of elements in array must be at least 1, got %i."), EnumNum);
+		UE_LOG(LogLearning, Error, TEXT("AddEnumArrayObservation: Number of elements in array must be at least 1, got %i."), EnumNum);
 		return nullptr;
 	}
 
-	UEnumArrayObservation* Observation = UE::Learning::Agents::Observations::Private::AddObservation<UEnumArrayObservation, UE::Learning::FFloatFeature>(InInteractor, Name, EnumNum * EnumType->NumEnums());
+	UEnumArrayObservation* Observation = UE::Learning::Agents::Observations::Private::AddObservation<UEnumArrayObservation, UE::Learning::FFloatFeature>(InInteractor, Name, TEXT("AddEnumArrayObservation"), EnumNum * EnumType->NumEnums());
 	if (Observation) { Observation->Enum = EnumType; }
 
 	return Observation;
@@ -455,7 +458,7 @@ void UEnumArrayObservation::VisualLog(const UE::Learning::FIndexSet Instances) c
 
 UTimeObservation* UTimeObservation::AddTimeObservation(ULearningAgentsInteractor* InInteractor, const FName Name, const float Scale)
 {
-	return UE::Learning::Agents::Observations::Private::AddObservation<UTimeObservation, UE::Learning::FTimeFeature>(InInteractor, Name, 1, Scale);
+	return UE::Learning::Agents::Observations::Private::AddObservation<UTimeObservation, UE::Learning::FTimeFeature>(InInteractor, Name, TEXT("AddTimeObservation"), 1, Scale);
 }
 
 void UTimeObservation::SetTimeObservation(const int32 AgentId, const float Time, const float RelativeTime)
@@ -504,11 +507,11 @@ UTimeArrayObservation* UTimeArrayObservation::AddTimeArrayObservation(ULearningA
 {
 	if (TimeNum < 1)
 	{
-		UE_LOG(LogLearning, Error, TEXT("Number of elements in array must be at least 1, got %i."), TimeNum);
+		UE_LOG(LogLearning, Error, TEXT("AddTimeArrayObservation: Number of elements in array must be at least 1, got %i."), TimeNum);
 		return nullptr;
 	}
 
-	return UE::Learning::Agents::Observations::Private::AddObservation<UTimeArrayObservation, UE::Learning::FTimeFeature>(InInteractor, Name, TimeNum, Scale);
+	return UE::Learning::Agents::Observations::Private::AddObservation<UTimeArrayObservation, UE::Learning::FTimeFeature>(InInteractor, Name, TEXT("AddTimeArrayObservation"), TimeNum, Scale);
 }
 
 void UTimeArrayObservation::SetTimeArrayObservation(const int32 AgentId, const TArray<float>& Times, const float RelativeTime)
@@ -563,7 +566,7 @@ void UTimeArrayObservation::VisualLog(const UE::Learning::FIndexSet Instances) c
 
 UAngleObservation* UAngleObservation::AddAngleObservation(ULearningAgentsInteractor* InInteractor, const FName Name, const float Scale)
 {
-	return UE::Learning::Agents::Observations::Private::AddObservation<UAngleObservation, UE::Learning::FAngleFeature>(InInteractor, Name, 1, Scale);
+	return UE::Learning::Agents::Observations::Private::AddObservation<UAngleObservation, UE::Learning::FAngleFeature>(InInteractor, Name, TEXT("AddAngleObservation"), 1, Scale);
 }
 
 void UAngleObservation::SetAngleObservation(const int32 AgentId, const float Angle, const float RelativeAngle)
@@ -615,11 +618,11 @@ UAngleArrayObservation* UAngleArrayObservation::AddAngleArrayObservation(ULearni
 {
 	if (AngleNum < 1)
 	{
-		UE_LOG(LogLearning, Error, TEXT("Number of elements in array must be at least 1, got %i."), AngleNum);
+		UE_LOG(LogLearning, Error, TEXT("AddAngleArrayObservation: Number of elements in array must be at least 1, got %i."), AngleNum);
 		return nullptr;
 	}
 
-	return UE::Learning::Agents::Observations::Private::AddObservation<UAngleArrayObservation, UE::Learning::FAngleFeature>(InInteractor, Name, AngleNum, Scale);
+	return UE::Learning::Agents::Observations::Private::AddObservation<UAngleArrayObservation, UE::Learning::FAngleFeature>(InInteractor, Name, TEXT("AddAngleArrayObservation"), AngleNum, Scale);
 }
 
 void UAngleArrayObservation::SetAngleArrayObservation(const int32 AgentId, const TArray<float>& Angles, const float RelativeAngle)
@@ -697,7 +700,7 @@ void UAngleArrayObservation::VisualLog(const UE::Learning::FIndexSet Instances)
 
 URotationObservation* URotationObservation::AddRotationObservation(ULearningAgentsInteractor* InInteractor, const FName Name, const float Scale)
 {
-	return UE::Learning::Agents::Observations::Private::AddObservation<URotationObservation, UE::Learning::FRotationFeature>(InInteractor, Name, 1, Scale);
+	return UE::Learning::Agents::Observations::Private::AddObservation<URotationObservation, UE::Learning::FRotationFeature>(InInteractor, Name, TEXT("AddRotationObservation"), 1, Scale);
 }
 
 void URotationObservation::SetRotationObservation(const int32 AgentId, const FRotator Rotation, const FRotator RelativeRotation)
@@ -760,11 +763,11 @@ URotationArrayObservation* URotationArrayObservation::AddRotationArrayObservatio
 {
 	if (RotationNum < 1)
 	{
-		UE_LOG(LogLearning, Error, TEXT("Number of elements in array must be at least 1, got %i."), RotationNum);
+		UE_LOG(LogLearning, Error, TEXT("AddRotationArrayObservation: Number of elements in array must be at least 1, got %i."), RotationNum);
 		return nullptr;
 	}
 
-	return UE::Learning::Agents::Observations::Private::AddObservation<URotationArrayObservation, UE::Learning::FRotationFeature>(InInteractor, Name, RotationNum, Scale);
+	return UE::Learning::Agents::Observations::Private::AddObservation<URotationArrayObservation, UE::Learning::FRotationFeature>(InInteractor, Name, TEXT("AddRotationArrayObservation"), RotationNum, Scale);
 }
 
 void URotationArrayObservation::SetRotationArrayObservation(const int32 AgentId, const TArray<FRotator>& Rotations, const FRotator RelativeRotation)
@@ -865,7 +868,7 @@ void URotationArrayObservation::VisualLog(const UE::Learning::FIndexSet Instance
 
 UDirectionObservation* UDirectionObservation::AddDirectionObservation(ULearningAgentsInteractor* InInteractor, const FName Name, const float Scale)
 {
-	return UE::Learning::Agents::Observations::Private::AddObservation<UDirectionObservation, UE::Learning::FDirectionFeature>(InInteractor, Name, 1, Scale);
+	return UE::Learning::Agents::Observations::Private::AddObservation<UDirectionObservation, UE::Learning::FDirectionFeature>(InInteractor, Name, TEXT("AddDirectionObservation"), 1, Scale);
 }
 
 void UDirectionObservation::SetDirectionObservation(const int32 AgentId, const FVector Direction, const FRotator RelativeRotation)
@@ -927,11 +930,11 @@ UDirectionArrayObservation* UDirectionArrayObservation::AddDirectionArrayObserva
 {
 	if (DirectionNum < 1)
 	{
-		UE_LOG(LogLearning, Error, TEXT("Number of elements in array must be at least 1, got %i."), DirectionNum);
+		UE_LOG(LogLearning, Error, TEXT("AddDirectionArrayObservation: Number of elements in array must be at least 1, got %i."), DirectionNum);
 		return nullptr;
 	}
 
-	return UE::Learning::Agents::Observations::Private::AddObservation<UDirectionArrayObservation, UE::Learning::FDirectionFeature>(InInteractor, Name, DirectionNum, Scale);
+	return UE::Learning::Agents::Observations::Private::AddObservation<UDirectionArrayObservation, UE::Learning::FDirectionFeature>(InInteractor, Name, TEXT("AddDirectionArrayObservation"), DirectionNum, Scale);
 }
 
 void UDirectionArrayObservation::SetDirectionArrayObservation(const int32 AgentId, const TArray<FVector>& Directions, const FRotator RelativeRotation)
@@ -1008,7 +1011,7 @@ void UDirectionArrayObservation::VisualLog(const UE::Learning::FIndexSet Instanc
 
 UPlanarDirectionObservation* UPlanarDirectionObservation::AddPlanarDirectionObservation(ULearningAgentsInteractor* InInteractor, const FName Name, const float Scale, const FVector Axis0, const FVector Axis1)
 {
-	return UE::Learning::Agents::Observations::Private::AddObservation<UPlanarDirectionObservation, UE::Learning::FPlanarDirectionFeature>(InInteractor, Name, 1, Scale, Axis0, Axis1);
+	return UE::Learning::Agents::Observations::Private::AddObservation<UPlanarDirectionObservation, UE::Learning::FPlanarDirectionFeature>(InInteractor, Name, TEXT("AddPlanarDirectionObservation"), 1, Scale, Axis0, Axis1);
 }
 
 void UPlanarDirectionObservation::SetPlanarDirectionObservation(const int32 AgentId, const FVector Direction, const FRotator RelativeRotation)
@@ -1076,7 +1079,13 @@ void UPlanarDirectionObservation::VisualLog(const UE::Learning::FIndexSet Instan
 
 UPlanarDirectionArrayObservation* UPlanarDirectionArrayObservation::AddPlanarDirectionArrayObservation(ULearningAgentsInteractor* InInteractor, const FName Name, const int32 DirectionNum, const float Scale, const FVector Axis0, const FVector Axis1)
 {
-	return UE::Learning::Agents::Observations::Private::AddObservation<UPlanarDirectionArrayObservation, UE::Learning::FPlanarDirectionFeature>(InInteractor, Name, DirectionNum, Scale, Axis0, Axis1);
+	if (DirectionNum < 1)
+	{
+		UE_LOG(LogLearning, Error, TEXT("AddPlanarDirectionArrayObservation: Number of elements in array must be at least 1, got %i."), DirectionNum);
+		return nullptr;
+	}
+
+	return UE::Learning::Agents::Observations::Private::AddObservation<UPlanarDirectionArrayObservation, UE::Learning::FPlanarDirectionFeature>(InInteractor, Name, TEXT("AddPlanarDirectionArrayObservation"), DirectionNum, Scale, Axis0, Axis1);
 }
 
 void UPlanarDirectionArrayObservation::SetPlanarDirectionArrayObservation(const int32 AgentId, const TArray<FVector>& Directions, const FRotator RelativeRotation)
@@ -1165,7 +1174,7 @@ void UPlanarDirectionArrayObservation::VisualLog(const UE::Learning::FIndexSet I
 
 UPositionObservation* UPositionObservation::AddPositionObservation(ULearningAgentsInteractor* InInteractor, const FName Name, const float Scale)
 {
-	return UE::Learning::Agents::Observations::Private::AddObservation<UPositionObservation, UE::Learning::FPositionFeature>(InInteractor, Name, 1, Scale);
+	return UE::Learning::Agents::Observations::Private::AddObservation<UPositionObservation, UE::Learning::FPositionFeature>(InInteractor, Name, TEXT("AddPositionObservation"), 1, Scale);
 }
 
 void UPositionObservation::SetPositionObservation(const int32 AgentId, const FVector Position, const FVector RelativePosition, const FRotator RelativeRotation)
@@ -1228,11 +1237,11 @@ UPositionArrayObservation* UPositionArrayObservation::AddPositionArrayObservatio
 {
 	if (PositionNum < 1)
 	{
-		UE_LOG(LogLearning, Error, TEXT("Number of elements in array must be at least 1, got %i."), PositionNum);
+		UE_LOG(LogLearning, Error, TEXT("AddPositionArrayObservation: Number of elements in array must be at least 1, got %i."), PositionNum);
 		return nullptr;
 	}
 
-	return UE::Learning::Agents::Observations::Private::AddObservation<UPositionArrayObservation, UE::Learning::FPositionFeature>(InInteractor, Name, PositionNum, Scale);
+	return UE::Learning::Agents::Observations::Private::AddObservation<UPositionArrayObservation, UE::Learning::FPositionFeature>(InInteractor, Name, TEXT("AddPositionArrayObservation"), PositionNum, Scale);
 }
 
 void UPositionArrayObservation::SetPositionArrayObservation(const int32 AgentId, const TArray<FVector>& Positions, const FVector RelativePosition, const FRotator RelativeRotation)
@@ -1310,7 +1319,7 @@ void UPositionArrayObservation::VisualLog(const UE::Learning::FIndexSet Instance
 
 UScalarPositionObservation* UScalarPositionObservation::AddScalarPositionObservation(ULearningAgentsInteractor* InInteractor, const FName Name, const float Scale)
 {
-	return UE::Learning::Agents::Observations::Private::AddObservation<UScalarPositionObservation, UE::Learning::FScalarPositionFeature>(InInteractor, Name, 1, Scale);
+	return UE::Learning::Agents::Observations::Private::AddObservation<UScalarPositionObservation, UE::Learning::FScalarPositionFeature>(InInteractor, Name, TEXT("AddScalarPositionObservation"), 1, Scale);
 }
 
 void UScalarPositionObservation::SetScalarPositionObservation(const int32 AgentId, const float Position, const float RelativePosition)
@@ -1362,11 +1371,11 @@ UScalarPositionArrayObservation* UScalarPositionArrayObservation::AddScalarPosit
 {
 	if (PositionNum < 1)
 	{
-		UE_LOG(LogLearning, Error, TEXT("Number of elements in array must be at least 1, got %i."), PositionNum);
+		UE_LOG(LogLearning, Error, TEXT("AddScalarPositionArrayObservation: Number of elements in array must be at least 1, got %i."), PositionNum);
 		return nullptr;
 	}
 
-	return UE::Learning::Agents::Observations::Private::AddObservation<UScalarPositionArrayObservation, UE::Learning::FScalarPositionFeature>(InInteractor, Name, PositionNum, Scale);
+	return UE::Learning::Agents::Observations::Private::AddObservation<UScalarPositionArrayObservation, UE::Learning::FScalarPositionFeature>(InInteractor, Name, TEXT("AddScalarPositionArrayObservation"), PositionNum, Scale);
 }
 
 void UScalarPositionArrayObservation::SetScalarPositionArrayObservation(const int32 AgentId, const TArray<float>& Positions, const float RelativePosition)
@@ -1445,7 +1454,7 @@ void UScalarPositionArrayObservation::VisualLog(const UE::Learning::FIndexSet In
 
 UPlanarPositionObservation* UPlanarPositionObservation::AddPlanarPositionObservation(ULearningAgentsInteractor* InInteractor, const FName Name, const float Scale, const FVector Axis0, const FVector Axis1)
 {
-	return UE::Learning::Agents::Observations::Private::AddObservation<UPlanarPositionObservation, UE::Learning::FPlanarPositionFeature>(InInteractor, Name, 1, Scale, Axis0, Axis1);
+	return UE::Learning::Agents::Observations::Private::AddObservation<UPlanarPositionObservation, UE::Learning::FPlanarPositionFeature>(InInteractor, Name, TEXT("AddPlanarPositionObservation"), 1, Scale, Axis0, Axis1);
 }
 
 void UPlanarPositionObservation::SetPlanarPositionObservation(const int32 AgentId, const FVector Position, const FVector RelativePosition, const FRotator RelativeRotation)
@@ -1516,11 +1525,11 @@ UPlanarPositionArrayObservation* UPlanarPositionArrayObservation::AddPlanarPosit
 {
 	if (PositionNum < 1)
 	{
-		UE_LOG(LogLearning, Error, TEXT("Number of elements in array must be at least 1, got %i."), PositionNum);
+		UE_LOG(LogLearning, Error, TEXT("AddPlanarPositionArrayObservation: Number of elements in array must be at least 1, got %i."), PositionNum);
 		return nullptr;
 	}
 
-	return UE::Learning::Agents::Observations::Private::AddObservation<UPlanarPositionArrayObservation, UE::Learning::FPlanarPositionFeature>(InInteractor, Name, PositionNum, Scale, Axis0, Axis1);
+	return UE::Learning::Agents::Observations::Private::AddObservation<UPlanarPositionArrayObservation, UE::Learning::FPlanarPositionFeature>(InInteractor, Name, TEXT("AddPlanarPositionArrayObservation"), PositionNum, Scale, Axis0, Axis1);
 }
 
 void UPlanarPositionArrayObservation::SetPlanarPositionArrayObservation(const int32 AgentId, const TArray<FVector>& Positions, const FVector RelativePosition, const FRotator RelativeRotation)
@@ -1609,7 +1618,7 @@ void UPlanarPositionArrayObservation::VisualLog(const UE::Learning::FIndexSet In
 
 UVelocityObservation* UVelocityObservation::AddVelocityObservation(ULearningAgentsInteractor* InInteractor, const FName Name, const float Scale)
 {
-	return UE::Learning::Agents::Observations::Private::AddObservation<UVelocityObservation, UE::Learning::FVelocityFeature>(InInteractor, Name, 1, Scale);
+	return UE::Learning::Agents::Observations::Private::AddObservation<UVelocityObservation, UE::Learning::FVelocityFeature>(InInteractor, Name, TEXT("AddVelocityObservation"), 1, Scale);
 }
 
 void UVelocityObservation::SetVelocityObservation(const int32 AgentId, const FVector Velocity, const FRotator RelativeRotation)
@@ -1671,11 +1680,11 @@ UVelocityArrayObservation* UVelocityArrayObservation::AddVelocityArrayObservatio
 {
 	if (VelocityNum < 1)
 	{
-		UE_LOG(LogLearning, Error, TEXT("Number of elements in array must be at least 1, got %i."), VelocityNum);
+		UE_LOG(LogLearning, Error, TEXT("AddVelocityArrayObservation: Number of elements in array must be at least 1, got %i."), VelocityNum);
 		return nullptr;
 	}
 
-	return UE::Learning::Agents::Observations::Private::AddObservation<UVelocityArrayObservation, UE::Learning::FVelocityFeature>(InInteractor, Name, VelocityNum, Scale);
+	return UE::Learning::Agents::Observations::Private::AddObservation<UVelocityArrayObservation, UE::Learning::FVelocityFeature>(InInteractor, Name, TEXT("AddVelocityArrayObservation"), VelocityNum, Scale);
 }
 
 void UVelocityArrayObservation::SetVelocityArrayObservation(const int32 AgentId, const TArray<FVector>& Velocities, const FRotator RelativeRotation)
@@ -1752,7 +1761,7 @@ void UVelocityArrayObservation::VisualLog(const UE::Learning::FIndexSet Instance
 
 UScalarVelocityObservation* UScalarVelocityObservation::AddScalarVelocityObservation(ULearningAgentsInteractor* InInteractor, const FName Name, const float Scale)
 {
-	return UE::Learning::Agents::Observations::Private::AddObservation<UScalarVelocityObservation, UE::Learning::FScalarVelocityFeature>(InInteractor, Name, 1, Scale);
+	return UE::Learning::Agents::Observations::Private::AddObservation<UScalarVelocityObservation, UE::Learning::FScalarVelocityFeature>(InInteractor, Name, TEXT("AddScalarVelocityObservation"), 1, Scale);
 }
 
 void UScalarVelocityObservation::SetScalarVelocityObservation(const int32 AgentId, const float Velocity)
@@ -1801,11 +1810,11 @@ UScalarVelocityArrayObservation* UScalarVelocityArrayObservation::AddScalarVeloc
 {
 	if (VelocityNum < 1)
 	{
-		UE_LOG(LogLearning, Error, TEXT("Number of elements in array must be at least 1, got %i."), VelocityNum);
+		UE_LOG(LogLearning, Error, TEXT("AddScalarVelocityArrayObservation: Number of elements in array must be at least 1, got %i."), VelocityNum);
 		return nullptr;
 	}
 
-	return UE::Learning::Agents::Observations::Private::AddObservation<UScalarVelocityArrayObservation, UE::Learning::FScalarVelocityFeature>(InInteractor, Name, VelocityNum, Scale);
+	return UE::Learning::Agents::Observations::Private::AddObservation<UScalarVelocityArrayObservation, UE::Learning::FScalarVelocityFeature>(InInteractor, Name, TEXT("AddScalarVelocityArrayObservation"), VelocityNum, Scale);
 }
 
 void UScalarVelocityArrayObservation::SetScalarVelocityArrayObservation(const int32 AgentId, const TArray<float>& Velocities)
@@ -1879,7 +1888,7 @@ void UScalarVelocityArrayObservation::VisualLog(const UE::Learning::FIndexSet In
 
 UPlanarVelocityObservation* UPlanarVelocityObservation::AddPlanarVelocityObservation(ULearningAgentsInteractor* InInteractor, const FName Name, const float Scale, const FVector Axis0, const FVector Axis1)
 {
-	return UE::Learning::Agents::Observations::Private::AddObservation<UPlanarVelocityObservation, UE::Learning::FPlanarVelocityFeature>(InInteractor, Name, 1, Scale, Axis0, Axis1);
+	return UE::Learning::Agents::Observations::Private::AddObservation<UPlanarVelocityObservation, UE::Learning::FPlanarVelocityFeature>(InInteractor, Name, TEXT("AddPlanarVelocityObservation"), 1, Scale, Axis0, Axis1);
 }
 
 void UPlanarVelocityObservation::SetPlanarVelocityObservation(const int32 AgentId, const FVector Velocity, const FRotator RelativeRotation)
@@ -1950,11 +1959,11 @@ UPlanarVelocityArrayObservation* UPlanarVelocityArrayObservation::AddPlanarVeloc
 {
 	if (VelocityNum < 1)
 	{
-		UE_LOG(LogLearning, Error, TEXT("Number of elements in array must be at least 1, got %i."), VelocityNum);
+		UE_LOG(LogLearning, Error, TEXT("AddPlanarVelocityArrayObservation: Number of elements in array must be at least 1, got %i."), VelocityNum);
 		return nullptr;
 	}
 
-	return UE::Learning::Agents::Observations::Private::AddObservation<UPlanarVelocityArrayObservation, UE::Learning::FPlanarVelocityFeature>(InInteractor, Name, VelocityNum, Scale, Axis0, Axis1);
+	return UE::Learning::Agents::Observations::Private::AddObservation<UPlanarVelocityArrayObservation, UE::Learning::FPlanarVelocityFeature>(InInteractor, Name, TEXT("AddPlanarVelocityArrayObservation"), VelocityNum, Scale, Axis0, Axis1);
 }
 
 void UPlanarVelocityArrayObservation::SetPlanarVelocityArrayObservation(const int32 AgentId, const TArray<FVector>& Velocities, const FRotator RelativeRotation)
@@ -2041,7 +2050,7 @@ void UPlanarVelocityArrayObservation::VisualLog(const UE::Learning::FIndexSet In
 
 UAngularVelocityObservation* UAngularVelocityObservation::AddAngularVelocityObservation(ULearningAgentsInteractor* InInteractor, const FName Name, const float Scale)
 {
-	return UE::Learning::Agents::Observations::Private::AddObservation<UAngularVelocityObservation, UE::Learning::FAngularVelocityFeature>(InInteractor, Name, 1, Scale);
+	return UE::Learning::Agents::Observations::Private::AddObservation<UAngularVelocityObservation, UE::Learning::FAngularVelocityFeature>(InInteractor, Name, TEXT("AddAngularVelocityObservation"), 1, Scale);
 }
 
 void UAngularVelocityObservation::SetAngularVelocityObservation(const int32 AgentId, const FVector AngularVelocity, const FRotator RelativeRotation)
@@ -2103,11 +2112,11 @@ UAngularVelocityArrayObservation* UAngularVelocityArrayObservation::AddAngularVe
 {
 	if (AngularVelocityNum < 1)
 	{
-		UE_LOG(LogLearning, Error, TEXT("Number of elements in array must be at least 1, got %i."), AngularVelocityNum);
+		UE_LOG(LogLearning, Error, TEXT("AddAngularVelocityArrayObservation: Number of elements in array must be at least 1, got %i."), AngularVelocityNum);
 		return nullptr;
 	}
 
-	return UE::Learning::Agents::Observations::Private::AddObservation<UAngularVelocityArrayObservation, UE::Learning::FAngularVelocityFeature>(InInteractor, Name, AngularVelocityNum, Scale);
+	return UE::Learning::Agents::Observations::Private::AddObservation<UAngularVelocityArrayObservation, UE::Learning::FAngularVelocityFeature>(InInteractor, Name, TEXT("AddAngularVelocityArrayObservation"), AngularVelocityNum, Scale);
 }
 
 void UAngularVelocityArrayObservation::SetAngularVelocityArrayObservation(const int32 AgentId, const TArray<FVector>& AngularVelocities, const FRotator RelativeRotation)
@@ -2184,7 +2193,7 @@ void UAngularVelocityArrayObservation::VisualLog(const UE::Learning::FIndexSet I
 
 UScalarAngularVelocityObservation* UScalarAngularVelocityObservation::AddScalarAngularVelocityObservation(ULearningAgentsInteractor* InInteractor, const FName Name, const float Scale)
 {
-	return UE::Learning::Agents::Observations::Private::AddObservation<UScalarAngularVelocityObservation, UE::Learning::FScalarAngularVelocityFeature>(InInteractor, Name, 1, Scale);
+	return UE::Learning::Agents::Observations::Private::AddObservation<UScalarAngularVelocityObservation, UE::Learning::FScalarAngularVelocityFeature>(InInteractor, Name, TEXT("AddScalarAngularVelocityObservation"), 1, Scale);
 }
 
 void UScalarAngularVelocityObservation::SetScalarAngularVelocityObservation(const int32 AgentId, const float AngularVelocity)
@@ -2233,11 +2242,11 @@ UScalarAngularVelocityArrayObservation* UScalarAngularVelocityArrayObservation::
 {
 	if (AngularVelocityNum < 1)
 	{
-		UE_LOG(LogLearning, Error, TEXT("Number of elements in array must be at least 1, got %i."), AngularVelocityNum);
+		UE_LOG(LogLearning, Error, TEXT("AddScalarAngularVelocityArrayObservation: Number of elements in array must be at least 1, got %i."), AngularVelocityNum);
 		return nullptr;
 	}
 
-	return UE::Learning::Agents::Observations::Private::AddObservation<UScalarAngularVelocityArrayObservation, UE::Learning::FScalarAngularVelocityFeature>(InInteractor, Name, AngularVelocityNum, Scale);
+	return UE::Learning::Agents::Observations::Private::AddObservation<UScalarAngularVelocityArrayObservation, UE::Learning::FScalarAngularVelocityFeature>(InInteractor, Name, TEXT("AddScalarAngularVelocityArrayObservation"), AngularVelocityNum, Scale);
 }
 
 void UScalarAngularVelocityArrayObservation::SetScalarAngularVelocityArrayObservation(const int32 AgentId, const TArray<float>& AngularVelocities)
diff --git a/Engine/Plugins/Experimental/LearningAgents/Source/LearningAgents/Private/LearningAgentsPolicy.cpp b/Engine/Plugins/Experimental/LearningAgents/Source/LearningAgents/Private/LearningAgentsPolicy.cpp
--- a/Engine/Plugins/Experimental/LearningAgents/Source/LearningAgents/Private/LearningAgentsPolicy.cpp
+++ b/Engine/Plugins/Experimental/LearningAgents/Source/LearningAgents/Private/LearningAgentsPolicy.cpp
@@ -108,67 +108,33 @@ UE::Learning::FNeuralNetworkPolicyFunction& ULearningAgentsPolicy::GetPolicyObje
 	return *PolicyObject;
 }
 
-void ULearningAgentsPolicy::LoadPolicyFromSnapshot(const FDirectoryPath& Directory, const FString Filename)
+void ULearningAgentsPolicy::LoadPolicyFromSnapshot(const FFilePath& File)
 {
 	if (!IsSetup())
 	{
-		UE_LOG(LogLearning, Error, TEXT("%s: Setup not run."), *GetName());
+		UE_LOG(LogLearning, Error, TEXT("%s: Setup not complete."), *GetName());
 		return;
 	}
 
-	TArray64<uint8> NetworkData;
-	const FString FilePath = Directory.Path + FGenericPlatformMisc::GetDefaultPathSeparator() + Filename;
-	if (FFileHelper::LoadFileToArray(NetworkData, *FilePath))
-	{
-		const int32 TotalByteNum = UE::Learning::FNeuralNetwork::GetSerializationByteNum(
-			Network->NeuralNetwork->GetInputNum(),
-			Network->NeuralNetwork->GetOutputNum(),
-			Network->NeuralNetwork->GetHiddenNum(),
-			Network->NeuralNetwork->GetLayerNum());
-
-		if (NetworkData.Num() != TotalByteNum)
-		{
-			UE_LOG(LogLearning, Error, TEXT("%s: Failed to load network from file %s. File size incorrect."), *GetName(), *FilePath);
-			return;
-		}
-
-		Network->NeuralNetwork->DeserializeFromBytes(NetworkData);
-	}
-	else
-	{
-		UE_LOG(LogLearning, Error, TEXT("%s: Failed to load network. File not found: %s"), *GetName(), *FilePath);
-	}
+	Network->LoadNetworkFromSnapshot(File);
 }
 
-void ULearningAgentsPolicy::SavePolicyToSnapshot(const FDirectoryPath& Directory, const FString Filename) const
+void ULearningAgentsPolicy::SavePolicyToSnapshot(const FFilePath& File) const
 {
 	if (!IsSetup())
 	{
-		UE_LOG(LogLearning, Error, TEXT("%s: Setup not run."), *GetName());
+		UE_LOG(LogLearning, Error, TEXT("%s: Setup not complete."), *GetName());
 		return;
 	}
 
-	TArray64<uint8> NetworkData;
-	NetworkData.SetNumUninitialized(UE::Learning::FNeuralNetwork::GetSerializationByteNum(
-		Network->NeuralNetwork->GetInputNum(),
-		Network->NeuralNetwork->GetOutputNum(),
-		Network->NeuralNetwork->GetHiddenNum(),
-		Network->NeuralNetwork->GetLayerNum()));
-
-	Network->NeuralNetwork->SerializeToBytes(NetworkData);
-
-	const FString FilePath = Directory.Path + FGenericPlatformMisc::GetDefaultPathSeparator() + Filename;
-	if (!FFileHelper::SaveArrayToFile(NetworkData, *FilePath))
-	{
-		UE_LOG(LogLearning, Error, TEXT("%s: Failed to save network to file: %s"), *GetName(), *FilePath);
-	}
+	Network->SaveNetworkToSnapshot(File);
 }
 
 void ULearningAgentsPolicy::LoadPolicyFromAsset(const ULearningAgentsNeuralNetwork* NeuralNetworkAsset)
 {
 	if (!IsSetup())
 	{
-		UE_LOG(LogLearning, Error, TEXT("%s: Setup not run."), *GetName());
+		UE_LOG(LogLearning, Error, TEXT("%s: Setup not complete."), *GetName());
 		return;
 	}
 
@@ -192,7 +158,7 @@ void ULearningAgentsPolicy::SavePolicyToAsset(ULearningAgentsNeuralNetwork* Neur
 {
 	if (!IsSetup())
 	{
-		UE_LOG(LogLearning, Error, TEXT("%s: Setup not run."), *GetName());
+		UE_LOG(LogLearning, Error, TEXT("%s: Setup not complete."), *GetName());
 		return;
 	}
 
@@ -231,7 +197,7 @@ void ULearningAgentsPolicy::EvaluatePolicy()
 
 	if (!IsSetup())
 	{
-		UE_LOG(LogLearning, Error, TEXT("%s: Setup not run."), *GetName());
+		UE_LOG(LogLearning, Error, TEXT("%s: Setup not complete."), *GetName());
 		return;
 	}
 
@@ -259,7 +225,7 @@ void ULearningAgentsPolicy::RunInference()
 
 	if (!IsSetup())
 	{
-		UE_LOG(LogLearning, Error, TEXT("%s: Setup not run."), *GetName());
+		UE_LOG(LogLearning, Error, TEXT("%s: Setup not complete."), *GetName());
 		return;
 	}
 
@@ -271,7 +237,7 @@ float ULearningAgentsPolicy::GetAgentActionNoiseScale(const int32 AgentId) const
 {
 	if (!IsSetup())
 	{
-		UE_LOG(LogLearning, Error, TEXT("%s: Setup not run."), *GetName());
+		UE_LOG(LogLearning, Error, TEXT("%s: Setup not complete."), *GetName());
 		return 0.0f;
 	}
 
@@ -289,7 +255,7 @@ void ULearningAgentsPolicy::SetAgentActionNoiseScale(const int32 AgentId, const
 {
 	if (!IsSetup())
 	{
-		UE_LOG(LogLearning, Error, TEXT("%s: Setup not run."), *GetName());
+		UE_LOG(LogLearning, Error, TEXT("%s: Setup not complete."), *GetName());
 		return;
 	}
 
@@ -307,7 +273,7 @@ void ULearningAgentsPolicy::SetAllAgentsActionNoiseScale(const float ActionNoise
 {
 	if (!IsSetup())
 	{
-		UE_LOG(LogLearning, Error, TEXT("%s: Setup not run."), *GetName());
+		UE_LOG(LogLearning, Error, TEXT("%s: Setup not complete."), *GetName());
 		return;
 	}
 
diff --git a/Engine/Plugins/Experimental/LearningAgents/Source/LearningAgents/Public/LearningAgentsCritic.h b/Engine/Plugins/Experimental/LearningAgents/Source/LearningAgents/Public/LearningAgentsCritic.h
--- a/Engine/Plugins/Experimental/LearningAgents/Source/LearningAgents/Public/LearningAgentsCritic.h
+++ b/Engine/Plugins/Experimental/LearningAgents/Source/LearningAgents/Public/LearningAgentsCritic.h
@@ -17,7 +17,7 @@ namespace UE::Learning
 	struct FNeuralNetworkCriticFunction;
 }
 
-struct FDirectoryPath;
+struct FFilePath;
 class ULearningAgentsNeuralNetwork;
 
 /** The configurable settings for a ULearningAgentsCritic. */
@@ -66,19 +66,17 @@ public:
 
 	/**
 	* Load a snapshot's weights into this critic.
-	* @param Directory The directory the snapshot file is in.
-	* @param Filename The filename of the snapshot, including the file extension.
+	* @param File The snapshot file.
 	*/
-	UFUNCTION(BlueprintCallable, Category = "LearningAgents")
-	void LoadCriticFromSnapshot(const FDirectoryPath& Directory, const FString Filename);
+	UFUNCTION(BlueprintCallable, Category = "LearningAgents", meta = (RelativePath))
+	void LoadCriticFromSnapshot(const FFilePath& File);
 
 	/**
 	* Save this critic's weights into a snapshot.
-	* @param Directory The directory to save the snapshot file in.
-	* @param Filename The filename of the snapshot, including the file extension.
+	* @param File The snapshot file.
 	*/
-	UFUNCTION(BlueprintCallable, Category = "LearningAgents")
-	void SaveCriticToSnapshot(const FDirectoryPath& Directory, const FString Filename) const;
+	UFUNCTION(BlueprintCallable, Category = "LearningAgents", meta = (RelativePath))
+	void SaveCriticToSnapshot(const FFilePath& File) const;
 
 	/**
 	* Load a ULearningAgentsNeuralNetwork asset's weights into this critic.
diff --git a/Engine/Plugins/Experimental/LearningAgents/Source/LearningAgents/Public/LearningAgentsHelpers.h b/Engine/Plugins/Experimental/LearningAgents/Source/LearningAgents/Public/LearningAgentsHelpers.h
--- a/Engine/Plugins/Experimental/LearningAgents/Source/LearningAgents/Public/LearningAgentsHelpers.h
+++ b/Engine/Plugins/Experimental/LearningAgents/Source/LearningAgents/Public/LearningAgentsHelpers.h
@@ -48,6 +48,18 @@ public:
 	UFUNCTION(BlueprintCallable, Category = "LearningAgents")
 	static USplineComponentHelper* AddSplineComponentHelper(ULearningAgentsManagerComponent* InManagerComponent, const FName Name = NAME_None);
 
+	UFUNCTION(BlueprintPure = false, Category = "LearningAgents", meta = (AgentId = "-1"))
+	FVector GetNearestPositionOnSpline(const int32 AgentId, const USplineComponent* SplineComponent, const FVector Position, const ESplineCoordinateSpace::Type CoordinateSpace = ESplineCoordinateSpace::World) const;
+
+	UFUNCTION(BlueprintPure = false, Category = "LearningAgents", meta = (AgentId = "-1"))
+	float GetDistanceAlongSplineAtPosition(const int32 AgentId, const USplineComponent* SplineComponent, const FVector Position, const ESplineCoordinateSpace::Type CoordinateSpace = ESplineCoordinateSpace::World) const;
+
+	UFUNCTION(BlueprintPure = false, Category = "LearningAgents", meta = (AgentId = "-1"))
+	FVector GetPositionAtDistanceAlongSpline(const int32 AgentId, const USplineComponent* SplineComponent, const float DistanceAlongSpline, const ESplineCoordinateSpace::Type CoordinateSpace = ESplineCoordinateSpace::World) const;
+
+	UFUNCTION(BlueprintPure = false, Category = "LearningAgents", meta = (AgentId = "-1"))
+	FVector GetDirectionAtDistanceAlongSpline(const int32 AgentId, const USplineComponent* SplineComponent, const float DistanceAlongSpline, const ESplineCoordinateSpace::Type CoordinateSpace = ESplineCoordinateSpace::World) const;
+
 	UFUNCTION(BlueprintPure = false, Category = "LearningAgents", meta = (AgentId = "-1"))
 	float GetProportionAlongSpline(const int32 AgentId, const USplineComponent* SplineComponent, const float DistanceAlongSpline) const;
 
diff --git a/Engine/Plugins/Experimental/LearningAgents/Source/LearningAgents/Public/LearningAgentsManager.h b/Engine/Plugins/Experimental/LearningAgents/Source/LearningAgents/Public/LearningAgentsManager.h
--- a/Engine/Plugins/Experimental/LearningAgents/Source/LearningAgents/Public/LearningAgentsManager.h
+++ b/Engine/Plugins/Experimental/LearningAgents/Source/LearningAgents/Public/LearningAgentsManager.h
@@ -109,7 +109,26 @@ public:
 	 * @return The agent object.
 	 */
 	UFUNCTION(BlueprintPure, Category = "LearningAgents", meta = (AgentId = "-1", DeterminesOutputType = "AgentClass"))
-	UObject* GetAgent(const int32 AgentId, const TSubclassOf<UObject> AgentClass);
+	UObject* GetAgent(const int32 AgentId, const TSubclassOf<UObject> AgentClass) const;
+
+	/**
+	 * Gets the agents associated with a set of ids. Calling this from blueprint with the appropriate AgentClass will 
+	 * automatically cast the object to the given type.
+	 * @param AgentIds The ids of the agents to get.
+	 * @param AgentClass The class to cast the agent objects to (in blueprint).
+	 * @param OutAgents The output array of agent objects.
+	 */
+	UFUNCTION(BlueprintPure = false, Category = "LearningAgents", meta = (DeterminesOutputType = "AgentClass", DynamicOutputParam = "OutAgents"))
+	void GetAgents(const TArray<int32>& AgentIds, const TSubclassOf<UObject> AgentClass, TArray<UObject*>& OutAgents) const;
+
+	/**
+	 * Gets the current added agents. Calling this from blueprint with the appropriate AgentClass will automatically
+	 * cast the object to the given type.
+	 * @param AgentClass The class to cast the agent objects to (in blueprint).
+	 * @param OutAgents The output array of agent objects.
+	 */
+	UFUNCTION(BlueprintPure = false, Category = "LearningAgents", meta = (DeterminesOutputType = "AgentClass", DynamicOutputParam = "OutAgents"))
+	void GetAddedAgents(const TSubclassOf<UObject> AgentClass, TArray<UObject*>& OutAgents) const;
 
 	/**
 	* Returns true if the given object is an agent in this manager; Otherwise, false.
diff --git a/Engine/Plugins/Experimental/LearningAgents/Source/LearningAgents/Public/LearningAgentsNeuralNetwork.h b/Engine/Plugins/Experimental/LearningAgents/Source/LearningAgents/Public/LearningAgentsNeuralNetwork.h
--- a/Engine/Plugins/Experimental/LearningAgents/Source/LearningAgents/Public/LearningAgentsNeuralNetwork.h
+++ b/Engine/Plugins/Experimental/LearningAgents/Source/LearningAgents/Public/LearningAgentsNeuralNetwork.h
@@ -50,5 +50,21 @@ public:
 	/** Serialize this neural network to/from the given archive. */
 	virtual void Serialize(FArchive& Ar) override;
 
+	/**
+	* Load this network from a snapshot.
+	* @param File The snapshot file.
+	*/
+	UFUNCTION(BlueprintCallable, Category = "LearningAgents", meta = (RelativePath))
+	void LoadNetworkFromSnapshot(const FFilePath& File);
+
+	/**
+	* Save this network into a snapshot.
+	* @param File The snapshot file.
+	*/
+	UFUNCTION(BlueprintCallable, Category = "LearningAgents", meta = (RelativePath))
+	void SaveNetworkToSnapshot(const FFilePath& File) const;
+
+public:
+
 	TSharedPtr<UE::Learning::FNeuralNetwork> NeuralNetwork;
 };
diff --git a/Engine/Plugins/Experimental/LearningAgents/Source/LearningAgents/Public/LearningAgentsPolicy.h b/Engine/Plugins/Experimental/LearningAgents/Source/LearningAgents/Public/LearningAgentsPolicy.h
--- a/Engine/Plugins/Experimental/LearningAgents/Source/LearningAgents/Public/LearningAgentsPolicy.h
+++ b/Engine/Plugins/Experimental/LearningAgents/Source/LearningAgents/Public/LearningAgentsPolicy.h
@@ -17,7 +17,7 @@ namespace UE::Learning
 	struct FNeuralNetworkPolicyFunction;
 }
 
-struct FDirectoryPath;
+struct FFilePath;
 class ULearningAgentsNeuralNetwork;
 
 /** The configurable settings for a ULearningAgentsPolicy. */
@@ -85,19 +85,17 @@ public:
 
 	/**
 	* Load a snapshot's weights into this policy.
-	* @param Directory The directory the snapshot file is in.
-	* @param Filename The filename of the snapshot, including the file extension.
+	* @param File The snapshot file.
 	*/
-	UFUNCTION(BlueprintCallable, Category = "LearningAgents")
-	void LoadPolicyFromSnapshot(const FDirectoryPath& Directory, const FString Filename);
+	UFUNCTION(BlueprintCallable, Category = "LearningAgents", meta = (RelativePath))
+	void LoadPolicyFromSnapshot(const FFilePath& File);
 
 	/**
 	* Save this policy's weights into a snapshot.
-	* @param Directory The directory to save the snapshot file in.
-	* @param Filename The filename of the snapshot, including the file extension.
+	* @param File The snapshot file.
 	*/
-	UFUNCTION(BlueprintCallable, Category = "LearningAgents")
-	void SavePolicyToSnapshot(const FDirectoryPath& Directory, const FString Filename) const;
+	UFUNCTION(BlueprintCallable, Category = "LearningAgents", meta = (RelativePath))
+	void SavePolicyToSnapshot(const FFilePath& File) const;
 
 	/**
 	* Load a ULearningAgentsNeuralNetwork asset's weights into this policy.
diff --git a/Engine/Plugins/Experimental/LearningAgents/Source/LearningAgentsTraining/Private/LearningAgentsCompletions.cpp b/Engine/Plugins/Experimental/LearningAgents/Source/LearningAgentsTraining/Private/LearningAgentsCompletions.cpp
--- a/Engine/Plugins/Experimental/LearningAgents/Source/LearningAgentsTraining/Private/LearningAgentsCompletions.cpp
+++ b/Engine/Plugins/Experimental/LearningAgents/Source/LearningAgentsTraining/Private/LearningAgentsCompletions.cpp
@@ -15,11 +15,11 @@
 namespace UE::Learning::Agents::Completions::Private
 {
 	template<typename CompletionUObject, typename CompletionFObject, typename... InArgTypes>
-	CompletionUObject* AddCompletion(ULearningAgentsTrainer* InAgentTrainer, const FName Name, InArgTypes&& ...Args)
+	CompletionUObject* AddCompletion(ULearningAgentsTrainer* InAgentTrainer, const FName Name, const TCHAR* FunctionName, InArgTypes&& ...Args)
 	{
 		if (!InAgentTrainer)
 		{
-			UE_LOG(LogLearning, Error, TEXT("InAgentTrainer is nullptr."));
+			UE_LOG(LogLearning, Error, TEXT("%s: InAgentTrainer is nullptr."), FunctionName);
 			return nullptr;
 		}
 
@@ -42,7 +42,7 @@ namespace UE::Learning::Agents::Completions::Private
 
 UConditionalCompletion* UConditionalCompletion::AddConditionalCompletion(ULearningAgentsTrainer* InAgentTrainer, const FName Name, const ELearningAgentsCompletion InCompletionMode)
 {
-	return UE::Learning::Agents::Completions::Private::AddCompletion<UConditionalCompletion, UE::Learning::FConditionalCompletion>(InAgentTrainer, Name, UE::Learning::Agents::GetCompletionMode(InCompletionMode));
+	return UE::Learning::Agents::Completions::Private::AddCompletion<UConditionalCompletion, UE::Learning::FConditionalCompletion>(InAgentTrainer, Name, TEXT("AddConditionalCompletion"), UE::Learning::Agents::GetCompletionMode(InCompletionMode));
 }
 
 void UConditionalCompletion::SetConditionalCompletion(const int32 AgentId, const bool bIsComplete)
@@ -84,7 +84,7 @@ void UConditionalCompletion::VisualLog(const UE::Learning::FIndexSet Instances)
 
 UTimeElapsedCompletion* UTimeElapsedCompletion::AddTimeElapsedCompletion(ULearningAgentsTrainer* InAgentTrainer, const FName Name, const float Threshold, const ELearningAgentsCompletion InCompletionMode)
 {
-	return UE::Learning::Agents::Completions::Private::AddCompletion<UTimeElapsedCompletion, UE::Learning::FTimeElapsedCompletion>(InAgentTrainer, Name, Threshold, UE::Learning::Agents::GetCompletionMode(InCompletionMode));
+	return UE::Learning::Agents::Completions::Private::AddCompletion<UTimeElapsedCompletion, UE::Learning::FTimeElapsedCompletion>(InAgentTrainer, Name, TEXT("AddTimeElapsedCompletion"), Threshold, UE::Learning::Agents::GetCompletionMode(InCompletionMode));
 }
 
 void UTimeElapsedCompletion::SetTimeElapsedCompletion(const int32 AgentId, const float Time)
@@ -137,6 +137,7 @@ UPlanarPositionDifferenceCompletion* UPlanarPositionDifferenceCompletion::AddPla
 	return UE::Learning::Agents::Completions::Private::AddCompletion<UPlanarPositionDifferenceCompletion, UE::Learning::FPlanarPositionDifferenceCompletion>(
 		InAgentTrainer,
 		Name,
+		TEXT("AddPlanarPositionDifferenceCompletion"),
 		1,
 		Threshold,
 		UE::Learning::Agents::GetCompletionMode(InCompletionMode),
@@ -235,6 +236,7 @@ UPlanarPositionSimilarityCompletion* UPlanarPositionSimilarityCompletion::AddPla
 	return UE::Learning::Agents::Completions::Private::AddCompletion<UPlanarPositionSimilarityCompletion, UE::Learning::FPlanarPositionSimilarityCompletion>(
 		InAgentTrainer,
 		Name,
+		TEXT("AddPlanarPositionSimilarityCompletion"),
 		1,
 		Threshold,
 		UE::Learning::Agents::GetCompletionMode(InCompletionMode),
diff --git a/Engine/Plugins/Experimental/LearningAgents/Source/LearningAgentsTraining/Private/LearningAgentsDataStorage.cpp b/Engine/Plugins/Experimental/LearningAgents/Source/LearningAgentsTraining/Private/LearningAgentsDataStorage.cpp
deleted file mode 100644
diff --git a/Engine/Plugins/Experimental/LearningAgents/Source/LearningAgentsTraining/Private/LearningAgentsImitationTrainer.cpp b/Engine/Plugins/Experimental/LearningAgents/Source/LearningAgentsTraining/Private/LearningAgentsImitationTrainer.cpp
--- a/Engine/Plugins/Experimental/LearningAgents/Source/LearningAgentsTraining/Private/LearningAgentsImitationTrainer.cpp
+++ b/Engine/Plugins/Experimental/LearningAgents/Source/LearningAgentsTraining/Private/LearningAgentsImitationTrainer.cpp
@@ -11,7 +11,7 @@
 #include "LearningLog.h"
 #include "LearningNeuralNetwork.h"
 #include "LearningImitationTrainer.h"
-#include "LearningAgentsDataStorage.h"
+#include "LearningAgentsRecording.h"
 #include "LearningAgentsPolicy.h"
 #include "LearningNeuralNetworkObject.h"
 #include "Misc/Paths.h"
@@ -32,13 +32,14 @@ void ULearningAgentsImitationTrainer::EndPlay(const EEndPlayReason::Type EndPlay
 
 void ULearningAgentsImitationTrainer::BeginTraining(
 	ULearningAgentsPolicy* InPolicy, 
-	const TArray<ULearningAgentsRecord*>& Records, 
-	const FLearningAgentsImitationTrainerTrainingSettings& TrainingSettings,
+	const ULearningAgentsRecording* Recording,
+	const FLearningAgentsImitationTrainerTrainingSettings& ImitationTrainerTrainingSettings,
+	const FLearningAgentsTrainerPathSettings& ImitationTrainerPathSettings,
 	const bool bReinitializePolicyNetwork)
 {
 	if (IsTraining())
 	{
-		UE_LOG(LogLearning, Error, TEXT("%s: Already Training!"), *GetName());
+		UE_LOG(LogLearning, Error, TEXT("%s: Cannot begin training as we are already training!"), *GetName());
 		return;
 	}
 
@@ -54,79 +55,95 @@ void ULearningAgentsImitationTrainer::BeginTraining(
 		return;
 	}
 
-	if (Records.IsEmpty())
+	if (!Recording)
 	{
-		UE_LOG(LogLearning, Error, TEXT("%s: Records list is empty! Imitation Trainer needs at least one record."), *GetName());
+		UE_LOG(LogLearning, Error, TEXT("%s: Recording is nullptr."), *GetName());
+		return;
+	}
+
+	if (Recording->Records.IsEmpty())
+	{
+		UE_LOG(LogLearning, Error, TEXT("%s: Recording is empty!"), *GetName());
 		return;
 	}
 
 	Policy = InPolicy;
 
-	// Get Number of Steps
+	// Check Paths
 
-	int32 StepNum = 0;
-	for (const ULearningAgentsRecord* Record : Records)
+	const FString PythonExecutablePath = UE::Learning::Trainer::GetPythonExecutablePath(ImitationTrainerPathSettings.GetEditorEnginePath());
+
+	if (!FPaths::FileExists(PythonExecutablePath))
 	{
-		if (!Record)
-		{
-			UE_LOG(LogLearning, Warning, TEXT("%s: Record is nullptr."));
-			continue;
-		}
+		UE_LOG(LogLearning, Error, TEXT("%s: Can't find Python executable \"%s\"."), *GetName(), *PythonExecutablePath);
+		return;
+	}
+	const FString PythonContentPath = UE::Learning::Trainer::GetPythonContentPath(ImitationTrainerPathSettings.GetEditorEnginePath());
 
-		// Assuming everyone has been trimmed already
-		for (const TLearningArray<2, float>& Chunk : Record->GetObservations())
-		{
-			StepNum += Chunk.Num<0>();
-		}
+	if (!FPaths::DirectoryExists(PythonContentPath))
+	{
+		UE_LOG(LogLearning, Error, TEXT("%s: Can't find LearningAgents plugin Content \"%s\"."), *GetName(), *PythonContentPath);
+		return;
 	}
 
-	// Parse into Flat Arrays
+	const FString SitePackagesPath = UE::Learning::Trainer::GetSitePackagesPath(ImitationTrainerPathSettings.GetEditorEnginePath());
+
+	if (!FPaths::DirectoryExists(SitePackagesPath))
+	{
+		UE_LOG(LogLearning, Error, TEXT("%s: Can't find Python site-packages \"%s\"."), *GetName(), *SitePackagesPath);
+		return;
+	}
+
+	const FString IntermediatePath = UE::Learning::Trainer::GetIntermediatePath(ImitationTrainerPathSettings.GetIntermediatePath());
+
+	// Sizes
 
 	const int32 PolicyInputNum = Policy->GetPolicyNetwork().GetInputNum();
 	const int32 PolicyOutputNum = Policy->GetPolicyNetwork().GetOutputNum();
 
-	RecordedObservations.SetNumUninitialized({ StepNum, PolicyInputNum });
-	RecordedActions.SetNumUninitialized({ StepNum, PolicyOutputNum / 2 });
+	// Get Number of Steps
 
-	int32 DataIndex = 0;
-	for (const ULearningAgentsRecord* Record : Records)
+	int32 TotalSampleNum = 0;
+	for (const FLearningAgentsRecord& Record : Recording->Records)
 	{
-		if (!Record) { continue; }
-
-		TConstArrayView<TLearningArray<2, float>> RecordObsChunks = Record->GetObservations();
-		TConstArrayView<TLearningArray<2, float>> RecordActionChunks = Record->GetActions();
+		if (Record.ObservationDimNum != PolicyInputNum)
+		{
+			UE_LOG(LogLearning, Warning, TEXT("%s: Record has wrong dimensionality for observations, got %i, policy expected %i."), *GetName(), Record.ObservationDimNum, PolicyInputNum);
+			continue;
+		}
 
-		for (int32 Idx = 0; Idx < RecordObsChunks.Num(); Idx++)
+		if (Record.ActionDimNum != PolicyOutputNum / 2)
 		{
-			if (RecordObsChunks[Idx].Num<1>() == PolicyInputNum &&
-				RecordActionChunks[Idx].Num<1>() == PolicyOutputNum / 2)
-			{
-				int32 ChunkLength = RecordObsChunks[Idx].Num<0>();
-				UE::Learning::Array::Copy(RecordedObservations.Slice(DataIndex, ChunkLength), RecordObsChunks[Idx]);
-				UE::Learning::Array::Copy(RecordedActions.Slice(DataIndex, ChunkLength), RecordActionChunks[Idx]);
-				DataIndex += ChunkLength;
-			}
-			else
-			{
-				UE_LOG(LogLearning, Warning, TEXT("%s: Record input or output size does not match policy."), *GetName());
-			}
+			UE_LOG(LogLearning, Warning, TEXT("%s: Record has wrong dimensionality for actions, got %i, policy expected %i."), *GetName(), Record.ActionDimNum, PolicyOutputNum / 2);
+			continue;
 		}
-	}
 
-	// Return if no valid records found
+		TotalSampleNum += Record.SampleNum;
+	}
 
-	if (DataIndex == 0)
+	if (TotalSampleNum == 0)
 	{
-		UE_LOG(LogLearning, Warning, TEXT("%s: input contains no valid training data."), *GetName());
-		RecordedObservations.Empty();
-		RecordedActions.Empty();
+		UE_LOG(LogLearning, Warning, TEXT("%s: Recording contains no valid training data."), *GetName());
 		return;
 	}
 
-	// Resize to final size
+	// Copy into Flat Arrays
+
+	RecordedObservations.SetNumUninitialized({ TotalSampleNum, PolicyInputNum });
+	RecordedActions.SetNumUninitialized({ TotalSampleNum, PolicyOutputNum / 2 });
+
+	int32 SampleIdx = 0;
+	for (const FLearningAgentsRecord& Record : Recording->Records)
+	{
+		if (Record.ObservationDimNum != PolicyInputNum) { continue; }
+		if (Record.ActionDimNum != PolicyOutputNum / 2) { continue; }
+
+		UE::Learning::Array::Copy(RecordedObservations.Slice(SampleIdx, Record.SampleNum), Record.Observations);
+		UE::Learning::Array::Copy(RecordedActions.Slice(SampleIdx, Record.SampleNum), Record.Actions);
+		SampleIdx += Record.SampleNum;
+	}
 
-	RecordedObservations.SetNumUninitialized({ DataIndex, PolicyInputNum });
-	RecordedActions.SetNumUninitialized({ DataIndex, PolicyOutputNum / 2 });
+	UE_LEARNING_CHECK(SampleIdx == TotalSampleNum);
 
 	// Begin Training Properly
 
@@ -136,30 +153,15 @@ void ULearningAgentsImitationTrainer::BeginTraining(
 	bIsTrainingComplete = false;
 	bRequestImitationTrainingStop = false;
 
-#if WITH_EDITOR
-	const FString PythonExecutablePath = UE::Learning::Trainer::DefaultEditorPythonExecutablePath();
-	const FString SitePackagesPath = UE::Learning::Trainer::DefaultEditorSitePackagesPath();
-	const FString PythonContentPath = UE::Learning::Trainer::DefaultEditorPythonContentPath();
-	const FString IntermediatePath = UE::Learning::Trainer::DefaultEditorIntermediatePath();
-#else
-	UE_LEARNING_CHECKF(PLATFORM_WINDOWS || PLATFORM_MAC || PLATFORM_LINUX, TEXT("Python only supported on Windows, Mac, and Linux."));
-	const FString PythonExe = PLATFORM_WINDOWS ? TEXT("python.exe") : TEXT("bin/python");
-	const FString EnginePath = IFileManager::Get().ConvertToAbsolutePathForExternalAppForRead(*(FPaths::RootDir() / TEXT("../../../../../../Engine")));
-	const FString PythonExecutablePath = EnginePath / TEXT("Binaries/ThirdParty/Python3") / FPlatformMisc::GetUBTPlatform() / PythonExe;
-	const FString SitePackagesPath = EnginePath / TEXT("Plugins/Experimental/PythonFoundationPackages/Content/Python/Lib/") / FPlatformMisc::GetUBTPlatform() / TEXT("site-packages");
-	const FString PythonContentPath = EnginePath / TEXT("Plugins/Experimental/LearningAgents/Content/Python/");
-	const FString IntermediatePath = EnginePath / TEXT("Plugins/Experimental/LearningAgents/Intermediate");
-#endif
-
 	UE::Learning::FImitationTrainerTrainingSettings ImitationTrainingSettings;
-	ImitationTrainingSettings.IterationNum = TrainingSettings.NumberOfIterations;
-	ImitationTrainingSettings.LearningRateActor = TrainingSettings.LearningRate;
-	ImitationTrainingSettings.LearningRateDecay = TrainingSettings.LearningRateDecay;
-	ImitationTrainingSettings.WeightDecay = TrainingSettings.WeightDecay;
-	ImitationTrainingSettings.BatchSize = TrainingSettings.BatchSize;
-	ImitationTrainingSettings.Seed = TrainingSettings.RandomSeed;
-	ImitationTrainingSettings.Device = TrainingSettings.Device == ELearningAgentsTrainerDevice::CPU ? UE::Learning::ETrainerDevice::CPU : UE::Learning::ETrainerDevice::GPU;
-	ImitationTrainingSettings.bUseTensorboard = TrainingSettings.bUseTensorboard;
+	ImitationTrainingSettings.IterationNum = ImitationTrainerTrainingSettings.NumberOfIterations;
+	ImitationTrainingSettings.LearningRateActor = ImitationTrainerTrainingSettings.LearningRate;
+	ImitationTrainingSettings.LearningRateDecay = ImitationTrainerTrainingSettings.LearningRateDecay;
+	ImitationTrainingSettings.WeightDecay = ImitationTrainerTrainingSettings.WeightDecay;
+	ImitationTrainingSettings.BatchSize = ImitationTrainerTrainingSettings.BatchSize;
+	ImitationTrainingSettings.Seed = ImitationTrainerTrainingSettings.RandomSeed;
+	ImitationTrainingSettings.Device = ImitationTrainerTrainingSettings.Device == ELearningAgentsTrainerDevice::CPU ? UE::Learning::ETrainerDevice::CPU : UE::Learning::ETrainerDevice::GPU;
+	ImitationTrainingSettings.bUseTensorboard = ImitationTrainerTrainingSettings.bUseTensorboard;
 
 	const UE::Learning::EImitationTrainerFlags TrainerFlags = 
 		bReinitializePolicyNetwork ? 
@@ -196,7 +198,7 @@ void ULearningAgentsImitationTrainer::EndTraining()
 {
 	if (!IsTraining())
 	{
-		UE_LOG(LogLearning, Warning, TEXT("%s: Not Training!"), *GetName());
+		UE_LOG(LogLearning, Warning, TEXT("%s: Cannot end training as we are not training!"), *GetName());
 		return;
 	}
 
diff --git a/Engine/Plugins/Experimental/LearningAgents/Source/LearningAgentsTraining/Private/LearningAgentsRecorder.cpp b/Engine/Plugins/Experimental/LearningAgents/Source/LearningAgentsTraining/Private/LearningAgentsRecorder.cpp
--- a/Engine/Plugins/Experimental/LearningAgents/Source/LearningAgentsTraining/Private/LearningAgentsRecorder.cpp
+++ b/Engine/Plugins/Experimental/LearningAgents/Source/LearningAgentsTraining/Private/LearningAgentsRecorder.cpp
@@ -4,12 +4,83 @@
 
 #include "LearningAgentsManager.h"
 #include "LearningAgentsInteractor.h"
-#include "LearningAgentsDataStorage.h"
+#include "LearningAgentsRecording.h"
 #include "LearningFeatureObject.h"
 #include "LearningLog.h"
 #include "LearningTrainer.h"
 
+#include "UObject/Package.h"
 #include "Engine/World.h"
+#include "Misc/Paths.h"
+
+FLearningAgentsRecorderPathSettings::FLearningAgentsRecorderPathSettings()
+{
+	IntermediateRelativePath.Path = FPaths::ProjectIntermediateDir();
+}
+
+TLearningArrayView<1, float> ULearningAgentsRecorder::FAgentRecordBuffer::GetObservation(const int32 SampleIdx)
+{
+	return Observations[SampleIdx / ChunkSize][SampleIdx % ChunkSize];
+}
+
+TLearningArrayView<1, float> ULearningAgentsRecorder::FAgentRecordBuffer::GetAction(const int32 SampleIdx)
+{
+	return Actions[SampleIdx / ChunkSize][SampleIdx % ChunkSize];
+}
+
+TLearningArrayView<1, const float> ULearningAgentsRecorder::FAgentRecordBuffer::GetObservation(const int32 SampleIdx) const
+{
+	return Observations[SampleIdx / ChunkSize][SampleIdx % ChunkSize];
+}
+
+TLearningArrayView<1, const float> ULearningAgentsRecorder::FAgentRecordBuffer::GetAction(const int32 SampleIdx) const
+{
+	return Actions[SampleIdx / ChunkSize][SampleIdx % ChunkSize];
+}
+
+void ULearningAgentsRecorder::FAgentRecordBuffer::Push(
+	const TLearningArrayView<1, const float> Observation,
+	const TLearningArrayView<1, const float> Action)
+{
+	if (SampleNum / ChunkSize <= Observations.Num())
+	{
+		Observations.AddDefaulted_GetRef().SetNumUninitialized({ ChunkSize, Observation.Num() });
+		Actions.AddDefaulted_GetRef().SetNumUninitialized({ ChunkSize, Action.Num() });
+	}
+
+	UE::Learning::Array::Copy(GetObservation(SampleNum), Observation);
+	UE::Learning::Array::Copy(GetAction(SampleNum), Action);
+	SampleNum++;
+}
+
+bool ULearningAgentsRecorder::FAgentRecordBuffer::IsEmpty() const
+{
+	return SampleNum == 0;
+}
+
+void ULearningAgentsRecorder::FAgentRecordBuffer::Empty()
+{
+	SampleNum = 0;
+	Observations.Empty();
+	Actions.Empty();
+}
+
+void ULearningAgentsRecorder::FAgentRecordBuffer::CopyToRecord(FLearningAgentsRecord& Record) const
+{
+	UE_LEARNING_CHECK(SampleNum > 0);
+
+	Record.SampleNum = SampleNum;
+	Record.ObservationDimNum = GetObservation(0).Num();
+	Record.ActionDimNum = GetAction(0).Num();
+	Record.Observations.SetNumUninitialized({ SampleNum, Observations[0].Num<1>() });
+	Record.Actions.SetNumUninitialized({ SampleNum, Actions[0].Num<1>() });
+
+	for (int32 SampleIdx = 0; SampleIdx < SampleNum; SampleIdx++)
+	{
+		UE::Learning::Array::Copy(Record.Observations[SampleIdx], GetObservation(SampleIdx));
+		UE::Learning::Array::Copy(Record.Actions[SampleIdx], GetAction(SampleIdx));
+	}
+}
 
 ULearningAgentsRecorder::ULearningAgentsRecorder() : ULearningAgentsManagerComponent() {}
 ULearningAgentsRecorder::ULearningAgentsRecorder(FVTableHelper& Helper) : ULearningAgentsRecorder() {}
@@ -25,7 +96,10 @@ void ULearningAgentsRecorder::EndPlay(const EEndPlayReason::Type EndPlayReason)
 	Super::EndPlay(EndPlayReason);
 }
 
-void ULearningAgentsRecorder::SetupRecorder(ALearningAgentsManager* InAgentManager, ULearningAgentsInteractor* InInteractor)
+void ULearningAgentsRecorder::SetupRecorder(
+	ALearningAgentsManager* InAgentManager, 
+	ULearningAgentsInteractor* InInteractor,
+	const FLearningAgentsRecorderPathSettings& RecorderPathSettings)
 {
 	if (IsSetup())
 	{
@@ -41,7 +115,7 @@ void ULearningAgentsRecorder::SetupRecorder(ALearningAgentsManager* InAgentManag
 
 	if (!InAgentManager->IsManagerSetup())
 	{
-		UE_LOG(LogLearning, Error, TEXT("%s's SetupManager() must be run before %s can be setup."), *InAgentManager->GetName(), *GetName());
+		UE_LOG(LogLearning, Error, TEXT("%s: %s's SetupManager must be run before it can be used."), *GetName(), *InAgentManager->GetName());
 		return;
 	}
 
@@ -61,36 +135,26 @@ void ULearningAgentsRecorder::SetupRecorder(ALearningAgentsManager* InAgentManag
 
 	Interactor = InInteractor;
 
-	DataStorage = NewObject<ULearningAgentsDataStorage>(this, TEXT("DataStorage"));
-#if WITH_EDITOR
-	DataDirectory.Path = UE::Learning::Trainer::DefaultEditorIntermediatePath() / TEXT("Recordings");
-#else
-	UE_LOG(LogLearning, Error, TEXT("ULearningAgentsRecorder: DataDirectory was not set. This is non-editor build so need a directory setting."));
-#endif
-
-	bIsSetup = true;
-}
-
-bool ULearningAgentsRecorder::AddAgent(const int32 AgentId)
-{
-	bool bSuccess = Super::AddAgent(AgentId);
+	Recording = NewObject<ULearningAgentsRecording>(this, TEXT("Recording"));
+	RecordingDirectory = RecorderPathSettings.IntermediateRelativePath.Path / TEXT("LearningAgents") / RecorderPathSettings.RecordingsSubdirectory;
 
-	if (bSuccess && IsRecording() && !CurrentRecords.Contains(AgentId))
-	{
-		CurrentRecords.Add(AgentId, DataStorage->CreateRecord(FName(Interactor->GetName() + "_id" + FString::FromInt(AgentId)), Interactor));
-	}
+	RecordBuffers.Empty();
+	RecordBuffers.SetNum(AgentManager->GetMaxInstanceNum());
 
-	return bSuccess;
+	bIsSetup = true;
 }
 
 bool ULearningAgentsRecorder::RemoveAgent(const int32 AgentId)
 {
 	bool bSuccess = Super::RemoveAgent(AgentId);
 
-	if (bSuccess && IsRecording() && CurrentRecords.Contains(AgentId))
+	if (bSuccess)
 	{
-		CurrentRecords[AgentId]->Trim();
-		CurrentRecords.Remove(AgentId);
+		if (!RecordBuffers[AgentId].IsEmpty())
+		{
+			RecordBuffers[AgentId].CopyToRecord(Recording->Records.AddDefaulted_GetRef());
+			RecordBuffers[AgentId].Empty();
+		}
 	}
 
 	return bSuccess;
@@ -98,15 +162,21 @@ bool ULearningAgentsRecorder::RemoveAgent(const int32 AgentId)
 
 void ULearningAgentsRecorder::AddExperience()
 {
+	if (!IsSetup())
+	{
+		UE_LOG(LogLearning, Error, TEXT("%s: Setup not complete."), *GetName());
+		return;
+	}
+
 	if (!IsRecording())
 	{
-		UE_LOG(LogLearning, Warning, TEXT("Trying to add experience but we aren't currently recording. Call BeginRecording() before AddExperience()."));
+		UE_LOG(LogLearning, Error, TEXT("%s: Trying to add experience but we aren't currently recording. Call BeginRecording before AddExperience."), *GetName());
 		return;
 	}
 
 	for (const int32 AgentId : AddedAgentSet)
 	{
-		CurrentRecords[AgentId]->AddExperience(
+		RecordBuffers[AgentId].Push(
 			Interactor->GetObservationFeature().FeatureBuffer()[AgentId], 
 			Interactor->GetActionFeature().FeatureBuffer()[AgentId]);
 	}
@@ -119,45 +189,175 @@ bool ULearningAgentsRecorder::IsRecording() const
 
 void ULearningAgentsRecorder::EndRecording()
 {
-	if (!IsRecording())
+	if (!IsSetup())
 	{
-		UE_LOG(LogLearning, Error, TEXT("Not Recording!"));
+		UE_LOG(LogLearning, Error, TEXT("%s: Setup not complete."), *GetName());
 		return;
 	}
 
-	if (bSaveDataOnEndPlay)
+	if (!IsRecording())
 	{
-		DataStorage->SaveAllRecords(DataDirectory);
+		UE_LOG(LogLearning, Error, TEXT("%s: Cannot end recording as we are not currently recording!"), *GetName());
+		return;
 	}
 
+	// Write to Recording
+
 	for (const int32 AgentId : AddedAgentSet)
 	{
-		CurrentRecords[AgentId]->Trim();
+		if (!RecordBuffers[AgentId].IsEmpty())
+		{
+			RecordBuffers[AgentId].CopyToRecord(Recording->Records.AddDefaulted_GetRef());
+			RecordBuffers[AgentId].Empty();
+		}
 	}
 
-	CurrentRecords.Empty();
+	// Save Recording to Intermediate Directory
+
+	FFilePath RecordingFilePath;
+	RecordingFilePath.FilePath = RecordingDirectory / FString::Printf(TEXT("%s_%s.bin"), *GetName(), *FDateTime::Now().ToFormattedString(TEXT("%Y-%m-%d_%H-%M-%S")));
+	SaveRecordingToFile(RecordingFilePath);
 
 	bIsRecording = false;
 }
 
-void ULearningAgentsRecorder::BeginRecording()
+void ULearningAgentsRecorder::LoadRecordingFromFile(const FFilePath& File)
 {
 	if (!IsSetup())
 	{
-		UE_LOG(LogLearning, Error, TEXT("Setup must be run before recording can begin."));
+		UE_LOG(LogLearning, Error, TEXT("%s: Setup not complete."), *GetName());
+		return;
+	}
+
+	Recording->LoadRecordingFromFile(File);
+}
+
+void ULearningAgentsRecorder::SaveRecordingToFile(const FFilePath& File) const
+{
+	if (!IsSetup())
+	{
+		UE_LOG(LogLearning, Error, TEXT("%s: Setup not complete."), *GetName());
+		return;
+	}
+
+	Recording->SaveRecordingToFile(File);
+}
+
+void ULearningAgentsRecorder::LoadRecordingFromAsset(const ULearningAgentsRecording* Asset)
+{
+	if (!IsSetup())
+	{
+		UE_LOG(LogLearning, Error, TEXT("%s: Setup not complete."), *GetName());
+		return;
+	}
+
+	if (!Asset)
+	{
+		UE_LOG(LogLearning, Error, TEXT("%s: Asset is nullptr."), *GetName());
+		return;
+	}
+
+	Recording->Records = Asset->Records;
+}
+
+void ULearningAgentsRecorder::SaveRecordingToAsset(ULearningAgentsRecording* Asset) const
+{
+	if (!IsSetup())
+	{
+		UE_LOG(LogLearning, Error, TEXT("%s: Setup not complete."), *GetName());
+		return;
+	}
+
+	if (!Asset)
+	{
+		UE_LOG(LogLearning, Error, TEXT("%s: Asset is nullptr."), *GetName());
+		return;
+	}
+
+	Asset->Records = Recording->Records;
+
+	// Manually mark the package as dirty since just using `Modify` prevents 
+	// marking packages as dirty during PIE which is most likely when this
+	// is being used.
+	if (UPackage* Package = Asset->GetPackage())
+	{
+		const bool bIsDirty = Package->IsDirty();
+
+		if (!bIsDirty)
+		{
+			Package->SetDirtyFlag(true);
+		}
+
+		Package->PackageMarkedDirtyEvent.Broadcast(Package, bIsDirty);
+	}
+}
+
+void ULearningAgentsRecorder::AppendRecordingToAsset(ULearningAgentsRecording* Asset) const
+{
+	if (!IsSetup())
+	{
+		UE_LOG(LogLearning, Error, TEXT("%s: Setup not complete."), *GetName());
+		return;
+	}
+
+	if (!Asset)
+	{
+		UE_LOG(LogLearning, Error, TEXT("%s: Asset is invalid."), *GetName());
+		return;
+	}
+
+	Asset->Records.Append(Recording->Records);
+
+	// Manually mark the package as dirty since just using `Modify` prevents 
+	// marking packages as dirty during PIE which is most likely when this
+	// is being used.
+	if (UPackage* Package = Asset->GetPackage())
+	{
+		const bool bIsDirty = Package->IsDirty();
+
+		if (!bIsDirty)
+		{
+			Package->SetDirtyFlag(true);
+		}
+
+		Package->PackageMarkedDirtyEvent.Broadcast(Package, bIsDirty);
+	}
+}
+
+void ULearningAgentsRecorder::BeginRecording(bool bReinitializeRecording)
+{
+	if (!IsSetup())
+	{
+		UE_LOG(LogLearning, Error, TEXT("%s: Setup not complete."), *GetName());
 		return;
 	}
 
 	if (IsRecording())
 	{
-		UE_LOG(LogLearning, Error, TEXT("Already Recording!"));
+		UE_LOG(LogLearning, Error, TEXT("%s: Cannot being recording as we are already Recording!"), *GetName());
 		return;
 	}
 
+	if (bReinitializeRecording)
+	{
+		Recording->Records.Empty();
+	}
+
 	for (const int32 AgentId : AddedAgentSet)
 	{
-		CurrentRecords.Add(AgentId, DataStorage->CreateRecord(FName(Interactor->GetName() + "_id" + FString::FromInt(AgentId)), Interactor));
+		RecordBuffers[AgentId].Empty();
 	}
 
 	bIsRecording = true;
 }
+
+const ULearningAgentsRecording* ULearningAgentsRecorder::GetCurrentRecording() const
+{
+	if (!IsSetup())
+	{
+		UE_LOG(LogLearning, Error, TEXT("%s: Setup not complete."), *GetName());
+		return nullptr;
+	}
+
+	return Recording;
+}
\ No newline at end of file
diff --git a/Engine/Plugins/Experimental/LearningAgents/Source/LearningAgentsTraining/Private/LearningAgentsRecording.cpp b/Engine/Plugins/Experimental/LearningAgents/Source/LearningAgentsTraining/Private/LearningAgentsRecording.cpp
new file mode 100644
--- /dev/null
+++ b/Engine/Plugins/Experimental/LearningAgents/Source/LearningAgentsTraining/Private/LearningAgentsRecording.cpp
@@ -0,0 +1,140 @@
+// Copyright Epic Games, Inc. All Rights Reserved.
+
+#include "LearningAgentsRecording.h"
+
+#include "LearningLog.h"
+
+#include "UObject/Package.h"
+#include "Misc/FileHelper.h"
+
+bool FLearningAgentsRecord::Serialize(FArchive& Ar)
+{
+	Ar << SampleNum;
+	Ar << ObservationDimNum;
+	Ar << ActionDimNum;
+	UE::Learning::Array::Serialize(Ar, Observations);
+	UE::Learning::Array::Serialize(Ar, Actions);
+	return true;
+}
+
+ULearningAgentsRecording::ULearningAgentsRecording() {}
+ULearningAgentsRecording::ULearningAgentsRecording(FVTableHelper& Helper) {}
+ULearningAgentsRecording::~ULearningAgentsRecording() {}
+
+namespace UE::Learning::Agents::Recording::Private
+{
+	static constexpr int32 MagicNumber = 0x06b5fb26;
+	static constexpr int32 VersionNumber = 1;
+}
+
+void ULearningAgentsRecording::LoadRecordingFromFile(const FFilePath& File)
+{
+	TArray<uint8> RecordingData;
+
+	if (FFileHelper::LoadFileToArray(RecordingData, *File.FilePath))
+	{
+		if (RecordingData.Num() < sizeof(int32) * 3)
+		{
+			UE_LOG(LogLearning, Error, TEXT("%s: Failed to load network. Incorrect Format."), *GetName());
+			return;
+		}
+
+		int32 Offset = 0;
+
+		int32 MagicNumber;
+		UE::Learning::DeserializeFromBytes(Offset, RecordingData, MagicNumber);
+
+		if (MagicNumber != UE::Learning::Agents::Recording::Private::MagicNumber)
+		{
+			UE_LOG(LogLearning, Error, TEXT("%s: Failed to load recording. Incorrect Magic Number."), *GetName());
+			return;
+		}
+
+		int32 VersionNumber;
+		UE::Learning::DeserializeFromBytes(Offset, RecordingData, VersionNumber);
+
+		if (VersionNumber != UE::Learning::Agents::Recording::Private::VersionNumber)
+		{
+			UE_LOG(LogLearning, Error, TEXT("%s: Failed to load recording. Unsupported Version Number %i."), *GetName(), VersionNumber);
+			return;
+		}
+
+		int32 RecordNum;
+		UE::Learning::DeserializeFromBytes(Offset, RecordingData, RecordNum);
+
+		Records.SetNum(RecordNum);
+
+		for (int32 RecordIdx = 0; RecordIdx < RecordNum; RecordIdx++)
+		{
+			UE::Learning::DeserializeFromBytes(Offset, RecordingData, Records[RecordIdx].SampleNum);
+			UE::Learning::DeserializeFromBytes(Offset, RecordingData, Records[RecordIdx].ObservationDimNum);
+			UE::Learning::DeserializeFromBytes(Offset, RecordingData, Records[RecordIdx].ActionDimNum);
+			UE::Learning::Array::DeserializeFromBytes(Offset, RecordingData, Records[RecordIdx].Observations);
+			UE::Learning::Array::DeserializeFromBytes(Offset, RecordingData, Records[RecordIdx].Actions);
+		}
+
+		UE_LEARNING_CHECK(Offset == RecordingData.Num());
+	}
+	else
+	{
+		UE_LOG(LogLearning, Error, TEXT("%s: Failed to load recording. File not found: \"%s\""), *GetName(), *File.FilePath);
+	}
+
+	// Manually mark the package as dirty since just using `Modify` prevents 
+	// marking packages as dirty during PIE which is most likely when this
+	// is being used.
+	if (UPackage* Package = GetPackage())
+	{
+		const bool bIsDirty = Package->IsDirty();
+
+		if (!bIsDirty)
+		{
+			Package->SetDirtyFlag(true);
+		}
+
+		Package->PackageMarkedDirtyEvent.Broadcast(Package, bIsDirty);
+	}
+}
+
+void ULearningAgentsRecording::SaveRecordingToFile(const FFilePath& File) const
+{
+	TArray<uint8> RecordingData;
+
+	int32 TotalByteNum = 
+		sizeof(int32) + // Magic Num
+		sizeof(int32) + // Version Num
+		sizeof(int32);  // Record Num
+
+	for (int32 RecordIdx = 0; RecordIdx < Records.Num(); RecordIdx++)
+	{
+		TotalByteNum +=
+			sizeof(int32) + // SampleNum
+			sizeof(int32) + // ObservationDimNum
+			sizeof(int32) + // ActionDimNum
+			UE::Learning::Array::SerializationByteNum<2, float>(Records[RecordIdx].Observations.Shape()) + // Observations
+			UE::Learning::Array::SerializationByteNum<2, float>(Records[RecordIdx].Actions.Shape()); // Actions
+	}
+
+	RecordingData.SetNumUninitialized(TotalByteNum);
+
+	int32 Offset = 0;
+	UE::Learning::SerializeToBytes(Offset, RecordingData, UE::Learning::Agents::Recording::Private::MagicNumber);
+	UE::Learning::SerializeToBytes(Offset, RecordingData, UE::Learning::Agents::Recording::Private::VersionNumber);
+	UE::Learning::SerializeToBytes(Offset, RecordingData, Records.Num());
+
+	for (int32 RecordIdx = 0; RecordIdx < Records.Num(); RecordIdx++)
+	{
+		UE::Learning::SerializeToBytes(Offset, RecordingData, Records[RecordIdx].SampleNum);
+		UE::Learning::SerializeToBytes(Offset, RecordingData, Records[RecordIdx].ObservationDimNum);
+		UE::Learning::SerializeToBytes(Offset, RecordingData, Records[RecordIdx].ActionDimNum);
+		UE::Learning::Array::SerializeToBytes(Offset, RecordingData, Records[RecordIdx].Observations);
+		UE::Learning::Array::SerializeToBytes(Offset, RecordingData, Records[RecordIdx].Actions);
+	}
+
+	UE_LEARNING_CHECK(Offset == RecordingData.Num());
+
+	if (!FFileHelper::SaveArrayToFile(RecordingData, *File.FilePath))
+	{
+		UE_LOG(LogLearning, Error, TEXT("%s: Failed to save recording to file: \"%s\""), *GetName(), *File.FilePath);
+	}
+}
\ No newline at end of file
diff --git a/Engine/Plugins/Experimental/LearningAgents/Source/LearningAgentsTraining/Private/LearningAgentsRewards.cpp b/Engine/Plugins/Experimental/LearningAgents/Source/LearningAgentsTraining/Private/LearningAgentsRewards.cpp
--- a/Engine/Plugins/Experimental/LearningAgents/Source/LearningAgentsTraining/Private/LearningAgentsRewards.cpp
+++ b/Engine/Plugins/Experimental/LearningAgents/Source/LearningAgentsTraining/Private/LearningAgentsRewards.cpp
@@ -14,11 +14,11 @@
 namespace UE::Learning::Agents::Rewards::Private
 {
 	template<typename RewardUObject, typename RewardFObject, typename... InArgTypes>
-	RewardUObject* AddReward(ULearningAgentsTrainer* InAgentTrainer, const FName Name, InArgTypes&& ...Args)
+	RewardUObject* AddReward(ULearningAgentsTrainer* InAgentTrainer, const FName Name, const TCHAR* FunctionName, InArgTypes&& ...Args)
 	{
 		if (!InAgentTrainer)
 		{
-			UE_LOG(LogLearning, Error, TEXT("InAgentTrainer is nullptr."));
+			UE_LOG(LogLearning, Error, TEXT("%s: InAgentTrainer is nullptr."), FunctionName);
 			return nullptr;
 		}
 
@@ -41,7 +41,7 @@ namespace UE::Learning::Agents::Rewards::Private
 
 UFloatReward* UFloatReward::AddFloatReward(ULearningAgentsTrainer* InAgentTrainer, const FName Name, const float Weight)
 {
-	return UE::Learning::Agents::Rewards::Private::AddReward<UFloatReward, UE::Learning::FFloatReward>(InAgentTrainer, Name, Weight);
+	return UE::Learning::Agents::Rewards::Private::AddReward<UFloatReward, UE::Learning::FFloatReward>(InAgentTrainer, Name, TEXT("AddFloatReward"), Weight);
 }
 
 void UFloatReward::SetFloatReward(const int32 AgentId, const float Reward)
@@ -85,7 +85,7 @@ void UFloatReward::VisualLog(const UE::Learning::FIndexSet Instances) const
 
 UScalarVelocityReward* UScalarVelocityReward::AddScalarVelocityReward(ULearningAgentsTrainer* InAgentTrainer, const FName Name, const float Weight, const float Scale)
 {
-	return UE::Learning::Agents::Rewards::Private::AddReward<UScalarVelocityReward, UE::Learning::FScalarVelocityReward>(InAgentTrainer, Name, Weight, Scale);
+	return UE::Learning::Agents::Rewards::Private::AddReward<UScalarVelocityReward, UE::Learning::FScalarVelocityReward>(InAgentTrainer, Name, TEXT("AddScalarVelocityReward"), Weight, Scale);
 }
 
 void UScalarVelocityReward::SetScalarVelocityReward(int32 AgentId, float Velocity)
@@ -129,7 +129,7 @@ void UScalarVelocityReward::VisualLog(const UE::Learning::FIndexSet Instances) c
 
 ULocalDirectionalVelocityReward* ULocalDirectionalVelocityReward::AddLocalDirectionalVelocityReward(ULearningAgentsTrainer* InAgentTrainer, const FName Name, const float Weight, const float Scale, const FVector Axis)
 {
-	return UE::Learning::Agents::Rewards::Private::AddReward<ULocalDirectionalVelocityReward, UE::Learning::FLocalDirectionalVelocityReward>(InAgentTrainer, Name, Weight, Scale, Axis);
+	return UE::Learning::Agents::Rewards::Private::AddReward<ULocalDirectionalVelocityReward, UE::Learning::FLocalDirectionalVelocityReward>(InAgentTrainer, Name, TEXT("AddLocalDirectionalVelocityReward"), Weight, Scale, Axis);
 }
 
 void ULocalDirectionalVelocityReward::SetLocalDirectionalVelocityReward(const int32 AgentId, const FVector Velocity, const FRotator RelativeRotation)
@@ -216,7 +216,7 @@ UPlanarPositionDifferencePenalty* UPlanarPositionDifferencePenalty::AddPlanarPos
 	const FVector Axis0,
 	const FVector Axis1)
 {
-	return UE::Learning::Agents::Rewards::Private::AddReward<UPlanarPositionDifferencePenalty, UE::Learning::FPlanarPositionDifferencePenalty>(InAgentTrainer, Name, Weight, Scale, Threshold, Axis0, Axis1);
+	return UE::Learning::Agents::Rewards::Private::AddReward<UPlanarPositionDifferencePenalty, UE::Learning::FPlanarPositionDifferencePenalty>(InAgentTrainer, Name, TEXT("AddPlanarPositionDifferencePenalty"), Weight, Scale, Threshold, Axis0, Axis1);
 }
 
 void UPlanarPositionDifferencePenalty::SetPlanarPositionDifferencePenalty(const int32 AgentId, const FVector Position0, const FVector Position1)
@@ -311,7 +311,13 @@ UPositionArraySimilarityReward* UPositionArraySimilarityReward::AddPositionArray
 	const float Weight,
 	const float Scale)
 {
-	return UE::Learning::Agents::Rewards::Private::AddReward<UPositionArraySimilarityReward, UE::Learning::FPositionArraySimilarityReward>(InAgentTrainer, Name, PositionNum, Weight, Scale);
+	if (PositionNum < 1)
+	{
+		UE_LOG(LogLearning, Error, TEXT("AddPositionArraySimilarityReward: Number of elements in array must be at least 1, got %i."), PositionNum);
+		return nullptr;
+	}
+
+	return UE::Learning::Agents::Rewards::Private::AddReward<UPositionArraySimilarityReward, UE::Learning::FPositionArraySimilarityReward>(InAgentTrainer, Name, TEXT("AddPositionArraySimilarityReward"), PositionNum, Weight, Scale);
 }
 
 void UPositionArraySimilarityReward::SetPositionArraySimilarityReward(
diff --git a/Engine/Plugins/Experimental/LearningAgents/Source/LearningAgentsTraining/Private/LearningAgentsTrainer.cpp b/Engine/Plugins/Experimental/LearningAgents/Source/LearningAgentsTraining/Private/LearningAgentsTrainer.cpp
--- a/Engine/Plugins/Experimental/LearningAgents/Source/LearningAgentsTraining/Private/LearningAgentsTrainer.cpp
+++ b/Engine/Plugins/Experimental/LearningAgents/Source/LearningAgentsTraining/Private/LearningAgentsTrainer.cpp
@@ -52,6 +52,32 @@ namespace UE::Learning::Agents
 		}
 	}
 }
+
+FLearningAgentsTrainerPathSettings::FLearningAgentsTrainerPathSettings()
+{
+	EditorEngineRelativePath.Path = FPaths::EngineDir();
+	IntermediateRelativePath.Path = FPaths::ProjectIntermediateDir();
+}
+
+FString FLearningAgentsTrainerPathSettings::GetEditorEnginePath() const
+{
+#if WITH_EDITOR
+	return EditorEngineRelativePath.Path;
+#else
+	if (NonEditorEngineRelativePath.IsEmpty())
+	{
+		UE_LOG(LogLearning, Warning, TEXT("GetEditorEnginePath: NonEditorEngineRelativePath not set"));
+	}
+
+	return NonEditorEngineRelativePath;
+#endif
+}
+
+FString FLearningAgentsTrainerPathSettings::GetIntermediatePath() const
+{
+	return IntermediateRelativePath.Path;
+}
+
 ULearningAgentsTrainer::ULearningAgentsTrainer() : ULearningAgentsManagerComponent() {}
 ULearningAgentsTrainer::ULearningAgentsTrainer(FVTableHelper& Helper) : ULearningAgentsTrainer() {}
 ULearningAgentsTrainer::~ULearningAgentsTrainer() {}
@@ -132,6 +158,12 @@ void ULearningAgentsTrainer::SetupTrainer(
 		AgentManager->GetInstanceData().ToSharedRef(),
 		AgentManager->GetMaxInstanceNum());
 
+	if (RewardObjects.Num() == 0)
+	{
+		UE_LOG(LogLearning, Error, TEXT("%s: No rewards added to Trainer during SetupRewards."), *GetName());
+		return;
+	}
+
 	// Setup Completions
 	CompletionObjects.Empty();
 	CompletionFeatures.Empty();
@@ -141,6 +173,12 @@ void ULearningAgentsTrainer::SetupTrainer(
 		AgentManager->GetInstanceData().ToSharedRef(),
 		AgentManager->GetMaxInstanceNum());
 
+	if (RewardObjects.Num() == 0)
+	{
+		// Not an error or warning because it's fine to run training without any completions.
+		UE_LOG(LogLearning, Display, TEXT("%s: No completions added to Trainer during SetupCompletions."), *GetName());
+	}
+
 	// Create Episode Buffer
 	EpisodeBuffer = MakeUnique<UE::Learning::FEpisodeBuffer>();
 	EpisodeBuffer->Resize(
@@ -234,15 +272,16 @@ const bool ULearningAgentsTrainer::IsTraining() const
 }
 
 void ULearningAgentsTrainer::BeginTraining(
-	const FLearningAgentsTrainerTrainingSettings& TrainingSettings, 
+	const FLearningAgentsTrainerTrainingSettings& TrainerTrainingSettings,
 	const FLearningAgentsTrainerGameSettings& TrainerGameSettings,
+	const FLearningAgentsTrainerPathSettings& TrainerPathSettings,
 	const FLearningAgentsCriticSettings& CriticSettings,
 	const bool bReinitializePolicyNetwork,
 	const bool bReinitializeCriticNetwork)
 {
 	if (!IsSetup())
 	{
-		UE_LOG(LogLearning, Error, TEXT("%s: Setup not run."), *GetName());
+		UE_LOG(LogLearning, Error, TEXT("%s: Setup not complete."), *GetName());
 		return;
 	}
 
@@ -252,6 +291,34 @@ void ULearningAgentsTrainer::BeginTraining(
 		return;
 	}
 
+	// Check Paths
+
+	const FString PythonExecutablePath = UE::Learning::Trainer::GetPythonExecutablePath(TrainerPathSettings.GetEditorEnginePath());
+
+	if (!FPaths::FileExists(PythonExecutablePath))
+	{
+		UE_LOG(LogLearning, Error, TEXT("%s: Can't find Python executable \"%s\"."), *GetName(), *PythonExecutablePath);
+		return;
+	}
+
+	const FString PythonContentPath = UE::Learning::Trainer::GetPythonContentPath(TrainerPathSettings.GetEditorEnginePath());
+
+	if (!FPaths::DirectoryExists(PythonContentPath))
+	{
+		UE_LOG(LogLearning, Error, TEXT("%s: Can't find LearningAgents plugin Content \"%s\"."), *GetName(), *PythonContentPath);
+		return;
+	}
+
+	const FString SitePackagesPath = UE::Learning::Trainer::GetSitePackagesPath(TrainerPathSettings.GetEditorEnginePath());
+
+	if (!FPaths::DirectoryExists(SitePackagesPath))
+	{
+		UE_LOG(LogLearning, Error, TEXT("%s: Can't find Python site-packages \"%s\"."), *GetName(), *SitePackagesPath);
+		return;
+	}
+
+	const FString IntermediatePath = UE::Learning::Trainer::GetIntermediatePath(TrainerPathSettings.GetIntermediatePath());
+
 	// Record GameState Settings
 
 	bFixedTimestepUsed = FApp::UseFixedTimeStep();
@@ -305,42 +372,15 @@ void ULearningAgentsTrainer::BeginTraining(
 
 	// Start Trainer
 
-#if WITH_EDITOR
-	const FString PythonExecutablePath = UE::Learning::Trainer::DefaultEditorPythonExecutablePath();
-	const FString SitePackagesPath = UE::Learning::Trainer::DefaultEditorSitePackagesPath();
-	const FString PythonContentPath = UE::Learning::Trainer::DefaultEditorPythonContentPath();
-	const FString IntermediatePath = UE::Learning::Trainer::DefaultEditorIntermediatePath();
-#else
-	// If we want to run training in a cooked, non-editor build, then by default we wont have access to python or the 
-	// learning training scripts - these are editor-only and will be stripped during the cooking process.
-	//
-	// However, running training in non-editor builds can be very important - we probably want to disable rendering 
-	// and sound while we are training to make experience gathering as fast as possible - and for any non-trivial game 
-	// is simply may not be realistic to run it for a long time in play-in-editor mode.
-	//
-	// For this reason even in non-editor builds we let you provide paths the `python.exe` provided by the editor, 
-	// as well as PythonFoundationPackages site-packages, and the Learning training scripts. This allows you to 
-	// run training when these things actually exist somewhere on your machine, which will usually be the case on 
-	// a normal development machine.
-
-	UE_LEARNING_CHECKF(PLATFORM_WINDOWS || PLATFORM_MAC || PLATFORM_LINUX, TEXT("Python only supported on Windows, Mac, and Linux."));
-	const FString PythonExe = PLATFORM_WINDOWS ? TEXT("python.exe") : TEXT("bin/python");
-	const FString EnginePath = IFileManager::Get().ConvertToAbsolutePathForExternalAppForRead(*(FPaths::RootDir() / TEXT("../../../../../../Engine")));
-	const FString PythonExecutablePath = EnginePath / TEXT("Binaries/ThirdParty/Python3") / FPlatformMisc::GetUBTPlatform() / PythonExe;
-	const FString SitePackagesPath = EnginePath / TEXT("Plugins/Experimental/PythonFoundationPackages/Content/Python/Lib/") / FPlatformMisc::GetUBTPlatform() / TEXT("site-packages");
-	const FString PythonContentPath = EnginePath / TEXT("Plugins/Experimental/LearningAgents/Content/Python/");
-	const FString IntermediatePath = EnginePath / TEXT("Plugins/Experimental/LearningAgents/Intermediate");
-#endif
-
 	UE::Learning::FPPOTrainerTrainingSettings PPOTrainingSettings;
-	PPOTrainingSettings.IterationNum = TrainingSettings.NumberOfIterations;
-	PPOTrainingSettings.bUseTensorboard = TrainingSettings.bUseTensorboard;
-	PPOTrainingSettings.InitialActionScale = TrainingSettings.InitialActionScale;
-	PPOTrainingSettings.DiscountFactor = TrainingSettings.DiscountFactor;
-	PPOTrainingSettings.Seed = TrainingSettings.RandomSeed;
-	PPOTrainingSettings.TrimEpisodeStartStepNum = TrainingSettings.NumberOfStepsToTrimAtStartOfEpisode;
-	PPOTrainingSettings.TrimEpisodeEndStepNum = TrainingSettings.NumberOfStepsToTrimAtEndOfEpisode;
-	PPOTrainingSettings.Device = TrainingSettings.Device == ELearningAgentsTrainerDevice::CPU ? UE::Learning::ETrainerDevice::CPU : UE::Learning::ETrainerDevice::GPU;
+	PPOTrainingSettings.IterationNum = TrainerTrainingSettings.NumberOfIterations;
+	PPOTrainingSettings.bUseTensorboard = TrainerTrainingSettings.bUseTensorboard;
+	PPOTrainingSettings.InitialActionScale = TrainerTrainingSettings.InitialActionScale;
+	PPOTrainingSettings.DiscountFactor = TrainerTrainingSettings.DiscountFactor;
+	PPOTrainingSettings.Seed = TrainerTrainingSettings.RandomSeed;
+	PPOTrainingSettings.TrimEpisodeStartStepNum = TrainerTrainingSettings.NumberOfStepsToTrimAtStartOfEpisode;
+	PPOTrainingSettings.TrimEpisodeEndStepNum = TrainerTrainingSettings.NumberOfStepsToTrimAtEndOfEpisode;
+	PPOTrainingSettings.Device = TrainerTrainingSettings.Device == ELearningAgentsTrainerDevice::CPU ? UE::Learning::ETrainerDevice::CPU : UE::Learning::ETrainerDevice::GPU;
 
 	UE::Learning::FPPOTrainerNetworkSettings PPONetworkSettings;
 	PPONetworkSettings.PolicyActionNoiseMin = Policy->GetPolicyObject().Settings.ActionNoiseMin;
@@ -494,7 +534,7 @@ void ULearningAgentsTrainer::EvaluateRewards()
 
 	if (!IsSetup())
 	{
-		UE_LOG(LogLearning, Error, TEXT("%s: Setup not run."), *GetName());
+		UE_LOG(LogLearning, Error, TEXT("%s: Setup not complete."), *GetName());
 		return;
 	}
 
@@ -519,7 +559,7 @@ void ULearningAgentsTrainer::EvaluateCompletions()
 
 	if (!IsSetup())
 	{
-		UE_LOG(LogLearning, Error, TEXT("%s: Setup not run."), *GetName());
+		UE_LOG(LogLearning, Error, TEXT("%s: Setup not complete."), *GetName());
 		return;
 	}
 
@@ -542,7 +582,7 @@ void ULearningAgentsTrainer::ProcessExperience()
 {
 	if (!IsSetup())
 	{
-		UE_LOG(LogLearning, Error, TEXT("%s: Setup not run."), *GetName());
+		UE_LOG(LogLearning, Error, TEXT("%s: Setup not complete."), *GetName());
 		return;
 	}
 
@@ -655,7 +695,7 @@ void ULearningAgentsTrainer::ResetAllEpisodes()
 
 	if (!IsSetup())
 	{
-		UE_LOG(LogLearning, Error, TEXT("%s: Setup not run."), *GetName());
+		UE_LOG(LogLearning, Error, TEXT("%s: Setup not complete."), *GetName());
 		return;
 	}
 
@@ -665,22 +705,29 @@ void ULearningAgentsTrainer::ResetAllEpisodes()
 }
 
 void ULearningAgentsTrainer::RunTraining(
-	const FLearningAgentsTrainerTrainingSettings& TrainingSettings,
+	const FLearningAgentsTrainerTrainingSettings& TrainerTrainingSettings,
 	const FLearningAgentsTrainerGameSettings& TrainerGameSettings,
+	const FLearningAgentsTrainerPathSettings& TrainerPathSettings,
 	const FLearningAgentsCriticSettings& CriticSettings,
 	const bool bReinitializePolicyNetwork,
 	const bool bReinitializeCriticNetwork)
 {
 	if (!IsSetup())
 	{
-		UE_LOG(LogLearning, Error, TEXT("%s: Setup not run."), *GetName());
+		UE_LOG(LogLearning, Error, TEXT("%s: Setup not complete."), *GetName());
 		return;
 	}
 
 	// If we aren't training yet, then start training and do the first inference step.
 	if (!IsTraining())
 	{
-		BeginTraining(TrainingSettings, TrainerGameSettings, CriticSettings, bReinitializePolicyNetwork, bReinitializeCriticNetwork);
+		BeginTraining(
+			TrainerTrainingSettings,
+			TrainerGameSettings, 
+			TrainerPathSettings, 
+			CriticSettings, 
+			bReinitializePolicyNetwork, 
+			bReinitializeCriticNetwork);
 
 		if (!IsTraining())
 		{
@@ -702,7 +749,7 @@ float ULearningAgentsTrainer::GetReward(const int32 AgentId) const
 {
 	if (!IsSetup())
 	{
-		UE_LOG(LogLearning, Error, TEXT("%s: Setup not run."), *GetName());
+		UE_LOG(LogLearning, Error, TEXT("%s: Setup not complete."), *GetName());
 		return 0.0f;
 	}
 
@@ -719,7 +766,7 @@ bool ULearningAgentsTrainer::IsCompleted(const int32 AgentId, ELearningAgentsCom
 {
 	if (!IsSetup())
 	{
-		UE_LOG(LogLearning, Error, TEXT("%s: Setup not run."), *GetName());
+		UE_LOG(LogLearning, Error, TEXT("%s: Setup not complete."), *GetName());
 		OutCompletion = ELearningAgentsCompletion::Termination;
 		return false;
 	}
diff --git a/Engine/Plugins/Experimental/LearningAgents/Source/LearningAgentsTraining/Public/LearningAgentsDataStorage.h b/Engine/Plugins/Experimental/LearningAgents/Source/LearningAgentsTraining/Public/LearningAgentsDataStorage.h
deleted file mode 100644
diff --git a/Engine/Plugins/Experimental/LearningAgents/Source/LearningAgentsTraining/Public/LearningAgentsImitationTrainer.h b/Engine/Plugins/Experimental/LearningAgents/Source/LearningAgentsTraining/Public/LearningAgentsImitationTrainer.h
--- a/Engine/Plugins/Experimental/LearningAgents/Source/LearningAgentsTraining/Public/LearningAgentsImitationTrainer.h
+++ b/Engine/Plugins/Experimental/LearningAgents/Source/LearningAgentsTraining/Public/LearningAgentsImitationTrainer.h
@@ -2,7 +2,7 @@
 
 #pragma once
 
-#include "LearningAgentsTrainer.h" // Included for ELearningAgentsTrainerDevice
+#include "LearningAgentsTrainer.h" // Included for ELearningAgentsTrainerDevice and FLearningAgentsTrainerPathSettings
 #include "LearningArray.h"
 
 #include "Templates/SharedPointer.h"
@@ -18,6 +18,7 @@ namespace UE::Learning
 
 class ULearningAgentsInteractor;
 class ULearningAgentsPolicy;
+class ULearningAgentsRecording;
 
 /** The configurable settings for the training process. */
 USTRUCT(BlueprintType, Category = "LearningAgents")
@@ -94,15 +95,17 @@ public:
 	/**
 	* Begins the training process with the provided settings.
 	* @param InPolicy The policy to train.
-	* @param Records The data to train on.
-	* @param TrainingSettings The settings for this training run.
+	* @param Recording The data to train on.
+	* @param ImitationTrainerTrainingSettings The settings for this training run.
+	* @param ImitationTrainerPathSettings The path settings used by the imitation trainer.
 	* @param bReinitializePolicyNetwork If true, reinitialize the policy. Set this to false if your policy is pre-trained.
 	*/
 	UFUNCTION(BlueprintCallable, Category = "LearningAgents")
 	void BeginTraining(
 		ULearningAgentsPolicy* InPolicy,
-		const TArray<ULearningAgentsRecord*>& Records,
-		const FLearningAgentsImitationTrainerTrainingSettings& TrainingSettings = FLearningAgentsImitationTrainerTrainingSettings(),
+		const ULearningAgentsRecording* Recording,
+		const FLearningAgentsImitationTrainerTrainingSettings& ImitationTrainerTrainingSettings = FLearningAgentsImitationTrainerTrainingSettings(),
+		const FLearningAgentsTrainerPathSettings& ImitationTrainerPathSettings = FLearningAgentsTrainerPathSettings(),
 		const bool bReinitializePolicyNetwork = true);
 
 	/** Stops the training process. */
diff --git a/Engine/Plugins/Experimental/LearningAgents/Source/LearningAgentsTraining/Public/LearningAgentsRecorder.h b/Engine/Plugins/Experimental/LearningAgents/Source/LearningAgentsTraining/Public/LearningAgentsRecorder.h
--- a/Engine/Plugins/Experimental/LearningAgents/Source/LearningAgentsTraining/Public/LearningAgentsRecorder.h
+++ b/Engine/Plugins/Experimental/LearningAgents/Source/LearningAgentsTraining/Public/LearningAgentsRecorder.h
@@ -10,8 +10,27 @@
 
 #include "LearningAgentsRecorder.generated.h"
 
-class ULearningAgentsDataStorage;
-class ULearningAgentsRecord;
+struct FLearningAgentsRecord;
+class ULearningAgentsRecording;
+
+/** The path settings for the recorder. */
+USTRUCT(BlueprintType, Category = "LearningAgents")
+struct FLearningAgentsRecorderPathSettings
+{
+	GENERATED_BODY()
+
+public:
+
+	FLearningAgentsRecorderPathSettings();
+
+	/** The relative path to the Intermediate directory. Defaults to FPaths::ProjectIntermediateDir. */
+	UPROPERTY(EditAnywhere, Category = "LearningAgents", meta = (RelativePath))
+	FDirectoryPath IntermediateRelativePath;
+
+	/** The name of the sub-directory to use in the intermediate directory */
+	UPROPERTY(EditAnywhere, Category = "LearningAgents", meta = (RelativePath))
+	FString RecordingsSubdirectory = TEXT("Recordings");
+};
 
 /** A component that can be used to create recordings of training data for imitation learning. */
 UCLASS(BlueprintType, Blueprintable)
@@ -32,16 +51,19 @@ public:
 
 	/**
 	* Initializes this object and runs the setup functions for the underlying data storage.
+	* @param InAgentManager The agent manager we are associated with.
 	* @param InInteractor The agent interactor we are recording with.
+	* @param RecorderPathSettings The path settings used by the recorder.
 	*/
 	UFUNCTION(BlueprintCallable, Category = "LearningAgents")
-	void SetupRecorder(ALearningAgentsManager* InAgentManager, ULearningAgentsInteractor* InInteractor);
+	void SetupRecorder(
+		ALearningAgentsManager* InAgentManager, 
+		ULearningAgentsInteractor* InInteractor,
+		const FLearningAgentsRecorderPathSettings& RecorderPathSettings = FLearningAgentsRecorderPathSettings());
 
 public:
 
 	//~ Begin ULearningAgentsManagerComponent Interface
-	virtual bool AddAgent(const int32 AgentId) override;
-
 	virtual bool RemoveAgent(const int32 AgentId) override;
 	//~ End ULearningAgentsManagerComponent Interface
 
@@ -49,25 +71,58 @@ public:
 public:
 
 	/**
-	* Adds experience to the added agents' recordings. Call this after ULearningAgentsInteractor::EncodeObservations and
-	* either ULearningAgentsController::EncodeActions (if recording a human/AI demonstration) or
-	* ULearningAgentsInteractor::DecodeActions (if recording another policy).
+	* Begins the recording of the observations and actions of each added agent.
+	* @param bReinitializeRecording If to reinitialize the current recording object.
 	*/
 	UFUNCTION(BlueprintCallable, Category = "LearningAgents")
-	void AddExperience();
+	void BeginRecording(bool bReinitializeRecording = true);
 
-	/** Begin new recordings for each added agent. */
-	UFUNCTION(BlueprintCallable, Category = "LearningAgents")
-	void BeginRecording();
-
-	/** End all recordings. */
+	/**
+	* Ends the recording of the observations and actions of each agent and stores them in the current recording object.
+	*/
 	UFUNCTION(BlueprintCallable, Category = "LearningAgents")
 	void EndRecording();
 
-	/** Returns true if the recorder is currently recording; Otherwise, false. */
+	/** Returns true if recording is active; Otherwise, false. */
 	UFUNCTION(BlueprintCallable, Category = "LearningAgents")
 	bool IsRecording() const;
 
+	/**
+	* While recording, adds the current observations and actions of the added agents to the internal buffer. Call this 
+	* after ULearningAgentsInteractor::EncodeObservations and either ULearningAgentsController::EncodeActions (if 
+	* recording a human/AI demonstration) or ULearningAgentsInteractor::DecodeActions (if recording another policy).
+	*/
+	UFUNCTION(BlueprintCallable, Category = "LearningAgents")
+	void AddExperience();
+
+	/** Gets the current recording object. Note: this may be empty until EndRecording has been called. */
+	UFUNCTION(BlueprintPure, Category = "LearningAgents")
+	const ULearningAgentsRecording* GetCurrentRecording() const;
+
+// ----- Load / Save -----
+public:
+
+	/** Loads the current recording object from a file */
+	UFUNCTION(BlueprintCallable, Category = "LearningAgents", meta = (RelativePath))
+	void LoadRecordingFromFile(const FFilePath& File);
+
+	/** Saves the current recording object to a file */
+	UFUNCTION(BlueprintCallable, Category = "LearningAgents", meta = (RelativePath))
+	void SaveRecordingToFile(const FFilePath& File) const;
+
+	/** Loads the current recording object from the given recording asset */
+	UFUNCTION(BlueprintCallable, Category = "LearningAgents")
+	void LoadRecordingFromAsset(const ULearningAgentsRecording* Asset);
+
+	/** Saves the current recording object to the given recording asset */
+	UFUNCTION(BlueprintCallable, Category = "LearningAgents", meta = (DevelopmentOnly))
+	void SaveRecordingToAsset(ULearningAgentsRecording* Asset) const;
+
+	/** Appends the current recording object to the given recording asset */
+	UFUNCTION(BlueprintCallable, Category = "LearningAgents", meta = (DevelopmentOnly))
+	void AppendRecordingToAsset(ULearningAgentsRecording* Asset) const;
+
+
 // ----- Private Data ----- 
 private:
 
@@ -75,29 +130,42 @@ private:
 	UPROPERTY(VisibleAnywhere, Transient, Category = "LearningAgents")
 	TObjectPtr<ULearningAgentsInteractor> Interactor;
 
-// ----- Recorder Configuration -----
+	/** The current recording object. */
+	UPROPERTY(VisibleAnywhere, Transient, Category = "LearningAgents")
+	TObjectPtr<ULearningAgentsRecording> Recording;
+
+	/** True if recording is currently in-progress. Otherwise, false. */
+	UPROPERTY(VisibleAnywhere, Transient, Category = "LearningAgents")
+	bool bIsRecording = false;
+
+// ----- Private Data ----- 
 private:
 
-	/** Directory where records will be saved. If not set, SetupRecorder will automatically set this to the editor's default intermediate folder. */
-	UPROPERTY(EditDefaultsOnly, Category = "LearningAgents")
-	FDirectoryPath DataDirectory;
+	FString RecordingDirectory;
 
-	/** If true, recorder will automatically save all records on EndRecording. Set this to false if you want to manually save records. */
-	UPROPERTY(EditAnywhere, Category = "LearningAgents")
-	bool bSaveDataOnEndPlay = true;
+	/** Basic structure used to buffer the observations and actions of each agent. */
+	struct FAgentRecordBuffer
+	{
+		const int32 ChunkSize = 1024;
+		int32 SampleNum = 0;
+		TArray<TLearningArray<2, float>, TInlineAllocator<16>> Observations;
+		TArray<TLearningArray<2, float>, TInlineAllocator<16>> Actions;
 
-// ----- Recorder State -----
-private:
+		TLearningArrayView<1, float> GetObservation(const int32 SampleIdx);
+		TLearningArrayView<1, float> GetAction(const int32 SampleIdx);
+		TLearningArrayView<1, const float> GetObservation(const int32 SampleIdx) const;
+		TLearningArrayView<1, const float> GetAction(const int32 SampleIdx) const;
 
-	/** True if recording is currently in-progress. Otherwise, false. */
-	UPROPERTY(VisibleAnywhere, Transient, Category = "LearningAgents")
-	bool bIsRecording = false;
+		bool IsEmpty() const;
 
-	/** The data storage manager. It can be used to save/load agent records. */
-	UPROPERTY(VisibleAnywhere, Transient, Category = "LearningAgents")
-	TObjectPtr<ULearningAgentsDataStorage> DataStorage;
+		void Empty();
 
-	/** All records which are currently being written to. */
-	UPROPERTY(VisibleAnywhere, Transient, Category = "LearningAgents")
-	TMap<int32, TObjectPtr<ULearningAgentsRecord>> CurrentRecords;
+		void Push(
+			const TLearningArrayView<1, const float> Observation,
+			const TLearningArrayView<1, const float> Action);
+
+		void CopyToRecord(FLearningAgentsRecord& Record) const;
+	};
+
+	TArray<FAgentRecordBuffer, TInlineAllocator<32>> RecordBuffers;
 };
diff --git a/Engine/Plugins/Experimental/LearningAgents/Source/LearningAgentsTraining/Public/LearningAgentsRecording.h b/Engine/Plugins/Experimental/LearningAgents/Source/LearningAgentsTraining/Public/LearningAgentsRecording.h
new file mode 100644
--- /dev/null
+++ b/Engine/Plugins/Experimental/LearningAgents/Source/LearningAgentsTraining/Public/LearningAgentsRecording.h
@@ -0,0 +1,75 @@
+// Copyright Epic Games, Inc. All Rights Reserved.
+
+#pragma once
+
+#include "LearningArray.h"
+
+#include "Containers/Array.h"
+#include "UObject/NameTypes.h"
+#include "Engine/DataAsset.h"
+
+#include "LearningAgentsRecording.generated.h"
+
+/** A single recording of a series of observations and actions. */
+USTRUCT(BlueprintType)
+struct LEARNINGAGENTSTRAINING_API FLearningAgentsRecord
+{
+	GENERATED_BODY()
+
+public:
+
+	/** The number of observations and actions recorded. */
+	UPROPERTY(BlueprintReadOnly, VisibleAnywhere, Category = "LearningAgents")
+	int32 SampleNum = 0;
+
+	/** The number of dimensions in the observation vector for this record */
+	UPROPERTY(BlueprintReadOnly, VisibleAnywhere, Category = "LearningAgents")
+	int32 ObservationDimNum = 0;
+
+	/** The number of dimensions in the action vector for this record */
+	UPROPERTY(BlueprintReadOnly, VisibleAnywhere, Category = "LearningAgents")
+	int32 ActionDimNum = 0;
+
+	bool Serialize(FArchive& Ar);
+
+	TLearningArray<2, float> Observations;
+	TLearningArray<2, float> Actions;
+};
+
+template<>
+struct TStructOpsTypeTraits<FLearningAgentsRecord> : public TStructOpsTypeTraitsBase2<FLearningAgentsRecord>
+{
+	enum
+	{
+		WithSerializer = true,
+	};
+};
+
+/** Data asset representing an array of records. */
+UCLASS(BlueprintType)
+class LEARNINGAGENTSTRAINING_API ULearningAgentsRecording : public UDataAsset
+{
+	GENERATED_BODY()
+
+public:
+	// These constructors/destructors are needed to make forward declarations happy
+	ULearningAgentsRecording();
+	ULearningAgentsRecording(FVTableHelper& Helper);
+	virtual ~ULearningAgentsRecording();
+
+public:
+
+	/** Load this recording from a file. */
+	UFUNCTION(BlueprintCallable, Category = "LearningAgents", meta = (RelativePath))
+	void LoadRecordingFromFile(const FFilePath& File);
+
+	/** Save this recording to a file. */
+	UFUNCTION(BlueprintCallable, Category = "LearningAgents", meta = (RelativePath))
+	void SaveRecordingToFile(const FFilePath& File) const;
+
+public:
+
+	/** Set of records. */
+	UPROPERTY(EditInstanceOnly, Category = "LearningAgents")
+	TArray<FLearningAgentsRecord> Records;
+};
diff --git a/Engine/Plugins/Experimental/LearningAgents/Source/LearningAgentsTraining/Public/LearningAgentsTrainer.h b/Engine/Plugins/Experimental/LearningAgents/Source/LearningAgentsTraining/Public/LearningAgentsTrainer.h
--- a/Engine/Plugins/Experimental/LearningAgents/Source/LearningAgentsTraining/Public/LearningAgentsTrainer.h
+++ b/Engine/Plugins/Experimental/LearningAgents/Source/LearningAgentsTraining/Public/LearningAgentsTrainer.h
@@ -181,6 +181,54 @@ public:
 	ELearningAgentsTrainerDevice Device = ELearningAgentsTrainerDevice::GPU;
 };
 
+/** The path settings for the trainer. */
+USTRUCT(BlueprintType, Category = "LearningAgents")
+struct FLearningAgentsTrainerPathSettings
+{
+	GENERATED_BODY()
+
+public:
+
+	FLearningAgentsTrainerPathSettings();
+
+	/** The relative path to the engine for editor builds. Defaults to FPaths::EngineDir. */
+	UPROPERTY(EditAnywhere, Category = "LearningAgents", meta = (RelativePath))
+	FDirectoryPath EditorEngineRelativePath;
+
+	/**
+	* The relative path to the editor engine folder for non-editor builds.
+	* 
+	* If we want to run training in cooked, non-editor builds, then by default we wont have access to python and the
+	* LearningAgents training scripts - these are editor-only things and are stripped during the cooking process.
+	*
+	* However, running training in non-editor builds can be very important - we probably want to disable rendering
+	* and sound while we are training to make experience gathering as fast as possible - and for any non-trivial game
+	* is simply may not be realistic to run it for a long time in play-in-editor.
+	*
+	* For this reason even in non-editor builds we let you provide the path where all of these editor-only things can 
+	* be found. This allows you to run training when these things actually exist somewhere accessible to the executable, 
+	* which will usually be the case on a normal development machine or cloud machine if it is set up that way.
+	*
+	* Since non-editor builds can be produced in a number of different ways, this is not set by default and cannot 
+	* use a directory picker since it is relative to the final location of where your cooked, non-editor executable 
+	* will exist rather than the current with-editor executable.
+	*/
+	UPROPERTY(EditAnywhere, Category = "LearningAgents")
+	FString NonEditorEngineRelativePath;
+
+	/** The relative path to the Intermediate directory. Defaults to FPaths::ProjectIntermediateDir. */
+	UPROPERTY(EditAnywhere, Category = "LearningAgents", meta = (RelativePath))
+	FDirectoryPath IntermediateRelativePath;
+
+public:
+
+	/** Gets the Relative Editor Engine Path accounting for if this is an editor build or not  */
+	FString GetEditorEnginePath() const;
+
+	/** Gets the Relative Intermediate Path  */
+	FString GetIntermediatePath() const;
+};
+
 /**
 * The ULearningAgentsTrainer is the core class for reinforcement learning training. It has a few responsibilities:
 *   1) It keeps track of which agents are gathering training data.
@@ -305,16 +353,18 @@ public:
 
 	/**
 	* Begins the training process with the provided settings.
-	* @param TrainingSettings The settings for this training run.
-	* @param GameSettings The settings that will affect the game's simulation.
+	* @param TrainerTrainingSettings The settings for this training run.
+	* @param TrainerGameSettings The settings that will affect the game's simulation.
+	* @param TrainerPathSettings The path settings used by the trainer.
 	* @param CriticSettings The settings for the critic (if we are using one).
 	* @param bReinitializePolicyNetwork If true, reinitialize the policy. Set this to false if your policy is pre-trained, e.g. with imitation learning.
 	* @param bReinitializeCriticNetwork If true, reinitialize the critic. Set this to false if your critic is pre-trained.
 	*/
 	UFUNCTION(BlueprintCallable, Category = "LearningAgents")
 	void BeginTraining(
-		const FLearningAgentsTrainerTrainingSettings& TrainingSettings = FLearningAgentsTrainerTrainingSettings(),
-		const FLearningAgentsTrainerGameSettings& GameSettings = FLearningAgentsTrainerGameSettings(),
+		const FLearningAgentsTrainerTrainingSettings& TrainerTrainingSettings = FLearningAgentsTrainerTrainingSettings(),
+		const FLearningAgentsTrainerGameSettings& TrainerGameSettings = FLearningAgentsTrainerGameSettings(),
+		const FLearningAgentsTrainerPathSettings& TrainerPathSettings = FLearningAgentsTrainerPathSettings(),
 		const FLearningAgentsCriticSettings& CriticSettings = FLearningAgentsCriticSettings(),
 		const bool bReinitializePolicyNetwork = true,
 		const bool bReinitializeCriticNetwork = true);
@@ -356,16 +406,18 @@ public:
 	* Convenience function that runs a basic training loop. If training has not been started, it will start it, and 
 	* then call RunInference. On each following call to this function, it will call EvaluateRewards, 
 	* EvaluateCompletions, and ProcessExperience, followed by RunInference.
-	* @param TrainingSettings The settings for this training run.
-	* @param GameSettings The settings that will affect the game's simulation.
+	* @param TrainerTrainingSettings The settings for this training run.
+	* @param TrainerGameSettings The settings that will affect the game's simulation.
+	* @param TrainerPathSettings The path settings used by the trainer.
 	* @param CriticSettings The settings for the critic (if we are using one).
 	* @param bReinitializePolicyNetwork If true, reinitialize the policy. Set this to false if your policy is pre-trained, e.g. with imitation learning.
 	* @param bReinitializeCriticNetwork If true, reinitialize the critic. Set this to false if your critic is pre-trained.
 	*/
 	UFUNCTION(BlueprintCallable, Category = "LearningAgents")
 	void RunTraining(
-		const FLearningAgentsTrainerTrainingSettings& TrainingSettings = FLearningAgentsTrainerTrainingSettings(),
-		const FLearningAgentsTrainerGameSettings& GameSettings = FLearningAgentsTrainerGameSettings(),
+		const FLearningAgentsTrainerTrainingSettings& TrainerTrainingSettings = FLearningAgentsTrainerTrainingSettings(),
+		const FLearningAgentsTrainerGameSettings& TrainerGameSettings = FLearningAgentsTrainerGameSettings(),
+		const FLearningAgentsTrainerPathSettings& TrainerPathSettings = FLearningAgentsTrainerPathSettings(),
 		const FLearningAgentsCriticSettings& CriticSettings = FLearningAgentsCriticSettings(),
 		const bool bReinitializePolicyNetwork = true,
 		const bool bReinitializeCriticNetwork = true);
diff --git a/Engine/Plugins/Experimental/LearningAgents/Source/LearningTraining/Private/LearningImitationTrainer.cpp b/Engine/Plugins/Experimental/LearningAgents/Source/LearningTraining/Private/LearningImitationTrainer.cpp
--- a/Engine/Plugins/Experimental/LearningAgents/Source/LearningTraining/Private/LearningImitationTrainer.cpp
+++ b/Engine/Plugins/Experimental/LearningAgents/Source/LearningTraining/Private/LearningImitationTrainer.cpp
@@ -13,6 +13,7 @@
 #include "Misc/MonitoredProcess.h"
 #include "Misc/Guid.h"
 #include "Misc/FileHelper.h"
+#include "Misc/Paths.h"
 
 #include "Dom/JsonObject.h"
 #include "Serialization/JsonSerializer.h"
@@ -36,7 +37,7 @@ int32 ULearningSocketImitationTrainerServerCommandlet::Main(const FString& Comma
 
 	UCommandlet::ParseCommandLine(*Commandline, Tokens, Switches, Params);
 
-	const FString* PythonExecutiblePathParam = Params.Find(TEXT("PythonExecutiblePath"));
+	const FString* PythonExecutablePathParam = Params.Find(TEXT("PythonExecutablePath"));
 	const FString* SitePackagesPathParam = Params.Find(TEXT("SitePackagesPath"));
 	const FString* PythonContentPathParam = Params.Find(TEXT("PythonContentPath"));
 	const FString* IntermediatePathParam = Params.Find(TEXT("IntermediatePath"));
@@ -45,13 +46,13 @@ int32 ULearningSocketImitationTrainerServerCommandlet::Main(const FString& Comma
 	const FString* LogSettingsParam = Params.Find(TEXT("LogSettings"));
 
 #if WITH_EDITOR
-	const FString PythonExecutiblePath = PythonExecutiblePathParam ? *PythonExecutiblePathParam : UE::Learning::Trainer::DefaultEditorPythonExecutablePath();
-	const FString SitePackagesPath = SitePackagesPathParam ? *SitePackagesPathParam : UE::Learning::Trainer::DefaultEditorSitePackagesPath();
-	const FString PythonContentPath = PythonContentPathParam ? *PythonContentPathParam : UE::Learning::Trainer::DefaultEditorPythonContentPath();
-	const FString IntermediatePath = IntermediatePathParam ? *IntermediatePathParam : UE::Learning::Trainer::DefaultEditorIntermediatePath();
+	const FString PythonExecutablePath = PythonExecutablePathParam ? *PythonExecutablePathParam : UE::Learning::Trainer::GetPythonExecutablePath(FPaths::EngineDir());
+	const FString SitePackagesPath = SitePackagesPathParam ? *SitePackagesPathParam : UE::Learning::Trainer::GetSitePackagesPath(FPaths::EngineDir());
+	const FString PythonContentPath = PythonContentPathParam ? *PythonContentPathParam : UE::Learning::Trainer::GetPythonContentPath(FPaths::EngineDir());
+	const FString IntermediatePath = IntermediatePathParam ? *IntermediatePathParam : UE::Learning::Trainer::GetIntermediatePath(FPaths::EngineDir());
 #else
 	UE_LEARNING_NOT_IMPLEMENTED();
-	const FString PythonExecutiblePath = TEXT("");
+	const FString PythonExecutablePath = TEXT("");
 	const FString SitePackagesPath = TEXT("");
 	const FString PythonContentPath = TEXT("");
 	const FString IntermediatePath = TEXT("");
@@ -79,7 +80,7 @@ int32 ULearningSocketImitationTrainerServerCommandlet::Main(const FString& Comma
 	}
 
 	UE_LOG(LogLearning, Display, TEXT("---  Imitation Training Server Arguments ---"));
-	UE_LOG(LogLearning, Display, TEXT("PythonExecutiblePath: %s"), *PythonExecutiblePath);
+	UE_LOG(LogLearning, Display, TEXT("PythonExecutablePath: %s"), *PythonExecutablePath);
 	UE_LOG(LogLearning, Display, TEXT("SitePackagesPath: %s"), *SitePackagesPath);
 	UE_LOG(LogLearning, Display, TEXT("PythonContentPath: %s"), *PythonContentPath);
 	UE_LOG(LogLearning, Display, TEXT("IntermediatePath: %s"), *IntermediatePath);
@@ -88,7 +89,7 @@ int32 ULearningSocketImitationTrainerServerCommandlet::Main(const FString& Comma
 	UE_LOG(LogLearning, Display, TEXT("LogSettings: %s"), LogSettings == UE::Learning::ELogSetting::Normal ? TEXT("Normal") : TEXT("Silent"));
 
 	UE::Learning::FSocketImitationTrainerServerProcess ServerProcess(
-		PythonExecutiblePath,
+		PythonExecutablePath,
 		SitePackagesPath,
 		PythonContentPath,
 		IntermediatePath,
@@ -109,7 +110,7 @@ namespace UE::Learning
 {
 	FSharedMemoryImitationTrainer::FSharedMemoryImitationTrainer(
 		const FString& TaskName,
-		const FString& PythonExecutiblePath,
+		const FString& PythonExecutablePath,
 		const FString& SitePackagesPath,
 		const FString& PythonContentPath,
 		const FString& IntermediatePath,
@@ -122,6 +123,10 @@ namespace UE::Learning
 		const ESubprocessFlags TrainingProcessFlags,
 		const ELogSetting LogSettings)
 	{
+		UE_LEARNING_CHECK(FPaths::FileExists(PythonExecutablePath));
+		UE_LEARNING_CHECK(FPaths::DirectoryExists(PythonContentPath));
+		UE_LEARNING_CHECK(FPaths::DirectoryExists(SitePackagesPath));
+
 		const int32 TotalPolicyByteNum = FNeuralNetwork::GetSerializationByteNum(
 			ObservationDimNum,
 			2 * ActionDimNum,
@@ -143,6 +148,7 @@ namespace UE::Learning
 
 		// Write Config
 
+		IFileManager& FileManager = IFileManager::Get();
 		const FString TimeStamp = FDateTime::Now().ToFormattedString(TEXT("%Y-%m-%d_%H-%M-%S"));
 		const FString TrainerMethod = TEXT("Imitation");
 		const FString TrainerType = TEXT("SharedMemory");
@@ -154,8 +160,8 @@ namespace UE::Learning
 		ConfigObject->SetStringField(TEXT("TrainerType"), TrainerType);
 		ConfigObject->SetStringField(TEXT("TimeStamp"), *TimeStamp);
 
-		ConfigObject->SetStringField(TEXT("SitePackagesPath"), *SitePackagesPath);
-		ConfigObject->SetStringField(TEXT("IntermediatePath"), *IntermediatePath);
+		ConfigObject->SetStringField(TEXT("SitePackagesPath"), *FileManager.ConvertToAbsolutePathForExternalAppForRead(*SitePackagesPath));
+		ConfigObject->SetStringField(TEXT("IntermediatePath"), *FileManager.ConvertToAbsolutePathForExternalAppForRead(*IntermediatePath));
 
 		ConfigObject->SetStringField(TEXT("PolicyGuid"), *Policy.Guid.ToString(EGuidFormats::DigitsWithHyphensInBraces));
 		ConfigObject->SetStringField(TEXT("ControlsGuid"), *Controls.Guid.ToString(EGuidFormats::DigitsWithHyphensInBraces));
@@ -194,10 +200,12 @@ namespace UE::Learning
 
 		// Start Python Sub-process
 
-		const FString CommandLineArguments = FString::Printf(TEXT("\"%s\" SharedMemory \"%s\""), *(PythonContentPath / TEXT("train_imitation.py")), *ConfigPath);
+		const FString CommandLineArguments = FString::Printf(TEXT("\"%s\" SharedMemory \"%s\""), 
+			*FileManager.ConvertToAbsolutePathForExternalAppForRead(*(PythonContentPath / TEXT("train_imitation.py"))),
+			*FileManager.ConvertToAbsolutePathForExternalAppForRead(*ConfigPath));
 
 		TrainingProcess = MakeShared<FMonitoredProcess>(
-			PythonExecutiblePath, 
+			FileManager.ConvertToAbsolutePathForExternalAppForRead(*PythonExecutablePath),
 			CommandLineArguments,
 			!(TrainingProcessFlags & ESubprocessFlags::ShowWindow),
 			!(TrainingProcessFlags & ESubprocessFlags::NoRedirectOutput));
@@ -342,7 +350,7 @@ namespace UE::Learning
 
 
 	FSocketImitationTrainerServerProcess::FSocketImitationTrainerServerProcess(
-		const FString& PythonExecutiblePath,
+		const FString& PythonExecutablePath,
 		const FString& SitePackagesPath,
 		const FString& PythonContentPath,
 		const FString& IntermediatePath,
@@ -351,11 +359,21 @@ namespace UE::Learning
 		const ESubprocessFlags TrainingProcessFlags,
 		const ELogSetting LogSettings)
 	{
+		UE_LEARNING_CHECK(FPaths::FileExists(PythonExecutablePath));
+		UE_LEARNING_CHECK(FPaths::DirectoryExists(PythonContentPath));
+		UE_LEARNING_CHECK(FPaths::DirectoryExists(SitePackagesPath));
+
+		IFileManager& FileManager = IFileManager::Get();
 		const FString CommandLineArguments = FString::Printf(TEXT("\"%s\" Socket \"%s:%i\" \"%s\" \"%s\" %i"),
-			*(PythonContentPath / TEXT("train_imitation.py")), IpAddress, Port, *SitePackagesPath, *IntermediatePath, LogSettings == ELogSetting::Normal ? 1 : 0);
+			*FileManager.ConvertToAbsolutePathForExternalAppForRead(*(PythonContentPath / TEXT("train_imitation.py"))), 
+			IpAddress, 
+			Port, 
+			*FileManager.ConvertToAbsolutePathForExternalAppForRead(*SitePackagesPath),
+			*FileManager.ConvertToAbsolutePathForExternalAppForRead(*IntermediatePath),
+			LogSettings == ELogSetting::Normal ? 1 : 0);
 
 		TrainingProcess = MakeShared<FMonitoredProcess>(
-			PythonExecutiblePath, 
+			FileManager.ConvertToAbsolutePathForExternalAppForRead(*PythonExecutablePath),
 			CommandLineArguments,
 			!(TrainingProcessFlags & ESubprocessFlags::ShowWindow),
 			!(TrainingProcessFlags & ESubprocessFlags::NoRedirectOutput));
diff --git a/Engine/Plugins/Experimental/LearningAgents/Source/LearningTraining/Private/LearningPPOTrainer.cpp b/Engine/Plugins/Experimental/LearningAgents/Source/LearningTraining/Private/LearningPPOTrainer.cpp
--- a/Engine/Plugins/Experimental/LearningAgents/Source/LearningTraining/Private/LearningPPOTrainer.cpp
+++ b/Engine/Plugins/Experimental/LearningAgents/Source/LearningTraining/Private/LearningPPOTrainer.cpp
@@ -16,6 +16,7 @@
 #include "Misc/FileHelper.h"
 #include "Misc/CommandLine.h"
 #include "Misc/Parse.h"
+#include "Misc/Paths.h"
 
 #include "Dom/JsonObject.h"
 #include "Serialization/JsonSerializer.h"
@@ -48,10 +49,10 @@ int32 ULearningSocketPPOTrainerServerCommandlet::Main(const FString& Commandline
 	const FString* LogSettingsParam = Params.Find(TEXT("LogSettings"));
 
 #if WITH_EDITOR
-	const FString PythonExecutiblePath = PythonExecutiblePathParam ? *PythonExecutiblePathParam : UE::Learning::Trainer::DefaultEditorPythonExecutablePath();
-	const FString SitePackagesPath = SitePackagesPathParam ? *SitePackagesPathParam : UE::Learning::Trainer::DefaultEditorSitePackagesPath();
-	const FString PythonContentPath = PythonContentPathParam ? *PythonContentPathParam : UE::Learning::Trainer::DefaultEditorPythonContentPath();
-	const FString IntermediatePath = IntermediatePathParam ? *IntermediatePathParam : UE::Learning::Trainer::DefaultEditorIntermediatePath();
+	const FString PythonExecutiblePath = PythonExecutiblePathParam ? *PythonExecutiblePathParam : UE::Learning::Trainer::GetPythonExecutablePath(FPaths::EngineDir());
+	const FString SitePackagesPath = SitePackagesPathParam ? *SitePackagesPathParam : UE::Learning::Trainer::GetSitePackagesPath(FPaths::EngineDir());
+	const FString PythonContentPath = PythonContentPathParam ? *PythonContentPathParam : UE::Learning::Trainer::GetPythonContentPath(FPaths::EngineDir());
+	const FString IntermediatePath = IntermediatePathParam ? *IntermediatePathParam : UE::Learning::Trainer::GetIntermediatePath(FPaths::EngineDir());
 #else
 	UE_LEARNING_NOT_IMPLEMENTED();
 	const FString PythonExecutiblePath = TEXT("");
@@ -125,6 +126,10 @@ namespace UE::Learning
 		const uint16 ProcessNum,
 		const ESubprocessFlags MultiProcessFlags)
 	{
+		UE_LEARNING_CHECK(FPaths::FileExists(PythonExecutablePath));
+		UE_LEARNING_CHECK(FPaths::DirectoryExists(PythonContentPath));
+		UE_LEARNING_CHECK(FPaths::DirectoryExists(SitePackagesPath));
+
 		const int32 ObservationVectorDimensionNum = ReplayBuffer.GetObservations().Num<1>();
 		const int32 ActionVectorDimensionNum = ReplayBuffer.GetActions().Num<1>();
 
@@ -216,6 +221,7 @@ namespace UE::Learning
 
 			// Write Config
 
+			IFileManager& FileManager = IFileManager::Get();
 			const FString TimeStamp = FDateTime::Now().ToFormattedString(TEXT("%Y-%m-%d_%H-%M-%S"));
 			const FString TrainerMethod = TEXT("PPO");
 			const FString TrainerType = TEXT("SharedMemory");
@@ -227,8 +233,8 @@ namespace UE::Learning
 			ConfigObject->SetStringField(TEXT("TrainerType"), TrainerType);
 			ConfigObject->SetStringField(TEXT("TimeStamp"), *TimeStamp);
 
-			ConfigObject->SetStringField(TEXT("SitePackagesPath"), *SitePackagesPath);
-			ConfigObject->SetStringField(TEXT("IntermediatePath"), *IntermediatePath);
+			ConfigObject->SetStringField(TEXT("SitePackagesPath"), *FileManager.ConvertToAbsolutePathForExternalAppForRead(*SitePackagesPath));
+			ConfigObject->SetStringField(TEXT("IntermediatePath"), *FileManager.ConvertToAbsolutePathForExternalAppForRead(*IntermediatePath));
 
 			ConfigObject->SetStringField(TEXT("PolicyGuid"), *Policy.Guid.ToString(EGuidFormats::DigitsWithHyphensInBraces));
 			if (bRequiresCritic) { ConfigObject->SetStringField(TEXT("CriticGuid"), *Critic.Guid.ToString(EGuidFormats::DigitsWithHyphensInBraces)); }
@@ -294,10 +300,12 @@ namespace UE::Learning
 
 			// Start Python Training Sub-process
 
-			const FString CommandLineArguments = FString::Printf(TEXT("\"%s\" SharedMemory \"%s\""), *(PythonContentPath / TEXT("train_ppo.py")), *ConfigPath);
+			const FString CommandLineArguments = FString::Printf(TEXT("\"%s\" SharedMemory \"%s\""), 
+				*FileManager.ConvertToAbsolutePathForExternalAppForRead(*(PythonContentPath / TEXT("train_ppo.py"))), 
+				*FileManager.ConvertToAbsolutePathForExternalAppForRead(*ConfigPath));
 
 			TrainingProcess = MakeShared<FMonitoredProcess>(
-				PythonExecutablePath, 
+				FileManager.ConvertToAbsolutePathForExternalAppForRead(*PythonExecutablePath),
 				CommandLineArguments, 
 				!(TrainingProcessFlags & ESubprocessFlags::ShowWindow),
 				!(TrainingProcessFlags & ESubprocessFlags::NoRedirectOutput));
@@ -534,7 +542,7 @@ namespace UE::Learning
 	}
 
 	FSocketPPOTrainerServerProcess::FSocketPPOTrainerServerProcess(
-		const FString& PythonExecutiblePath,
+		const FString& PythonExecutablePath,
 		const FString& SitePackagesPath,
 		const FString& PythonContentPath,
 		const FString& IntermediatePath,
@@ -543,11 +551,21 @@ namespace UE::Learning
 		const ESubprocessFlags TrainingProcessFlags,
 		const ELogSetting LogSettings)
 	{
+		UE_LEARNING_CHECK(FPaths::FileExists(PythonExecutablePath));
+		UE_LEARNING_CHECK(FPaths::DirectoryExists(PythonContentPath));
+		UE_LEARNING_CHECK(FPaths::DirectoryExists(SitePackagesPath));
+
+		IFileManager& FileManager = IFileManager::Get();
 		const FString CommandLineArguments = FString::Printf(TEXT("\"%s\" Socket \"%s:%i\" \"%s\" \"%s\" %i"), 
-			*(PythonContentPath / TEXT("train_ppo.py")), IpAddress, Port, *SitePackagesPath, *IntermediatePath, LogSettings == ELogSetting::Normal ? 1 : 0);
+			*FileManager.ConvertToAbsolutePathForExternalAppForRead(*(PythonContentPath / TEXT("train_ppo.py"))), 
+			IpAddress, 
+			Port, 
+			*FileManager.ConvertToAbsolutePathForExternalAppForRead(*SitePackagesPath),
+			*FileManager.ConvertToAbsolutePathForExternalAppForRead(*IntermediatePath), 
+			LogSettings == ELogSetting::Normal ? 1 : 0);
 
 		TrainingProcess = MakeShared<FMonitoredProcess>(
-			PythonExecutiblePath, 
+			FileManager.ConvertToAbsolutePathForExternalAppForRead(*PythonExecutablePath),
 			CommandLineArguments,
 			!(TrainingProcessFlags & ESubprocessFlags::ShowWindow),
 			!(TrainingProcessFlags & ESubprocessFlags::NoRedirectOutput));
diff --git a/Engine/Plugins/Experimental/LearningAgents/Source/LearningTraining/Private/LearningSharedMemoryTraining.cpp b/Engine/Plugins/Experimental/LearningAgents/Source/LearningTraining/Private/LearningSharedMemoryTraining.cpp
--- a/Engine/Plugins/Experimental/LearningAgents/Source/LearningTraining/Private/LearningSharedMemoryTraining.cpp
+++ b/Engine/Plugins/Experimental/LearningAgents/Source/LearningTraining/Private/LearningSharedMemoryTraining.cpp
@@ -81,7 +81,8 @@ namespace UE::Learning::SharedMemoryTraining
 		// Read the policy
 		{
 			FScopeNullableWriteLock ScopeLock(NetworkLock);
-			OutNetwork.DeserializeFromBytes(Policy);
+			int32 Offset = 0;
+			OutNetwork.DeserializeFromBytes(Offset, Policy);
 		}
 
 		// Confirm we have read the policy
@@ -121,7 +122,8 @@ namespace UE::Learning::SharedMemoryTraining
 		// Read the critic
 		{
 			FScopeNullableWriteLock ScopeLock(NetworkLock);
-			OutNetwork.DeserializeFromBytes(Critic);
+			int32 Offset = 0;
+			OutNetwork.DeserializeFromBytes(Offset, Critic);
 		}
 
 		// Confirm we have read the critic
@@ -161,7 +163,8 @@ namespace UE::Learning::SharedMemoryTraining
 		// Write the policy
 		{
 			FScopeNullableReadLock ScopeLock(NetworkLock);
-			Network.SerializeToBytes(Policy);
+			int32 Offset = 0;
+			Network.SerializeToBytes(Offset, Policy);
 		}
 
 		// Confirm we have written the policy
@@ -201,7 +204,8 @@ namespace UE::Learning::SharedMemoryTraining
 		// Write the critic
 		{
 			FScopeNullableReadLock ScopeLock(NetworkLock);
-			Network.SerializeToBytes(Critic);
+			int32 Offset = 0;
+			Network.SerializeToBytes(Offset, Critic);
 		}
 
 		// Confirm we have written the critic
diff --git a/Engine/Plugins/Experimental/LearningAgents/Source/LearningTraining/Private/LearningSocketTraining.cpp b/Engine/Plugins/Experimental/LearningAgents/Source/LearningTraining/Private/LearningSocketTraining.cpp
--- a/Engine/Plugins/Experimental/LearningAgents/Source/LearningTraining/Private/LearningSocketTraining.cpp
+++ b/Engine/Plugins/Experimental/LearningAgents/Source/LearningTraining/Private/LearningSocketTraining.cpp
@@ -119,7 +119,8 @@ namespace UE::Learning::SocketTraining
 
 		{
 			FScopeNullableWriteLock ScopeLock(NetworkLock);
-			OutNetwork.DeserializeFromBytes(OutNetworkBuffer);
+			int32 Offset = 0;
+			OutNetwork.DeserializeFromBytes(Offset, OutNetworkBuffer);
 		}
 
 		return ETrainerResponse::Success;
@@ -152,7 +153,8 @@ namespace UE::Learning::SocketTraining
 
 		{
 			FScopeNullableWriteLock ScopeLock(NetworkLock);
-			OutNetwork.DeserializeFromBytes(OutNetworkBuffer);
+			int32 Offset = 0;
+			OutNetwork.DeserializeFromBytes(Offset, OutNetworkBuffer);
 		}
 
 		return ETrainerResponse::Success;
@@ -235,7 +237,8 @@ namespace UE::Learning::SocketTraining
 
 		{
 			FScopeNullableReadLock ScopeLock(NetworkLock);
-			Network.SerializeToBytes(NetworkBuffer);
+			int32 Offset = 0;
+			Network.SerializeToBytes(Offset, NetworkBuffer);
 		}
 
 		const uint8 Signal = (uint8)ESignal::SendPolicy;
@@ -263,7 +266,8 @@ namespace UE::Learning::SocketTraining
 
 		{
 			FScopeNullableReadLock ScopeLock(NetworkLock);
-			Network.SerializeToBytes(NetworkBuffer);
+			int32 Offset = 0;
+			Network.SerializeToBytes(Offset, NetworkBuffer);
 		}
 
 		const uint8 Signal = (uint8)ESignal::SendCritic;
diff --git a/Engine/Plugins/Experimental/LearningAgents/Source/LearningTraining/Private/LearningTrainer.cpp b/Engine/Plugins/Experimental/LearningAgents/Source/LearningTraining/Private/LearningTrainer.cpp
--- a/Engine/Plugins/Experimental/LearningAgents/Source/LearningTraining/Private/LearningTrainer.cpp
+++ b/Engine/Plugins/Experimental/LearningAgents/Source/LearningTraining/Private/LearningTrainer.cpp
@@ -2,9 +2,7 @@
 
 #include "LearningTrainer.h"
 
-#include "Misc/Paths.h"
-#include "Interfaces/IPluginManager.h"
-#include "HAL/FileManager.h"
+#include "HAL/Platform.h"
 
 namespace UE::Learning::Trainer
 {
@@ -43,34 +41,26 @@ namespace UE::Learning::Trainer
 		return FMath::Pow(0.5f, 1.0f / FMath::Max(HalfLifeSteps, 1));
 	}
 
-#if WITH_EDITOR
-
-	FString DefaultEditorPythonExecutablePath()
+	FString GetPythonExecutablePath(const FString& EngineDir)
 	{
 		UE_LEARNING_CHECKF(PLATFORM_WINDOWS || PLATFORM_MAC || PLATFORM_LINUX, TEXT("Python only supported on Windows, Mac, and Linux."));
 
-		return IFileManager::Get().ConvertToAbsolutePathForExternalAppForRead(
-			*(FPaths::EngineDir() / TEXT("Binaries/ThirdParty/Python3") / FPlatformMisc::GetUBTPlatform() / (PLATFORM_WINDOWS ? TEXT("python.exe") : TEXT("bin/python"))));
+		return EngineDir / TEXT("Binaries/ThirdParty/Python3") / FPlatformMisc::GetUBTPlatform() / (PLATFORM_WINDOWS ? TEXT("python.exe") : TEXT("bin/python"));
 	}
 
-	FString DefaultEditorSitePackagesPath()
+	FString GetSitePackagesPath(const FString& EngineDir)
 	{
-		return IFileManager::Get().ConvertToAbsolutePathForExternalAppForRead(
-			*(IPluginManager::Get().FindPlugin(TEXT("PythonFoundationPackages"))->GetBaseDir() / TEXT("Content/Python/Lib") / FPlatformMisc::GetUBTPlatform() / TEXT("site-packages")));
+		return EngineDir / TEXT("Plugins/Experimental/PythonFoundationPackages/Content/Python/Lib") / FPlatformMisc::GetUBTPlatform() / TEXT("site-packages");
 	}
 
-	FString DefaultEditorPythonContentPath()
+	FString GetPythonContentPath(const FString& EngineDir)
 	{
-		return  IFileManager::Get().ConvertToAbsolutePathForExternalAppForRead(
-			*(IPluginManager::Get().FindPlugin(TEXT("LearningAgents"))->GetBaseDir() / TEXT("Content/Python")));
+		return EngineDir / TEXT("Plugins/Experimental/LearningAgents/Content/Python/");
 	}
 
-	FString DefaultEditorIntermediatePath()
+	FString GetIntermediatePath(const FString& IntermediateDir)
 	{
-		return IFileManager::Get().ConvertToAbsolutePathForExternalAppForRead(
-			*(IPluginManager::Get().FindPlugin(TEXT("LearningAgents"))->GetBaseDir() / TEXT("Intermediate")));
+		return IntermediateDir / TEXT("LearningAgents");
 	}
 
-#endif
-
 }
\ No newline at end of file
diff --git a/Engine/Plugins/Experimental/LearningAgents/Source/LearningTraining/Public/LearningTrainer.h b/Engine/Plugins/Experimental/LearningAgents/Source/LearningTraining/Public/LearningTrainer.h
--- a/Engine/Plugins/Experimental/LearningAgents/Source/LearningTraining/Public/LearningTrainer.h
+++ b/Engine/Plugins/Experimental/LearningAgents/Source/LearningTraining/Public/LearningTrainer.h
@@ -101,12 +101,25 @@ namespace UE::Learning
 		*/
 		LEARNINGTRAINING_API float DiscountFactorFromHalfLifeSteps(const int32 HalfLifeSteps);
 
-#if WITH_EDITOR
-		LEARNINGTRAINING_API FString DefaultEditorPythonExecutablePath();
-		LEARNINGTRAINING_API FString DefaultEditorSitePackagesPath();
-		LEARNINGTRAINING_API FString DefaultEditorPythonContentPath();
-		LEARNINGTRAINING_API FString DefaultEditorIntermediatePath();
-#endif
+		/**
+		* Gets the python executable path from the engine directory.
+		*/
+		LEARNINGTRAINING_API FString GetPythonExecutablePath(const FString& EngineDir);
+
+		/**
+		* Gets the PythonFoundationPackages site-packages path from the engine directory.
+		*/
+		LEARNINGTRAINING_API FString GetSitePackagesPath(const FString& EngineDir);
+
+		/**
+		* Gets the LearningAgents Content path from the engine directory.
+		*/
+		LEARNINGTRAINING_API FString GetPythonContentPath(const FString& EngineDir);
+
+		/**
+		* Gets the LearningAgents Intermediate path from the intermediate directory.
+		*/
+		LEARNINGTRAINING_API FString GetIntermediatePath(const FString& IntermediateDir);
 	}
 
 }
\ No newline at end of file
