LearningAgents: Fixed issue with debug plane rendering. Added additional Trainer settings. A couple of other minor fixes and documentation updates.

diff --git a/Engine/Plugins/Experimental/LearningAgents/Source/LearningAgents/Private/LearningAgentsDebug.cpp b/Engine/Plugins/Experimental/LearningAgents/Source/LearningAgents/Private/LearningAgentsDebug.cpp
--- a/Engine/Plugins/Experimental/LearningAgents/Source/LearningAgents/Private/LearningAgentsDebug.cpp
+++ b/Engine/Plugins/Experimental/LearningAgents/Source/LearningAgents/Private/LearningAgentsDebug.cpp
@@ -4,6 +4,16 @@
 
 namespace UE::Learning::Agents::Debug
 {
+	FMatrix PlaneMatrix(const FQuat Rotation, const FVector Position, const FVector Axis0, const FVector Axis1)
+	{
+		// Here we reconstruct the transformation matrix from the plane axes assuming the first axis is forward, and second is to the right.
+		const FVector ForwardAxis = Rotation.RotateVector(Axis0.GetSafeNormal(UE_SMALL_NUMBER, FVector::ForwardVector));
+		const FVector InitialRightAxis = Rotation.RotateVector(Axis1.GetSafeNormal(UE_SMALL_NUMBER, FVector::RightVector));
+		const FVector UpAxis = ForwardAxis.Cross(InitialRightAxis).GetSafeNormal(UE_SMALL_NUMBER, FVector::UpVector);
+		const FVector RightAxis = UpAxis.Cross(ForwardAxis).GetSafeNormal(UE_SMALL_NUMBER, FVector::RightVector);
+		return FMatrix(ForwardAxis, RightAxis, UpAxis, Position);
+	}
+
 	FVector GridOffsetForIndex(const int32 Idx, const int32 Num, const float Spacing, const float Height)
 	{
 		const int32 RowNum = FMath::Max(FMath::CeilToInt(FMath::Sqrt((float)Num)), 1);
diff --git a/Engine/Plugins/Experimental/LearningAgents/Source/LearningAgents/Private/LearningAgentsObservations.cpp b/Engine/Plugins/Experimental/LearningAgents/Source/LearningAgents/Private/LearningAgentsObservations.cpp
--- a/Engine/Plugins/Experimental/LearningAgents/Source/LearningAgents/Private/LearningAgentsObservations.cpp
+++ b/Engine/Plugins/Experimental/LearningAgents/Source/LearningAgents/Private/LearningAgentsObservations.cpp
@@ -1011,7 +1011,14 @@ void UDirectionArrayObservation::VisualLog(const UE::Learning::FIndexSet Instanc
 
 UPlanarDirectionObservation* UPlanarDirectionObservation::AddPlanarDirectionObservation(ULearningAgentsInteractor* InInteractor, const FName Name, const float Scale, const FVector Axis0, const FVector Axis1)
 {
-	return UE::Learning::Agents::Observations::Private::AddObservation<UPlanarDirectionObservation, UE::Learning::FPlanarDirectionFeature>(InInteractor, Name, TEXT("AddPlanarDirectionObservation"), 1, Scale, Axis0, Axis1);
+	return UE::Learning::Agents::Observations::Private::AddObservation<UPlanarDirectionObservation, UE::Learning::FPlanarDirectionFeature>(
+		InInteractor, 
+		Name, 
+		TEXT("AddPlanarDirectionObservation"), 
+		1, 
+		Scale, 
+		Axis0.GetSafeNormal(UE_SMALL_NUMBER, FVector::ForwardVector), 
+		Axis1.GetSafeNormal(UE_SMALL_NUMBER, FVector::RightVector));
 }
 
 void UPlanarDirectionObservation::SetPlanarDirectionObservation(const int32 AgentId, const FVector Direction, const FRotator RelativeRotation)
@@ -1085,7 +1092,14 @@ UPlanarDirectionArrayObservation* UPlanarDirectionArrayObservation::AddPlanarDir
 		return nullptr;
 	}
 
-	return UE::Learning::Agents::Observations::Private::AddObservation<UPlanarDirectionArrayObservation, UE::Learning::FPlanarDirectionFeature>(InInteractor, Name, TEXT("AddPlanarDirectionArrayObservation"), DirectionNum, Scale, Axis0, Axis1);
+	return UE::Learning::Agents::Observations::Private::AddObservation<UPlanarDirectionArrayObservation, UE::Learning::FPlanarDirectionFeature>(
+		InInteractor, 
+		Name, 
+		TEXT("AddPlanarDirectionArrayObservation"), 
+		DirectionNum, 
+		Scale, 
+		Axis0.GetSafeNormal(UE_SMALL_NUMBER, FVector::ForwardVector),
+		Axis1.GetSafeNormal(UE_SMALL_NUMBER, FVector::RightVector));
 }
 
 void UPlanarDirectionArrayObservation::SetPlanarDirectionArrayObservation(const int32 AgentId, const TArray<FVector>& Directions, const FRotator RelativeRotation)
@@ -1454,7 +1468,14 @@ void UScalarPositionArrayObservation::VisualLog(const UE::Learning::FIndexSet In
 
 UPlanarPositionObservation* UPlanarPositionObservation::AddPlanarPositionObservation(ULearningAgentsInteractor* InInteractor, const FName Name, const float Scale, const FVector Axis0, const FVector Axis1)
 {
-	return UE::Learning::Agents::Observations::Private::AddObservation<UPlanarPositionObservation, UE::Learning::FPlanarPositionFeature>(InInteractor, Name, TEXT("AddPlanarPositionObservation"), 1, Scale, Axis0, Axis1);
+	return UE::Learning::Agents::Observations::Private::AddObservation<UPlanarPositionObservation, UE::Learning::FPlanarPositionFeature>(
+		InInteractor, 
+		Name, 
+		TEXT("AddPlanarPositionObservation"), 
+		1, 
+		Scale, 
+		Axis0.GetSafeNormal(UE_SMALL_NUMBER, FVector::ForwardVector),
+		Axis1.GetSafeNormal(UE_SMALL_NUMBER, FVector::RightVector));
 }
 
 void UPlanarPositionObservation::SetPlanarPositionObservation(const int32 AgentId, const FVector Position, const FVector RelativePosition, const FRotator RelativeRotation)
@@ -1529,7 +1550,14 @@ UPlanarPositionArrayObservation* UPlanarPositionArrayObservation::AddPlanarPosit
 		return nullptr;
 	}
 
-	return UE::Learning::Agents::Observations::Private::AddObservation<UPlanarPositionArrayObservation, UE::Learning::FPlanarPositionFeature>(InInteractor, Name, TEXT("AddPlanarPositionArrayObservation"), PositionNum, Scale, Axis0, Axis1);
+	return UE::Learning::Agents::Observations::Private::AddObservation<UPlanarPositionArrayObservation, UE::Learning::FPlanarPositionFeature>(
+		InInteractor, 
+		Name, 
+		TEXT("AddPlanarPositionArrayObservation"), 
+		PositionNum, 
+		Scale, 
+		Axis0.GetSafeNormal(UE_SMALL_NUMBER, FVector::ForwardVector),
+		Axis1.GetSafeNormal(UE_SMALL_NUMBER, FVector::RightVector));
 }
 
 void UPlanarPositionArrayObservation::SetPlanarPositionArrayObservation(const int32 AgentId, const TArray<FVector>& Positions, const FVector RelativePosition, const FRotator RelativeRotation)
@@ -1888,7 +1916,14 @@ void UScalarVelocityArrayObservation::VisualLog(const UE::Learning::FIndexSet In
 
 UPlanarVelocityObservation* UPlanarVelocityObservation::AddPlanarVelocityObservation(ULearningAgentsInteractor* InInteractor, const FName Name, const float Scale, const FVector Axis0, const FVector Axis1)
 {
-	return UE::Learning::Agents::Observations::Private::AddObservation<UPlanarVelocityObservation, UE::Learning::FPlanarVelocityFeature>(InInteractor, Name, TEXT("AddPlanarVelocityObservation"), 1, Scale, Axis0, Axis1);
+	return UE::Learning::Agents::Observations::Private::AddObservation<UPlanarVelocityObservation, UE::Learning::FPlanarVelocityFeature>(
+		InInteractor, 
+		Name, 
+		TEXT("AddPlanarVelocityObservation"), 
+		1, 
+		Scale, 
+		Axis0.GetSafeNormal(UE_SMALL_NUMBER, FVector::ForwardVector),
+		Axis1.GetSafeNormal(UE_SMALL_NUMBER, FVector::RightVector));
 }
 
 void UPlanarVelocityObservation::SetPlanarVelocityObservation(const int32 AgentId, const FVector Velocity, const FRotator RelativeRotation)
@@ -1963,7 +1998,14 @@ UPlanarVelocityArrayObservation* UPlanarVelocityArrayObservation::AddPlanarVeloc
 		return nullptr;
 	}
 
-	return UE::Learning::Agents::Observations::Private::AddObservation<UPlanarVelocityArrayObservation, UE::Learning::FPlanarVelocityFeature>(InInteractor, Name, TEXT("AddPlanarVelocityArrayObservation"), VelocityNum, Scale, Axis0, Axis1);
+	return UE::Learning::Agents::Observations::Private::AddObservation<UPlanarVelocityArrayObservation, UE::Learning::FPlanarVelocityFeature>(
+		InInteractor, 
+		Name, 
+		TEXT("AddPlanarVelocityArrayObservation"), 
+		VelocityNum, 
+		Scale, 
+		Axis0.GetSafeNormal(UE_SMALL_NUMBER, FVector::ForwardVector),
+		Axis1.GetSafeNormal(UE_SMALL_NUMBER, FVector::RightVector));
 }
 
 void UPlanarVelocityArrayObservation::SetPlanarVelocityArrayObservation(const int32 AgentId, const TArray<FVector>& Velocities, const FRotator RelativeRotation)
diff --git a/Engine/Plugins/Experimental/LearningAgents/Source/LearningAgents/Public/LearningAgentsActions.h b/Engine/Plugins/Experimental/LearningAgents/Source/LearningAgents/Public/LearningAgentsActions.h
--- a/Engine/Plugins/Experimental/LearningAgents/Source/LearningAgents/Public/LearningAgentsActions.h
+++ b/Engine/Plugins/Experimental/LearningAgents/Source/LearningAgents/Public/LearningAgentsActions.h
@@ -36,6 +36,7 @@ class LEARNINGAGENTS_API ULearningAgentsAction : public UObject
 
 public:
 
+	/** Reference to the Interactor this action is associated with. */
 	UPROPERTY(VisibleAnywhere, Transient, Category = "LearningAgents")
 	TObjectPtr<ULearningAgentsInteractor> Interactor;
 
diff --git a/Engine/Plugins/Experimental/LearningAgents/Source/LearningAgents/Public/LearningAgentsDebug.h b/Engine/Plugins/Experimental/LearningAgents/Source/LearningAgents/Public/LearningAgentsDebug.h
--- a/Engine/Plugins/Experimental/LearningAgents/Source/LearningAgents/Public/LearningAgentsDebug.h
+++ b/Engine/Plugins/Experimental/LearningAgents/Source/LearningAgents/Public/LearningAgentsDebug.h
@@ -31,7 +31,15 @@
 	UE_VLOG_OBOX(Owner, Category, Verbosity, Box, Matrix, Color, Format, ##__VA_ARGS__)
 
 #define UE_LEARNING_AGENTS_VLOG_PLANE(Owner, Category, Verbosity, Location, Rotation, Axis0, Axis1, Color, Format, ...) \
-	UE_LEARNING_AGENTS_VLOG_OBOX(Owner, Category, Verbosity, FBox(-25.0f * (Axis0 + Axis1), 25.0f * (Axis0 + Axis1)), FTransform(Rotation, Location, FVector::OneVector).ToMatrixNoScale(), Color, Format, ##__VA_ARGS__)
+	UE_LEARNING_AGENTS_VLOG_OBOX(Owner, Category, Verbosity, FBox(FVector(-25.0f, -25.0f, -0.1f), FVector(25.0f, 25.0f, 0.1f)), UE::Learning::Agents::Debug::PlaneMatrix(Rotation, Location, Axis0, Axis1), Color, Format, ##__VA_ARGS__)
+
+#define UE_LEARNING_AGENTS_VLOG_MATRIX(Owner, Category, Verbosity, Matrix, Color, Format, ...) \
+	{ \
+		UE_LEARNING_AGENTS_VLOG_SEGMENT(Owner, Category, Verbosity, Matrix.TransformPosition(FVector::ZeroVector), Matrix.TransformPosition(15.0f * FVector::ForwardVector), FColor::Red, TEXT("")); \
+		UE_LEARNING_AGENTS_VLOG_SEGMENT(Owner, Category, Verbosity, Matrix.TransformPosition(FVector::ZeroVector), Matrix.TransformPosition(15.0f * FVector::RightVector), FColor::Green, TEXT("")); \
+		UE_LEARNING_AGENTS_VLOG_SEGMENT(Owner, Category, Verbosity, Matrix.TransformPosition(FVector::ZeroVector), Matrix.TransformPosition(15.0f * FVector::UpVector), FColor::Blue, TEXT("")); \
+		UE_LEARNING_AGENTS_VLOG_STRING(Owner, Category, Verbosity, Matrix.TransformPosition(FVector::ZeroVector) + FVector(0.0f, 0.0f, 20.0f), Color, Format, ##__VA_ARGS__); \
+	}
 
 #define UE_LEARNING_AGENTS_VLOG_TRANSFORM(Owner, Category, Verbosity, Location, Rotation, Color, Format, ...) \
 	{ \
@@ -54,6 +62,8 @@
 
 namespace UE::Learning::Agents::Debug
 {
+	LEARNINGAGENTS_API FMatrix PlaneMatrix(const FQuat Rotation, const FVector Position, const FVector Axis0, const FVector Axis1);
+
 	LEARNINGAGENTS_API FVector GridOffsetForIndex(const int32 Idx, const int32 Num, const float Spacing = 20.0f, const float Height = 20.0f);
 
 	LEARNINGAGENTS_API FString FloatArrayToStatsString(const TLearningArrayView<1, const float> Array);
diff --git a/Engine/Plugins/Experimental/LearningAgents/Source/LearningAgents/Public/LearningAgentsHelpers.h b/Engine/Plugins/Experimental/LearningAgents/Source/LearningAgents/Public/LearningAgentsHelpers.h
--- a/Engine/Plugins/Experimental/LearningAgents/Source/LearningAgents/Public/LearningAgentsHelpers.h
+++ b/Engine/Plugins/Experimental/LearningAgents/Source/LearningAgents/Public/LearningAgentsHelpers.h
@@ -26,6 +26,7 @@ class LEARNINGAGENTS_API ULearningAgentsHelper : public UObject
 
 public:
 
+	/** Reference to the Manager Component this helper is associated with. */
 	UPROPERTY(VisibleAnywhere, Transient, Category = "LearningAgents")
 	TObjectPtr<ULearningAgentsManagerComponent> ManagerComponent;
 
@@ -45,27 +46,91 @@ class LEARNINGAGENTS_API USplineComponentHelper : public ULearningAgentsHelper
 
 public:
 
+	/**
+	* Adds a new spline component helper to the given manager component.
+	* @param InManagerComponent The manager component to add this helper to (such as an Interactor or Trainer).
+	* @param Name The name of this new helper. Used for debugging.
+	* @return The newly created helper.
+	*/
 	UFUNCTION(BlueprintCallable, Category = "LearningAgents")
 	static USplineComponentHelper* AddSplineComponentHelper(ULearningAgentsManagerComponent* InManagerComponent, const FName Name = NAME_None);
 
+	/**
+	* Gets the position on a spline closest to the provided position.
+	* @param AgentId The agent id to run the helper for.
+	* @param SplineComponent The spline component to use.
+	* @param Position The position to find the closest position to.
+	* @param CoordinateSpace The coordinate space to use for the spline.
+	* @return The closest position on the spline.
+	*/
 	UFUNCTION(BlueprintPure = false, Category = "LearningAgents", meta = (AgentId = "-1"))
 	FVector GetNearestPositionOnSpline(const int32 AgentId, const USplineComponent* SplineComponent, const FVector Position, const ESplineCoordinateSpace::Type CoordinateSpace = ESplineCoordinateSpace::World) const;
 
+	/**
+	* Gets the distance along a spline closest to the provided position.
+	* @param AgentId The agent id to run the helper for.
+	* @param SplineComponent The spline component to use.
+	* @param Position The position to find the closest position to.
+	* @param CoordinateSpace The coordinate space to use for the spline.
+	* @return The distance along the spline of the closest position.
+	*/
 	UFUNCTION(BlueprintPure = false, Category = "LearningAgents", meta = (AgentId = "-1"))
 	float GetDistanceAlongSplineAtPosition(const int32 AgentId, const USplineComponent* SplineComponent, const FVector Position, const ESplineCoordinateSpace::Type CoordinateSpace = ESplineCoordinateSpace::World) const;
 
+	/**
+	* Gets the position along a spline at the given distance.
+	* @param AgentId The agent id to run the helper for.
+	* @param SplineComponent The spline component to use.
+	* @param DistanceAlongSpline The distance along the spline to use.
+	* @param CoordinateSpace The coordinate space to use for the spline.
+	* @return The position at the given distance along the spline.
+	*/
 	UFUNCTION(BlueprintPure = false, Category = "LearningAgents", meta = (AgentId = "-1"))
 	FVector GetPositionAtDistanceAlongSpline(const int32 AgentId, const USplineComponent* SplineComponent, const float DistanceAlongSpline, const ESplineCoordinateSpace::Type CoordinateSpace = ESplineCoordinateSpace::World) const;
 
+	/**
+	* Gets the direction along a spline at the given distance.
+	* @param AgentId The agent id to run the helper for.
+	* @param SplineComponent The spline component to use.
+	* @param DistanceAlongSpline The distance along the spline to use.
+	* @param CoordinateSpace The coordinate space to use for the spline.
+	* @return The direction at the given distance along the spline.
+	*/
 	UFUNCTION(BlueprintPure = false, Category = "LearningAgents", meta = (AgentId = "-1"))
 	FVector GetDirectionAtDistanceAlongSpline(const int32 AgentId, const USplineComponent* SplineComponent, const float DistanceAlongSpline, const ESplineCoordinateSpace::Type CoordinateSpace = ESplineCoordinateSpace::World) const;
 
+	/**
+	* Gets the proportion along a spline in the range 0-1 for a given distance.
+	* @param AgentId The agent id to run the helper for.
+	* @param SplineComponent The spline component to use.
+	* @param DistanceAlongSpline The distance along the spline to use.
+	* @return The proportion along the spline in the range 0-1.
+	*/
 	UFUNCTION(BlueprintPure = false, Category = "LearningAgents", meta = (AgentId = "-1"))
 	float GetProportionAlongSpline(const int32 AgentId, const USplineComponent* SplineComponent, const float DistanceAlongSpline) const;
 
+	/**
+	* Gets the proportion along a spline encoded as an angle between -180 and 180 degrees for a given distance. 
+	* This should be used for looped splines.
+	* @param AgentId The agent id to run the helper for.
+	* @param SplineComponent The spline component to use.
+	* @param DistanceAlongSpline The distance along the spline to use.
+	* @return The proportion along the spline encoded as an angle between -180 and 180 degrees
+	*/
 	UFUNCTION(BlueprintPure = false, Category = "LearningAgents", meta = (AgentId = "-1"))
 	float GetProportionAlongSplineAsAngle(const int32 AgentId, const USplineComponent* SplineComponent, const float DistanceAlongSpline) const;
 
+	/**
+	* Gets an array of positions sampled along a spline between some starting and stopping distance. Deals properly
+	* for splines which are looped.
+	* @param OutPositions Output array of spline positions.
+	* @param AgentId The agent id to run the helper for.
+	* @param SplineComponent The spline component to use.
+	* @param PositionNum The number of positions to sample along the spline.
+	* @param StartDistanceAlongSpline The starting distance along the spline.
+	* @param StopDistanceAlongSpline The stopping distance along the spline.
+	* @param CoordinateSpace The coordinate space to use for the spline.
+	*/
 	UFUNCTION(BlueprintPure = false, Category = "LearningAgents", meta = (AgentId = "-1"))
 	void GetPositionsAlongSpline(
 		TArray<FVector>& OutPositions,
@@ -76,6 +141,16 @@ public:
 		const float StopDistanceAlongSpline,
 		const ESplineCoordinateSpace::Type CoordinateSpace = ESplineCoordinateSpace::World) const;
 
+	/**
+	* Gets the velocity a point is traveling along a spline.
+	* @param AgentId The agent id to run the helper for.
+	* @param SplineComponent The spline component to use.
+	* @param Position The position of the point.
+	* @param Velocity The velocity of the point.
+	* @param FiniteDifferenceDelta The finite difference delta used to estimate the velocity along the spline.
+	* @param CoordinateSpace The coordinate space to use for the spline.
+	* @return The scalar velocity the point is traveling along the spline.
+	*/
 	UFUNCTION(BlueprintPure = false, Category = "LearningAgents", meta = (AgentId = "-1"))
 	float GetVelocityAlongSpline(
 		const int32 AgentId,
@@ -96,12 +171,34 @@ class LEARNINGAGENTS_API UProjectionHelper : public ULearningAgentsHelper
 
 public:
 
+	/**
+	* Adds a new projection helper to the given manager component.
+	* @param InManagerComponent The manager component to add this helper to (such as an Interactor or Trainer).
+	* @param Name The name of this new helper. Used for debugging.
+	* @return The newly created helper.
+	*/
 	UFUNCTION(BlueprintCallable, Category = "LearningAgents")
 	static UProjectionHelper* AddProjectionHelper(ULearningAgentsManagerComponent* InManagerComponent, const FName Name = NAME_None);
 
+	/**
+	* Projects a transform onto the zero height ground plane resulting in translation only in XY and rotation only around the Z.
+	* @param AgentId The agent id to run the helper for.
+	* @param Transform The transform to project.
+	* @param LocalForwardVector The local forward direction.
+	* @return The transform projected onto the ground plane.
+	*/
 	UFUNCTION(BlueprintPure = false, Category = "LearningAgents", meta = (AgentId = "-1"))
 	FTransform ProjectTransformOntoGroundPlane(const int32 AgentId, const FTransform Transform, const FVector LocalForwardVector = FVector::ForwardVector) const;
 
+	/**
+	* Projects a position and rotation onto the zero height ground plane resulting in translation only in XY and rotation only around the Z.
+	* @param OutPosition The output projected position.
+	* @param OutRotation The output projected rotation.
+	* @param AgentId The agent id to run the helper for.
+	* @param InPosition The input position.
+	* @param InRotation The input rotation.
+	* @param LocalForwardVector The local forward direction.
+	*/
 	UFUNCTION(BlueprintPure = false, Category = "LearningAgents", meta = (AgentId = "-1"))
 	void ProjectPositionRotationOntoGroundPlane(
 		FVector& OutPosition, 
@@ -122,9 +219,22 @@ class LEARNINGAGENTS_API UMeshComponentHelper : public ULearningAgentsHelper
 
 public:
 
+	/**
+	* Adds a new mesh component helper to the given manager component.
+	* @param InManagerComponent The manager component to add this helper to (such as an Interactor or Trainer).
+	* @param Name The name of this new helper. Used for debugging.
+	* @return The newly created helper.
+	*/
 	UFUNCTION(BlueprintCallable, Category = "LearningAgents")
 	static UMeshComponentHelper* AddMeshComponentHelper(ULearningAgentsManagerComponent* InManagerComponent, const FName Name = NAME_None);
 
+	/**
+	* Gets the bone positions for a set of bones of the mesh component in world space.
+	* @param OutBonePositions The output array of bone positions.
+	* @param AgentId The agent id to run the helper for.
+	* @param MeshComponent The mesh component to use.
+	* @param BoneNames The names of the bones to use.
+	*/
 	UFUNCTION(BlueprintPure = false, Category = "LearningAgents", meta = (AgentId = "-1"))
 	void GetMeshBonePositions(TArray<FVector>& OutBonePositions, const int32 AgentId, const UMeshComponent* MeshComponent, const TArray<FName>& BoneNames) const;
 };
@@ -139,9 +249,29 @@ class LEARNINGAGENTS_API URayCastHelper : public ULearningAgentsHelper
 
 public:
 
+	/**
+	* Adds a new ray cast helper to the given manager component.
+	* @param InManagerComponent The manager component to add this helper to (such as an Interactor or Trainer).
+	* @param Name The name of this new helper. Used for debugging.
+	* @return The newly created helper.
+	*/
 	UFUNCTION(BlueprintCallable, Category = "LearningAgents")
 	static URayCastHelper* AddRayCastHelper(ULearningAgentsManagerComponent* InManagerComponent, const FName Name = NAME_None);
 
+	/**
+	* Samples a grid of heights from above.
+	* @param OutHeights The output array of heights for each ray in the grid.
+	* @param AgentId The agent id to run the helper for.
+	* @param Position The central position of the grid.
+	* @param Rotation The orientation of the grid.
+	* @param RowNum The number of rows in the grid.
+	* @param ColNum The number of columns in the grid.
+	* @param RowWidth The width of grid rows.
+	* @param ColWidth The width of grid columns.
+	* @param MaxHeight The height at which to start ray casts from.
+	* @param MinHeight The height at which to send ray casts to.
+	* @param CollisionChannel The collision channel to use for the ray-casts.
+	*/
 	UFUNCTION(BlueprintPure=false, Category = "LearningAgents", meta = (AgentId = "-1"))
 	void RayCastGridHeights(
 		TArray<float>& OutHeights,
diff --git a/Engine/Plugins/Experimental/LearningAgents/Source/LearningAgents/Public/LearningAgentsObservations.h b/Engine/Plugins/Experimental/LearningAgents/Source/LearningAgents/Public/LearningAgentsObservations.h
--- a/Engine/Plugins/Experimental/LearningAgents/Source/LearningAgents/Public/LearningAgentsObservations.h
+++ b/Engine/Plugins/Experimental/LearningAgents/Source/LearningAgents/Public/LearningAgentsObservations.h
@@ -44,6 +44,7 @@ class LEARNINGAGENTS_API ULearningAgentsObservation : public UObject
 
 public:
 
+	/** Reference to the Interactor this observation is associated with. */
 	UPROPERTY(VisibleAnywhere, Transient, Category = "LearningAgents")
 	TObjectPtr<ULearningAgentsInteractor> Interactor;
 
diff --git a/Engine/Plugins/Experimental/LearningAgents/Source/LearningAgentsTraining/Private/LearningAgentsCompletions.cpp b/Engine/Plugins/Experimental/LearningAgents/Source/LearningAgentsTraining/Private/LearningAgentsCompletions.cpp
--- a/Engine/Plugins/Experimental/LearningAgents/Source/LearningAgentsTraining/Private/LearningAgentsCompletions.cpp
+++ b/Engine/Plugins/Experimental/LearningAgents/Source/LearningAgentsTraining/Private/LearningAgentsCompletions.cpp
@@ -141,8 +141,8 @@ UPlanarPositionDifferenceCompletion* UPlanarPositionDifferenceCompletion::AddPla
 		1,
 		Threshold,
 		UE::Learning::Agents::GetCompletionMode(InCompletionMode),
-		Axis0,
-		Axis1);
+		Axis0.GetSafeNormal(UE_SMALL_NUMBER, FVector::ForwardVector),
+		Axis1.GetSafeNormal(UE_SMALL_NUMBER, FVector::RightVector));
 }
 
 void UPlanarPositionDifferenceCompletion::SetPlanarPositionDifferenceCompletion(const int32 AgentId, const FVector Position0, const FVector Position1)
@@ -240,8 +240,8 @@ UPlanarPositionSimilarityCompletion* UPlanarPositionSimilarityCompletion::AddPla
 		1,
 		Threshold,
 		UE::Learning::Agents::GetCompletionMode(InCompletionMode),
-		Axis0,
-		Axis1);
+		Axis0.GetSafeNormal(UE_SMALL_NUMBER, FVector::ForwardVector),
+		Axis1.GetSafeNormal(UE_SMALL_NUMBER, FVector::RightVector));
 }
 
 void UPlanarPositionSimilarityCompletion::SetPlanarPositionSimilarityCompletion(const int32 AgentId, const FVector Position0, const FVector Position1)
diff --git a/Engine/Plugins/Experimental/LearningAgents/Source/LearningAgentsTraining/Private/LearningAgentsImitationTrainer.cpp b/Engine/Plugins/Experimental/LearningAgents/Source/LearningAgentsTraining/Private/LearningAgentsImitationTrainer.cpp
--- a/Engine/Plugins/Experimental/LearningAgents/Source/LearningAgentsTraining/Private/LearningAgentsImitationTrainer.cpp
+++ b/Engine/Plugins/Experimental/LearningAgents/Source/LearningAgentsTraining/Private/LearningAgentsImitationTrainer.cpp
@@ -160,7 +160,7 @@ void ULearningAgentsImitationTrainer::BeginTraining(
 	ImitationTrainingSettings.WeightDecay = ImitationTrainerTrainingSettings.WeightDecay;
 	ImitationTrainingSettings.BatchSize = ImitationTrainerTrainingSettings.BatchSize;
 	ImitationTrainingSettings.Seed = ImitationTrainerTrainingSettings.RandomSeed;
-	ImitationTrainingSettings.Device = ImitationTrainerTrainingSettings.Device == ELearningAgentsTrainerDevice::CPU ? UE::Learning::ETrainerDevice::CPU : UE::Learning::ETrainerDevice::GPU;
+	ImitationTrainingSettings.Device = UE::Learning::Agents::GetTrainerDevice(ImitationTrainerTrainingSettings.Device);
 	ImitationTrainingSettings.bUseTensorboard = ImitationTrainerTrainingSettings.bUseTensorboard;
 
 	const UE::Learning::EImitationTrainerFlags TrainerFlags = 
diff --git a/Engine/Plugins/Experimental/LearningAgents/Source/LearningAgentsTraining/Private/LearningAgentsRecorder.cpp b/Engine/Plugins/Experimental/LearningAgents/Source/LearningAgentsTraining/Private/LearningAgentsRecorder.cpp
--- a/Engine/Plugins/Experimental/LearningAgents/Source/LearningAgentsTraining/Private/LearningAgentsRecorder.cpp
+++ b/Engine/Plugins/Experimental/LearningAgents/Source/LearningAgentsTraining/Private/LearningAgentsRecorder.cpp
@@ -334,7 +334,7 @@ void ULearningAgentsRecorder::BeginRecording(bool bReinitializeRecording)
 
 	if (IsRecording())
 	{
-		UE_LOG(LogLearning, Error, TEXT("%s: Cannot being recording as we are already Recording!"), *GetName());
+		UE_LOG(LogLearning, Error, TEXT("%s: Cannot begin recording as we are already Recording!"), *GetName());
 		return;
 	}
 
diff --git a/Engine/Plugins/Experimental/LearningAgents/Source/LearningAgentsTraining/Private/LearningAgentsRewards.cpp b/Engine/Plugins/Experimental/LearningAgents/Source/LearningAgentsTraining/Private/LearningAgentsRewards.cpp
--- a/Engine/Plugins/Experimental/LearningAgents/Source/LearningAgentsTraining/Private/LearningAgentsRewards.cpp
+++ b/Engine/Plugins/Experimental/LearningAgents/Source/LearningAgentsTraining/Private/LearningAgentsRewards.cpp
@@ -129,7 +129,7 @@ void UScalarVelocityReward::VisualLog(const UE::Learning::FIndexSet Instances) c
 
 ULocalDirectionalVelocityReward* ULocalDirectionalVelocityReward::AddLocalDirectionalVelocityReward(ULearningAgentsTrainer* InAgentTrainer, const FName Name, const float Weight, const float Scale, const FVector Axis)
 {
-	return UE::Learning::Agents::Rewards::Private::AddReward<ULocalDirectionalVelocityReward, UE::Learning::FLocalDirectionalVelocityReward>(InAgentTrainer, Name, TEXT("AddLocalDirectionalVelocityReward"), Weight, Scale, Axis);
+	return UE::Learning::Agents::Rewards::Private::AddReward<ULocalDirectionalVelocityReward, UE::Learning::FLocalDirectionalVelocityReward>(InAgentTrainer, Name, TEXT("AddLocalDirectionalVelocityReward"), Weight, Scale, Axis.GetSafeNormal(UE_SMALL_NUMBER, FVector::ForwardVector));
 }
 
 void ULocalDirectionalVelocityReward::SetLocalDirectionalVelocityReward(const int32 AgentId, const FVector Velocity, const FRotator RelativeRotation)
@@ -216,7 +216,15 @@ UPlanarPositionDifferencePenalty* UPlanarPositionDifferencePenalty::AddPlanarPos
 	const FVector Axis0,
 	const FVector Axis1)
 {
-	return UE::Learning::Agents::Rewards::Private::AddReward<UPlanarPositionDifferencePenalty, UE::Learning::FPlanarPositionDifferencePenalty>(InAgentTrainer, Name, TEXT("AddPlanarPositionDifferencePenalty"), Weight, Scale, Threshold, Axis0, Axis1);
+	return UE::Learning::Agents::Rewards::Private::AddReward<UPlanarPositionDifferencePenalty, UE::Learning::FPlanarPositionDifferencePenalty>(
+		InAgentTrainer, 
+		Name, 
+		TEXT("AddPlanarPositionDifferencePenalty"), 
+		Weight, 
+		Scale, 
+		Threshold, 
+		Axis0.GetSafeNormal(UE_SMALL_NUMBER, FVector::ForwardVector),
+		Axis1.GetSafeNormal(UE_SMALL_NUMBER, FVector::RightVector));
 }
 
 void UPlanarPositionDifferencePenalty::SetPlanarPositionDifferencePenalty(const int32 AgentId, const FVector Position0, const FVector Position1)
diff --git a/Engine/Plugins/Experimental/LearningAgents/Source/LearningAgentsTraining/Private/LearningAgentsTrainer.cpp b/Engine/Plugins/Experimental/LearningAgents/Source/LearningAgentsTraining/Private/LearningAgentsTrainer.cpp
--- a/Engine/Plugins/Experimental/LearningAgents/Source/LearningAgentsTraining/Private/LearningAgentsTrainer.cpp
+++ b/Engine/Plugins/Experimental/LearningAgents/Source/LearningAgentsTraining/Private/LearningAgentsTrainer.cpp
@@ -53,6 +53,29 @@ namespace UE::Learning::Agents
 	}
 }
 
+namespace UE::Learning::Agents
+{
+	ELearningAgentsTrainerDevice GetLearningAgentsTrainerDevice(const ETrainerDevice Device)
+	{
+		switch (Device)
+		{
+		case ETrainerDevice::CPU: return ELearningAgentsTrainerDevice::CPU;
+		case ETrainerDevice::GPU: return ELearningAgentsTrainerDevice::GPU;
+		default:UE_LOG(LogLearning, Error, TEXT("Unknown Trainer Device.")); return ELearningAgentsTrainerDevice::CPU;
+		}
+	}
+
+	ETrainerDevice GetTrainerDevice(const ELearningAgentsTrainerDevice Device)
+	{
+		switch (Device)
+		{
+		case ELearningAgentsTrainerDevice::CPU: return ETrainerDevice::CPU;
+		case ELearningAgentsTrainerDevice::GPU: return ETrainerDevice::GPU;
+		default:UE_LOG(LogLearning, Error, TEXT("Unknown Trainer Device.")); return ETrainerDevice::CPU;
+		}
+	}
+}
+
 FLearningAgentsTrainerPathSettings::FLearningAgentsTrainerPathSettings()
 {
 	EditorEngineRelativePath.Path = FPaths::EngineDir();
@@ -374,13 +397,24 @@ void ULearningAgentsTrainer::BeginTraining(
 
 	UE::Learning::FPPOTrainerTrainingSettings PPOTrainingSettings;
 	PPOTrainingSettings.IterationNum = TrainerTrainingSettings.NumberOfIterations;
-	PPOTrainingSettings.bUseTensorboard = TrainerTrainingSettings.bUseTensorboard;
+	PPOTrainingSettings.LearningRatePolicy = TrainerTrainingSettings.LearningRatePolicy;
+	PPOTrainingSettings.LearningRateCritic = TrainerTrainingSettings.LearningRateCritic;
+	PPOTrainingSettings.LearningRateDecay = TrainerTrainingSettings.LearningRateDecay;
+	PPOTrainingSettings.WeightDecay = TrainerTrainingSettings.WeightDecay;
 	PPOTrainingSettings.InitialActionScale = TrainerTrainingSettings.InitialActionScale;
-	PPOTrainingSettings.DiscountFactor = TrainerTrainingSettings.DiscountFactor;
-	PPOTrainingSettings.Seed = TrainerTrainingSettings.RandomSeed;
+	PPOTrainingSettings.BatchSize = TrainerTrainingSettings.BatchSize;
+	PPOTrainingSettings.EpsilonClip = TrainerTrainingSettings.EpsilonClip;
+	PPOTrainingSettings.ActionRegularizationWeight = TrainerTrainingSettings.ActionRegularizationWeight;
+	PPOTrainingSettings.EntropyWeight = TrainerTrainingSettings.EntropyWeight;
+	PPOTrainingSettings.GaeLambda = TrainerTrainingSettings.GaeLambda;
+	PPOTrainingSettings.bClipAdvantages = TrainerTrainingSettings.bClipAdvantages;
+	PPOTrainingSettings.bAdvantageNormalization = TrainerTrainingSettings.bAdvantageNormalization;
 	PPOTrainingSettings.TrimEpisodeStartStepNum = TrainerTrainingSettings.NumberOfStepsToTrimAtStartOfEpisode;
 	PPOTrainingSettings.TrimEpisodeEndStepNum = TrainerTrainingSettings.NumberOfStepsToTrimAtEndOfEpisode;
-	PPOTrainingSettings.Device = TrainerTrainingSettings.Device == ELearningAgentsTrainerDevice::CPU ? UE::Learning::ETrainerDevice::CPU : UE::Learning::ETrainerDevice::GPU;
+	PPOTrainingSettings.Seed = TrainerTrainingSettings.RandomSeed;
+	PPOTrainingSettings.DiscountFactor = TrainerTrainingSettings.DiscountFactor;
+	PPOTrainingSettings.Device = UE::Learning::Agents::GetTrainerDevice(TrainerTrainingSettings.Device);
+	PPOTrainingSettings.bUseTensorboard = TrainerTrainingSettings.bUseTensorboard;
 
 	UE::Learning::FPPOTrainerNetworkSettings PPONetworkSettings;
 	PPONetworkSettings.PolicyActionNoiseMin = Policy->GetPolicyObject().Settings.ActionNoiseMin;
diff --git a/Engine/Plugins/Experimental/LearningAgents/Source/LearningAgentsTraining/Public/LearningAgentsCompletions.h b/Engine/Plugins/Experimental/LearningAgents/Source/LearningAgentsTraining/Public/LearningAgentsCompletions.h
--- a/Engine/Plugins/Experimental/LearningAgents/Source/LearningAgentsTraining/Public/LearningAgentsCompletions.h
+++ b/Engine/Plugins/Experimental/LearningAgents/Source/LearningAgentsTraining/Public/LearningAgentsCompletions.h
@@ -43,6 +43,7 @@ class LEARNINGAGENTSTRAINING_API ULearningAgentsCompletion : public UObject
 
 public:
 
+	/** Reference to the Trainer this completion is associated with. */
 	UPROPERTY(VisibleAnywhere, Transient, Category = "LearningAgents")
 	TObjectPtr<ULearningAgentsTrainer> AgentTrainer;
 
diff --git a/Engine/Plugins/Experimental/LearningAgents/Source/LearningAgentsTraining/Public/LearningAgentsImitationTrainer.h b/Engine/Plugins/Experimental/LearningAgents/Source/LearningAgentsTraining/Public/LearningAgentsImitationTrainer.h
--- a/Engine/Plugins/Experimental/LearningAgents/Source/LearningAgentsTraining/Public/LearningAgentsImitationTrainer.h
+++ b/Engine/Plugins/Experimental/LearningAgents/Source/LearningAgentsTraining/Public/LearningAgentsImitationTrainer.h
@@ -28,30 +28,30 @@ struct FLearningAgentsImitationTrainerTrainingSettings
 
 public:
 
-	/** The number of iterations to run before training is complete. */
+	/** The number of iterations to run before ending training. */
 	UPROPERTY(EditAnywhere, Category = "LearningAgents", meta = (ClampMin = "0", UIMin = "0"))
 	int32 NumberOfIterations = 1000000;
 
-	/** Learning rate of the policy network. Typical values are between 0.001f and 0.0001f. */
-	UPROPERTY(EditAnywhere, Category = "LearningAgents", meta = (ClampMin = "0.0", UIMin = "0.0"))
+	/** Learning rate of the policy network. Typical values are between 0.001 and 0.0001. */
+	UPROPERTY(EditAnywhere, Category = "LearningAgents", meta = (ClampMin = "0.0", ClampMax = "1.0", UIMin = "0.0", UIMax = "1.0"))
 	float LearningRate = 0.0001f;
 
 	/** Ratio by which to decay the learning rate every 1000 iterations. */
-	UPROPERTY(EditAnywhere, Category = "LearningAgents", meta = (ClampMin = "0.0", UIMin = "0.0"))
+	UPROPERTY(EditAnywhere, Category = "LearningAgents", meta = (ClampMin = "0.0", ClampMax = "1.0", UIMin = "0.0", UIMax = "1.0"))
 	float LearningRateDecay = 0.99f;
 
 	/**
 	* Amount of weight decay to apply to the network. Larger values encourage network weights to be smaller but too
 	* large a value can cause the network weights to collapse to all zeros.
 	*/
-	UPROPERTY(EditAnywhere, Category = "LearningAgents", meta = (ClampMin = "0.0", UIMin = "0.0"))
+	UPROPERTY(EditAnywhere, Category = "LearningAgents", meta = (ClampMin = "0.0", ClampMax = "1.0", UIMin = "0.0", UIMax = "1.0"))
 	float WeightDecay = 0.001f;
 
 	/**
 	* Batch size to use for training. Smaller values tend to produce better results at the cost of slowing down
-	* training.
+	* training. Large batch sizes are much more computationally efficient when training on the GPU.
 	*/
-	UPROPERTY(EditAnywhere, Category = "LearningAgents", meta = (ClampMin = "1", UIMin = "1"))
+	UPROPERTY(EditAnywhere, Category = "LearningAgents", meta = (ClampMin = "0", ClampMax = "4096", UIMin = "0", UIMax = "4096"))
 	uint32 BatchSize = 128;
 
 	/** The seed used for any random sampling the trainer will perform, e.g. for weight initialization. */
@@ -62,7 +62,7 @@ public:
 	UPROPERTY(EditAnywhere, Category = "LearningAgents")
 	ELearningAgentsTrainerDevice Device = ELearningAgentsTrainerDevice::CPU;
 
-	/** If true, TensorBoard logs will be emitted to Intermediate\TensorBoard. Otherwise, no logs will be emitted. */
+	/** If true, TensorBoard logs will be emitted to the intermediate directory. */
 	UPROPERTY(EditAnywhere, Category = "LearningAgents")
 	bool bUseTensorboard = false;
 };
diff --git a/Engine/Plugins/Experimental/LearningAgents/Source/LearningAgentsTraining/Public/LearningAgentsRewards.h b/Engine/Plugins/Experimental/LearningAgents/Source/LearningAgentsTraining/Public/LearningAgentsRewards.h
--- a/Engine/Plugins/Experimental/LearningAgents/Source/LearningAgentsTraining/Public/LearningAgentsRewards.h
+++ b/Engine/Plugins/Experimental/LearningAgents/Source/LearningAgentsTraining/Public/LearningAgentsRewards.h
@@ -39,6 +39,7 @@ class LEARNINGAGENTSTRAINING_API ULearningAgentsReward : public UObject
 
 public:
 
+	/** Reference to the Trainer this reward is associated with. */
 	UPROPERTY(VisibleAnywhere, Transient, Category = "LearningAgents")
 	TObjectPtr<ULearningAgentsTrainer> AgentTrainer;
 
diff --git a/Engine/Plugins/Experimental/LearningAgents/Source/LearningAgentsTraining/Public/LearningAgentsTrainer.h b/Engine/Plugins/Experimental/LearningAgents/Source/LearningAgentsTraining/Public/LearningAgentsTrainer.h
--- a/Engine/Plugins/Experimental/LearningAgents/Source/LearningAgentsTraining/Public/LearningAgentsTrainer.h
+++ b/Engine/Plugins/Experimental/LearningAgents/Source/LearningAgentsTraining/Public/LearningAgentsTrainer.h
@@ -23,6 +23,7 @@ namespace UE::Learning
 	struct FSumReward;
 	struct FCompletionObject;
 	enum class ECompletionMode : uint8;
+	enum class ETrainerDevice : uint8;
 }
 
 class ALearningAgentsManager;
@@ -42,6 +43,7 @@ enum class ELearningAgentsCompletion : uint8
 	/** Episode ended early and zero reward was expected for all future steps. */
 	Termination	UMETA(DisplayName = "Termination"),
 };
+
 namespace UE::Learning::Agents
 {
 	/** Get the learning agents completion from the UE::Learning completion. */
@@ -128,6 +130,15 @@ enum class ELearningAgentsTrainerDevice : uint8
 	GPU,
 };
 
+namespace UE::Learning::Agents
+{
+	/** Get the learning agents trainer device from the UE::Learning trainer device. */
+	LEARNINGAGENTSTRAINING_API ELearningAgentsTrainerDevice GetLearningAgentsTrainerDevice(const ETrainerDevice Device);
+
+	/** Get the UE::Learning trainer device from the learning agents trainer device. */
+	LEARNINGAGENTSTRAINING_API ETrainerDevice GetTrainerDevice(const ELearningAgentsTrainerDevice Device);
+}
+
 /** The configurable settings for the training process. */
 USTRUCT(BlueprintType, Category = "LearningAgents")
 struct FLearningAgentsTrainerTrainingSettings
@@ -136,35 +147,89 @@ struct FLearningAgentsTrainerTrainingSettings
 
 public:
 
-	/** The number of iterations to run before training is complete. */
+	/** The number of iterations to run before ending training. */
 	UPROPERTY(EditAnywhere, Category = "LearningAgents", meta = (ClampMin = "0", UIMin = "0"))
 	int32 NumberOfIterations = 1000000;
 
-	/** If true, TensorBoard logs will be emitted to Intermediate\TensorBoard. Otherwise, no logs will be emitted. */
-	UPROPERTY(EditAnywhere, Category = "LearningAgents")
-	bool bUseTensorboard = false;
+	/** Learning rate of the policy network. Typical values are between 0.001 and 0.0001. */
+	UPROPERTY(EditAnywhere, Category = "LearningAgents", meta = (ClampMin = "0.0", ClampMax = "1.0", UIMin = "0.0", UIMax = "1.0"))
+	float LearningRatePolicy = 0.0001f;
+
+	/**
+	* Learning rate of the critic network. To avoid instability generally the critic should have a larger learning 
+	* rate than the policy. Typically this can be set to 10x the rate of the policy.
+	*/
+	UPROPERTY(EditAnywhere, Category = "LearningAgents", meta = (ClampMin = "0.0", ClampMax = "1.0", UIMin = "0.0", UIMax = "1.0"))
+	float LearningRateCritic = 0.001f;
+
+	/** Ratio by which to decay the learning rate every 1000 iterations. */
+	UPROPERTY(EditAnywhere, Category = "LearningAgents", meta = (ClampMin = "0.0", ClampMax = "1.0", UIMin = "0.0", UIMax = "1.0"))
+	float LearningRateDecay = 0.99f;
+
+	/**
+	* Amount of weight decay to apply to the network. Larger values encourage network weights to be smaller but too 
+	* large a value can cause the network weights to collapse to all zeros.
+	*/
+	UPROPERTY(EditAnywhere, Category = "LearningAgents", meta = (ClampMin = "0.0", ClampMax = "1.0", UIMin = "0.0", UIMax = "1.0"))
+	float WeightDecay = 0.001f;
 
 	/**
 	* The initial scaling for the weights of the output layer of the neural network. Typically, you would use this to
-	* scale down the initial weights as it can stabilize the initial training and speed up convergence.
+	* scale down the initial actions as it can stabilize the training and speed up convergence.
 	*/
-	UPROPERTY(EditAnywhere, Category = "LearningAgents", meta = (ClampMin = "0.0", UIMin = "0.0"))
+	UPROPERTY(EditAnywhere, Category = "LearningAgents", meta = (ClampMin = "0.0", ClampMax = "10.0", UIMin = "0.0", UIMax = "10.0"))
 	float InitialActionScale = 0.1f;
 
 	/**
-	* The discount factor to use during training. This affects how much the agent cares about future rewards vs
-	* near-term rewards. Should typically be a value less than but near 1. 
+	* Batch size to use for training. Smaller values tend to produce better results at the cost of slowing down 
+	* training. Large batch sizes are much more computationally efficient when training on the GPU.
 	*/
-	UPROPERTY(EditAnywhere, Category = "LearningAgents", meta = (ClampMin = "0.0", UIMin = "0.0"))
-	float DiscountFactor = 0.99f;
+	UPROPERTY(EditAnywhere, Category = "LearningAgents", meta = (ClampMin = "0", ClampMax = "4096", UIMin = "0", UIMax = "4096"))
+	int32 BatchSize = 128;
 
-	/** The seed used for any random sampling the trainer will perform, e.g. for weight initialization. */
-	UPROPERTY(EditAnywhere, Category = "LearningAgents", meta = (ClampMin = "0", UIMin = "0"))
-	int32 RandomSeed = 1234;
+	/**
+	* Clipping ratio to apply to policy updates. Keeps the training "on-policy". Larger values may speed up training at 
+	* the cost of stability. Conversely, too small values will keep the policy from being able to learn an 
+	* optimal policy.
+	*/
+	UPROPERTY(EditAnywhere, Category = "LearningAgents", meta = (ClampMin = "0.0", ClampMax = "1.0", UIMin = "0.0", UIMax = "1.0"))
+	float EpsilonClip = 0.2f;
+
+	/**
+	* Weight used to regularize actions. Larger values will encourage smaller actions but too large will cause actions 
+	* to become always zero.
+	*/
+	UPROPERTY(EditAnywhere, Category = "LearningAgents", meta = (ClampMin = "0.0", ClampMax = "1.0", UIMin = "0.0", UIMax = "1.0"))
+	float ActionRegularizationWeight = 0.001f;
+
+	/**
+	* Weighting used for the entropy bonus. Larger values encourage larger action noise and therefore greater 
+	* exploration but can make actions very noisy.
+	*/
+	UPROPERTY(EditAnywhere, Category = "LearningAgents", meta = (ClampMin = "0.0", ClampMax = "1.0", UIMin = "0.0", UIMax = "1.0"))
+	float EntropyWeight = 0.01f;
+
+	/**
+	* This is used in the Generalized Advantage Estimation as what is essentially an exponential smoothing/decay. 
+	* Typical values should be between 0.9 and 1.0.
+	*/
+	UPROPERTY(EditAnywhere, Category = "LearningAgents", meta = (ClampMin = "0.0", ClampMax = "1.0", UIMin = "0.0", UIMax = "1.0"))
+	float GaeLambda = 0.9f;
+
+	/**
+	* When true, very large or small advantages will be clipped. This has few downsides and helps with numerical 
+	* stability.
+	*/
+	UPROPERTY(EditAnywhere, Category = "LearningAgents")
+	bool bClipAdvantages = true;
+
+	/** When true, advantages are normalized. This tends to makes training more robust to adjustments of the scale of rewards. */
+	UPROPERTY(EditAnywhere, Category = "LearningAgents")
+	bool bAdvantageNormalization = true;
 
-	/** 
+	/**
 	* The number of steps to trim from the start of the episode, e.g. can be useful if some things are still getting
-	* setup at the start of the episode.
+	* setup at the start of the episode and you don't want them used for training.
 	*/
 	UPROPERTY(EditAnywhere, Category = "LearningAgents", meta = (ClampMin = "0", UIMin = "0"))
 	int32 NumberOfStepsToTrimAtStartOfEpisode = 0;
@@ -176,9 +241,24 @@ public:
 	UPROPERTY(EditAnywhere, Category = "LearningAgents", meta = (ClampMin = "0", UIMin = "0"))
 	int32 NumberOfStepsToTrimAtEndOfEpisode = 0;
 
+	/** The seed used for any random sampling the trainer will perform, e.g. for weight initialization. */
+	UPROPERTY(EditAnywhere, Category = "LearningAgents", meta = (ClampMin = "0", UIMin = "0"))
+	int32 RandomSeed = 1234;
+
+	/**
+	* The discount factor to use during training. This affects how much the agent cares about future rewards vs
+	* near-term rewards. Should typically be a value less than but near 1.0. 
+	*/
+	UPROPERTY(EditAnywhere, Category = "LearningAgents", meta = (ClampMin = "0.0", ClampMax = "1.0", UIMin = "0.0", UIMax = "1.0"))
+	float DiscountFactor = 0.99f;
+
 	/** The device to train on. */
 	UPROPERTY(EditAnywhere, Category = "LearningAgents")
 	ELearningAgentsTrainerDevice Device = ELearningAgentsTrainerDevice::GPU;
+
+	/** If true, TensorBoard logs will be emitted to the intermediate directory. */
+	UPROPERTY(EditAnywhere, Category = "LearningAgents")
+	bool bUseTensorboard = false;
 };
 
 /** The path settings for the trainer. */
diff --git a/Engine/Plugins/Experimental/LearningAgents/Source/LearningTraining/Public/LearningPPOTrainer.h b/Engine/Plugins/Experimental/LearningAgents/Source/LearningTraining/Public/LearningPPOTrainer.h
--- a/Engine/Plugins/Experimental/LearningAgents/Source/LearningTraining/Public/LearningPPOTrainer.h
+++ b/Engine/Plugins/Experimental/LearningAgents/Source/LearningTraining/Public/LearningPPOTrainer.h
@@ -80,7 +80,7 @@ namespace UE::Learning
 		float LearningRatePolicy = 0.0001f;
 
 		// Learning rate of the critic network. To avoid instability generally the critic 
-		// should have a larger learning rate then the actor.
+		// should have a larger learning rate than the policy.
 		float LearningRateCritic = 0.001f;
 
 		// Ratio by which to decay the learning rate every 1000 iterations.
@@ -96,12 +96,13 @@ namespace UE::Learning
 		float InitialActionScale = 0.1f;
 
 		// Batch size to use for training. Smaller values tend to produce better results 
-		// at the cost of slowing down training.
+		// at the cost of slowing down training. Large batch sizes are much more computationally efficient 
+		// when training on the GPU.
 		uint32 BatchSize = 128;
 
 		// Clipping ratio to apply to policy updates. Keeps the training "on-policy". 
 		// Larger values may speed up training at the cost of stability. Conversely, too small 
-		// values will keep the policy from being unable to learn an optimal policy.
+		// values will keep the policy from being able to learn an optimal policy.
 		float EpsilonClip = 0.2f;
 
 		// Weight used to regularize actions. Larger values will encourage smaller actions but too large
@@ -116,12 +117,12 @@ namespace UE::Learning
 		// an exponential smoothing/decay. Typical values should be between 0.9 and 1.0.
 		float GaeLambda = 0.9f;
 
-		// If to clip very large or small advantages. This has few downsides and helps 
-		// with numerical stability.
+		// When true, very large or small advantages will be clipped. This has few downsides 
+		// and helps with numerical stability.
 		bool bClipAdvantages = true;
 
-		// If to normalize advantages. This tends to makes training more robust to adjustments 
-		// of the scale of rewards.
+		// When true, advantages are normalized. This tends to makes training more robust to 
+		// adjustments of the scale of rewards. 
 		bool bAdvantageNormalization = true;
 
 		// Number of steps to trim from the start of each episode during training. This can
