LearningAgents: Added iteration checks for verifying order of operations. Fixed issue with inference for multiple agents. Fixed live recompilation issues.

diff --git a/Engine/Plugins/Experimental/LearningAgents/Source/Learning/Private/Learning.ispc b/Engine/Plugins/Experimental/LearningAgents/Source/Learning/Private/Learning.ispc
--- a/Engine/Plugins/Experimental/LearningAgents/Source/Learning/Private/Learning.ispc
+++ b/Engine/Plugins/Experimental/LearningAgents/Source/Learning/Private/Learning.ispc
@@ -175,12 +175,39 @@ export void LearningLayerMatMulPlusBias(
 	}
 }
 
+export void LearningLayerMatMulVecPlusBias(
+	uniform float Output[],
+	const uniform float Input[],
+	const uniform float Weights[],
+	const uniform float Biases[],
+	const uniform int RowNum,
+	const uniform int ColNum)
+{
+	foreach(ColIdx = 0 ... ColNum)
+	{
+		Output[ColIdx] = Biases[ColIdx];
+	}
+
+	for (uniform int RowIdx = 0; RowIdx < RowNum; RowIdx++)
+	{
+		uniform const float Value = Input[RowIdx];
+
+		if (Value != 0.0)
+		{
+			foreach(ColIdx = 0 ... ColNum)
+			{
+				Output[ColIdx] += Value * Weights[RowIdx * ColNum + ColIdx];
+			}
+		}
+	}
+}
+
+
 export void LearningLayerReLU(
 	uniform float InputOutput[],
-	const uniform int BatchNum,
-	const uniform int HiddenNum)
+	const uniform int Num)
 {
-	foreach(Idx = 0 ... BatchNum * HiddenNum)
+	foreach(Idx = 0 ... Num)
 	{
 		InputOutput[Idx] = max(InputOutput[Idx], 0.0f);
 	}
@@ -188,10 +215,9 @@ export void LearningLayerReLU(
 
 export void LearningLayerELU(
 	uniform float InputOutput[],
-	const uniform int BatchNum,
-	const uniform int HiddenNum)
+	const uniform int Num)
 {
-	foreach (Idx = 0 ... BatchNum * HiddenNum)
+	foreach (Idx = 0 ... Num)
 	{
 		InputOutput[Idx] = InputOutput[Idx] > 0.0f ? InputOutput[Idx] : LearningInvExpApprox(-InputOutput[Idx]) - 1.0f;
 	}
@@ -205,10 +231,9 @@ inline float LearningTanH(const float X, const float Min=-10.0f, const float Max
 
 export void LearningLayerTanH(
 	uniform float InputOutput[],
-	const uniform int BatchNum,
-	const uniform int HiddenNum)
+	const uniform int Num)
 {
-	foreach(Idx = 0 ... BatchNum * HiddenNum)
+	foreach(Idx = 0 ... Num)
 	{
 		InputOutput[Idx] = LearningTanH(InputOutput[Idx]);
 	}
@@ -245,6 +270,30 @@ export void LearningLayerActionNoise(
 	}
 }
 
+export void LearningLayerActionNoiseSingleBatch(
+	uniform float Output[],
+	uniform float OutputMean[],
+	uniform float OutputStd[],
+	const uniform float Input[],
+	const uniform uint32 Seed,
+	const uniform float ActionNoiseScale,
+	const uniform int OutputNum,
+	const uniform float LogActionNoiseMin,
+	const uniform float LogActionNoiseMax)
+{
+	foreach(OutputIdx = 0 ... OutputNum)
+	{
+		OutputMean[OutputIdx] = Input[OutputIdx];
+		OutputStd[OutputIdx] = ActionNoiseScale *
+			exp(LearningSigmoid(Input[OutputNum + OutputIdx]) * (LogActionNoiseMax - LogActionNoiseMin) + LogActionNoiseMin);
+
+		Output[OutputIdx] = LearningRandomGaussian(
+			Seed ^ 0xab744615 ^ LearningRandomInt(OutputIdx ^ 0xf8a88a27),
+			OutputMean[OutputIdx],
+			OutputStd[OutputIdx]);
+	}
+}
+
 /////////////////////////////////////////////////////////
 // 
 // Functions used by various different Optimizers
@@ -507,6 +556,19 @@ export void LearningCombineFeature(
 	}
 }
 
+export void LearningCombineFeatureSingleInstance(
+	uniform float Output[],
+	uniform const float Buffer[],
+	uniform const int DimOffset,
+	uniform const int DimNum,
+	uniform const float Scale)
+{
+	foreach(DimIdx = 0 ... DimNum)
+	{
+		Output[DimOffset + DimIdx] = Scale * Buffer[DimIdx];
+	}
+}
+
 export void LearningSeparateFeature(
 	uniform float Buffer[],
 	uniform const float Input[],
@@ -523,6 +585,20 @@ export void LearningSeparateFeature(
 	}
 }
 
+export void LearningSeparateFeatureSingleInstance(
+	uniform float Buffer[],
+	uniform const float Input[],
+	uniform const int DimOffset,
+	uniform const int DimNum,
+	uniform const float Scale,
+	uniform const float Epsilon)
+{
+	foreach(DimIdx = 0 ... DimNum)
+	{
+		Buffer[DimIdx] = Input[DimOffset + DimIdx] / max(Scale, Epsilon);
+	}
+}
+
 /////////////////////////////////////////////////////////
 // 
 // Functions related to K-Means Clustering
diff --git a/Engine/Plugins/Experimental/LearningAgents/Source/Learning/Private/LearningFeatureObject.cpp b/Engine/Plugins/Experimental/LearningAgents/Source/Learning/Private/LearningFeatureObject.cpp
--- a/Engine/Plugins/Experimental/LearningAgents/Source/Learning/Private/LearningFeatureObject.cpp
+++ b/Engine/Plugins/Experimental/LearningAgents/Source/Learning/Private/LearningFeatureObject.cpp
@@ -726,9 +726,9 @@ namespace UE::Learning
 
 				UE_LEARNING_CHECK(FeatureInput.Num<1>() == FeatureSize);
 
-				if (UE_LEARNING_ISPC && Instances.IsSlice())
-				{
 #if UE_LEARNING_ISPC
+				if (Instances.IsSlice())
+				{
 					ispc::LearningCombineFeature(
 						Feature.Slice(Instances.GetSliceStart(), Instances.GetSliceNum()).GetData(),
 						FeatureInput.Slice(Instances.GetSliceStart(), Instances.GetSliceNum()).GetData(),
@@ -737,18 +737,28 @@ namespace UE::Learning
 						Feature.Num<1>(),
 						Instances.GetSliceNum(),
 						Scale);
-#endif
 				}
 				else
 				{
 					for (const int32 InstanceIdx : Instances)
 					{
-						for (int32 DimIdx = 0; DimIdx < FeatureSize; DimIdx++)
-						{
-							Feature[InstanceIdx][FeatureOffset + DimIdx] = Scale * FeatureInput[InstanceIdx][DimIdx];
-						}
+						ispc::LearningCombineFeatureSingleInstance(
+							Feature[InstanceIdx].GetData(),
+							FeatureInput[InstanceIdx].GetData(),
+							FeatureOffset,
+							FeatureSize,
+							Scale);
+					}
+				}
+#else
+				for (const int32 InstanceIdx : Instances)
+				{
+					for (int32 DimIdx = 0; DimIdx < FeatureSize; DimIdx++)
+					{
+						Feature[InstanceIdx][FeatureOffset + DimIdx] = Scale * FeatureInput[InstanceIdx][DimIdx];
 					}
 				}
+#endif
 			}
 		}
 
@@ -775,9 +785,9 @@ namespace UE::Learning
 
 				UE_LEARNING_CHECK(FeatureOutput.Num<1>() == FeatureSize);
 
-				if (UE_LEARNING_ISPC && Instances.IsSlice())
-				{
 #if UE_LEARNING_ISPC
+				if (Instances.IsSlice())
+				{
 					ispc::LearningSeparateFeature(
 						FeatureOutput.Slice(Instances.GetSliceStart(), Instances.GetSliceNum()).GetData(),
 						Feature.Slice(Instances.GetSliceStart(), Instances.GetSliceNum()).GetData(),
@@ -787,18 +797,29 @@ namespace UE::Learning
 						Instances.GetSliceNum(),
 						Scale,
 						UE_SMALL_NUMBER);
-#endif
 				}
 				else
 				{
 					for (const int32 InstanceIdx : Instances)
 					{
-						for (int32 DimIdx = 0; DimIdx < FeatureSize; DimIdx++)
-						{
-							FeatureOutput[InstanceIdx][DimIdx] = Feature[InstanceIdx][FeatureOffset + DimIdx] / FMath::Max(Scale, UE_SMALL_NUMBER);
-						}
+						ispc::LearningSeparateFeatureSingleInstance(
+							FeatureOutput[InstanceIdx].GetData(),
+							Feature[InstanceIdx].GetData(),
+							FeatureOffset,
+							FeatureSize,
+							Scale,
+							UE_SMALL_NUMBER);
+					}
+				}
+#else
+				for (const int32 InstanceIdx : Instances)
+				{
+					for (int32 DimIdx = 0; DimIdx < FeatureSize; DimIdx++)
+					{
+						FeatureOutput[InstanceIdx][DimIdx] = Feature[InstanceIdx][FeatureOffset + DimIdx] / FMath::Max(Scale, UE_SMALL_NUMBER);
 					}
 				}
+#endif
 			}
 
 			Features[FeatureIdx]->Decode(Instances);
diff --git a/Engine/Plugins/Experimental/LearningAgents/Source/Learning/Private/LearningNeuralNetworkObject.cpp b/Engine/Plugins/Experimental/LearningAgents/Source/Learning/Private/LearningNeuralNetworkObject.cpp
--- a/Engine/Plugins/Experimental/LearningAgents/Source/Learning/Private/LearningNeuralNetworkObject.cpp
+++ b/Engine/Plugins/Experimental/LearningAgents/Source/Learning/Private/LearningNeuralNetworkObject.cpp
@@ -3,7 +3,6 @@
 #include "LearningNeuralNetworkObject.h"
 
 #include "LearningNeuralNetwork.h"
-#include "LearningEigen.h"
 #include "LearningRandom.h"
 
 #if UE_LEARNING_ISPC
@@ -46,9 +45,9 @@ namespace UE::Learning
 			const int32 RowNum = Weights.Num<0>();
 			const int32 ColNum = Weights.Num<1>();
 
-			if (UE_LEARNING_ISPC && Instances.IsSlice())
-			{
 #if UE_LEARNING_ISPC
+			if (Instances.IsSlice())
+			{
 				ispc::LearningLayerMatMulPlusBias(
 					Output.Slice(Instances.GetSliceStart(), Instances.GetSliceNum()).GetData(),
 					Input.Slice(Instances.GetSliceStart(), Instances.GetSliceNum()).GetData(),
@@ -57,15 +56,39 @@ namespace UE::Learning
 					Instances.GetSliceNum(),
 					RowNum,
 					ColNum);
-#endif
 			}
 			else
 			{
-				OutEigenMatrix(Output).noalias() =
-					(InEigenMatrix(Weights).transpose() *
-						InEigenMatrix(Input).transpose()).transpose().rowwise() +
-					InEigenRowVector(Biases);
+				for (const int32 InstanceIdx : Instances)
+				{
+					ispc::LearningLayerMatMulVecPlusBias(
+						Output[InstanceIdx].GetData(),
+						Input[InstanceIdx].GetData(),
+						Weights.GetData(),
+						Biases.GetData(),
+						RowNum,
+						ColNum);
+				}
 			}
+#else
+			for (const int32 InstanceIdx : Instances)
+			{
+				Array::Copy(Output[InstanceIdx], Biases);
+
+				for (int32 RowIdx = 0; RowIdx < RowNum; RowIdx++)
+				{
+					const float Value = Input[InstanceIdx][RowIdx];
+
+					if (Value != 0.0)
+					{
+						for (int32 ColIdx = 0; ColIdx < ColNum; ColIdx++)
+						{
+							Output[InstanceIdx][ColIdx] += Value * Weights[RowIdx][ColIdx];
+						}
+					}
+				}
+			}
+#endif
 		}
 
 		static inline void ActivationReLU(TLearningArrayView<2, float> InputOutput, const FIndexSet Instances)
@@ -74,25 +97,31 @@ namespace UE::Learning
 
 			const int32 HiddenNum = InputOutput.Num<1>();
 
-			if (UE_LEARNING_ISPC && Instances.IsSlice())
-			{
 #if UE_LEARNING_ISPC
+			if (Instances.IsSlice())
+			{
 				ispc::LearningLayerReLU(
 					InputOutput.Slice(Instances.GetSliceStart(), Instances.GetSliceNum()).GetData(),
-					Instances.GetSliceNum(),
-					HiddenNum);
-#endif
+					Instances.GetSliceNum() * HiddenNum);
 			}
 			else
 			{
 				for (const int32 InstanceIdx : Instances)
 				{
-					for (int32 HiddenIdx = 0; HiddenIdx < HiddenNum; HiddenIdx++)
-					{
-						InputOutput[InstanceIdx][HiddenIdx] = ReLU(InputOutput[InstanceIdx][HiddenIdx]);
-					}
+					ispc::LearningLayerReLU(
+						InputOutput[InstanceIdx].GetData(),
+						HiddenNum);
+				}
+			}
+#else
+			for (const int32 InstanceIdx : Instances)
+			{
+				for (int32 HiddenIdx = 0; HiddenIdx < HiddenNum; HiddenIdx++)
+				{
+					InputOutput[InstanceIdx][HiddenIdx] = ReLU(InputOutput[InstanceIdx][HiddenIdx]);
 				}
 			}
+#endif
 		}
 
 		static inline void ActivationELU(TLearningArrayView<2, float> InputOutput, const FIndexSet Instances)
@@ -101,25 +130,31 @@ namespace UE::Learning
 
 			const int32 HiddenNum = InputOutput.Num<1>();
 
-			if (UE_LEARNING_ISPC && Instances.IsSlice())
-			{
 #if UE_LEARNING_ISPC
+			if (Instances.IsSlice())
+			{
 				ispc::LearningLayerELU(
 					InputOutput.Slice(Instances.GetSliceStart(), Instances.GetSliceNum()).GetData(),
-					Instances.GetSliceNum(),
-					HiddenNum);
-#endif
+					Instances.GetSliceNum() * HiddenNum);
 			}
 			else
 			{
 				for (const int32 InstanceIdx : Instances)
 				{
-					for (int32 HiddenIdx = 0; HiddenIdx < HiddenNum; HiddenIdx++)
-					{
-						InputOutput[InstanceIdx][HiddenIdx] = ELU(InputOutput[InstanceIdx][HiddenIdx]);
-					}
+					ispc::LearningLayerELU(
+						InputOutput[InstanceIdx].GetData(),
+						HiddenNum);
+				}
+			}
+#else
+			for (const int32 InstanceIdx : Instances)
+			{
+				for (int32 HiddenIdx = 0; HiddenIdx < HiddenNum; HiddenIdx++)
+				{
+					InputOutput[InstanceIdx][HiddenIdx] = ELU(InputOutput[InstanceIdx][HiddenIdx]);
 				}
 			}
+#endif
 		}
 
 		static inline void ActivationTanH(TLearningArrayView<2, float> InputOutput, const FIndexSet Instances)
@@ -128,25 +163,31 @@ namespace UE::Learning
 
 			const int32 HiddenNum = InputOutput.Num<1>();
 
-			if (UE_LEARNING_ISPC && Instances.IsSlice())
-			{
 #if UE_LEARNING_ISPC
+			if (Instances.IsSlice())
+			{
 				ispc::LearningLayerTanH(
 					InputOutput.Slice(Instances.GetSliceStart(), Instances.GetSliceNum()).GetData(),
-					Instances.GetSliceNum(),
-					HiddenNum);
-#endif
+					Instances.GetSliceNum() * HiddenNum);
 			}
 			else
 			{
 				for (const int32 InstanceIdx : Instances)
 				{
-					for (int32 HiddenIdx = 0; HiddenIdx < HiddenNum; HiddenIdx++)
-					{
-						InputOutput[InstanceIdx][HiddenIdx] = TanH(InputOutput[InstanceIdx][HiddenIdx]);
-					}
+					ispc::LearningLayerTanH(
+						InputOutput[InstanceIdx].GetData(),
+						HiddenNum);
 				}
 			}
+#else
+			for (const int32 InstanceIdx : Instances)
+			{
+				for (int32 HiddenIdx = 0; HiddenIdx < HiddenNum; HiddenIdx++)
+				{
+					InputOutput[InstanceIdx][HiddenIdx] = TanH(InputOutput[InstanceIdx][HiddenIdx]);
+				}
+			}
+#endif
 		}
 
 		static inline void ActionNoise(
@@ -165,9 +206,9 @@ namespace UE::Learning
 			const int32 InputNum = Input.Num<1>();
 			const int32 OutputNum = Output.Num<1>();
 
-			if (UE_LEARNING_ISPC && Instances.IsSlice())
-			{
 #if UE_LEARNING_ISPC
+			if (Instances.IsSlice())
+			{
 				ispc::LearningLayerActionNoise(
 					Output.Slice(Instances.GetSliceStart(), Instances.GetSliceNum()).GetData(),
 					OutputMean.Slice(Instances.GetSliceStart(), Instances.GetSliceNum()).GetData(),
@@ -180,25 +221,39 @@ namespace UE::Learning
 					OutputNum,
 					LogActionNoiseMin,
 					LogActionNoiseMax);
-#endif
 			}
 			else
 			{
 				for (const int32 InstanceIdx : Instances)
 				{
-					for (int32 OutputIdx = 0; OutputIdx < OutputNum; OutputIdx++)
-					{
-						OutputMean[InstanceIdx][OutputIdx] = Input[InstanceIdx][OutputIdx];
-						OutputStd[InstanceIdx][OutputIdx] = ActionNoiseScale[InstanceIdx] *
-							FMath::Exp(Sigmoid(Input[InstanceIdx][OutputNum + OutputIdx]) * (LogActionNoiseMax - LogActionNoiseMin) + LogActionNoiseMin);
-
-						Output[InstanceIdx][OutputIdx] = Random::Gaussian(
-							Seed[InstanceIdx] ^ 0xab744615 ^ Random::Int(OutputIdx ^ 0xf8a88a27),
-							OutputMean[InstanceIdx][OutputIdx],
-							OutputStd[InstanceIdx][OutputIdx]);
-					}
+					ispc::LearningLayerActionNoiseSingleBatch(
+						Output[InstanceIdx].GetData(),
+						OutputMean[InstanceIdx].GetData(),
+						OutputStd[InstanceIdx].GetData(),
+						Input[InstanceIdx].GetData(),
+						Seed[InstanceIdx],
+						ActionNoiseScale[InstanceIdx],
+						OutputNum,
+						LogActionNoiseMin,
+						LogActionNoiseMax);
 				}
 			}
+#else
+			for (const int32 InstanceIdx : Instances)
+			{
+				for (int32 OutputIdx = 0; OutputIdx < OutputNum; OutputIdx++)
+				{
+					OutputMean[InstanceIdx][OutputIdx] = Input[InstanceIdx][OutputIdx];
+					OutputStd[InstanceIdx][OutputIdx] = ActionNoiseScale[InstanceIdx] *
+						FMath::Exp(Sigmoid(Input[InstanceIdx][OutputNum + OutputIdx]) * (LogActionNoiseMax - LogActionNoiseMin) + LogActionNoiseMin);
+
+					Output[InstanceIdx][OutputIdx] = Random::Gaussian(
+						Seed[InstanceIdx] ^ 0xab744615 ^ Random::Int(OutputIdx ^ 0xf8a88a27),
+						OutputMean[InstanceIdx][OutputIdx],
+						OutputStd[InstanceIdx][OutputIdx]);
+				}
+			}
+#endif
 
 			Random::ResampleStateArray(Seed, Instances);
 		}
diff --git a/Engine/Plugins/Experimental/LearningAgents/Source/Learning/Public/LearningArray.h b/Engine/Plugins/Experimental/LearningAgents/Source/Learning/Public/LearningArray.h
--- a/Engine/Plugins/Experimental/LearningAgents/Source/Learning/Public/LearningArray.h
+++ b/Engine/Plugins/Experimental/LearningAgents/Source/Learning/Public/LearningArray.h
@@ -1081,6 +1081,16 @@ namespace UE::Learning
 		{
 			return Format<float>(Array, [](const float& Value) { return FString::Printf(TEXT("%6.3f"), Value); }, MaxItemNum);
 		}
+
+		static inline FString FormatInt32(const TLearningArrayView<1, const int32> Array, const int32 MaxItemNum = 16)
+		{
+			return Format<int32>(Array, [](const int32& Value) { return FString::Printf(TEXT("%i"), Value); }, MaxItemNum);
+		}
+
+		static inline FString FormatUint64(const TLearningArrayView<1, const uint64> Array, const int32 MaxItemNum = 16)
+		{
+			return Format<uint64>(Array, [](const uint64& Value) { return FString::Printf(TEXT("%ul"), Value); }, MaxItemNum);
+		}
 	}
 }
 
diff --git a/Engine/Plugins/Experimental/LearningAgents/Source/LearningAgents/Private/LearningAgentsActions.cpp b/Engine/Plugins/Experimental/LearningAgents/Source/LearningAgents/Private/LearningAgentsActions.cpp
--- a/Engine/Plugins/Experimental/LearningAgents/Source/LearningAgents/Private/LearningAgentsActions.cpp
+++ b/Engine/Plugins/Experimental/LearningAgents/Source/LearningAgents/Private/LearningAgentsActions.cpp
@@ -22,15 +22,26 @@ namespace UE::Learning::Agents::Actions::Private
 			return nullptr;
 		}
 
+		if (!InInteractor->HasAgentManager())
+		{
+			UE_LOG(LogLearning, Error, TEXT("%s: Must be attached to a LearningAgentsManager Actor."), *InInteractor->GetName());
+			return nullptr;
+		}
+
 		ActionUObject* Action = NewObject<ActionUObject>(InInteractor, Name);
 
 		Action->Interactor = InInteractor;
 		Action->FeatureObject = MakeShared<ActionFObject>(
 			Action->GetFName(),
 			InInteractor->GetAgentManager()->GetInstanceData().ToSharedRef(),
-			InInteractor->GetAgentManager()->GetMaxInstanceNum(),
+			InInteractor->GetAgentManager()->GetMaxAgentNum(),
 			Forward<InArgTypes>(Args)...);
 
+		Action->AgentGetIteration.SetNumUninitialized({ InInteractor->GetAgentManager()->GetMaxAgentNum() });
+		Action->AgentSetIteration.SetNumUninitialized({ InInteractor->GetAgentManager()->GetMaxAgentNum() });
+		UE::Learning::Array::Set<1, uint64>(Action->AgentGetIteration, INDEX_NONE);
+		UE::Learning::Array::Set<1, uint64>(Action->AgentSetIteration, INDEX_NONE);
+
 		// We assume all supported action feature objects can be encoded and decoded
 		UE_LEARNING_CHECK(Action->FeatureObject->IsEncodable() && Action->FeatureObject->IsDecodable());
 
@@ -47,7 +58,7 @@ UFloatAction* UFloatAction::AddFloatAction(ULearningAgentsInteractor* InInteract
 	return UE::Learning::Agents::Actions::Private::AddAction<UFloatAction,UE::Learning::FFloatFeature>(InInteractor, Name, TEXT("AddFloatAction"), 1, Scale);
 }
 
-float UFloatAction::GetFloatAction(const int32 AgentId) const
+float UFloatAction::GetFloatAction(const int32 AgentId)
 {
 	if (!Interactor->HasAgent(AgentId))
 	{
@@ -55,6 +66,7 @@ float UFloatAction::GetFloatAction(const int32 AgentId) const
 		return 0.0f;
 	}
 
+	AgentGetIteration[AgentId]++;
 	return FeatureObject->InstanceData->ConstView(FeatureObject->ValueHandle)[AgentId][0];
 }
 
@@ -67,6 +79,7 @@ void UFloatAction::SetFloatAction(const int32 AgentId, const float Value)
 	}
 
 	FeatureObject->InstanceData->View(FeatureObject->ValueHandle)[AgentId][0] = Value;
+	AgentSetIteration[AgentId]++;
 }
 
 #if UE_LEARNING_AGENTS_ENABLE_VISUAL_LOG
@@ -84,11 +97,11 @@ void UFloatAction::VisualLog(const UE::Learning::FIndexSet Instances) const
 			UE_LEARNING_AGENTS_VLOG_STRING(this, LogLearning, Display,
 				Actor->GetActorLocation(),
 				VisualLogColor.ToFColor(true),
-				TEXT("Agent %i\nScale: [% 6.2f]\nValue: [% 6.2f]\nEncoded: [% 6.3f]"),
+				TEXT("Agent %i\nScale: [% 6.2f]\nValue: %s\nEncoded: %s"),
 				Instance,
 				FeatureObject->Scale,
-				ValueView[Instance][0],
-				FeatureView[Instance][0]);
+				*UE::Learning::Array::FormatFloat(ValueView[Instance]),
+				*UE::Learning::Array::FormatFloat(FeatureView[Instance]));
 		}
 	}
 }
@@ -105,7 +118,7 @@ UFloatArrayAction* UFloatArrayAction::AddFloatArrayAction(ULearningAgentsInterac
 	return UE::Learning::Agents::Actions::Private::AddAction<UFloatArrayAction, UE::Learning::FFloatFeature>(InInteractor, Name, TEXT("AddFloatArrayAction"), Num, Scale);
 }
 
-void UFloatArrayAction::GetFloatArrayAction(const int32 AgentId, TArray<float>& OutValues) const
+void UFloatArrayAction::GetFloatArrayAction(const int32 AgentId, TArray<float>& OutValues)
 {
 	if (!Interactor->HasAgent(AgentId))
 	{
@@ -116,6 +129,8 @@ void UFloatArrayAction::GetFloatArrayAction(const int32 AgentId, TArray<float>&
 
 	const TLearningArrayView<2, const float> View = FeatureObject->InstanceData->ConstView(FeatureObject->ValueHandle);
 
+	AgentGetIteration[AgentId]++;
+
 	OutValues.SetNumUninitialized(View.Num<1>());
 	UE::Learning::Array::Copy<1, float>(OutValues, View[AgentId]);
 }
@@ -137,6 +152,7 @@ void UFloatArrayAction::SetFloatArrayAction(const int32 AgentId, const TArray<fl
 	}
 
 	UE::Learning::Array::Copy<1, float>(View[AgentId], Values);
+	AgentSetIteration[AgentId]++;
 }
 
 #if UE_LEARNING_AGENTS_ENABLE_VISUAL_LOG
@@ -171,7 +187,7 @@ UVectorAction* UVectorAction::AddVectorAction(ULearningAgentsInteractor* InInter
 	return UE::Learning::Agents::Actions::Private::AddAction<UVectorAction, UE::Learning::FFloatFeature>(InInteractor, Name, TEXT("AddVectorAction"), 3, Scale);
 }
 
-FVector UVectorAction::GetVectorAction(const int32 AgentId) const
+FVector UVectorAction::GetVectorAction(const int32 AgentId)
 {
 	if (!Interactor->HasAgent(AgentId))
 	{
@@ -181,6 +197,8 @@ FVector UVectorAction::GetVectorAction(const int32 AgentId) const
 
 	const TLearningArrayView<2, const float> View = FeatureObject->InstanceData->ConstView(FeatureObject->ValueHandle);
 
+	AgentGetIteration[AgentId]++;
+
 	return FVector(View[AgentId][0], View[AgentId][1], View[AgentId][2]);
 }
 
@@ -193,10 +211,11 @@ void UVectorAction::SetVectorAction(const int32 AgentId, const FVector InAction)
 	}
 
 	const TLearningArrayView<2, float> View = FeatureObject->InstanceData->View(FeatureObject->ValueHandle);
-
 	View[AgentId][0] = InAction.X;
 	View[AgentId][1] = InAction.Y;
 	View[AgentId][2] = InAction.Z;
+
+	AgentSetIteration[AgentId]++;
 }
 
 #if UE_LEARNING_AGENTS_ENABLE_VISUAL_LOG
@@ -248,7 +267,7 @@ UVectorArrayAction* UVectorArrayAction::AddVectorArrayAction(ULearningAgentsInte
 	return UE::Learning::Agents::Actions::Private::AddAction<UVectorArrayAction, UE::Learning::FFloatFeature>(InInteractor, Name, TEXT("AddVectorArrayAction"), Num * 3, Scale);
 }
 
-void UVectorArrayAction::GetVectorArrayAction(const int32 AgentId, TArray<FVector>& OutVectors) const
+void UVectorArrayAction::GetVectorArrayAction(const int32 AgentId, TArray<FVector>& OutVectors)
 {
 	if (!Interactor->HasAgent(AgentId))
 	{
@@ -259,6 +278,8 @@ void UVectorArrayAction::GetVectorArrayAction(const int32 AgentId, TArray<FVecto
 
 	const TLearningArrayView<2, const float> View = FeatureObject->InstanceData->ConstView(FeatureObject->ValueHandle);
 
+	AgentGetIteration[AgentId]++;
+
 	OutVectors.SetNumUninitialized(View.Num<1>() / 3);
 
 	for (int32 VectorIdx = 0; VectorIdx < View.Num<1>() / 3; VectorIdx++)
@@ -292,6 +313,8 @@ void UVectorArrayAction::SetVectorArrayAction(const int32 AgentId, const TArray<
 		View[AgentId][VectorIdx * 3 + 1] = Vectors[VectorIdx].Y;
 		View[AgentId][VectorIdx * 3 + 2] = Vectors[VectorIdx].Z;
 	}
+
+	AgentSetIteration[AgentId]++;
 }
 
 #if UE_LEARNING_AGENTS_ENABLE_VISUAL_LOG
@@ -346,7 +369,7 @@ URotationAction* URotationAction::AddRotationAction(ULearningAgentsInteractor* I
 	return UE::Learning::Agents::Actions::Private::AddAction<URotationAction, UE::Learning::FRotationVectorFeature>(InInteractor, Name, TEXT("AddRotationAction"), 1, Scale);
 }
 
-FRotator URotationAction::GetRotationAction(const int32 AgentId) const
+FRotator URotationAction::GetRotationAction(const int32 AgentId)
 {
 	if (!Interactor->HasAgent(AgentId))
 	{
@@ -356,10 +379,12 @@ FRotator URotationAction::GetRotationAction(const int32 AgentId) const
 
 	const TLearningArrayView<2, const FVector> View = FeatureObject->InstanceData->ConstView(FeatureObject->RotationVectorsHandle);
 
+	AgentGetIteration[AgentId]++;
+
 	return FQuat::MakeFromRotationVector((UE_TWO_PI / 180.0f) * View[AgentId][0]).Rotator();
 }
 
-FVector URotationAction::GetRotationActionAsRotationVector(const int32 AgentId) const
+FVector URotationAction::GetRotationActionAsRotationVector(const int32 AgentId)
 {
 	if (!Interactor->HasAgent(AgentId))
 	{
@@ -370,7 +395,7 @@ FVector URotationAction::GetRotationActionAsRotationVector(const int32 AgentId)
 	return FeatureObject->InstanceData->ConstView(FeatureObject->RotationVectorsHandle)[AgentId][0];
 }
 
-FQuat URotationAction::GetRotationActionAsQuat(const int32 AgentId) const
+FQuat URotationAction::GetRotationActionAsQuat(const int32 AgentId)
 {
 	if (!Interactor->HasAgent(AgentId))
 	{
@@ -380,6 +405,8 @@ FQuat URotationAction::GetRotationActionAsQuat(const int32 AgentId) const
 
 	const TLearningArrayView<2, const FVector> View = FeatureObject->InstanceData->ConstView(FeatureObject->RotationVectorsHandle);
 
+	AgentGetIteration[AgentId]++;
+
 	return FQuat::MakeFromRotationVector((UE_TWO_PI / 180.0f) * View[AgentId][0]);
 }
 
@@ -426,7 +453,7 @@ URotationArrayAction* URotationArrayAction::AddRotationArrayAction(ULearningAgen
 	return UE::Learning::Agents::Actions::Private::AddAction<URotationArrayAction, UE::Learning::FRotationVectorFeature>(InInteractor, Name, TEXT("AddRotationArrayAction"), RotationNum, Scale);
 }
 
-void URotationArrayAction::GetRotationArrayAction(const int32 AgentId, TArray<FRotator>& OutRotations) const
+void URotationArrayAction::GetRotationArrayAction(const int32 AgentId, TArray<FRotator>& OutRotations)
 {
 	if (!Interactor->HasAgent(AgentId))
 	{
@@ -436,6 +463,8 @@ void URotationArrayAction::GetRotationArrayAction(const int32 AgentId, TArray<FR
 
 	const TLearningArrayView<2, const FVector> View = FeatureObject->InstanceData->ConstView(FeatureObject->RotationVectorsHandle);
 
+	AgentGetIteration[AgentId]++;
+
 	OutRotations.SetNumUninitialized(View.Num<1>());
 	for (int32 RotationVectorIdx = 0; RotationVectorIdx < View.Num<1>(); RotationVectorIdx++)
 	{
@@ -443,7 +472,7 @@ void URotationArrayAction::GetRotationArrayAction(const int32 AgentId, TArray<FR
 	}
 }
 
-void URotationArrayAction::GetRotationArrayActionAsRotationVectors(const int32 AgentId, TArray<FVector>& OutRotationVectors) const
+void URotationArrayAction::GetRotationArrayActionAsRotationVectors(const int32 AgentId, TArray<FVector>& OutRotationVectors)
 {
 	if (!Interactor->HasAgent(AgentId))
 	{
@@ -453,6 +482,8 @@ void URotationArrayAction::GetRotationArrayActionAsRotationVectors(const int32 A
 
 	const TLearningArrayView<2, const FVector> View = FeatureObject->InstanceData->ConstView(FeatureObject->RotationVectorsHandle);
 
+	AgentGetIteration[AgentId]++;
+
 	OutRotationVectors.SetNumUninitialized(View.Num<1>());
 	for (int32 RotationVectorIdx = 0; RotationVectorIdx < View.Num<1>(); RotationVectorIdx++)
 	{
@@ -460,7 +491,7 @@ void URotationArrayAction::GetRotationArrayActionAsRotationVectors(const int32 A
 	}
 }
 
-void URotationArrayAction::GetRotationArrayActionAsQuats(const int32 AgentId, TArray<FQuat>& OutRotations) const
+void URotationArrayAction::GetRotationArrayActionAsQuats(const int32 AgentId, TArray<FQuat>& OutRotations)
 {
 	if (!Interactor->HasAgent(AgentId))
 	{
@@ -470,6 +501,8 @@ void URotationArrayAction::GetRotationArrayActionAsQuats(const int32 AgentId, TA
 
 	const TLearningArrayView<2, const FVector> View = FeatureObject->InstanceData->ConstView(FeatureObject->RotationVectorsHandle);
 
+	AgentGetIteration[AgentId]++;
+
 	OutRotations.SetNumUninitialized(View.Num<1>());
 	for (int32 RotationVectorIdx = 0; RotationVectorIdx < View.Num<1>(); RotationVectorIdx++)
 	{
diff --git a/Engine/Plugins/Experimental/LearningAgents/Source/LearningAgents/Private/LearningAgentsController.cpp b/Engine/Plugins/Experimental/LearningAgents/Source/LearningAgents/Private/LearningAgentsController.cpp
--- a/Engine/Plugins/Experimental/LearningAgents/Source/LearningAgents/Private/LearningAgentsController.cpp
+++ b/Engine/Plugins/Experimental/LearningAgents/Source/LearningAgents/Private/LearningAgentsController.cpp
@@ -8,9 +8,7 @@
 #include "LearningFeatureObject.h"
 #include "LearningLog.h"
 
-#include "EngineDefines.h"
-
-ULearningAgentsController::ULearningAgentsController() : ULearningAgentsManagerComponent() {}
+ULearningAgentsController::ULearningAgentsController(const FObjectInitializer& ObjectInitializer) : ULearningAgentsManagerComponent(ObjectInitializer) {}
 ULearningAgentsController::ULearningAgentsController(FVTableHelper& Helper) : ULearningAgentsController() {}
 ULearningAgentsController::~ULearningAgentsController() {}
 
@@ -19,7 +17,7 @@ void ULearningAgentsController::SetActions_Implementation(const TArray<int32>& A
 	// Can be overridden to get actions without blueprints
 }
 
-void ULearningAgentsController::SetupController(ALearningAgentsManager* InAgentManager, ULearningAgentsInteractor* InInteractor)
+void ULearningAgentsController::SetupController(ULearningAgentsInteractor* InInteractor)
 {
 	if (IsSetup())
 	{
@@ -27,21 +25,12 @@ void ULearningAgentsController::SetupController(ALearningAgentsManager* InAgentM
 		return;
 	}
 
-	if (!InAgentManager)
+	if (!Manager)
 	{
-		UE_LOG(LogLearning, Error, TEXT("%s: InAgentManager is nullptr."), *GetName());
+		UE_LOG(LogLearning, Error, TEXT("%s: Must be attached to a LearningAgentsManager Actor."), *GetName());
 		return;
 	}
 
-	if (!InAgentManager->IsManagerSetup())
-	{
-		UE_LOG(LogLearning, Error, TEXT("%s: %s's SetupManager must be run before it can be used."), *GetName(), *InAgentManager->GetName());
-		return;
-	}
-
-	// This manager is not referenced in this class but we need it in blueprints to call GetAgent()
-	AgentManager = InAgentManager;
-
 	if (!InInteractor)
 	{
 		UE_LOG(LogLearning, Error, TEXT("%s: InInteractor is nullptr."), *GetName());
@@ -69,18 +58,72 @@ void ULearningAgentsController::EncodeActions()
 		return;
 	}
 
-	SetActions(AddedAgentIds);
+	const TLearningArrayView<1, uint64> ActionEncodingAgentIteration = Interactor->GetActionEncodingAgentIteration();
 
-	Interactor->GetActionFeature().Encode(AddedAgentSet);
+	// Run Set Actions Callback
 
-#if ENABLE_VISUAL_LOG
-	for (const ULearningAgentsAction* ActionObject : Interactor->GetActionObjects())
+	SetActions(Manager->GetAllAgentIds());
+
+	// Check that all actions have had their setter run
+
+	ValidAgentStatus.SetNumUninitialized(Manager->GetMaxAgentNum());
+	ValidAgentStatus.SetRange(0, Manager->GetMaxAgentNum(), true);
+
+	for (ULearningAgentsAction* ActionObject : Interactor->GetActionObjects())
 	{
-		if (ActionObject)
+		for (const int32 AgentId : Manager->GetAllAgentSet())
 		{
-			ActionObject->VisualLog(AddedAgentSet);
+			if (ActionObject->AgentSetIteration[AgentId] == ActionEncodingAgentIteration[AgentId])
+			{
+				UE_LOG(LogLearning, Warning, TEXT("%s: Action %s for agent with id %i has not been set (got iteration %i, expected iteration %i) and so agent will not have actions encoded."), *GetName(), *ActionObject->GetName(), AgentId, ActionObject->AgentSetIteration[AgentId], ActionEncodingAgentIteration[AgentId] + 1);
+				ValidAgentStatus[AgentId] = false;
+				continue;
+			}
+
+			if (ActionObject->AgentSetIteration[AgentId] > ActionEncodingAgentIteration[AgentId] + 1)
+			{
+				UE_LOG(LogLearning, Warning, TEXT("%s: Action %s for agent with id %i appears to have been set multiple times (got iteration %i, expected iteration %i) and so agent will not have actions encoded."), *GetName(), *ActionObject->GetName(), AgentId, ActionObject->AgentSetIteration[AgentId], ActionEncodingAgentIteration[AgentId] + 1);
+				ValidAgentStatus[AgentId] = false;
+				continue;
+			}
+
+			if (ActionObject->AgentSetIteration[AgentId] != ActionEncodingAgentIteration[AgentId] + 1)
+			{
+				UE_LOG(LogLearning, Warning, TEXT("%s: Action %s for agent with id %i does not have a matching iteration number (got iteration %i, expected iteration %i) and so agent will not have actions encoded."), *GetName(), *ActionObject->GetName(), AgentId, ActionObject->AgentSetIteration[AgentId], ActionEncodingAgentIteration[AgentId] + 1);
+				ValidAgentStatus[AgentId] = false;
+				continue;
+			}
 		}
 	}
+
+	ValidAgentIds.Empty(Manager->GetAgentNum());
+
+	for (const int32 AgentId : Manager->GetAllAgentSet())
+	{
+		if (ValidAgentStatus[AgentId]) { ValidAgentIds.Add(AgentId); }
+	}
+
+	ValidAgentSet = ValidAgentIds;
+	ValidAgentSet.TryMakeSlice();
+
+	// Encode Actions
+
+	Interactor->GetActionFeature().Encode(ValidAgentSet);
+
+	// Increment Action Encoding Iteration
+
+	for (const int32 AgentId : ValidAgentSet)
+	{
+		ActionEncodingAgentIteration[AgentId]++;
+	}
+
+	// Visual Logger
+
+#if UE_LEARNING_AGENTS_ENABLE_VISUAL_LOG
+	for (const ULearningAgentsAction* ActionObject : Interactor->GetActionObjects())
+	{
+		ActionObject->VisualLog(ValidAgentSet);
+	}
 #endif
 }
 
diff --git a/Engine/Plugins/Experimental/LearningAgents/Source/LearningAgents/Private/LearningAgentsCritic.cpp b/Engine/Plugins/Experimental/LearningAgents/Source/LearningAgents/Private/LearningAgentsCritic.cpp
--- a/Engine/Plugins/Experimental/LearningAgents/Source/LearningAgents/Private/LearningAgentsCritic.cpp
+++ b/Engine/Plugins/Experimental/LearningAgents/Source/LearningAgents/Private/LearningAgentsCritic.cpp
@@ -14,11 +14,14 @@
 #include "Misc/Paths.h"
 #include "GameFramework/Actor.h"
 
-ULearningAgentsCritic::ULearningAgentsCritic() : ULearningAgentsManagerComponent() {}
+ULearningAgentsCritic::ULearningAgentsCritic(const FObjectInitializer& ObjectInitializer) : ULearningAgentsManagerComponent(ObjectInitializer) {}
 ULearningAgentsCritic::ULearningAgentsCritic(FVTableHelper& Helper) : ULearningAgentsCritic() {}
 ULearningAgentsCritic::~ULearningAgentsCritic() {}
 
-void ULearningAgentsCritic::SetupCritic(ALearningAgentsManager* InAgentManager, ULearningAgentsInteractor* InInteractor, const FLearningAgentsCriticSettings& CriticSettings)
+void ULearningAgentsCritic::SetupCritic(
+	ULearningAgentsInteractor* InInteractor, 
+	const FLearningAgentsCriticSettings& CriticSettings,
+	ULearningAgentsNeuralNetwork* NeuralNetworkAsset)
 {
 	if (IsSetup())
 	{
@@ -26,20 +29,12 @@ void ULearningAgentsCritic::SetupCritic(ALearningAgentsManager* InAgentManager,
 		return;
 	}
 
-	if (!InAgentManager)
+	if (!Manager)
 	{
-		UE_LOG(LogLearning, Error, TEXT("%s: InAgentManager is nullptr."), *GetName());
+		UE_LOG(LogLearning, Error, TEXT("%s: Must be attached to a LearningAgentsManager Actor."), *GetName());
 		return;
 	}
 
-	if (!InAgentManager->IsManagerSetup())
-	{
-		UE_LOG(LogLearning, Error, TEXT("%s: %s's SetupManager must be run before it can be used."), *GetName(), *InAgentManager->GetName());
-		return;
-	}
-
-	AgentManager = InAgentManager;
-
 	if (!InInteractor)
 	{
 		UE_LOG(LogLearning, Error, TEXT("%s: InInteractor is nullptr."), *GetName());
@@ -54,28 +49,105 @@ void ULearningAgentsCritic::SetupCritic(ALearningAgentsManager* InAgentManager,
 
 	Interactor = InInteractor;
 
-	// Setup Neural Network
-	Network = NewObject<ULearningAgentsNeuralNetwork>(this, TEXT("CriticNetwork"));
-	Network->NeuralNetwork = MakeShared<UE::Learning::FNeuralNetwork>();
-	Network->NeuralNetwork->Resize(
-		Interactor->GetObservationFeature().DimNum(),
-		1,
-		CriticSettings.HiddenLayerSize,
-		CriticSettings.LayerNum);
-	Network->NeuralNetwork->ActivationFunction = UE::Learning::Agents::GetActivationFunction(CriticSettings.ActivationFunction);
+	if (NeuralNetworkAsset)
+	{
+		// Use Existing Neural Network Asset
+
+		if (NeuralNetworkAsset->NeuralNetwork)
+		{
+			if (NeuralNetworkAsset->NeuralNetwork->GetInputNum() != Interactor->GetObservationFeature().DimNum() ||
+				NeuralNetworkAsset->NeuralNetwork->GetOutputNum() != 2 * Interactor->GetActionFeature().DimNum())
+			{
+				UE_LOG(LogLearning, Error, TEXT("%s: Neural Network Asset provided during Setup is incorrect size: Inputs and outputs don't match."), *GetName());
+				return;
+			}
+
+			if (NeuralNetworkAsset->NeuralNetwork->GetHiddenNum() != CriticSettings.HiddenLayerSize ||
+				NeuralNetworkAsset->NeuralNetwork->GetLayerNum() != CriticSettings.LayerNum ||
+				NeuralNetworkAsset->NeuralNetwork->ActivationFunction != UE::Learning::Agents::GetActivationFunction(CriticSettings.ActivationFunction))
+			{
+				UE_LOG(LogLearning, Warning, TEXT("%s: Neural Network Asset settings don't match those given by CriticSettings"), *GetName());
+			}
+
+			Network = NeuralNetworkAsset;
+		}
+		else
+		{
+			Network = NeuralNetworkAsset;
+			Network->NeuralNetwork = MakeShared<UE::Learning::FNeuralNetwork>();
+			Network->NeuralNetwork->Resize(
+				Interactor->GetObservationFeature().DimNum(),
+				1,
+				CriticSettings.HiddenLayerSize,
+				CriticSettings.LayerNum);
+			Network->NeuralNetwork->ActivationFunction = UE::Learning::Agents::GetActivationFunction(CriticSettings.ActivationFunction);
+		}
+	}
+	else
+	{
+		// Create New Neural Network Asset
+
+		Network = NewObject<ULearningAgentsNeuralNetwork>(this, TEXT("CriticNetwork"));
+		Network->NeuralNetwork = MakeShared<UE::Learning::FNeuralNetwork>();
+		Network->NeuralNetwork->Resize(
+			Interactor->GetObservationFeature().DimNum(),
+			1,
+			CriticSettings.HiddenLayerSize,
+			CriticSettings.LayerNum);
+		Network->NeuralNetwork->ActivationFunction = UE::Learning::Agents::GetActivationFunction(CriticSettings.ActivationFunction);
+	}
 
 	// Create Critic Object
 	CriticObject = MakeShared<UE::Learning::FNeuralNetworkCriticFunction>(
 		TEXT("CriticObject"),
-		AgentManager->GetInstanceData().ToSharedRef(),
-		AgentManager->GetMaxInstanceNum(),
+		Manager->GetInstanceData().ToSharedRef(),
+		Manager->GetMaxAgentNum(),
 		Network->NeuralNetwork.ToSharedRef());
 
-	AgentManager->GetInstanceData()->Link(Interactor->GetObservationFeature().FeatureHandle, CriticObject->InputHandle);
+	Manager->GetInstanceData()->Link(Interactor->GetObservationFeature().FeatureHandle, CriticObject->InputHandle);
+
+	DiscountedReturnAgentIteration.SetNumUninitialized({ Manager->GetMaxAgentNum() });
+	UE::Learning::Array::Set<1, uint64>(DiscountedReturnAgentIteration, INDEX_NONE);
+	UE::Learning::Array::Set<1, uint64>(DiscountedReturnAgentIteration, 0, Manager->GetAllAgentSet());
 
 	bIsSetup = true;
 }
 
+void ULearningAgentsCritic::OnAgentsAdded(const TArray<int32>& AgentIds)
+{
+	if (IsSetup())
+	{
+		UE::Learning::FIndexSet AgentSet = AgentIds;
+		AgentSet.TryMakeSlice();
+		UE::Learning::Array::Set<1, uint64>(DiscountedReturnAgentIteration, 0, AgentSet);
+	}
+}
+
+void ULearningAgentsCritic::OnAgentsRemoved(const TArray<int32>& AgentIds)
+{
+	if (IsSetup())
+	{
+		UE::Learning::FIndexSet AgentSet = AgentIds;
+		AgentSet.TryMakeSlice();
+		UE::Learning::Array::Set<1, uint64>(DiscountedReturnAgentIteration, INDEX_NONE, AgentSet);
+	}
+}
+
+void ULearningAgentsCritic::OnAgentsReset(const TArray<int32>& AgentIds)
+{
+	if (IsSetup())
+	{
+		UE::Learning::FIndexSet AgentSet = AgentIds;
+		AgentSet.TryMakeSlice();
+		UE::Learning::Array::Set<1, uint64>(DiscountedReturnAgentIteration, 0, AgentSet);
+	}
+}
+
+ULearningAgentsNeuralNetwork* ULearningAgentsCritic::GetNetworkAsset()
+{
+	return Network;
+}
+
 UE::Learning::FNeuralNetwork& ULearningAgentsCritic::GetCriticNetwork()
 {
 	return *Network->NeuralNetwork;
@@ -108,7 +180,7 @@ void ULearningAgentsCritic::SaveCriticToSnapshot(const FFilePath& File) const
 	Network->SaveNetworkToSnapshot(File);
 }
 
-void ULearningAgentsCritic::LoadCriticFromAsset(const ULearningAgentsNeuralNetwork* NeuralNetworkAsset)
+void ULearningAgentsCritic::UseCriticFromAsset(ULearningAgentsNeuralNetwork* NeuralNetworkAsset)
 {
 	if (!IsSetup())
 	{
@@ -122,17 +194,26 @@ void ULearningAgentsCritic::LoadCriticFromAsset(const ULearningAgentsNeuralNetwo
 		return;
 	}
 
+	if (NeuralNetworkAsset == Network)
+	{
+		UE_LOG(LogLearning, Error, TEXT("%s: Asset is same as the current network."), *GetName());
+		return;
+	}
+
 	if (NeuralNetworkAsset->NeuralNetwork->GetInputNum() != Network->NeuralNetwork->GetInputNum() ||
-		NeuralNetworkAsset->NeuralNetwork->GetOutputNum() != Network->NeuralNetwork->GetOutputNum())
+		NeuralNetworkAsset->NeuralNetwork->GetOutputNum() != Network->NeuralNetwork->GetOutputNum() ||
+		NeuralNetworkAsset->NeuralNetwork->GetLayerNum() != Network->NeuralNetwork->GetLayerNum() ||
+		NeuralNetworkAsset->NeuralNetwork->ActivationFunction != Network->NeuralNetwork->ActivationFunction)
 	{
-		UE_LOG(LogLearning, Error, TEXT("%s: Failed to load network from asset. Inputs and outputs don't match."), *GetName());
+		UE_LOG(LogLearning, Error, TEXT("%s: Failed to use asset as network settings don't match."), *GetName());
 		return;
 	}
 
-	*Network->NeuralNetwork = *NeuralNetworkAsset->NeuralNetwork;
+	Network = NeuralNetworkAsset;
+	CriticObject->NeuralNetwork = Network->NeuralNetwork.ToSharedRef();
 }
 
-void ULearningAgentsCritic::SaveCriticToAsset(ULearningAgentsNeuralNetwork* NeuralNetworkAsset) const
+void ULearningAgentsCritic::LoadCriticFromAsset(ULearningAgentsNeuralNetwork* NeuralNetworkAsset)
 {
 	if (!IsSetup())
 	{
@@ -140,49 +221,63 @@ void ULearningAgentsCritic::SaveCriticToAsset(ULearningAgentsNeuralNetwork* Neur
 		return;
 	}
 
-	if (!NeuralNetworkAsset)
+	Network->LoadNetworkFromAsset(NeuralNetworkAsset);
+}
+
+void ULearningAgentsCritic::SaveCriticToAsset(ULearningAgentsNeuralNetwork* NeuralNetworkAsset)
+{
+	if (!IsSetup())
 	{
-		UE_LOG(LogLearning, Error, TEXT("%s: Asset is invalid."), *GetName());
+		UE_LOG(LogLearning, Error, TEXT("%s: Setup not complete."), *GetName());
 		return;
 	}
 
-	if (!NeuralNetworkAsset->NeuralNetwork)
+	Network->SaveNetworkToAsset(NeuralNetworkAsset);
+}
+
+void ULearningAgentsCritic::EvaluateCritic()
+{
+	UE_LEARNING_TRACE_CPUPROFILER_EVENT_SCOPE(ULearningAgentsCritic::EvaluateCritic);
+
+	if (!IsSetup())
 	{
-		NeuralNetworkAsset->NeuralNetwork = MakeShared<UE::Learning::FNeuralNetwork>();
+		UE_LOG(LogLearning, Error, TEXT("%s: Setup not complete."), *GetName());
+		return;
 	}
 
-	*NeuralNetworkAsset->NeuralNetwork = *Network->NeuralNetwork;
+	// Check Agents actually have encoded observations.
 
-	// Manually mark the package as dirty since just using `Modify` prevents 
-	// marking packages as dirty during PIE which is most likely when this
-	// is being used.
-	if (UPackage* Package = NeuralNetworkAsset->GetPackage())
-	{
-		const bool bIsDirty = Package->IsDirty();
+	ValidAgentIds.Empty(Manager->GetAgentNum());
 
-		if (!bIsDirty)
+	for (const int32 AgentId : Manager->GetAllAgentSet())
+	{
+		if (Interactor->GetObservationEncodingAgentIteration()[AgentId] == 0)
 		{
-			Package->SetDirtyFlag(true);
+			UE_LOG(LogLearning, Warning, TEXT("%s: Agent with id %i has not made observations so critic will not be evaluated for it."), *GetName(), AgentId);
+			continue;
 		}
 
-		Package->PackageMarkedDirtyEvent.Broadcast(Package, bIsDirty);
+		ValidAgentIds.Add(AgentId);
 	}
-}
 
-void ULearningAgentsCritic::EvaluateCritic()
-{
-	UE_LEARNING_TRACE_CPUPROFILER_EVENT_SCOPE(ULearningAgentsCritic::EvaluateCritic);
+	ValidAgentSet = ValidAgentIds;
+	ValidAgentSet.TryMakeSlice();
 
-	if (!IsSetup())
+	// Evaluate Critic
+
+	CriticObject->Evaluate(ValidAgentSet);
+
+	// Increment Discounted Return Iteration
+
+	for (const int32 AgentId : ValidAgentSet)
 	{
-		UE_LOG(LogLearning, Error, TEXT("%s: Setup not complete."), *GetName());
-		return;
+		DiscountedReturnAgentIteration[AgentId]++;
 	}
-
-	CriticObject->Evaluate(AddedAgentSet);
+		
+	// Visual Logger
 
 #if UE_LEARNING_AGENTS_ENABLE_VISUAL_LOG
-	VisualLog(AddedAgentSet);
+	VisualLog(ValidAgentSet);
 #endif
 }
 
@@ -200,31 +295,35 @@ float ULearningAgentsCritic::GetEstimatedDiscountedReturn(const int32 AgentId) c
 		return 0.0f;
 	}
 
-	const TLearningArrayView<1, const float> CriticOutputView = CriticObject->InstanceData->ConstView(CriticObject->OutputHandle);
+	if (DiscountedReturnAgentIteration[AgentId] == 0)
+	{
+		UE_LOG(LogLearning, Error, TEXT("%s: Agent with id %d has not yet computed the estimated discounted return. Did you run EvaluateCritic?"), *GetName(), AgentId);
+		return 0.0f;
+	}
 
-	return CriticOutputView[AgentId];
+	return CriticObject->InstanceData->ConstView(CriticObject->OutputHandle)[AgentId];
 }
 
 #if UE_LEARNING_AGENTS_ENABLE_VISUAL_LOG
-void ULearningAgentsCritic::VisualLog(const UE::Learning::FIndexSet Instances) const
+void ULearningAgentsCritic::VisualLog(const UE::Learning::FIndexSet AgentSet) const
 {
 	UE_LEARNING_TRACE_CPUPROFILER_EVENT_SCOPE(ULearningAgentsCritic::VisualLog);
 
 	const TLearningArrayView<2, const float> InputView = CriticObject->InstanceData->ConstView(CriticObject->InputHandle);
 	const TLearningArrayView<1, const float> OutputView = CriticObject->InstanceData->ConstView(CriticObject->OutputHandle);
 
-	for (const int32 Instance : Instances)
+	for (const int32 AgentId : AgentSet)
 	{
-		if (const AActor* Actor = Cast<AActor>(Interactor->GetAgent(Instance)))
+		if (const AActor* Actor = Cast<AActor>(GetAgent(AgentId)))
 		{
 			UE_LEARNING_AGENTS_VLOG_STRING(this, LogLearning, Display,
 				Actor->GetActorLocation(),
 				VisualLogColor.ToFColor(true),
 				TEXT("Agent %i\nInput: %s\nInput Stats (Min/Max/Mean/Std): %s\nOutput: [% 6.3f]"),
-				Instance,
-				*UE::Learning::Array::FormatFloat(InputView[Instance]),
-				*UE::Learning::Agents::Debug::FloatArrayToStatsString(InputView[Instance]),
-				OutputView[Instance]);
+				AgentId,
+				*UE::Learning::Array::FormatFloat(InputView[AgentId]),
+				*UE::Learning::Agents::Debug::FloatArrayToStatsString(InputView[AgentId]),
+				OutputView[AgentId]);
 		}
 	}
 }
diff --git a/Engine/Plugins/Experimental/LearningAgents/Source/LearningAgents/Private/LearningAgentsHelpers.cpp b/Engine/Plugins/Experimental/LearningAgents/Source/LearningAgents/Private/LearningAgentsHelpers.cpp
--- a/Engine/Plugins/Experimental/LearningAgents/Source/LearningAgents/Private/LearningAgentsHelpers.cpp
+++ b/Engine/Plugins/Experimental/LearningAgents/Source/LearningAgents/Private/LearningAgentsHelpers.cpp
@@ -581,7 +581,7 @@ void UMeshComponentHelper::GetMeshBonePositions(TArray<FVector>& OutBonePosition
 		}
 	}
 
-#if ENABLE_VISUAL_LOG
+#if UE_LEARNING_AGENTS_ENABLE_VISUAL_LOG
 	{
 		for (int32 BoneIdx = 0; BoneIdx < BoneNum; BoneIdx++)
 		{
@@ -678,7 +678,7 @@ void URayCastHelper::RayCastGridHeights(
 		}
 	}
 
-#if ENABLE_VISUAL_LOG
+#if UE_LEARNING_AGENTS_ENABLE_VISUAL_LOG
 	{
 		UE_LEARNING_AGENTS_VLOG_TRANSFORM(this, LogLearning, Display,
 			Position,
diff --git a/Engine/Plugins/Experimental/LearningAgents/Source/LearningAgents/Private/LearningAgentsInteractor.cpp b/Engine/Plugins/Experimental/LearningAgents/Source/LearningAgents/Private/LearningAgentsInteractor.cpp
--- a/Engine/Plugins/Experimental/LearningAgents/Source/LearningAgents/Private/LearningAgentsInteractor.cpp
+++ b/Engine/Plugins/Experimental/LearningAgents/Source/LearningAgents/Private/LearningAgentsInteractor.cpp
@@ -12,11 +12,11 @@
 #include "LearningLog.h"
 #include "EngineDefines.h"
 
-ULearningAgentsInteractor::ULearningAgentsInteractor() : ULearningAgentsManagerComponent() {}
+ULearningAgentsInteractor::ULearningAgentsInteractor(const FObjectInitializer& ObjectInitializer) : ULearningAgentsManagerComponent(ObjectInitializer) {}
 ULearningAgentsInteractor::ULearningAgentsInteractor(FVTableHelper& Helper) : ULearningAgentsInteractor() {}
 ULearningAgentsInteractor::~ULearningAgentsInteractor() {}
 
-void ULearningAgentsInteractor::SetupInteractor(ALearningAgentsManager* InAgentManager)
+void ULearningAgentsInteractor::SetupInteractor()
 {
 	if (IsSetup())
 	{
@@ -24,28 +24,16 @@ void ULearningAgentsInteractor::SetupInteractor(ALearningAgentsManager* InAgentM
 		return;
 	}
 
-	if (!InAgentManager)
+	if (!Manager)
 	{
-		UE_LOG(LogLearning, Error, TEXT("%s: InAgentManager is nullptr."), *GetName());
+		UE_LOG(LogLearning, Error, TEXT("%s: Must be attached to a LearningAgentsManager Actor."), *GetName());
 		return;
 	}
 
-	if (!InAgentManager->IsManagerSetup())
-	{
-		UE_LOG(LogLearning, Error, TEXT("%s: %s's SetupManager must be run before it can be used."), *GetName(), *InAgentManager->GetName());
-		return;
-	}
-
-	AgentManager = InAgentManager;
-
 	// Setup Observations
 	ObservationObjects.Empty();
 	ObservationFeatures.Empty();
-	SetupObservations(this);
-	Observations = MakeShared<UE::Learning::FConcatenateFeature>(TEXT("Observations"),
-		TLearningArrayView<1, const TSharedRef<UE::Learning::FFeatureObject>>(ObservationFeatures),
-		InAgentManager->GetInstanceData().ToSharedRef(),
-		InAgentManager->GetMaxInstanceNum());
+	SetupObservations();
 
 	if (ObservationObjects.Num() == 0)
 	{
@@ -53,6 +41,11 @@ void ULearningAgentsInteractor::SetupInteractor(ALearningAgentsManager* InAgentM
 		return;
 	}
 
+	Observations = MakeShared<UE::Learning::FConcatenateFeature>(*(GetName() + TEXT("Observations")),
+		TLearningArrayView<1, const TSharedRef<UE::Learning::FFeatureObject>>(ObservationFeatures),
+		Manager->GetInstanceData().ToSharedRef(),
+		Manager->GetMaxAgentNum());
+
 	if (Observations->DimNum() == 0)
 	{
 		UE_LOG(LogLearning, Error, TEXT("%s: Observation vector is zero-sized - all added observations have no size."), *GetName());
@@ -62,11 +55,7 @@ void ULearningAgentsInteractor::SetupInteractor(ALearningAgentsManager* InAgentM
 	// Setup Actions
 	ActionObjects.Empty();
 	ActionFeatures.Empty();
-	SetupActions(this);
-	Actions = MakeShared<UE::Learning::FConcatenateFeature>(TEXT("Actions"),
-		TLearningArrayView<1, const TSharedRef<UE::Learning::FFeatureObject>>(ActionFeatures),
-		InAgentManager->GetInstanceData().ToSharedRef(),
-		InAgentManager->GetMaxInstanceNum());
+	SetupActions();
 
 	if (ActionObjects.Num() == 0)
 	{
@@ -74,15 +63,75 @@ void ULearningAgentsInteractor::SetupInteractor(ALearningAgentsManager* InAgentM
 		return;
 	}
 
+	Actions = MakeShared<UE::Learning::FConcatenateFeature>(*(GetName() + TEXT("Actions")),
+		TLearningArrayView<1, const TSharedRef<UE::Learning::FFeatureObject>>(ActionFeatures),
+		Manager->GetInstanceData().ToSharedRef(),
+		Manager->GetMaxAgentNum());
+
 	if (Actions->DimNum() == 0)
 	{
 		UE_LOG(LogLearning, Error, TEXT("%s: Action vector is zero-sized -  all added actions have no size."), *GetName());
 		return;
 	}
 
+	// Reset Agent iteration
+
+	ObservationEncodingAgentIteration.SetNumUninitialized({ Manager->GetMaxAgentNum() });
+	ActionEncodingAgentIteration.SetNumUninitialized({ Manager->GetMaxAgentNum() });
+	UE::Learning::Array::Set<1, uint64>(ObservationEncodingAgentIteration, INDEX_NONE);
+	UE::Learning::Array::Set<1, uint64>(ActionEncodingAgentIteration, INDEX_NONE);
+	SetAgentIteration(Manager->GetAllAgentSet(), 0);
+
 	bIsSetup = true;
 }
 
+void ULearningAgentsInteractor::OnAgentsAdded(const TArray<int32>& AgentIds)
+{
+	if (IsSetup())
+	{
+		UE::Learning::FIndexSet AgentSet = AgentIds;
+		AgentSet.TryMakeSlice();
+		SetAgentIteration(AgentSet, 0);
+	}
+}
+
+void ULearningAgentsInteractor::OnAgentsRemoved(const TArray<int32>& AgentIds)
+{
+	if (IsSetup())
+	{
+		UE::Learning::FIndexSet AgentSet = AgentIds;
+		AgentSet.TryMakeSlice();
+		SetAgentIteration(AgentSet, INDEX_NONE);
+	}
+}
+
+void ULearningAgentsInteractor::OnAgentsReset(const TArray<int32>& AgentIds)
+{
+	if (IsSetup())
+	{
+		UE::Learning::FIndexSet AgentSet = AgentIds;
+		AgentSet.TryMakeSlice();
+		SetAgentIteration(AgentSet, 0);
+	}
+}
+
+void ULearningAgentsInteractor::SetAgentIteration(const UE::Learning::FIndexSet AgentIds, const uint64 Value)
+{
+	UE::Learning::Array::Set<1, uint64>(ObservationEncodingAgentIteration, Value, AgentIds);
+	UE::Learning::Array::Set<1, uint64>(ActionEncodingAgentIteration, Value, AgentIds);
+
+	for (ULearningAgentsObservation* ObservationObject : ObservationObjects)
+	{
+		UE::Learning::Array::Set<1, uint64>(ObservationObject->AgentIteration, Value, AgentIds);
+	}
+
+	for (ULearningAgentsAction* ActionObject : ActionObjects)
+	{
+		UE::Learning::Array::Set<1, uint64>(ActionObject->AgentGetIteration, Value, AgentIds);
+		UE::Learning::Array::Set<1, uint64>(ActionObject->AgentSetIteration, Value, AgentIds);
+	}
+}
+
 UE::Learning::FFeatureObject& ULearningAgentsInteractor::GetObservationFeature() const
 {
 	return *Observations;
@@ -103,7 +152,17 @@ TConstArrayView<ULearningAgentsAction*> ULearningAgentsInteractor::GetActionObje
 	return ActionObjects;
 }
 
-void ULearningAgentsInteractor::SetupObservations_Implementation(ULearningAgentsInteractor* Interactor)
+TLearningArrayView<1, uint64> ULearningAgentsInteractor::GetObservationEncodingAgentIteration()
+{
+	return ObservationEncodingAgentIteration;
+}
+
+TLearningArrayView<1, uint64> ULearningAgentsInteractor::GetActionEncodingAgentIteration()
+{
+	return ActionEncodingAgentIteration;
+}
+
+void ULearningAgentsInteractor::SetupObservations_Implementation()
 {
 	// Can be overridden to setup observations without blueprints
 }
@@ -116,11 +175,12 @@ void ULearningAgentsInteractor::SetObservations_Implementation(const TArray<int3
 void ULearningAgentsInteractor::AddObservation(TObjectPtr<ULearningAgentsObservation> Object, const TSharedRef<UE::Learning::FFeatureObject>& Feature)
 {
 	UE_LEARNING_CHECK(!IsSetup());
+	UE_LEARNING_CHECK(Object);
 	ObservationObjects.Add(Object);
 	ObservationFeatures.Add(Feature);
 }
 
-void ULearningAgentsInteractor::SetupActions_Implementation(ULearningAgentsInteractor* Interactor)
+void ULearningAgentsInteractor::SetupActions_Implementation()
 {
 	// Can be overridden to setup actions without blueprints
 }
@@ -133,11 +193,12 @@ void ULearningAgentsInteractor::GetActions_Implementation(const TArray<int32>& A
 void ULearningAgentsInteractor::AddAction(TObjectPtr<ULearningAgentsAction> Object, const TSharedRef<UE::Learning::FFeatureObject>& Feature)
 {
 	UE_LEARNING_CHECK(!IsSetup());
+	UE_LEARNING_CHECK(Object);
 	ActionObjects.Add(Object);
 	ActionFeatures.Add(Feature);
 }
 
-void ULearningAgentsInteractor::EncodeObservations()
+void ULearningAgentsInteractor::EncodeObservations(const UE::Learning::FIndexSet AgentSet)
 {
 	UE_LEARNING_TRACE_CPUPROFILER_EVENT_SCOPE(ULearningAgentsInteractor::EncodeObservations);
 
@@ -147,22 +208,80 @@ void ULearningAgentsInteractor::EncodeObservations()
 		return;
 	}
 
-	SetObservations(AddedAgentIds);
+	// Run Set Observations Callback
 
-	Observations->Encode(AddedAgentSet);
+	ValidAgentIds.Empty(AgentSet.Num());
+	for (const int32 AgentId : AgentSet)
+	{
+		ValidAgentIds.Add(AgentId);
+	}
 
-#if ENABLE_VISUAL_LOG
-	for (const ULearningAgentsObservation* ObservationObject : ObservationObjects)
+	SetObservations(ValidAgentIds);
+
+	// Check that all observations have had their setter run
+
+	ValidAgentStatus.SetNumUninitialized(Manager->GetMaxAgentNum());
+	ValidAgentStatus.SetRange(0, Manager->GetMaxAgentNum(), true);
+
+	for (ULearningAgentsObservation* ObservationObject : ObservationObjects)
 	{
-		if (ObservationObject)
+		for (const int32 AgentId : AgentSet)
 		{
-			ObservationObject->VisualLog(AddedAgentSet);
+			if (ObservationObject->AgentIteration[AgentId] <= ObservationEncodingAgentIteration[AgentId])
+			{
+				UE_LOG(LogLearning, Warning, TEXT("%s: Observation %s for agent with id %i has not been set (got iteration %i, expected iteration %i) and so agent will not have observations encoded."), *GetName(), *ObservationObject->GetName(), AgentId, ObservationObject->AgentIteration[AgentId], ObservationEncodingAgentIteration[AgentId] + 1);
+				ValidAgentStatus[AgentId] = false;
+				continue;
+			}
+
+			if (ObservationObject->AgentIteration[AgentId] > ObservationEncodingAgentIteration[AgentId] + 1)
+			{
+				UE_LOG(LogLearning, Warning, TEXT("%s: Observation %s for agent with id %i appears to have been set multiple times (got iteration %i, expected iteration %i) and so agent will not have observations encoded."), *GetName(), *ObservationObject->GetName(), AgentId, ObservationObject->AgentIteration[AgentId], ObservationEncodingAgentIteration[AgentId] + 1);
+				ValidAgentStatus[AgentId] = false;
+				continue;
+			}
+
+			if (ObservationObject->AgentIteration[AgentId] != ObservationEncodingAgentIteration[AgentId] + 1)
+			{
+				UE_LOG(LogLearning, Warning, TEXT("%s: Observation %s for agent with id %i does not have a matching iteration number (got iteration %i, expected iteration %i) and so agent will not have observations encoded."), *GetName(), *ObservationObject->GetName(), AgentId, ObservationObject->AgentIteration[AgentId], ObservationEncodingAgentIteration[AgentId] + 1);
+				ValidAgentStatus[AgentId] = false;
+				continue;
+			}
 		}
 	}
+
+	ValidAgentIds.Empty(Manager->GetAgentNum());
+
+	for (const int32 AgentId : AgentSet)
+	{
+		if (ValidAgentStatus[AgentId]) { ValidAgentIds.Add(AgentId); }
+	}
+
+	ValidAgentSet = ValidAgentIds;
+	ValidAgentSet.TryMakeSlice();
+
+	// Encode Observations
+
+	Observations->Encode(ValidAgentSet);
+
+	// Increment Observation Encoding Iteration
+
+	for (const int32 AgentId : ValidAgentSet)
+	{
+		ObservationEncodingAgentIteration[AgentId]++;
+	}
+
+	// Visual Logger
+
+#if UE_LEARNING_AGENTS_ENABLE_VISUAL_LOG
+	for (const ULearningAgentsObservation* ObservationObject : ObservationObjects)
+	{
+		ObservationObject->VisualLog(ValidAgentSet);
+	}
 #endif
 }
 
-void ULearningAgentsInteractor::DecodeActions()
+void ULearningAgentsInteractor::DecodeActions(const UE::Learning::FIndexSet AgentSet)
 {
 	UE_LEARNING_TRACE_CPUPROFILER_EVENT_SCOPE(ULearningAgentsInteractor::DecodeActions);
 
@@ -172,20 +291,77 @@ void ULearningAgentsInteractor::DecodeActions()
 		return;
 	}
 
-	Actions->Decode(AddedAgentSet);
+	// Check Agents actually have encoded actions.
 
-	GetActions(AddedAgentIds);
+	ValidAgentIds.Empty(Manager->GetAgentNum());
 
-#if ENABLE_VISUAL_LOG
-	for (const ULearningAgentsAction* ActionObject : ActionObjects)
+	for (const int32 AgentId : AgentSet)
 	{
-		if (ActionObject)
+		if (ActionEncodingAgentIteration[AgentId] == 0)
 		{
-			ActionObject->VisualLog(AddedAgentSet);
+			UE_LOG(LogLearning, Warning, TEXT("%s: Agent with id %i does not have an encoded action vector so actions will not be decoded for it. Was EvaluatePolicy or EncodeActions run?"), *GetName(), AgentId);
+			continue;
 		}
+
+		ValidAgentIds.Add(AgentId);
+	}
+
+	ValidAgentSet = ValidAgentIds;
+	ValidAgentSet.TryMakeSlice();
+
+	// Decode Actions
+
+	Actions->Decode(ValidAgentSet);
+
+	// Run Get Actions Callback
+
+	GetActions(ValidAgentIds);
+
+	// Check that all actions have had their getter run
+
+	for (ULearningAgentsAction* ActionObject : ActionObjects)
+	{
+		for (const int32 AgentId : ValidAgentSet)
+		{
+			if (ActionObject->AgentGetIteration[AgentId] == ActionEncodingAgentIteration[AgentId] - 1)
+			{
+				UE_LOG(LogLearning, Warning, TEXT("%s: Action %s for agent with id %i appears to have not been got."), *GetName(), *ActionObject->GetName(), AgentId);
+				continue;
+			}
+				
+			if (ActionObject->AgentGetIteration[AgentId] > ActionEncodingAgentIteration[AgentId])
+			{
+				UE_LOG(LogLearning, Warning, TEXT("%s: Action %s for agent with id %i appears to have been got multiple times."), *GetName(), *ActionObject->GetName(), AgentId);
+				continue;
+			}
+				
+			if (ActionObject->AgentGetIteration[AgentId] != ActionEncodingAgentIteration[AgentId])
+			{
+				UE_LOG(LogLearning, Warning, TEXT("%s: Action %s for agent with id %i does not have a matching iteration number (got %i, expected %i)."), *GetName(), *ActionObject->GetName(), AgentId, ActionObject->AgentGetIteration[AgentId], ActionEncodingAgentIteration[AgentId]);
+				continue;
+			}
+		}
+	}
+
+	// Visual Logger
+
+#if UE_LEARNING_AGENTS_ENABLE_VISUAL_LOG
+	for (const ULearningAgentsAction* ActionObject : ActionObjects)
+	{
+		ActionObject->VisualLog(ValidAgentSet);
 	}
 #endif
 }
+
+void ULearningAgentsInteractor::EncodeObservations()
+{
+	EncodeObservations(Manager->GetAllAgentSet());
+}
+
+void ULearningAgentsInteractor::DecodeActions()
+{
+	DecodeActions(Manager->GetAllAgentSet());
+}
 void ULearningAgentsInteractor::GetObservationVector(const int32 AgentId, TArray<float>& OutObservationVector) const
 {
 	if (!IsSetup())
@@ -202,6 +378,13 @@ void ULearningAgentsInteractor::GetObservationVector(const int32 AgentId, TArray
 		return;
 	}
 
+	if (ObservationEncodingAgentIteration[AgentId] == 0)
+	{
+		UE_LOG(LogLearning, Error, TEXT("%s: Agent with id %d has not yet computed an observation vector. Have you run EncodeObservations?"), *GetName(), AgentId);
+		OutObservationVector.Empty();
+		return;
+	}
+
 	OutObservationVector.SetNumUninitialized(Observations->DimNum());
 	UE::Learning::Array::Copy<1, float>(OutObservationVector, Observations->FeatureBuffer()[AgentId]);
 }
@@ -222,6 +405,13 @@ void ULearningAgentsInteractor::GetActionVector(const int32 AgentId, TArray<floa
 		return;
 	}
 
+	if (ActionEncodingAgentIteration[AgentId] == 0)
+	{
+		UE_LOG(LogLearning, Error, TEXT("%s: Agent with id %d has not yet computed an action vector. Have you run EvaluatePolicy or EncodeActions?"), *GetName(), AgentId);
+		OutActionVector.Empty();
+		return;
+	}
+
 	OutActionVector.SetNumUninitialized(Actions->DimNum());
 	UE::Learning::Array::Copy<1, float>(OutActionVector, Actions->FeatureBuffer()[AgentId]);
 }
\ No newline at end of file
diff --git a/Engine/Plugins/Experimental/LearningAgents/Source/LearningAgents/Private/LearningAgentsManager.cpp b/Engine/Plugins/Experimental/LearningAgents/Source/LearningAgents/Private/LearningAgentsManager.cpp
--- a/Engine/Plugins/Experimental/LearningAgents/Source/LearningAgents/Private/LearningAgentsManager.cpp
+++ b/Engine/Plugins/Experimental/LearningAgents/Source/LearningAgents/Private/LearningAgentsManager.cpp
@@ -10,7 +10,7 @@
 #include "UObject/ScriptInterface.h"
 #include "EngineDefines.h"
 
-ALearningAgentsManager::ALearningAgentsManager() : AActor() {}
+ALearningAgentsManager::ALearningAgentsManager(const FObjectInitializer& ObjectInitializer) : AActor(ObjectInitializer) {}
 ALearningAgentsManager::ALearningAgentsManager(FVTableHelper& Helper) : ALearningAgentsManager() {}
 ALearningAgentsManager::~ALearningAgentsManager() {}
 
@@ -23,10 +23,15 @@ void ALearningAgentsManager::PostInitProperties()
 		return;
 	}
 
+	// Allocate Instance Data
+	InstanceData = MakeShared<UE::Learning::FArrayMap>();
+
 	// Pre-populate the vacant ids
-	OccupiedAgentIds.Reserve(MaxInstanceNum);
-	VacantAgentIds.Reserve(MaxInstanceNum);
-	for (int32 i = MaxInstanceNum - 1; i >= 0; i--)
+	OnEventAgentIds.Reserve(MaxAgentNum);
+	OccupiedAgentIds.Reserve(MaxAgentNum);
+	VacantAgentIds.Reserve(MaxAgentNum);
+	Agents.Reserve(MaxAgentNum);
+	for (int32 i = MaxAgentNum - 1; i >= 0; i--)
 	{
 		VacantAgentIds.Push(i);
 		Agents.Add(nullptr);
@@ -35,50 +40,9 @@ void ALearningAgentsManager::PostInitProperties()
 	UpdateAgentSets();
 }
 
-void ALearningAgentsManager::SetupManager()
+int32 ALearningAgentsManager::GetMaxAgentNum() const
 {
-	if (IsManagerSetup())
-	{
-		UE_LOG(LogLearning, Error, TEXT("%s: Setup already run!"), *GetName());
-		return;
-	}
-
-	InstanceData = MakeShared<UE::Learning::FArrayMap>();
-
-	GetComponents<ULearningAgentsManagerComponent>(CachedManagerComponents);
-	if (CachedManagerComponents.IsEmpty())
-	{
-		UE_LOG(LogLearning, Warning, TEXT("%s: Found zero attached manager components during setup."), *GetName());
-	}
-	else
-	{
-		for (const ULearningAgentsManagerComponent* ManagerComponent : CachedManagerComponents)
-		{
-			UE_LOG(LogLearning, Display, TEXT("Added %s manager component to %s's cache."), *ManagerComponent->GetName(), *GetName());
-		}
-	}
-
-	bIsSetup = true;
-}
-
-bool ALearningAgentsManager::IsManagerSetup() const
-{
-	return bIsSetup;
-}
-
-int32 ALearningAgentsManager::GetMaxInstanceNum() const
-{
-	return MaxInstanceNum;
-}
-
-const TSharedPtr<UE::Learning::FArrayMap>& ALearningAgentsManager::GetInstanceData() const
-{
-	return InstanceData;
-}
-
-TConstArrayView<TObjectPtr<UObject>> ALearningAgentsManager::GetAgents() const
-{
-	return Agents;
+	return MaxAgentNum;
 }
 
 int32 ALearningAgentsManager::AddAgent(UObject* Agent)
@@ -91,39 +55,86 @@ int32 ALearningAgentsManager::AddAgent(UObject* Agent)
 
 	if (VacantAgentIds.IsEmpty())
 	{
-		UE_LOG(LogLearning, Error,
-			TEXT("%s: Attempting to add an agent but we have no more vacant ids. Increase MaxInstanceNum (%d) or remove unused agents."),
-			*GetName(),
-			MaxInstanceNum);
+		UE_LOG(LogLearning, Error, TEXT("%s: Attempting to add an agent but we have no more vacant ids. Increase MaxAgentNum (%d) or remove unused agents."), *GetName(), MaxAgentNum);
 		return INDEX_NONE;
 	}
 
 	// Add Agent
-	const int32 NewAgentId = VacantAgentIds.Pop();
-	Agents[NewAgentId] = Agent;
-
-	OccupiedAgentIds.Add(NewAgentId);
 
+	const int32 AgentId = VacantAgentIds.Pop();
+	Agents[AgentId] = Agent;
+	OccupiedAgentIds.Add(AgentId);
 	UpdateAgentSets();
 
-	OnAgentAdded(NewAgentId);
+	OnEventAgentIds.SetNumUninitialized(0, false);
+	OnEventAgentIds.Add(AgentId);
+
+	// Call OnAgentsAdded Event
 
-	return NewAgentId;
+	UE_LOG(LogLearning, Display, TEXT("%s: Adding Agent %s with id %i."), *GetName(), *Agent->GetName(), AgentId);
+
+	TInlineComponentArray<ULearningAgentsManagerComponent*> ManagerComponents;
+	GetComponents<ULearningAgentsManagerComponent>(ManagerComponents);
+
+	for (ULearningAgentsManagerComponent* ManagerComponent : ManagerComponents)
+	{
+		ManagerComponent->OnAgentsAdded(OnEventAgentIds);
+	}
+
+	// Return new id
+
+	return AgentId;
 }
 
-void ALearningAgentsManager::OnAgentAdded_Implementation(const int32 AgentId)
+void ALearningAgentsManager::AddAgents(TArray<int32>& OutAgentIds, const TArray<UObject*>& InAgents)
 {
-	for (ULearningAgentsManagerComponent* ManagerComponent : CachedManagerComponents)
+	OnEventAgentIds.SetNumUninitialized(0, false);
+	OutAgentIds.SetNumUninitialized(InAgents.Num());
+
+	for (int32 AgentIdx = 0; AgentIdx < InAgents.Num(); AgentIdx++)
+	{
+		if (InAgents[AgentIdx] == nullptr)
+		{
+			UE_LOG(LogLearning, Error, TEXT("%s: Attempted to add an agent but agent is nullptr."), *GetName());
+			OutAgentIds[AgentIdx] = INDEX_NONE;
+			continue;
+		}
+
+		if (VacantAgentIds.IsEmpty())
+		{
+			UE_LOG(LogLearning, Error, TEXT("%s: Attempting to add an agent but we have no more vacant ids. Increase MaxAgentNum (%d) or remove unused agents."), *GetName(), MaxAgentNum);
+			OutAgentIds[AgentIdx] = INDEX_NONE;
+			continue;
+		}
+
+		// Add Agent
+
+		OutAgentIds[AgentIdx] = VacantAgentIds.Pop();
+		Agents[OutAgentIds[AgentIdx]] = InAgents[AgentIdx];
+		OccupiedAgentIds.Add(OutAgentIds[AgentIdx]);
+		OnEventAgentIds.Add(OutAgentIds[AgentIdx]);
+	}
+
+	UpdateAgentSets();
+
+	// Call OnAgentsAdded Event
+
+	UE_LOG(LogLearning, Display, TEXT("%s: Adding Agents %s."), *GetName(), *UE::Learning::Array::FormatInt32(OnEventAgentIds));
+
+	TInlineComponentArray<ULearningAgentsManagerComponent*> ManagerComponents;
+	GetComponents<ULearningAgentsManagerComponent>(ManagerComponents);
+
+	for (ULearningAgentsManagerComponent* ManagerComponent : ManagerComponents)
 	{
-		ManagerComponent->AddAgent(AgentId);
+		ManagerComponent->OnAgentsAdded(OnEventAgentIds);
 	}
 }
 
-void ALearningAgentsManager::RemoveAgentById(const int32 AgentId)
+void ALearningAgentsManager::RemoveAgent(const int32 AgentId)
 {
 	if (AgentId == INDEX_NONE)
 	{
-		UE_LOG(LogLearning, Warning, TEXT("%s: Attempting to remove an agent with id of INDEX_NONE."), *GetName());
+		UE_LOG(LogLearning, Error, TEXT("%s: Attempting to remove an agent with id of -1."), *GetName());
 		return;
 	}
 
@@ -131,114 +142,302 @@ void ALearningAgentsManager::RemoveAgentById(const int32 AgentId)
 
 	if (RemovedCount == 0)
 	{
-		UE_LOG(LogLearning, Warning, TEXT("%s: Trying to remove an agent with id of %i but it was not found."), *GetName(), AgentId);
+		UE_LOG(LogLearning, Error, TEXT("%s: Trying to remove an agent with id of %i but it was not found."), *GetName(), AgentId);
 		return;
 	}
 
+	UE_LEARNING_CHECKF(RemovedCount == 1, TEXT("Somehow agent id was included multiple times in set."));
+
 	// Remove Agent
+
 	VacantAgentIds.Push(AgentId);
 	Agents[AgentId] = nullptr;
+	UpdateAgentSets();
+
+	OnEventAgentIds.SetNumUninitialized(0, false);
+	OnEventAgentIds.Add(AgentId);
+
+	// Call OnAgentsRemoved Event
+
+	UE_LOG(LogLearning, Display, TEXT("%s: Removing Agent %i."), *GetName(), AgentId);
+
+	TInlineComponentArray<ULearningAgentsManagerComponent*> ManagerComponents;
+	GetComponents<ULearningAgentsManagerComponent>(ManagerComponents);
+
+	for (ULearningAgentsManagerComponent* ManagerComponent : ManagerComponents)
+	{
+		ManagerComponent->OnAgentsRemoved(OnEventAgentIds);
+	}
+}
+
+void ALearningAgentsManager::RemoveAgents(const TArray<int32>& AgentIds)
+{
+	OnEventAgentIds.SetNumUninitialized(0, false);
+
+	for (int32 AgentIdx = 0; AgentIdx < AgentIds.Num(); AgentIdx++)
+	{
+		if (AgentIds[AgentIdx] == INDEX_NONE)
+		{
+			UE_LOG(LogLearning, Error, TEXT("%s: Attempting to remove an agent with id of -1."), *GetName());
+			continue;
+		}
+
+		const int32 RemovedCount = OccupiedAgentIds.RemoveSingleSwap(AgentIds[AgentIdx], false);
+
+		if (RemovedCount == 0)
+		{
+			UE_LOG(LogLearning, Error, TEXT("%s: Trying to remove an agent with id of %i but it was not found."), *GetName(), AgentIds[AgentIdx]);
+			continue;
+		}
+
+		UE_LEARNING_CHECKF(RemovedCount == 1, TEXT("Somehow agent id was included multiple times in set."));
+
+		// Remove Agent
+
+		VacantAgentIds.Push(AgentIds[AgentIdx]);
+		Agents[AgentIds[AgentIdx]] = nullptr;
+		OnEventAgentIds.Add(AgentIds[AgentIdx]);
+	}
 
 	UpdateAgentSets();
 
-	OnAgentRemoved(AgentId);
+	// Call OnAgentsRemoved Event
+
+	UE_LOG(LogLearning, Display, TEXT("%s: Removing Agents %s."), *GetName(), *UE::Learning::Array::FormatInt32(OnEventAgentIds));
+
+	TInlineComponentArray<ULearningAgentsManagerComponent*> ManagerComponents;
+	GetComponents<ULearningAgentsManagerComponent>(ManagerComponents);
+
+	for (ULearningAgentsManagerComponent* ManagerComponent : ManagerComponents)
+	{
+		ManagerComponent->OnAgentsRemoved(OnEventAgentIds);
+	}
 }
 
-void ALearningAgentsManager::RemoveAgent(UObject* Agent)
+void ALearningAgentsManager::RemoveAllAgents()
 {
-	if (Agent == nullptr)
+	// We need to make a copy as OccupiedAgentIds will be modified during removal.
+	const TArray<int32> RemoveAgentIds = OccupiedAgentIds;
+	RemoveAgents(RemoveAgentIds);
+}
+
+void ALearningAgentsManager::ResetAgent(const int32 AgentId)
+{
+	if (AgentId == INDEX_NONE)
 	{
-		UE_LOG(LogLearning, Error, TEXT("%s: Attempted to remove agent but agent is nullptr."), *GetName());
+		UE_LOG(LogLearning, Error, TEXT("%s: Attempting to reset an agent with id of -1."), *GetName());
 		return;
 	}
 
-	int32 AgentId = INDEX_NONE;
-	const bool bIsFound = Agents.Find(Agent, AgentId);
-
-	if (!bIsFound)
+	if (!HasAgent(AgentId))
 	{
-		UE_LOG(LogLearning, Warning, TEXT("%s: Trying to remove agent %s but it was not found."), *Agent->GetName(), *GetName());
+		UE_LOG(LogLearning, Error, TEXT("%s: Trying to reset an agent with id of %i but it was not found."), *GetName(), AgentId);
 		return;
 	}
 
-	RemoveAgentById(AgentId);
+	OnEventAgentIds.SetNumUninitialized(0, false);
+	OnEventAgentIds.Add(AgentId);
+
+	// Call OnAgentsReset Event
+
+	UE_LOG(LogLearning, Display, TEXT("%s: Resetting Agent %i."), *GetName(), AgentId);
+
+	TInlineComponentArray<ULearningAgentsManagerComponent*> ManagerComponents;
+	GetComponents<ULearningAgentsManagerComponent>(ManagerComponents);
+
+	for (ULearningAgentsManagerComponent* ManagerComponent : ManagerComponents)
+	{
+		ManagerComponent->OnAgentsReset(OnEventAgentIds);
+	}
 }
 
-void ALearningAgentsManager::OnAgentRemoved_Implementation(const int32 AgentId)
+void ALearningAgentsManager::ResetAgents(const TArray<int32>& AgentIds)
 {
-	for (ULearningAgentsManagerComponent* ManagerComponent : CachedManagerComponents)
+	OnEventAgentIds.SetNumUninitialized(0, false);
+
+	for (int32 AgentIdx = 0; AgentIdx < AgentIds.Num(); AgentIdx++)
 	{
-		if (ManagerComponent->HasAgent(AgentId))
+		if (AgentIds[AgentIdx] == INDEX_NONE)
+		{
+			UE_LOG(LogLearning, Error, TEXT("%s: Attempting to reset an agent with id of -1."), *GetName());
+			continue;
+		}
+
+		if (!HasAgent(AgentIds[AgentIdx]))
 		{
-			ManagerComponent->RemoveAgent(AgentId);
+			UE_LOG(LogLearning, Error, TEXT("%s: Trying to reset an agent with id of %i but it was not found."), *GetName(), AgentIds[AgentIdx]);
+			continue;
 		}
+
+		OnEventAgentIds.Add(AgentIds[AgentIdx]);
 	}
-}
 
-bool ALearningAgentsManager::HasAgent(UObject* Agent) const
-{
-	return Agents.Find(Agent) ? true : false;
+	// Call OnAgentsReset Event
+
+	UE_LOG(LogLearning, Display, TEXT("%s: Resetting Agents %s."), *GetName(), *UE::Learning::Array::FormatInt32(OnEventAgentIds));
+
+	TInlineComponentArray<ULearningAgentsManagerComponent*> ManagerComponents;
+	GetComponents<ULearningAgentsManagerComponent>(ManagerComponents);
+
+	for (ULearningAgentsManagerComponent* ManagerComponent : ManagerComponents)
+	{
+		ManagerComponent->OnAgentsReset(OnEventAgentIds);
+	}
 }
 
-bool ALearningAgentsManager::HasAgentById(const int32 AgentId) const
+void ALearningAgentsManager::ResetAllAgents()
 {
-	return OccupiedAgentSet.Contains(AgentId);
+	ResetAgents(OccupiedAgentIds);
 }
 
 UObject* ALearningAgentsManager::GetAgent(const int32 AgentId, const TSubclassOf<UObject> AgentClass) const
 {
-	if (!OccupiedAgentSet.Contains(AgentId))
+	if (!HasAgent(AgentId))
 	{
 		UE_LOG(LogLearning, Error, TEXT("%s: AgentId %d not found. Be sure to only use AgentIds returned by AddAgent and check that the agent has not be removed."), *GetName(), AgentId);
 		return nullptr;
 	}
-
-	return Agents[AgentId];
+	else
+	{
+		return Agents[AgentId];
+	}
 }
 
-void ALearningAgentsManager::GetAgents(const TArray<int32>& AgentIds, const TSubclassOf<UObject> AgentClass, TArray<UObject*>& OutAgents) const
+void ALearningAgentsManager::GetAgents(TArray<UObject*>& OutAgents, const TArray<int32>& AgentIds, const TSubclassOf<UObject> AgentClass) const
 {
 	OutAgents.SetNumUninitialized(AgentIds.Num());
 
 	for (int32 AgentIdIdx = 0; AgentIdIdx < AgentIds.Num(); AgentIdIdx++)
 	{
-		if (!OccupiedAgentSet.Contains(AgentIds[AgentIdIdx]))
+		if (!HasAgent(AgentIds[AgentIdIdx]))
 		{
 			UE_LOG(LogLearning, Error, TEXT("%s: AgentId %d not found. Be sure to only use AgentIds returned by AddAgent and check that the agent has not be removed."), *GetName(), AgentIds[AgentIdIdx]);
-			OutAgents.Empty();
-			return;
+			OutAgents[AgentIdIdx] = nullptr;
 		}
+		else
+		{
+			OutAgents[AgentIdIdx] = Agents[AgentIds[AgentIdIdx]];
+		}
+	}
+}
+
+void ALearningAgentsManager::GetAllAgents(TArray<UObject*>& OutAgents, const TSubclassOf<UObject> AgentClass) const
+{
+	GetAgents(OutAgents, OccupiedAgentIds, AgentClass);
+}
+
+int32 ALearningAgentsManager::GetAgentId(UObject* Agent) const
+{
+	const int32 AgentId = Agents.Find(Agent);
 
-		OutAgents[AgentIdIdx] = Agents[AgentIds[AgentIdIdx]];
+	if (AgentId == INDEX_NONE)
+	{
+		UE_LOG(LogLearning, Error, TEXT("%s: Cannot get id for Agent %s, as it was not found in added agent set."), *GetName(), *Agent->GetName());
+		return INDEX_NONE;
+	}
+	else
+	{
+		return AgentId;
 	}
 }
 
-void ALearningAgentsManager::GetAddedAgents(const TSubclassOf<UObject> AgentClass, TArray<UObject*>& OutAgents) const
+void ALearningAgentsManager::GetAgentIds(TArray<int32>& OutAgentIds, const TArray<UObject*>& InAgents) const
 {
-	GetAgents(OccupiedAgentIds, AgentClass, OutAgents);
+	OutAgentIds.SetNumUninitialized(InAgents.Num());
+
+	for (int32 AgentIdx = 0; AgentIdx < InAgents.Num(); AgentIdx++)
+	{
+		const int32 AgentId = Agents.Find(InAgents[AgentIdx]);
+
+		if (AgentId == INDEX_NONE)
+		{
+			UE_LOG(LogLearning, Error, TEXT("%s: Cannot get id for Agent %s, as it was not found in added agent set."), *GetName(), *InAgents[AgentIdx]->GetName());
+			OutAgentIds[AgentIdx] = INDEX_NONE;
+		}
+		else
+		{
+			OutAgentIds[AgentIdx] = AgentId;
+		}
+	}
+}
+
+void ALearningAgentsManager::GetAllAgentIds(TArray<int32>& OutAgentIds) const
+{
+	OutAgentIds = OccupiedAgentIds;
+}
+
+int32 ALearningAgentsManager::GetAgentNum() const
+{
+	return OccupiedAgentIds.Num();
+}
+
+bool ALearningAgentsManager::HasAgentObject(UObject* Agent) const
+{
+	return Agents.Contains(Agent);
+}
+
+bool ALearningAgentsManager::HasAgent(const int32 AgentId) const
+{
+	return OccupiedAgentSet.Contains(AgentId);
+}
+
+void ALearningAgentsManager::AddManagerAsTickPrerequisiteOfAgents(const TArray<AActor*>& InAgents)
+{
+	for (AActor* Agent : InAgents)
+	{
+		Agent->AddTickPrerequisiteActor(this);
+	}
+}
+
+void ALearningAgentsManager::AddAgentsAsTickPrerequisiteOfManager(const TArray<AActor*>& InAgents)
+{
+	for (AActor* Agent : InAgents)
+	{
+		this->AddTickPrerequisiteActor(Agent);
+	}
 }
 
 const UObject* ALearningAgentsManager::GetAgent(const int32 AgentId) const
 {
-	UE_LEARNING_CHECK(OccupiedAgentSet.Contains(AgentId));
+	UE_LEARNING_CHECK(HasAgent(AgentId));
 	return Agents[AgentId];
 }
 
 UObject* ALearningAgentsManager::GetAgent(const int32 AgentId)
 {
-	UE_LEARNING_CHECK(OccupiedAgentSet.Contains(AgentId));
+	UE_LEARNING_CHECK(HasAgent(AgentId));
 	return Agents[AgentId];
 }
 
-void ALearningAgentsManager::GetAgentIds(TArray<int32>& OutAgentIds) const
+const TArray<int32>& ALearningAgentsManager::GetAllAgentIds() const
 {
-	OutAgentIds = OccupiedAgentIds;
+	return OccupiedAgentIds;
+}
+
+UE::Learning::FIndexSet ALearningAgentsManager::GetAllAgentSet() const
+{
+	return OccupiedAgentSet;
+}
+
+const TSharedPtr<UE::Learning::FArrayMap>& ALearningAgentsManager::GetInstanceData() const
+{
+	return InstanceData;
+}
+
+TConstArrayView<TObjectPtr<UObject>> ALearningAgentsManager::GetAgents() const
+{
+	return Agents;
 }
 
 void ALearningAgentsManager::UpdateAgentSets()
 {
+	// Keep OccupiedAgentIds sorted in ascending order for better memory access patterns.
+	OccupiedAgentIds.Sort();
 	OccupiedAgentSet = OccupiedAgentIds;
 	OccupiedAgentSet.TryMakeSlice();
+
+	// Keep VacantAgentIds sorted in descending order so that lowest ids are popped first.
+	VacantAgentIds.Sort([](const int32& Lhs, const int32& Rhs) { return Lhs > Rhs; });
 	VacantAgentSet = VacantAgentIds;
 	VacantAgentSet.TryMakeSlice();
-}
+}
\ No newline at end of file
diff --git a/Engine/Plugins/Experimental/LearningAgents/Source/LearningAgents/Private/LearningAgentsManagerComponent.cpp b/Engine/Plugins/Experimental/LearningAgents/Source/LearningAgents/Private/LearningAgentsManagerComponent.cpp
--- a/Engine/Plugins/Experimental/LearningAgents/Source/LearningAgents/Private/LearningAgentsManagerComponent.cpp
+++ b/Engine/Plugins/Experimental/LearningAgents/Source/LearningAgents/Private/LearningAgentsManagerComponent.cpp
@@ -4,94 +4,100 @@
 #include "LearningAgentsManager.h"
 #include "LearningLog.h"
 
-ULearningAgentsManagerComponent::ULearningAgentsManagerComponent() : UActorComponent() {}
+ULearningAgentsManagerComponent::ULearningAgentsManagerComponent(const FObjectInitializer& ObjectInitializer) : UActorComponent(ObjectInitializer) {}
 ULearningAgentsManagerComponent::ULearningAgentsManagerComponent(FVTableHelper& Helper) : ULearningAgentsManagerComponent() {}
 ULearningAgentsManagerComponent::~ULearningAgentsManagerComponent() {}
 
-bool ULearningAgentsManagerComponent::AddAgent(const int32 AgentId)
+void ULearningAgentsManagerComponent::PostInitProperties()
 {
-	if (HasAgent(AgentId))
+	Super::PostInitProperties();
+
+	if (HasAnyFlags(RF_ClassDefaultObject | RF_ArchetypeObject))
 	{
-		UE_LOG(LogLearning, Warning, TEXT("%s: AgentId %i is already added."), *GetName(), AgentId);
-		return false;
+		return;
 	}
 
-	AddedAgentIds.Add(AgentId);
-	AddedAgentSet = AddedAgentIds;
-	AddedAgentSet.TryMakeSlice();
-
-	return true;
-}
+	Manager = GetOwner<ALearningAgentsManager>();
 
-bool ULearningAgentsManagerComponent::RemoveAgent(const int32 AgentId)
-{
-	if (AddedAgentIds.RemoveSingleSwap(AgentId, false) == 0)
+	if (!Manager)
 	{
-		UE_LOG(LogLearning, Warning, TEXT("%s: Trying to remove an agent with id of %i but it was not found."), *GetName(), AgentId);
-		return false;
+		UE_LOG(LogLearning, Error, TEXT("%s: Must be attached to a LearningAgentsManager Actor."), *GetName());
 	}
+}
 
-	AddedAgentSet = AddedAgentIds;
-	AddedAgentSet.TryMakeSlice();
+void ULearningAgentsManagerComponent::OnAgentsAdded(const TArray<int32>& AgentIds) { }
 
-	return true;
-}
+void ULearningAgentsManagerComponent::OnAgentsRemoved(const TArray<int32>& AgentIds) { }
 
-bool ULearningAgentsManagerComponent::HasAgent(const int32 AgentId) const
-{
-	return AddedAgentSet.Contains(AgentId);
-}
+void ULearningAgentsManagerComponent::OnAgentsReset(const TArray<int32>& AgentIds) { }
 
 UObject* ULearningAgentsManagerComponent::GetAgent(const int32 AgentId, const TSubclassOf<UObject> AgentClass) const
 {
-	if (!AgentManager)
+	if (!Manager)
 	{
-		UE_LOG(LogLearning, Error, TEXT("%s: Agent manager is nullptr. Call setup on this component prior to getting agents."), *GetName());
+		UE_LOG(LogLearning, Error, TEXT("%s: Must be attached to a LearningAgentsManager Actor."), *GetName());
 		return nullptr;
 	}
 
-	if (!AddedAgentSet.Contains(AgentId))
+	return Manager->GetAgent(AgentId, AgentClass);
+}
+
+void ULearningAgentsManagerComponent::GetAgents(const TArray<int32>& AgentIds, const TSubclassOf<UObject> AgentClass, TArray<UObject*>& OutAgents) const
+{
+	if (!Manager)
 	{
-		UE_LOG(LogLearning, Error, TEXT("%s: AgentId %d not found. Be sure to only use AgentIds returned by AddAgent() and check that the agent has not be removed."), *GetName(), AgentId);
-		return nullptr;
+		UE_LOG(LogLearning, Error, TEXT("%s: Must be attached to a LearningAgentsManager Actor."), *GetName());
+		OutAgents.Empty();
+		return;
 	}
 
-	// Calling this overload since it will log about missing manager ids
-	return AgentManager->GetAgent(AgentId, AgentClass); 
+	Manager->GetAgents(OutAgents, AgentIds, AgentClass);
 }
 
-const UObject* ULearningAgentsManagerComponent::GetAgent(const int32 AgentId) const
+ALearningAgentsManager* ULearningAgentsManagerComponent::GetAgentManager(const TSubclassOf<ALearningAgentsManager> AgentManagerClass) const
 {
-	UE_LEARNING_CHECKF(AgentManager, TEXT("AgentManager is nullptr. Did we forget to call Setup on this component and set the manager?"));
-	UE_LEARNING_CHECKF(AddedAgentSet.Contains(AgentId), TEXT("AgentId not found. Make sure it was added via AddAgent()."));
+	if (!Manager)
+	{
+		UE_LOG(LogLearning, Error, TEXT("%s: Must be attached to a LearningAgentsManager Actor."), *GetName());
+		return nullptr;
+	}
+	
+	return Manager;
+}
 
-	return AgentManager->GetAgent(AgentId);
+const UObject* ULearningAgentsManagerComponent::GetAgent(const int32 AgentId) const
+{
+	UE_LEARNING_CHECKF(Manager, TEXT("Must be attached to a LearningAgentsManager Actor."));
+	return Manager->GetAgent(AgentId);
 }
 
 UObject* ULearningAgentsManagerComponent::GetAgent(const int32 AgentId)
 {
-	UE_LEARNING_CHECKF(AgentManager, TEXT("AgentManager is nullptr. Did we forget to call Setup on this component and set the manager?"));
-	UE_LEARNING_CHECKF(AddedAgentSet.Contains(AgentId), TEXT("AgentId not found. Make sure it was added via AddAgent()."));
-
-	return AgentManager->GetAgent(AgentId);
+	UE_LEARNING_CHECKF(Manager, TEXT("Must be attached to a LearningAgentsManager Actor."));
+	return Manager->GetAgent(AgentId);
 }
 
-ALearningAgentsManager* ULearningAgentsManagerComponent::GetAgentManager() const
+bool ULearningAgentsManagerComponent::HasAgent(const int32 AgentId) const
 {
-	UE_LEARNING_CHECKF(AgentManager, TEXT("AgentManager is nullptr. Did we forget to call Setup on this component and set the manager?"));
+	UE_LEARNING_CHECKF(Manager, TEXT("Must be attached to a LearningAgentsManager Actor."));
+	return Manager->HasAgent(AgentId);
+}
 
-	return AgentManager;
+const ALearningAgentsManager* ULearningAgentsManagerComponent::GetAgentManager() const
+{
+	UE_LEARNING_CHECKF(Manager, TEXT("Must be attached to a LearningAgentsManager Actor."));
+	return Manager;
 }
 
-ALearningAgentsManager* ULearningAgentsManagerComponent::GetAgentManager(const TSubclassOf<ALearningAgentsManager> AgentManagerClass) const
+ALearningAgentsManager* ULearningAgentsManagerComponent::GetAgentManager()
 {
-	if (!AgentManager)
-	{
-		UE_LOG(LogLearning, Error, TEXT("%s: AgentManager is nullptr. Did we forget to call Setup on this component and set the manager?"), *GetName());
-		return nullptr;
-	}
+	UE_LEARNING_CHECKF(Manager, TEXT("Must be attached to a LearningAgentsManager Actor."));
+	return Manager;
+}
 
-	return AgentManager;
+bool ULearningAgentsManagerComponent::HasAgentManager() const
+{
+	return Manager != nullptr;
 }
 
 bool ULearningAgentsManagerComponent::IsSetup() const
@@ -101,5 +107,6 @@ bool ULearningAgentsManagerComponent::IsSetup() const
 
 void ULearningAgentsManagerComponent::AddHelper(TObjectPtr<ULearningAgentsHelper> Object)
 {
+	UE_LEARNING_CHECK(Object);
 	HelperObjects.Add(Object);
 }
\ No newline at end of file
diff --git a/Engine/Plugins/Experimental/LearningAgents/Source/LearningAgents/Private/LearningAgentsNeuralNetwork.cpp b/Engine/Plugins/Experimental/LearningAgents/Source/LearningAgents/Private/LearningAgentsNeuralNetwork.cpp
--- a/Engine/Plugins/Experimental/LearningAgents/Source/LearningAgents/Private/LearningAgentsNeuralNetwork.cpp
+++ b/Engine/Plugins/Experimental/LearningAgents/Source/LearningAgents/Private/LearningAgentsNeuralNetwork.cpp
@@ -127,32 +127,24 @@ void ULearningAgentsNeuralNetwork::LoadNetworkFromSnapshot(const FFilePath& File
 			return;
 		}
 
+		if (!NeuralNetwork)
+		{
+			NeuralNetwork = MakeShared<UE::Learning::FNeuralNetwork>();
+		}
+
 		NeuralNetwork->DeserializeFromBytes(Offset, RecordingData);
 
 		UE_LEARNING_CHECK(Offset == RecordingData.Num());
+
+		ForceMarkDirty();
 	}
 	else
 	{
 		UE_LOG(LogLearning, Error, TEXT("%s: Failed to load network. File not found: \"%s\""), *GetName(), *File.FilePath);
 	}
-
-	// Manually mark the package as dirty since just using `Modify` prevents 
-	// marking packages as dirty during PIE which is most likely when this
-	// is being used.
-	if (UPackage* Package = GetPackage())
-	{
-		const bool bIsDirty = Package->IsDirty();
-
-		if (!bIsDirty)
-		{
-			Package->SetDirtyFlag(true);
-		}
-
-		Package->PackageMarkedDirtyEvent.Broadcast(Package, bIsDirty);
-	}
 }
 
-void ULearningAgentsNeuralNetwork::SaveNetworkToSnapshot(const FFilePath& File) const
+void ULearningAgentsNeuralNetwork::SaveNetworkToSnapshot(const FFilePath& File)
 {
 	TArray<uint8> RecordingData;
 
@@ -177,4 +169,79 @@ void ULearningAgentsNeuralNetwork::SaveNetworkToSnapshot(const FFilePath& File)
 	{
 		UE_LOG(LogLearning, Error, TEXT("%s: Failed to save network to file: \"%s\""), *GetName(), *File.FilePath);
 	}
+}
+
+void ULearningAgentsNeuralNetwork::LoadNetworkFromAsset(ULearningAgentsNeuralNetwork* NeuralNetworkAsset)
+{
+	if (!NeuralNetworkAsset || !NeuralNetworkAsset->NeuralNetwork)
+	{
+		UE_LOG(LogLearning, Error, TEXT("%s: Asset is invalid."), *GetName());
+		return;
+	}
+
+	if (NeuralNetworkAsset == this)
+	{
+		UE_LOG(LogLearning, Error, TEXT("%s: Asset is same as the current network."), *GetName());
+		return;
+	}
+
+	if (NeuralNetwork)
+	{
+		if (NeuralNetworkAsset->NeuralNetwork->GetInputNum() != NeuralNetwork->GetInputNum() ||
+			NeuralNetworkAsset->NeuralNetwork->GetOutputNum() != NeuralNetwork->GetOutputNum() ||
+			NeuralNetworkAsset->NeuralNetwork->GetLayerNum() != NeuralNetwork->GetLayerNum() ||
+			NeuralNetworkAsset->NeuralNetwork->ActivationFunction != NeuralNetwork->ActivationFunction)
+		{
+			UE_LOG(LogLearning, Error, TEXT("%s: Failed to load network from asset as settings don't match."), *GetName());
+			return;
+		}
+	}
+	else
+	{
+		NeuralNetwork = MakeShared<UE::Learning::FNeuralNetwork>();
+	}
+
+	*NeuralNetwork = *NeuralNetworkAsset->NeuralNetwork;
+	ForceMarkDirty();
+}
+
+void ULearningAgentsNeuralNetwork::SaveNetworkToAsset(ULearningAgentsNeuralNetwork* NeuralNetworkAsset)
+{
+	if (!NeuralNetworkAsset)
+	{
+		UE_LOG(LogLearning, Error, TEXT("%s: Asset is invalid."), *GetName());
+		return;
+	}
+
+	if (NeuralNetworkAsset == this)
+	{
+		UE_LOG(LogLearning, Error, TEXT("%s: Asset is same as the current network."), *GetName());
+		return;
+	}
+
+	if (!NeuralNetworkAsset->NeuralNetwork)
+	{
+		NeuralNetworkAsset->NeuralNetwork = MakeShared<UE::Learning::FNeuralNetwork>();
+	}
+
+	*NeuralNetworkAsset->NeuralNetwork = *NeuralNetwork;
+	NeuralNetworkAsset->ForceMarkDirty();
+}
+
+void ULearningAgentsNeuralNetwork::ForceMarkDirty()
+{
+	// Manually mark the package as dirty since just using `Modify` prevents 
+	// marking packages as dirty during PIE which is most likely when this
+	// is being used.
+	if (UPackage* Package = GetPackage())
+	{
+		const bool bIsDirty = Package->IsDirty();
+
+		if (!bIsDirty)
+		{
+			Package->SetDirtyFlag(true);
+		}
+
+		Package->PackageMarkedDirtyEvent.Broadcast(Package, bIsDirty);
+	}
 }
\ No newline at end of file
