LearningAgents: Added iteration checks for verifying order of operations. Fixed issue with inference for multiple agents. Fixed live recompilation issues.

diff --git a/Engine/Plugins/Experimental/LearningAgents/Source/LearningAgents/Private/LearningAgentsObservations.cpp b/Engine/Plugins/Experimental/LearningAgents/Source/LearningAgents/Private/LearningAgentsObservations.cpp
--- a/Engine/Plugins/Experimental/LearningAgents/Source/LearningAgents/Private/LearningAgentsObservations.cpp
+++ b/Engine/Plugins/Experimental/LearningAgents/Source/LearningAgents/Private/LearningAgentsObservations.cpp
@@ -22,15 +22,24 @@ namespace UE::Learning::Agents::Observations::Private
 			return nullptr;
 		}
 
+		if (!InInteractor->HasAgentManager())
+		{
+			UE_LOG(LogLearning, Error, TEXT("%s: Must be attached to a LearningAgentsManager Actor."), *InInteractor->GetName());
+			return nullptr;
+		}
+
 		ObservationUObject* Observation = NewObject<ObservationUObject>(InInteractor, Name);
 
 		Observation->Interactor = InInteractor;
 		Observation->FeatureObject = MakeShared<ObservationFObject>(
 			Observation->GetFName(),
 			InInteractor->GetAgentManager()->GetInstanceData().ToSharedRef(),
-			InInteractor->GetAgentManager()->GetMaxInstanceNum(),
+			InInteractor->GetAgentManager()->GetMaxAgentNum(),
 			Forward<InArgTypes>(Args)...);
 
+		Observation->AgentIteration.SetNumUninitialized({ InInteractor->GetAgentManager()->GetMaxAgentNum() });
+		UE::Learning::Array::Set<1, uint64>(Observation->AgentIteration, INDEX_NONE);
+
 		// We assume all supported observation feature objects can be encoded
 		UE_LEARNING_CHECK(Observation->FeatureObject->IsEncodable());
 
@@ -56,6 +65,7 @@ void UFloatObservation::SetFloatObservation(const int32 AgentId, const float Val
 	}
 
 	FeatureObject->InstanceData->View(FeatureObject->ValueHandle)[AgentId][0] = Value;
+	AgentIteration[AgentId]++;
 }
 
 #if UE_LEARNING_AGENTS_ENABLE_VISUAL_LOG
@@ -73,10 +83,10 @@ void UFloatObservation::VisualLog(const UE::Learning::FIndexSet Instances) const
 			UE_LEARNING_AGENTS_VLOG_STRING(this, LogLearning, Display,
 				Actor->GetActorLocation(),
 				VisualLogColor.ToFColor(true),
-				TEXT("Agent %i\nScale: [% 6.2f]\nValue: [% 6.2f]\nEncoded: [% 6.3f]"),
+				TEXT("Agent %i\nScale: [% 6.2f]\nValue: %s\nEncoded: %s"),
 				Instance,
 				FeatureObject->Scale,
-				ValueView[Instance][0],
+				*UE::Learning::Array::FormatFloat(ValueView[Instance]),
 				*UE::Learning::Array::FormatFloat(FeatureView[Instance]));
 		}
 	}
@@ -113,6 +123,7 @@ void UFloatArrayObservation::SetFloatArrayObservation(const int32 AgentId, const
 	}
 
 	UE::Learning::Array::Copy<1, float>(View[AgentId], Values);
+	AgentIteration[AgentId]++;
 }
 
 #if UE_LEARNING_AGENTS_ENABLE_VISUAL_LOG
@@ -160,6 +171,7 @@ void UVectorObservation::SetVectorObservation(const int32 AgentId, const FVector
 	View[AgentId][0] = Vector.X;
 	View[AgentId][1] = Vector.Y;
 	View[AgentId][2] = Vector.Z;
+	AgentIteration[AgentId]++;
 }
 
 #if UE_LEARNING_AGENTS_ENABLE_VISUAL_LOG
@@ -233,6 +245,8 @@ void UVectorArrayObservation::SetVectorArrayObservation(const int32 AgentId, con
 		View[AgentId][VectorIdx * 3 + 1] = Vectors[VectorIdx].Y;
 		View[AgentId][VectorIdx * 3 + 2] = Vectors[VectorIdx].Z;
 	}
+
+	AgentIteration[AgentId]++;
 }
 
 #if UE_LEARNING_AGENTS_ENABLE_VISUAL_LOG
@@ -321,6 +335,8 @@ void UEnumObservation::SetEnumObservation(const int32 AgentId, const uint8 EnumV
 	const TLearningArrayView<2, float> EnumView = FeatureObject->InstanceData->View(FeatureObject->ValueHandle);
 	UE::Learning::Array::Zero(EnumView[AgentId]);
 	EnumView[AgentId][Index] = 1.0f;
+
+	AgentIteration[AgentId]++;
 }
 
 #if UE_LEARNING_AGENTS_ENABLE_VISUAL_LOG
@@ -428,6 +444,8 @@ void UEnumArrayObservation::SetEnumArrayObservation(const int32 AgentId, const T
 
 		EnumViewSlice[Index] = 1.0f;
 	}
+
+	AgentIteration[AgentId]++;
 }
 
 #if UE_LEARNING_AGENTS_ENABLE_VISUAL_LOG
@@ -471,6 +489,7 @@ void UTimeObservation::SetTimeObservation(const int32 AgentId, const float Time,
 
 	FeatureObject->InstanceData->View(FeatureObject->TimeHandle)[AgentId][0] = Time;
 	FeatureObject->InstanceData->View(FeatureObject->RelativeTimeHandle)[AgentId] = RelativeTime;
+	AgentIteration[AgentId]++;
 }
 
 #if UE_LEARNING_AGENTS_ENABLE_VISUAL_LOG
@@ -533,6 +552,7 @@ void UTimeArrayObservation::SetTimeArrayObservation(const int32 AgentId, const T
 
 	UE::Learning::Array::Copy<1, float>(TimeView[AgentId], Times);
 	RelativeTimeView[AgentId] = RelativeTime;
+	AgentIteration[AgentId]++;
 }
 
 #if UE_LEARNING_AGENTS_ENABLE_VISUAL_LOG
@@ -579,6 +599,7 @@ void UAngleObservation::SetAngleObservation(const int32 AgentId, const float Ang
 
 	FeatureObject->InstanceData->View(FeatureObject->AngleHandle)[AgentId][0] = FMath::DegreesToRadians(Angle);
 	FeatureObject->InstanceData->View(FeatureObject->RelativeAngleHandle)[AgentId] = FMath::DegreesToRadians(RelativeAngle);
+	AgentIteration[AgentId]++;
 }
 
 #if UE_LEARNING_AGENTS_ENABLE_VISUAL_LOG
@@ -648,6 +669,7 @@ void UAngleArrayObservation::SetAngleArrayObservation(const int32 AgentId, const
 	}
 	
 	RelativeAngleView[AgentId] = FMath::DegreesToRadians(RelativeAngle);
+	AgentIteration[AgentId]++;
 }
 
 #if UE_LEARNING_AGENTS_ENABLE_VISUAL_LOG
@@ -718,6 +740,7 @@ void URotationObservation::SetRotationObservationFromQuat(const int32 AgentId, c
 
 	FeatureObject->InstanceData->View(FeatureObject->RotationHandle)[AgentId][0] = Rotation;
 	FeatureObject->InstanceData->View(FeatureObject->RelativeRotationHandle)[AgentId] = RelativeRotation;
+	AgentIteration[AgentId]++;
 }
 
 #if UE_LEARNING_AGENTS_ENABLE_VISUAL_LOG
@@ -793,6 +816,7 @@ void URotationArrayObservation::SetRotationArrayObservation(const int32 AgentId,
 	}
 
 	RelativeRotationView[AgentId] = RelativeRotation.Quaternion();
+	AgentIteration[AgentId]++;
 }
 
 void URotationArrayObservation::SetRotationArrayObservationFromQuats(const int32 AgentId, const TArray<FQuat>& Rotations, const FQuat RelativeRotation)
@@ -814,6 +838,7 @@ void URotationArrayObservation::SetRotationArrayObservationFromQuats(const int32
 
 	UE::Learning::Array::Copy<1, FQuat>(RotationView[AgentId], Rotations);
 	RelativeRotationView[AgentId] = RelativeRotation;
+	AgentIteration[AgentId]++;
 }
 
 #if UE_LEARNING_AGENTS_ENABLE_VISUAL_LOG
@@ -881,6 +906,7 @@ void UDirectionObservation::SetDirectionObservation(const int32 AgentId, const F
 
 	FeatureObject->InstanceData->View(FeatureObject->DirectionHandle)[AgentId][0] = Direction;
 	FeatureObject->InstanceData->View(FeatureObject->RelativeRotationHandle)[AgentId] = RelativeRotation.Quaternion();
+	AgentIteration[AgentId]++;
 }
 
 #if UE_LEARNING_AGENTS_ENABLE_VISUAL_LOG
@@ -956,6 +982,7 @@ void UDirectionArrayObservation::SetDirectionArrayObservation(const int32 AgentI
 
 	UE::Learning::Array::Copy<1, FVector>(DirectionView[AgentId], Directions);
 	RelativeRotationView[AgentId] = RelativeRotation.Quaternion();
+	AgentIteration[AgentId]++;
 }
 
 #if UE_LEARNING_AGENTS_ENABLE_VISUAL_LOG
@@ -1031,6 +1058,7 @@ void UPlanarDirectionObservation::SetPlanarDirectionObservation(const int32 Agen
 
 	FeatureObject->InstanceData->View(FeatureObject->DirectionHandle)[AgentId][0] = Direction;
 	FeatureObject->InstanceData->View(FeatureObject->RelativeRotationHandle)[AgentId] = RelativeRotation.Quaternion();
+	AgentIteration[AgentId]++;
 }
 
 #if UE_LEARNING_AGENTS_ENABLE_VISUAL_LOG
@@ -1121,6 +1149,7 @@ void UPlanarDirectionArrayObservation::SetPlanarDirectionArrayObservation(const
 
 	UE::Learning::Array::Copy<1, FVector>(DirectionView[AgentId], Directions);
 	RelativeRotationView[AgentId] = RelativeRotation.Quaternion();
+	AgentIteration[AgentId]++;
 }
 
 #if UE_LEARNING_AGENTS_ENABLE_VISUAL_LOG
@@ -1202,6 +1231,7 @@ void UPositionObservation::SetPositionObservation(const int32 AgentId, const FVe
 	FeatureObject->InstanceData->View(FeatureObject->PositionHandle)[AgentId][0] = Position;
 	FeatureObject->InstanceData->View(FeatureObject->RelativePositionHandle)[AgentId] = RelativePosition;
 	FeatureObject->InstanceData->View(FeatureObject->RelativeRotationHandle)[AgentId] = RelativeRotation.Quaternion();
+	AgentIteration[AgentId]++;
 }
 
 #if UE_LEARNING_AGENTS_ENABLE_VISUAL_LOG
@@ -1279,6 +1309,7 @@ void UPositionArrayObservation::SetPositionArrayObservation(const int32 AgentId,
 	RelativePositionView[AgentId] = RelativePosition;
 	RelativeRotationView[AgentId] = RelativeRotation.Quaternion();
 	UE::Learning::Array::Copy<1, FVector>(PositionView[AgentId], Positions);
+	AgentIteration[AgentId]++;
 }
 
 #if UE_LEARNING_AGENTS_ENABLE_VISUAL_LOG
@@ -1346,6 +1377,7 @@ void UScalarPositionObservation::SetScalarPositionObservation(const int32 AgentI
 
 	FeatureObject->InstanceData->View(FeatureObject->PositionHandle)[AgentId][0] = Position;
 	FeatureObject->InstanceData->View(FeatureObject->RelativePositionHandle)[AgentId] = RelativePosition;
+	AgentIteration[AgentId]++;
 }
 
 void UScalarPositionObservation::SetScalarPositionObservationWithAxis(const int32 AgentId, const FVector Position, const FVector RelativePosition, const FVector Axis)
@@ -1411,6 +1443,7 @@ void UScalarPositionArrayObservation::SetScalarPositionArrayObservation(const in
 
 	RelativePositionView[AgentId] = RelativePosition;
 	UE::Learning::Array::Copy<1, float>(PositionView[AgentId], Positions);
+	AgentIteration[AgentId]++;
 }
 
 void UScalarPositionArrayObservation::SetScalarPositionArrayObservationWithAxis(const int32 AgentId, const TArray<FVector>& Positions, const FVector RelativePosition, const FVector Axis)
@@ -1435,6 +1468,8 @@ void UScalarPositionArrayObservation::SetScalarPositionArrayObservationWithAxis(
 	{
 		PositionView[AgentId][PositionIdx] = Positions[PositionIdx].Dot(Axis);
 	}
+
+	AgentIteration[AgentId]++;
 }
 
 #if UE_LEARNING_AGENTS_ENABLE_VISUAL_LOG
@@ -1489,6 +1524,7 @@ void UPlanarPositionObservation::SetPlanarPositionObservation(const int32 AgentI
 	FeatureObject->InstanceData->View(FeatureObject->PositionHandle)[AgentId][0] = Position;
 	FeatureObject->InstanceData->View(FeatureObject->RelativePositionHandle)[AgentId] = RelativePosition;
 	FeatureObject->InstanceData->View(FeatureObject->RelativeRotationHandle)[AgentId] = RelativeRotation.Quaternion();
+	AgentIteration[AgentId]++;
 }
 
 #if UE_LEARNING_AGENTS_ENABLE_VISUAL_LOG
@@ -1581,6 +1617,7 @@ void UPlanarPositionArrayObservation::SetPlanarPositionArrayObservation(const in
 	RelativePositionView[AgentId] = RelativePosition;
 	RelativeRotationView[AgentId] = RelativeRotation.Quaternion();
 	UE::Learning::Array::Copy<1, FVector>(PositionView[AgentId], Positions);
+	AgentIteration[AgentId]++;
 }
 
 #if UE_LEARNING_AGENTS_ENABLE_VISUAL_LOG
@@ -1659,6 +1696,7 @@ void UVelocityObservation::SetVelocityObservation(const int32 AgentId, const FVe
 
 	FeatureObject->InstanceData->View(FeatureObject->VelocityHandle)[AgentId][0] = Velocity;
 	FeatureObject->InstanceData->View(FeatureObject->RelativeRotationHandle)[AgentId] = RelativeRotation.Quaternion();
+	AgentIteration[AgentId]++;
 }
 
 #if UE_LEARNING_AGENTS_ENABLE_VISUAL_LOG
@@ -1734,6 +1772,7 @@ void UVelocityArrayObservation::SetVelocityArrayObservation(const int32 AgentId,
 
 	UE::Learning::Array::Copy<1, FVector>(VelocityView[AgentId], Velocities);
 	RelativeRotationView[AgentId] = RelativeRotation.Quaternion();
+	AgentIteration[AgentId]++;
 }
 
 #if UE_LEARNING_AGENTS_ENABLE_VISUAL_LOG
@@ -1801,6 +1840,7 @@ void UScalarVelocityObservation::SetScalarVelocityObservation(const int32 AgentI
 	}
 
 	FeatureObject->InstanceData->View(FeatureObject->VelocityHandle)[AgentId][0] = Velocity;
+	AgentIteration[AgentId]++;
 }
 
 void UScalarVelocityObservation::SetScalarVelocityObservationWithAxis(const int32 AgentId, const FVector Velocity, const FVector Axis)
@@ -1862,6 +1902,7 @@ void UScalarVelocityArrayObservation::SetScalarVelocityArrayObservation(const in
 	}
 
 	UE::Learning::Array::Copy<1, float>(VelocityView[AgentId], Velocities);
+	AgentIteration[AgentId]++;
 }
 
 void UScalarVelocityArrayObservation::SetScalarVelocityArrayObservationWithAxis(const int32 AgentId, const TArray<FVector>& Velocities, const FVector Axis)
@@ -1884,6 +1925,8 @@ void UScalarVelocityArrayObservation::SetScalarVelocityArrayObservationWithAxis(
 	{
 		VelocityView[AgentId][VelocityIdx] = Velocities[VelocityIdx].Dot(Axis);
 	}
+
+	AgentIteration[AgentId]++;
 }
 
 #if UE_LEARNING_AGENTS_ENABLE_VISUAL_LOG
@@ -1936,6 +1979,7 @@ void UPlanarVelocityObservation::SetPlanarVelocityObservation(const int32 AgentI
 
 	FeatureObject->InstanceData->View(FeatureObject->VelocityHandle)[AgentId][0] = Velocity;
 	FeatureObject->InstanceData->View(FeatureObject->RelativeRotationHandle)[AgentId] = RelativeRotation.Quaternion();
+	AgentIteration[AgentId]++;
 }
 
 #if UE_LEARNING_AGENTS_ENABLE_VISUAL_LOG
@@ -2027,6 +2071,7 @@ void UPlanarVelocityArrayObservation::SetPlanarVelocityArrayObservation(const in
 
 	UE::Learning::Array::Copy<1, FVector>(VelocityView[AgentId], Velocities);
 	RelativeRotationView[AgentId] = RelativeRotation.Quaternion();
+	AgentIteration[AgentId]++;
 }
 
 #if UE_LEARNING_AGENTS_ENABLE_VISUAL_LOG
@@ -2105,6 +2150,7 @@ void UAngularVelocityObservation::SetAngularVelocityObservation(const int32 Agen
 
 	FeatureObject->InstanceData->View(FeatureObject->AngularVelocityHandle)[AgentId][0] = AngularVelocity;
 	FeatureObject->InstanceData->View(FeatureObject->RelativeRotationHandle)[AgentId] = RelativeRotation.Quaternion();
+	AgentIteration[AgentId]++;
 }
 
 #if UE_LEARNING_AGENTS_ENABLE_VISUAL_LOG
@@ -2180,6 +2226,7 @@ void UAngularVelocityArrayObservation::SetAngularVelocityArrayObservation(const
 
 	UE::Learning::Array::Copy<1, FVector>(AngularVelocityView[AgentId], AngularVelocities);
 	RelativeRotationView[AgentId] = RelativeRotation.Quaternion();
+	AgentIteration[AgentId]++;
 }
 
 #if UE_LEARNING_AGENTS_ENABLE_VISUAL_LOG
@@ -2247,6 +2294,7 @@ void UScalarAngularVelocityObservation::SetScalarAngularVelocityObservation(cons
 	}
 
 	FeatureObject->InstanceData->View(FeatureObject->AngularVelocityHandle)[AgentId][0] = AngularVelocity;
+	AgentIteration[AgentId]++;
 }
 
 void UScalarAngularVelocityObservation::SetScalarAngularVelocityObservationWithAxis(const int32 AgentId, const FVector AngularVelocity, const FVector Axis)
@@ -2308,6 +2356,7 @@ void UScalarAngularVelocityArrayObservation::SetScalarAngularVelocityArrayObserv
 	}
 
 	UE::Learning::Array::Copy<1, float>(AngularVelocityView[AgentId], AngularVelocities);
+	AgentIteration[AgentId]++;
 }
 
 void UScalarAngularVelocityArrayObservation::SetScalarAngularVelocityArrayObservationWithAxis(const int32 AgentId, const TArray<FVector>& AngularVelocities, const FVector Axis)
@@ -2330,6 +2379,7 @@ void UScalarAngularVelocityArrayObservation::SetScalarAngularVelocityArrayObserv
 	{
 		AngularVelocityView[AgentId][AngularVelocityIdx] = AngularVelocities[AngularVelocityIdx].Dot(Axis);
 	}
+	AgentIteration[AgentId]++;
 }
 
 #if UE_LEARNING_AGENTS_ENABLE_VISUAL_LOG
diff --git a/Engine/Plugins/Experimental/LearningAgents/Source/LearningAgents/Private/LearningAgentsPolicy.cpp b/Engine/Plugins/Experimental/LearningAgents/Source/LearningAgents/Private/LearningAgentsPolicy.cpp
--- a/Engine/Plugins/Experimental/LearningAgents/Source/LearningAgents/Private/LearningAgentsPolicy.cpp
+++ b/Engine/Plugins/Experimental/LearningAgents/Source/LearningAgents/Private/LearningAgentsPolicy.cpp
@@ -14,11 +14,14 @@
 #include "Misc/Paths.h"
 #include "GameFramework/Actor.h"
 
-ULearningAgentsPolicy::ULearningAgentsPolicy() : ULearningAgentsManagerComponent() {}
+ULearningAgentsPolicy::ULearningAgentsPolicy(const FObjectInitializer& ObjectInitializer) : ULearningAgentsManagerComponent(ObjectInitializer) {}
 ULearningAgentsPolicy::ULearningAgentsPolicy(FVTableHelper& Helper) : ULearningAgentsPolicy() {}
 ULearningAgentsPolicy::~ULearningAgentsPolicy() {}
 
-void ULearningAgentsPolicy::SetupPolicy(ALearningAgentsManager* InAgentManager, ULearningAgentsInteractor* InInteractor, const FLearningAgentsPolicySettings& PolicySettings)
+void ULearningAgentsPolicy::SetupPolicy(
+	ULearningAgentsInteractor* InInteractor, 
+	const FLearningAgentsPolicySettings& PolicySettings,
+	ULearningAgentsNeuralNetwork* NeuralNetworkAsset)
 {
 	if (IsSetup())
 	{
@@ -26,20 +29,12 @@ void ULearningAgentsPolicy::SetupPolicy(ALearningAgentsManager* InAgentManager,
 		return;
 	}
 
-	if (!InAgentManager)
+	if (!Manager)
 	{
-		UE_LOG(LogLearning, Error, TEXT("%s: InAgentManager is nullptr."), *GetName());
+		UE_LOG(LogLearning, Error, TEXT("%s: Must be attached to a LearningAgentsManager Actor."), *GetName());
 		return;
 	}
 
-	if (!InAgentManager->IsManagerSetup())
-	{
-		UE_LOG(LogLearning, Error, TEXT("%s: %s's SetupManager must be run before it can be used."), *GetName(), *InAgentManager->GetName());
-		return;
-	}
-
-	AgentManager = InAgentManager;
-
 	if (!InInteractor)
 	{
 		UE_LOG(LogLearning, Error, TEXT("%s: InInteractor is nullptr."), *GetName());
@@ -55,47 +50,78 @@ void ULearningAgentsPolicy::SetupPolicy(ALearningAgentsManager* InAgentManager,
 	Interactor = InInteractor;
 
 	// Setup Neural Network
-	Network = NewObject<ULearningAgentsNeuralNetwork>(this, TEXT("PolicyNetwork"));
-	Network->NeuralNetwork = MakeShared<UE::Learning::FNeuralNetwork>();
-	Network->NeuralNetwork->Resize(
-		Interactor->GetObservationFeature().DimNum(),
-		2 * Interactor->GetActionFeature().DimNum(),
-		PolicySettings.HiddenLayerSize,
-		PolicySettings.LayerNum);
-	Network->NeuralNetwork->ActivationFunction = UE::Learning::Agents::GetActivationFunction(PolicySettings.ActivationFunction);
+
+	if (NeuralNetworkAsset)
+	{
+		// Use Existing Neural Network Asset
+
+		if (NeuralNetworkAsset->NeuralNetwork)
+		{
+			if (NeuralNetworkAsset->NeuralNetwork->GetInputNum() != Interactor->GetObservationFeature().DimNum() ||
+				NeuralNetworkAsset->NeuralNetwork->GetOutputNum() != 2 * Interactor->GetActionFeature().DimNum())
+			{
+				UE_LOG(LogLearning, Error, TEXT("%s: Neural Network Asset provided during Setup is incorrect size: Inputs and outputs don't match what is required."), *GetName());
+				return;
+			}
+
+			if (NeuralNetworkAsset->NeuralNetwork->GetHiddenNum() != PolicySettings.HiddenLayerSize ||
+				NeuralNetworkAsset->NeuralNetwork->GetLayerNum() != PolicySettings.LayerNum ||
+				NeuralNetworkAsset->NeuralNetwork->ActivationFunction != UE::Learning::Agents::GetActivationFunction(PolicySettings.ActivationFunction))
+			{
+				UE_LOG(LogLearning, Warning, TEXT("%s: Neural Network Asset settings don't match those given by PolicySettings"), *GetName());
+			}
+
+			Network = NeuralNetworkAsset;
+		}
+		else
+		{
+			Network = NeuralNetworkAsset;
+			Network->NeuralNetwork = MakeShared<UE::Learning::FNeuralNetwork>();
+			Network->NeuralNetwork->Resize(
+				Interactor->GetObservationFeature().DimNum(),
+				2 * Interactor->GetActionFeature().DimNum(),
+				PolicySettings.HiddenLayerSize,
+				PolicySettings.LayerNum);
+			Network->NeuralNetwork->ActivationFunction = UE::Learning::Agents::GetActivationFunction(PolicySettings.ActivationFunction);
+		}
+	}
+	else
+	{
+		// Create New Neural Network Asset
+
+		Network = NewObject<ULearningAgentsNeuralNetwork>(this, TEXT("PolicyNetwork"));
+		Network->NeuralNetwork = MakeShared<UE::Learning::FNeuralNetwork>();
+		Network->NeuralNetwork->Resize(
+			Interactor->GetObservationFeature().DimNum(),
+			2 * Interactor->GetActionFeature().DimNum(),
+			PolicySettings.HiddenLayerSize,
+			PolicySettings.LayerNum);
+		Network->NeuralNetwork->ActivationFunction = UE::Learning::Agents::GetActivationFunction(PolicySettings.ActivationFunction);
+	}
 
 	// Create Policy Object
 	UE::Learning::FNeuralNetworkPolicyFunctionSettings PolicyFunctionSettings;
 	PolicyFunctionSettings.ActionNoiseMin = PolicySettings.ActionNoiseMin;
 	PolicyFunctionSettings.ActionNoiseMax = PolicySettings.ActionNoiseMax;
-	PolicyFunctionSettings.ActionNoiseScale = PolicySettings.InitialActionNoiseScale;
+	PolicyFunctionSettings.ActionNoiseScale = PolicySettings.ActionNoiseScale;
 
 	PolicyObject = MakeShared<UE::Learning::FNeuralNetworkPolicyFunction>(
 		TEXT("PolicyObject"),
-		AgentManager->GetInstanceData().ToSharedRef(),
-		AgentManager->GetMaxInstanceNum(),
+		Manager->GetInstanceData().ToSharedRef(),
+		Manager->GetMaxAgentNum(),
 		Network->NeuralNetwork.ToSharedRef(),
 		PolicySettings.ActionNoiseSeed,
 		PolicyFunctionSettings);
 
-	AgentManager->GetInstanceData()->Link(Interactor->GetObservationFeature().FeatureHandle, PolicyObject->InputHandle);
+	Manager->GetInstanceData()->Link(Interactor->GetObservationFeature().FeatureHandle, PolicyObject->InputHandle);
+	Manager->GetInstanceData()->Link(PolicyObject->OutputHandle, Interactor->GetActionFeature().FeatureHandle);
 
-	InitialActionNoiseScale = PolicySettings.InitialActionNoiseScale;
-	
 	bIsSetup = true;
 }
 
-bool ULearningAgentsPolicy::AddAgent(const int32 AgentId)
+ULearningAgentsNeuralNetwork* ULearningAgentsPolicy::GetNetworkAsset()
 {
-	bool bSuccess = Super::AddAgent(AgentId);
-
-	if (bSuccess)
-	{
-		// Reset the noise scale for this agent back to the initial value.
-		SetAgentActionNoiseScale(AgentId, InitialActionNoiseScale);
-	}
-
-	return bSuccess;
+	return Network;
 }
 
 UE::Learning::FNeuralNetwork& ULearningAgentsPolicy::GetPolicyNetwork()
@@ -130,7 +156,7 @@ void ULearningAgentsPolicy::SavePolicyToSnapshot(const FFilePath& File) const
 	Network->SaveNetworkToSnapshot(File);
 }
 
-void ULearningAgentsPolicy::LoadPolicyFromAsset(const ULearningAgentsNeuralNetwork* NeuralNetworkAsset)
+void ULearningAgentsPolicy::UsePolicyFromAsset(ULearningAgentsNeuralNetwork* NeuralNetworkAsset)
 {
 	if (!IsSetup())
 	{
@@ -144,17 +170,26 @@ void ULearningAgentsPolicy::LoadPolicyFromAsset(const ULearningAgentsNeuralNetwo
 		return;
 	}
 
+	if (NeuralNetworkAsset == Network)
+	{
+		UE_LOG(LogLearning, Error, TEXT("%s: Asset is same as the current network."), *GetName());
+		return;
+	}
+
 	if (NeuralNetworkAsset->NeuralNetwork->GetInputNum() != Network->NeuralNetwork->GetInputNum() ||
-		NeuralNetworkAsset->NeuralNetwork->GetOutputNum() != Network->NeuralNetwork->GetOutputNum())
+		NeuralNetworkAsset->NeuralNetwork->GetOutputNum() != Network->NeuralNetwork->GetOutputNum() ||
+		NeuralNetworkAsset->NeuralNetwork->GetLayerNum() != Network->NeuralNetwork->GetLayerNum() ||
+		NeuralNetworkAsset->NeuralNetwork->ActivationFunction != Network->NeuralNetwork->ActivationFunction)
 	{
-		UE_LOG(LogLearning, Error, TEXT("%s: Failed to load network from asset. Inputs and outputs don't match."), *GetName());
+		UE_LOG(LogLearning, Error, TEXT("%s: Failed to use asset as network settings don't match."), *GetName());
 		return;
 	}
 
-	*Network->NeuralNetwork = *NeuralNetworkAsset->NeuralNetwork;
+	Network = NeuralNetworkAsset;
+	PolicyObject->NeuralNetwork = Network->NeuralNetwork.ToSharedRef();
 }
 
-void ULearningAgentsPolicy::SavePolicyToAsset(ULearningAgentsNeuralNetwork* NeuralNetworkAsset) const
+void ULearningAgentsPolicy::LoadPolicyFromAsset(ULearningAgentsNeuralNetwork* NeuralNetworkAsset)
 {
 	if (!IsSetup())
 	{
@@ -162,60 +197,63 @@ void ULearningAgentsPolicy::SavePolicyToAsset(ULearningAgentsNeuralNetwork* Neur
 		return;
 	}
 
-	if (!NeuralNetworkAsset)
+	Network->LoadNetworkFromAsset(NeuralNetworkAsset);
+}
+
+void ULearningAgentsPolicy::SavePolicyToAsset(ULearningAgentsNeuralNetwork* NeuralNetworkAsset)
+{
+	if (!IsSetup())
 	{
-		UE_LOG(LogLearning, Error, TEXT("%s: Asset is invalid."), *GetName());
+		UE_LOG(LogLearning, Error, TEXT("%s: Setup not complete."), *GetName());
 		return;
 	}
 
-	if (!NeuralNetworkAsset->NeuralNetwork)
+	Network->SaveNetworkToAsset(NeuralNetworkAsset);
+}
+
+void ULearningAgentsPolicy::EvaluatePolicy()
+{
+	UE_LEARNING_TRACE_CPUPROFILER_EVENT_SCOPE(ULearningAgentsPolicy::EvaluatePolicy);
+
+	if (!IsSetup())
 	{
-		NeuralNetworkAsset->NeuralNetwork = MakeShared<UE::Learning::FNeuralNetwork>();
+		UE_LOG(LogLearning, Error, TEXT("%s: Setup not complete."), *GetName());
+		return;
 	}
 
-	*NeuralNetworkAsset->NeuralNetwork = *Network->NeuralNetwork;
+	// Check Agents actually have encoded observations.
 
-	// Manually mark the package as dirty since just using `Modify` prevents 
-	// marking packages as dirty during PIE which is most likely when this
-	// is being used.
-	if (UPackage* Package = NeuralNetworkAsset->GetPackage())
-	{
-		const bool bIsDirty = Package->IsDirty();
+	ValidAgentIds.Empty(Manager->GetAgentNum());
 
-		if (!bIsDirty)
+	for (const int32 AgentId : Manager->GetAllAgentSet())
+	{
+		if (Interactor->GetObservationEncodingAgentIteration()[AgentId] == 0)
 		{
-			Package->SetDirtyFlag(true);
+			UE_LOG(LogLearning, Warning, TEXT("%s: Agent with id %i has not made observations so policy will not be evaluated for it."), *GetName(), AgentId);
+			continue;
 		}
 
-		Package->PackageMarkedDirtyEvent.Broadcast(Package, bIsDirty);
+		ValidAgentIds.Add(AgentId);
 	}
-}
 
-void ULearningAgentsPolicy::EvaluatePolicy()
-{
-	UE_LEARNING_TRACE_CPUPROFILER_EVENT_SCOPE(ULearningAgentsPolicy::EvaluatePolicy);
+	ValidAgentSet = ValidAgentIds;
+	ValidAgentSet.TryMakeSlice();
 
-	if (!IsSetup())
+	// Evaluate Policy
+
+	PolicyObject->Evaluate(ValidAgentSet);
+
+	// Increment Action Encoding Iteration
+
+	for (const int32 AgentId : ValidAgentSet)
 	{
-		UE_LOG(LogLearning, Error, TEXT("%s: Setup not complete."), *GetName());
-		return;
+		Interactor->GetActionEncodingAgentIteration()[AgentId]++;
 	}
 
-	PolicyObject->Evaluate(AddedAgentIds);
-
-	// Copy the actions computed by the policy into the agent-type's feature buffer.
-	// 
-	// Normally we would just link these two handles, but in this case we want to allow
-	// for multiple difference policies to be used for different agents, so that means
-	// there may be multiple writers to the action feature vector handle and so therefore
-	// the handles cannot be linked and we need to do the copy manually
-	UE::Learning::Array::Copy(
-		PolicyObject->InstanceData->View(Interactor->GetActionFeature().FeatureHandle),
-		PolicyObject->InstanceData->ConstView(PolicyObject->OutputHandle),
-		AddedAgentIds);
+	// Visual Logger
 
 #if UE_LEARNING_AGENTS_ENABLE_VISUAL_LOG
-	VisualLog(AddedAgentIds);
+	VisualLog(ValidAgentSet);
 #endif
 }
 
@@ -233,7 +271,7 @@ void ULearningAgentsPolicy::RunInference()
 	EvaluatePolicy();
 	Interactor->DecodeActions();
 }
-float ULearningAgentsPolicy::GetAgentActionNoiseScale(const int32 AgentId) const
+float ULearningAgentsPolicy::GetActionNoiseScale() const
 {
 	if (!IsSetup())
 	{
@@ -241,17 +279,10 @@ float ULearningAgentsPolicy::GetAgentActionNoiseScale(const int32 AgentId) const
 		return 0.0f;
 	}
 
-	if (!HasAgent(AgentId))
-	{
-		UE_LOG(LogLearning, Error, TEXT("%s: AgentId %d not found in the agents set."), *GetName(), AgentId);
-		return 0.0f;
-	}
-
-	const TLearningArrayView<1, const float> ActionNoiseScaleView = PolicyObject->InstanceData->ConstView(PolicyObject->ActionNoiseScaleHandle);
-	return ActionNoiseScaleView[AgentId];
+	return PolicyObject->InstanceData->ConstView(PolicyObject->ActionNoiseScaleHandle)[0];
 }
 
-void ULearningAgentsPolicy::SetAgentActionNoiseScale(const int32 AgentId, const float ActionNoiseScale)
+void ULearningAgentsPolicy::SetActionNoiseScale(const float ActionNoiseScale)
 {
 	if (!IsSetup())
 	{
@@ -259,30 +290,12 @@ void ULearningAgentsPolicy::SetAgentActionNoiseScale(const int32 AgentId, const
 		return;
 	}
 
-	if (!HasAgent(AgentId))
-	{
-		UE_LOG(LogLearning, Error, TEXT("%s: AgentId %d not found in the agents set."), *GetName(), AgentId);
-		return;
-	}
 
-	const TLearningArrayView<1, float> ActionNoiseScaleView = PolicyObject->InstanceData->View(PolicyObject->ActionNoiseScaleHandle);
-	ActionNoiseScaleView[AgentId] = ActionNoiseScale;
-}
-
-void ULearningAgentsPolicy::SetAllAgentsActionNoiseScale(const float ActionNoiseScale)
-{
-	if (!IsSetup())
-	{
-		UE_LOG(LogLearning, Error, TEXT("%s: Setup not complete."), *GetName());
-		return;
-	}
-
-	const TLearningArrayView<1, float> ActionNoiseScaleView = PolicyObject->InstanceData->View(PolicyObject->ActionNoiseScaleHandle);
-	UE::Learning::Array::Set(ActionNoiseScaleView, ActionNoiseScale, AddedAgentSet);
+	UE::Learning::Array::Set(PolicyObject->InstanceData->View(PolicyObject->ActionNoiseScaleHandle), ActionNoiseScale);
 }
 
 #if UE_LEARNING_AGENTS_ENABLE_VISUAL_LOG
-void ULearningAgentsPolicy::VisualLog(const UE::Learning::FIndexSet Instances) const
+void ULearningAgentsPolicy::VisualLog(const UE::Learning::FIndexSet AgentSet) const
 {
 	UE_LEARNING_TRACE_CPUPROFILER_EVENT_SCOPE(ULearningAgentsPolicy::VisualLog);
 
@@ -292,23 +305,22 @@ void ULearningAgentsPolicy::VisualLog(const UE::Learning::FIndexSet Instances) c
 	const TLearningArrayView<2, const float> OutputStdView = PolicyObject->InstanceData->ConstView(PolicyObject->OutputStdHandle);
 	const TLearningArrayView<1, const float> ActionNoiseScaleView = PolicyObject->InstanceData->ConstView(PolicyObject->ActionNoiseScaleHandle);
 
-	for (const int32 Instance : Instances)
+	for (const int32 AgentId : AgentSet)
 	{
-		if (const AActor* Actor = Cast<AActor>(Interactor->GetAgent(Instance)))
+		if (const AActor* Actor = Cast<AActor>(GetAgent(AgentId)))
 		{
-
 			UE_LEARNING_AGENTS_VLOG_STRING(this, LogLearning, Display,
 				Actor->GetActorLocation(),
 				VisualLogColor.ToFColor(true),
 				TEXT("Agent %i\nAction Noise Scale: [% 6.3f]\nInput: %s\nInput Stats (Min/Max/Mean/Std): %s\nOutput Mean: %s\nOutput Std: %s\nOutput Sample: %s\nOutput Stats (Min/Max/Mean/Std): %s"),
-				Instance,
-				ActionNoiseScaleView[Instance],
-				*UE::Learning::Array::FormatFloat(InputView[Instance]),
-				*UE::Learning::Array::FormatFloat(OutputView[Instance]),
-				*UE::Learning::Array::FormatFloat(OutputMeanView[Instance]),
-				*UE::Learning::Array::FormatFloat(OutputStdView[Instance]),
-				*UE::Learning::Agents::Debug::FloatArrayToStatsString(InputView[Instance]),
-				*UE::Learning::Agents::Debug::FloatArrayToStatsString(OutputView[Instance]));
+				AgentId,
+				ActionNoiseScaleView[AgentId],
+				*UE::Learning::Array::FormatFloat(InputView[AgentId]),
+				*UE::Learning::Array::FormatFloat(OutputView[AgentId]),
+				*UE::Learning::Array::FormatFloat(OutputMeanView[AgentId]),
+				*UE::Learning::Array::FormatFloat(OutputStdView[AgentId]),
+				*UE::Learning::Agents::Debug::FloatArrayToStatsString(InputView[AgentId]),
+				*UE::Learning::Agents::Debug::FloatArrayToStatsString(OutputView[AgentId]));
 
 		}
 	}
diff --git a/Engine/Plugins/Experimental/LearningAgents/Source/LearningAgents/Public/LearningAgentsActions.h b/Engine/Plugins/Experimental/LearningAgents/Source/LearningAgents/Public/LearningAgentsActions.h
--- a/Engine/Plugins/Experimental/LearningAgents/Source/LearningAgents/Public/LearningAgentsActions.h
+++ b/Engine/Plugins/Experimental/LearningAgents/Source/LearningAgents/Public/LearningAgentsActions.h
@@ -25,10 +25,10 @@ class ULearningAgentsInteractor;
 //------------------------------------------------------------------
 
 /**
-* The base class for all actions. Actions define the outputs from your agents. Action getters are marked non-pure by
-* convention as many of them do non-trivial amounts of work that can cause performance issues when marked pure in 
-* blueprints.
-*/
+ * The base class for all actions. Actions define the outputs from your agents. Action getters are marked non-pure by
+ * convention as many of them do non-trivial amounts of work that can cause performance issues when marked pure in 
+ * blueprints.
+ */
 UCLASS(Abstract, BlueprintType)
 class LEARNINGAGENTS_API ULearningAgentsAction : public UObject
 {
@@ -40,6 +40,15 @@ public:
 	UPROPERTY(VisibleAnywhere, Transient, Category = "LearningAgents")
 	TObjectPtr<ULearningAgentsInteractor> Interactor;
 
+public:
+
+	/** Number of times this action has been got for all agents */
+	TLearningArray<1, uint64, TInlineAllocator<32>> AgentGetIteration;
+
+	/** Number of times this action has been set for all agents */
+	TLearningArray<1, uint64, TInlineAllocator<32>> AgentSetIteration;
+
+public:
 #if UE_LEARNING_AGENTS_ENABLE_VISUAL_LOG
 	/** Color used to draw this action in the visual log */
 	FLinearColor VisualLogColor = FColor::Blue;
@@ -60,28 +69,28 @@ class LEARNINGAGENTS_API UFloatAction : public ULearningAgentsAction
 public:
 
 	/**
-	* Adds a new float action to the given agent interactor. Call during ULearningAgentsInteractor::SetupActions event.
-	* @param InInteractor The agent interactor to add this action to.
-	* @param Name The name of this new action. Used for debugging.
-	* @param Scale Used to normalize the data for the action.
-	* @return The newly created action.
-	*/
-	UFUNCTION(BlueprintCallable, Category = "LearningAgents")
+	 * Adds a new float action to the given agent interactor. Call during ULearningAgentsInteractor::SetupActions event.
+	 * @param InInteractor The agent interactor to add this action to.
+	 * @param Name The name of this new action. Used for debugging.
+	 * @param Scale Used to normalize the data for the action.
+	 * @return The newly created action.
+	 */
+	UFUNCTION(BlueprintCallable, Category = "LearningAgents", meta = (DefaultToSelf = "InInteractor"))
 	static UFloatAction* AddFloatAction(ULearningAgentsInteractor* InInteractor, const FName Name = NAME_None, const float Scale = 1.0f);
 
 	/**
-	* Gets the data for this action. Call during ULearningAgentsInteractor::GetActions event.
-	* @param AgentId The agent id to get data for.
-	* @return The current action value.
-	*/
-	UFUNCTION(BlueprintPure = false, Category = "LearningAgents", meta = (AgentId = "-1"))
-	float GetFloatAction(const int32 AgentId) const;
+	 * Gets the data for this action. Call during ULearningAgentsInteractor::GetActions event.
+	 * @param AgentId The agent id to get data for.
+	 * @return The current action value.
+	 */
+	UFUNCTION(BlueprintCallable, Category = "LearningAgents", meta = (AgentId = "-1"))
+	float GetFloatAction(const int32 AgentId);
 
 	/**
-	* Sets the data for this action. Call during ULearningAgentsController::EncodeActions event.
-	* @param AgentId The agent id to set data for.
-	* @param Value The current action value.
-	*/	
+	 * Sets the data for this action. Call during ULearningAgentsController::EncodeActions event.
+	 * @param AgentId The agent id to set data for.
+	 * @param Value The current action value.
+	 */	
 	UFUNCTION(BlueprintCallable, Category = "LearningAgents", meta = (AgentId = "-1"))
 	void SetFloatAction(const int32 AgentId, const float Value);
 
@@ -102,29 +111,29 @@ class LEARNINGAGENTS_API UFloatArrayAction : public ULearningAgentsAction
 public:
 
 	/**
-	* Adds a new float array action to the given agent interactor. Call during ULearningAgentsInteractor::SetupActions event.
-	* @param InInteractor The agent interactor to add this action to.
-	* @param Name The name of this new action. Used for debugging.
-	* @param Num The number of floats in the array
-	* @param Scale Used to normalize the data for the action.
-	* @return The newly created action.
-	*/
-	UFUNCTION(BlueprintCallable, Category = "LearningAgents")
+	 * Adds a new float array action to the given agent interactor. Call during ULearningAgentsInteractor::SetupActions event.
+	 * @param InInteractor The agent interactor to add this action to.
+	 * @param Name The name of this new action. Used for debugging.
+	 * @param Num The number of floats in the array
+	 * @param Scale Used to normalize the data for the action.
+	 * @return The newly created action.
+	 */
+	UFUNCTION(BlueprintCallable, Category = "LearningAgents", meta = (DefaultToSelf = "InInteractor"))
 	static UFloatArrayAction* AddFloatArrayAction(ULearningAgentsInteractor* InInteractor, const FName Name = NAME_None, const int32 Num = 1, const float Scale = 1.0f);
 
 	/**
-	* Gets the data for this action. Call during ULearningAgentsInteractor::GetActions event.
-	* @param AgentId The agent id to get data for.
-	* @param OutValues The output array of floats
-	*/
-	UFUNCTION(BlueprintPure = false, Category = "LearningAgents", meta = (AgentId = "-1"))
-	void GetFloatArrayAction(const int32 AgentId, TArray<float>& OutValues) const;
+	 * Gets the data for this action. Call during ULearningAgentsInteractor::GetActions event.
+	 * @param AgentId The agent id to get data for.
+	 * @param OutValues The output array of floats
+	 */
+	UFUNCTION(BlueprintCallable, Category = "LearningAgents", meta = (AgentId = "-1"))
+	void GetFloatArrayAction(const int32 AgentId, TArray<float>& OutValues);
 
 	/**
-	* Sets the data for this action. Call during ULearningAgentsController::EncodeActions event.
-	* @param AgentId The agent id to set data for.
-	* @param Values The current action values.
-	*/
+	 * Sets the data for this action. Call during ULearningAgentsController::EncodeActions event.
+	 * @param AgentId The agent id to set data for.
+	 * @param Values The current action values.
+	 */
 	UFUNCTION(BlueprintCallable, Category = "LearningAgents", meta = (AgentId = "-1"))
 	void SetFloatArrayAction(const int32 AgentId, const TArray<float>& Values);
 
@@ -147,28 +156,28 @@ class LEARNINGAGENTS_API UVectorAction : public ULearningAgentsAction
 public:
 
 	/**
-	* Adds a new vector action to the given agent interactor. Call during ULearningAgentsInteractor::SetupActions event.
-	* @param InInteractor The agent interactor to add this action to.
-	* @param Name The name of this new action. Used for debugging.
-	* @param Scale Used to normalize the data for the action.
-	* @return The newly created action.
-	*/
-	UFUNCTION(BlueprintCallable, Category = "LearningAgents")
+	 * Adds a new vector action to the given agent interactor. Call during ULearningAgentsInteractor::SetupActions event.
+	 * @param InInteractor The agent interactor to add this action to.
+	 * @param Name The name of this new action. Used for debugging.
+	 * @param Scale Used to normalize the data for the action.
+	 * @return The newly created action.
+	 */
+	UFUNCTION(BlueprintCallable, Category = "LearningAgents", meta = (DefaultToSelf = "InInteractor"))
 	static UVectorAction* AddVectorAction(ULearningAgentsInteractor* InInteractor, const FName Name = NAME_None, const float Scale = 1.0f);
 
 	/**
-	* Gets the data for this action. Call during ULearningAgentsInteractor::GetActions event.
-	* @param AgentId The agent id to get data for.
-	* @return The current action values.
-	*/
-	UFUNCTION(BlueprintPure = false, Category = "LearningAgents", meta = (AgentId = "-1"))
-	FVector GetVectorAction(const int32 AgentId) const;
+	 * Gets the data for this action. Call during ULearningAgentsInteractor::GetActions event.
+	 * @param AgentId The agent id to get data for.
+	 * @return The current action values.
+	 */
+	UFUNCTION(BlueprintCallable, Category = "LearningAgents", meta = (AgentId = "-1"))
+	FVector GetVectorAction(const int32 AgentId);
 
 	/**
-	* Sets the data for this action. Call during ULearningAgentsController::EncodeActions event.
-	* @param AgentId The agent id to set data for.
-	* @param Value The current action values.
-	*/
+	 * Sets the data for this action. Call during ULearningAgentsController::EncodeActions event.
+	 * @param AgentId The agent id to set data for.
+	 * @param Value The current action values.
+	 */
 	UFUNCTION(BlueprintCallable, Category = "LearningAgents", meta = (AgentId = "-1"))
 	void SetVectorAction(const int32 AgentId, const FVector Value);
 
@@ -189,29 +198,29 @@ class LEARNINGAGENTS_API UVectorArrayAction : public ULearningAgentsAction
 public:
 
 	/**
-	* Adds a new vector action to the given agent interactor. Call during ULearningAgentsInteractor::SetupActions event.
-	* @param InInteractor The agent interactor to add this action to.
-	* @param Name The name of this new action. Used for debugging.
-	* @param Num The number of vectors in the array
-	* @param Scale Used to normalize the data for the action.
-	* @return The newly created action.
-	*/
-	UFUNCTION(BlueprintCallable, Category = "LearningAgents")
+	 * Adds a new vector action to the given agent interactor. Call during ULearningAgentsInteractor::SetupActions event.
+	 * @param InInteractor The agent interactor to add this action to.
+	 * @param Name The name of this new action. Used for debugging.
+	 * @param Num The number of vectors in the array
+	 * @param Scale Used to normalize the data for the action.
+	 * @return The newly created action.
+	 */
+	UFUNCTION(BlueprintCallable, Category = "LearningAgents", meta = (DefaultToSelf = "InInteractor"))
 	static UVectorArrayAction* AddVectorArrayAction(ULearningAgentsInteractor* InInteractor, const FName Name = NAME_None, const int32 Num = 1, const float Scale = 1.0f);
 
 	/**
-	* Gets the data for this action. Call during ULearningAgentsInteractor::GetActions event.
-	* @param AgentId The agent id to get data for.
-	* @param OutVectors The current action values.
-	*/
-	UFUNCTION(BlueprintPure = false, Category = "LearningAgents", meta = (AgentId = "-1"))
-	void GetVectorArrayAction(const int32 AgentId, TArray<FVector>& OutVectors) const;
+	 * Gets the data for this action. Call during ULearningAgentsInteractor::GetActions event.
+	 * @param AgentId The agent id to get data for.
+	 * @param OutVectors The current action values.
+	 */
+	UFUNCTION(BlueprintCallable, Category = "LearningAgents", meta = (AgentId = "-1"))
+	void GetVectorArrayAction(const int32 AgentId, TArray<FVector>& OutVectors);
 
 	/**
-	* Sets the data for this action. Call during ULearningAgentsController::EncodeActions event.
-	* @param AgentId The agent id to set data for.
-	* @param Vectors The current action values.
-	*/
+	 * Sets the data for this action. Call during ULearningAgentsController::EncodeActions event.
+	 * @param AgentId The agent id to set data for.
+	 * @param Vectors The current action values.
+	 */
 	UFUNCTION(BlueprintCallable, Category = "LearningAgents", meta = (AgentId = "-1"))
 	void SetVectorArrayAction(const int32 AgentId, const TArray<FVector>& Vectors);
 
@@ -234,38 +243,38 @@ class LEARNINGAGENTS_API URotationAction : public ULearningAgentsAction
 public:
 
 	/**
-	* Adds a new rotation action to the given agent interactor. Call during ULearningAgentsInteractor::SetupActions event.
-	* @param InInteractor The agent interactor to add this action to.
-	* @param Name The name of this new action. Used for debugging.
-	* @param Scale Used to normalize the data for the action.
-	* @return The newly created action.
-	*/
-	UFUNCTION(BlueprintCallable, Category = "LearningAgents")
+	 * Adds a new rotation action to the given agent interactor. Call during ULearningAgentsInteractor::SetupActions event.
+	 * @param InInteractor The agent interactor to add this action to.
+	 * @param Name The name of this new action. Used for debugging.
+	 * @param Scale Used to normalize the data for the action.
+	 * @return The newly created action.
+	 */
+	UFUNCTION(BlueprintCallable, Category = "LearningAgents", meta = (DefaultToSelf = "InInteractor"))
 	static URotationAction* AddRotationAction(ULearningAgentsInteractor* InInteractor, const FName Name = NAME_None, const float Scale = 180.0f);
 
 	/**
-	* Gets the data for this action as a rotator. Call during ULearningAgentsInteractor::GetActions event.
-	* @param AgentId The agent id to get data for.
-	* @return The current action value.
-	*/
-	UFUNCTION(BlueprintPure = false, Category = "LearningAgents", meta = (AgentId = "-1"))
-	FRotator GetRotationAction(const int32 AgentId) const;
+	 * Gets the data for this action as a rotator. Call during ULearningAgentsInteractor::GetActions event.
+	 * @param AgentId The agent id to get data for.
+	 * @return The current action value.
+	 */
+	UFUNCTION(BlueprintCallable, Category = "LearningAgents", meta = (AgentId = "-1"))
+	FRotator GetRotationAction(const int32 AgentId);
 
 	/**
-	* Gets the data for this action as a rotation vector. Call during ULearningAgentsInteractor::GetActions event.
-	* @param AgentId The agent id to get data for.
-	* @return The current action value.
-	*/
-	UFUNCTION(BlueprintPure = false, Category = "LearningAgents", meta = (AgentId = "-1"))
-	FVector GetRotationActionAsRotationVector(const int32 AgentId) const;
+	 * Gets the data for this action as a rotation vector. Call during ULearningAgentsInteractor::GetActions event.
+	 * @param AgentId The agent id to get data for.
+	 * @return The current action value.
+	 */
+	UFUNCTION(BlueprintCallable, Category = "LearningAgents", meta = (AgentId = "-1"))
+	FVector GetRotationActionAsRotationVector(const int32 AgentId);
 
 	/**
-	* Gets the data for this action as a quaternion. Call during ULearningAgentsInteractor::GetActions event.
-	* @param AgentId The agent id to get data for.
-	* @return The current action value.
-	*/
-	UFUNCTION(BlueprintPure = false, Category = "LearningAgents", meta = (AgentId = "-1"))
-	FQuat GetRotationActionAsQuat(const int32 AgentId) const;
+	 * Gets the data for this action as a quaternion. Call during ULearningAgentsInteractor::GetActions event.
+	 * @param AgentId The agent id to get data for.
+	 * @return The current action value.
+	 */
+	UFUNCTION(BlueprintCallable, Category = "LearningAgents", meta = (AgentId = "-1"))
+	FQuat GetRotationActionAsQuat(const int32 AgentId);
 
 #if UE_LEARNING_AGENTS_ENABLE_VISUAL_LOG
 	/** Describes this action to the visual logger for debugging purposes. */
@@ -284,39 +293,39 @@ class LEARNINGAGENTS_API URotationArrayAction : public ULearningAgentsAction
 public:
 
 	/**
-	* Adds a new rotation array action to the given agent interactor. Call during ULearningAgentsInteractor::SetupActions event.
-	* @param InInteractor The agent interactor to add this action to.
-	* @param Name The name of this new action. Used for debugging.
-	* @param RotationNum The number of rotations in the array.
-	* @param Scale Used to normalize the data for the action.
-	* @return The newly created action.
-	*/
-	UFUNCTION(BlueprintCallable, Category = "LearningAgents")
+	 * Adds a new rotation array action to the given agent interactor. Call during ULearningAgentsInteractor::SetupActions event.
+	 * @param InInteractor The agent interactor to add this action to.
+	 * @param Name The name of this new action. Used for debugging.
+	 * @param RotationNum The number of rotations in the array.
+	 * @param Scale Used to normalize the data for the action.
+	 * @return The newly created action.
+	 */
+	UFUNCTION(BlueprintCallable, Category = "LearningAgents", meta = (DefaultToSelf = "InInteractor"))
 	static URotationArrayAction* AddRotationArrayAction(ULearningAgentsInteractor* InInteractor, const FName Name = NAME_None, const int32 RotationNum = 1, const float Scale = 180.0f);
 
 	/**
-	* Gets the data for this action as rotators. Call during ULearningAgentsInteractor::GetActions event.
-	* @param AgentId The agent id to get data for.
-	* @param OutRotations The current action values.
-	*/
-	UFUNCTION(BlueprintPure = false, Category = "LearningAgents", meta = (AgentId = "-1"))
-	void GetRotationArrayAction(const int32 AgentId, TArray<FRotator>& OutRotations) const;
+	 * Gets the data for this action as rotators. Call during ULearningAgentsInteractor::GetActions event.
+	 * @param AgentId The agent id to get data for.
+	 * @param OutRotations The current action values.
+	 */
+	UFUNCTION(BlueprintCallable, Category = "LearningAgents", meta = (AgentId = "-1"))
+	void GetRotationArrayAction(const int32 AgentId, TArray<FRotator>& OutRotations);
 
 	/**
-	* Gets the data for this action as rotation vectors. Call during ULearningAgentsInteractor::GetActions event.
-	* @param AgentId The agent id to get data for.
-	* @param OutRotationVectors The current action values.
-	*/
-	UFUNCTION(BlueprintPure = false, Category = "LearningAgents", meta = (AgentId = "-1"))
-	void GetRotationArrayActionAsRotationVectors(const int32 AgentId, TArray<FVector>& OutRotationVectors) const;
+	 * Gets the data for this action as rotation vectors. Call during ULearningAgentsInteractor::GetActions event.
+	 * @param AgentId The agent id to get data for.
+	 * @param OutRotationVectors The current action values.
+	 */
+	UFUNCTION(BlueprintCallable, Category = "LearningAgents", meta = (AgentId = "-1"))
+	void GetRotationArrayActionAsRotationVectors(const int32 AgentId, TArray<FVector>& OutRotationVectors);
 
 	/**
-	* Gets the data for this action as quaternions. Call during ULearningAgentsInteractor::GetActions event.
-	* @param AgentId The agent id to get data for.
-	* @param OutRotations The current action values.
-	*/
-	UFUNCTION(BlueprintPure = false, Category = "LearningAgents", meta = (AgentId = "-1"))
-	void GetRotationArrayActionAsQuats(const int32 AgentId, TArray<FQuat>& OutRotations) const;
+	 * Gets the data for this action as quaternions. Call during ULearningAgentsInteractor::GetActions event.
+	 * @param AgentId The agent id to get data for.
+	 * @param OutRotations The current action values.
+	 */
+	UFUNCTION(BlueprintCallable, Category = "LearningAgents", meta = (AgentId = "-1"))
+	void GetRotationArrayActionAsQuats(const int32 AgentId, TArray<FQuat>& OutRotations);
 
 #if UE_LEARNING_AGENTS_ENABLE_VISUAL_LOG
 	/** Describes this action to the visual logger for debugging purposes. */
diff --git a/Engine/Plugins/Experimental/LearningAgents/Source/LearningAgents/Public/LearningAgentsController.h b/Engine/Plugins/Experimental/LearningAgents/Source/LearningAgents/Public/LearningAgentsController.h
--- a/Engine/Plugins/Experimental/LearningAgents/Source/LearningAgents/Public/LearningAgentsController.h
+++ b/Engine/Plugins/Experimental/LearningAgents/Source/LearningAgents/Public/LearningAgentsController.h
@@ -7,10 +7,10 @@
 #include "LearningAgentsController.generated.h"
 
 /**
-* A controller provides a method for injecting actions into the learning agents system from some other existing
-* behavior, e.g. we may want to gather demonstrations from a human or AI behavior tree controlling our agent(s)
-* for imitation learning purposes.
-*/
+ * A controller provides a method for injecting actions into the learning agents system from some other existing
+ * behavior, e.g. we may want to gather demonstrations from a human or AI behavior tree controlling our agent(s)
+ * for imitation learning purposes.
+ */
 UCLASS(Abstract, BlueprintType, Blueprintable)
 class LEARNINGAGENTS_API ULearningAgentsController : public ULearningAgentsManagerComponent
 {
@@ -19,42 +19,42 @@ class LEARNINGAGENTS_API ULearningAgentsController : public ULearningAgentsManag
 public:
 
 	// These constructors/destructors are needed to make forward declarations happy
-	ULearningAgentsController();
+	ULearningAgentsController(const FObjectInitializer& ObjectInitializer = FObjectInitializer::Get());
 	ULearningAgentsController(FVTableHelper& Helper);
 	virtual ~ULearningAgentsController();
 
-	/** Initializes this object to be used with the given manager and agent interactor. */
+	/** Initializes this object to be used with the given agent interactor. */
 	UFUNCTION(BlueprintCallable, Category = "LearningAgents")
-	void SetupController(ALearningAgentsManager* InAgentManager, ULearningAgentsInteractor* InInteractor);
+	void SetupController(ULearningAgentsInteractor* InInteractor);
 
 	/**
-	* During this event, you should set the actions of your agents.
-	* @param AgentIds The list of agent ids to set actions for.
-	* @see LearningAgentsActions.h for the list of available actions.
-	* @see ULearningAgentsInteractor::GetAgent to get the agent corresponding to each id.
-	*/
+	 * During this event, you should set the actions of your agents.
+	 * @param AgentIds The list of agent ids to set actions for.
+	 * @see LearningAgentsActions.h for the list of available actions.
+	 * @see ULearningAgentsInteractor::GetAgent to get the agent corresponding to each id.
+	 */
 	UFUNCTION(BlueprintNativeEvent, Category = "LearningAgents")
 	void SetActions(const TArray<int32>& AgentIds);
 
 	/**
-	* Call this function when it is time to gather all the actions for your agents. This should be called roughly 
-	* whenever you are calling ULearningAgentsInteractor::EncodeObservations. This will call this controller's SetActions event.
-	*/
+	 * Call this function when it is time to gather all the actions for your agents. This should be called roughly 
+	 * whenever you are calling ULearningAgentsInteractor::EncodeObservations. This will call this controller's SetActions event.
+	 */
 	UFUNCTION(BlueprintCallable, Category = "LearningAgents")
 	void EncodeActions();
 
 	/**
-	* Calls EncodeObservations, followed by EncodeActions, followed by DecodeActions
-	*/
+	 * Calls EncodeObservations, followed by EncodeActions, followed by DecodeActions
+	 */
 	UFUNCTION(BlueprintCallable, Category = "LearningAgents")
 	void RunController();
 
 protected:
 
 	/**
-	* Gets the agent interactor associated with this component.
-	* @param AgentClass The class to cast the agent interactor to (in blueprint).
-	*/
+	 * Gets the agent interactor associated with this component.
+	 * @param AgentClass The class to cast the agent interactor to (in blueprint).
+	 */
 	UFUNCTION(BlueprintPure, Category = "LearningAgents", meta = (DeterminesOutputType = "InteractorClass"))
 	ULearningAgentsInteractor* GetInteractor(const TSubclassOf<ULearningAgentsInteractor> InteractorClass) const;
 
@@ -63,4 +63,9 @@ private:
 	/** The agent interactor this controller is associated with. */
 	UPROPERTY(VisibleAnywhere, Transient, Category = "LearningAgents")
 	TObjectPtr<ULearningAgentsInteractor> Interactor;
+
+	/** Temp buffers used to record the set of agents that are valid for encoding */
+	TArray<int32> ValidAgentIds;
+	UE::Learning::FIndexSet ValidAgentSet;
+	TBitArray<> ValidAgentStatus;
 };
diff --git a/Engine/Plugins/Experimental/LearningAgents/Source/LearningAgents/Public/LearningAgentsCritic.h b/Engine/Plugins/Experimental/LearningAgents/Source/LearningAgents/Public/LearningAgentsCritic.h
--- a/Engine/Plugins/Experimental/LearningAgents/Source/LearningAgents/Public/LearningAgentsCritic.h
+++ b/Engine/Plugins/Experimental/LearningAgents/Source/LearningAgents/Public/LearningAgentsCritic.h
@@ -50,73 +50,99 @@ class LEARNINGAGENTS_API ULearningAgentsCritic : public ULearningAgentsManagerCo
 public:
 
 	// These constructors/destructors are needed to make forward declarations happy
-	ULearningAgentsCritic();
+	ULearningAgentsCritic(const FObjectInitializer& ObjectInitializer = FObjectInitializer::Get());
 	ULearningAgentsCritic(FVTableHelper& Helper);
 	virtual ~ULearningAgentsCritic();
 
-	/** Initializes this object to be used with the given agent interactor and critic settings. */
+	/**
+	 * Initializes this object to be used with the given agent interactor and critic settings.
+	 * @param InInteractor The input Interactor component
+	 * @param CriticSettings The critic settings to use
+	 * @param NeuralNetworkAsset Optional Network Asset to use. If provided must match the given CriticSettings. If not
+	 * provided or asset is empty then a new neural network object will be created according to the given 
+	 * CriticSettings and used.
+	 */
 	UFUNCTION(BlueprintCallable, Category = "LearningAgents")
 	void SetupCritic(
-		ALearningAgentsManager* InAgentManager,
 		ULearningAgentsInteractor* InInteractor,
-		const FLearningAgentsCriticSettings& CriticSettings = FLearningAgentsCriticSettings());
+		const FLearningAgentsCriticSettings& CriticSettings = FLearningAgentsCriticSettings(),
+		ULearningAgentsNeuralNetwork* NeuralNetworkAsset = nullptr);
+
+public:
+
+	//~ Begin ULearningAgentsManagerComponent Interface
+	virtual void OnAgentsAdded(const TArray<int32>& AgentIds) override;
+	virtual void OnAgentsRemoved(const TArray<int32>& AgentIds) override;
+	virtual void OnAgentsReset(const TArray<int32>& AgentIds) override;
+	//~ End ULearningAgentsManagerComponent Interface
 
 // ----- Load / Save -----
 public:
 
 	/**
-	* Load a snapshot's weights into this critic.
-	* @param File The snapshot file.
-	*/
+	 * Load a snapshot's weights into this critic.
+	 * @param File The snapshot file.
+	 */
 	UFUNCTION(BlueprintCallable, Category = "LearningAgents", meta = (RelativePath))
 	void LoadCriticFromSnapshot(const FFilePath& File);
 
 	/**
-	* Save this critic's weights into a snapshot.
-	* @param File The snapshot file.
-	*/
+	 * Save this critic's weights into a snapshot.
+	 * @param File The snapshot file.
+	 */
 	UFUNCTION(BlueprintCallable, Category = "LearningAgents", meta = (RelativePath))
 	void SaveCriticToSnapshot(const FFilePath& File) const;
 
 	/**
-	* Load a ULearningAgentsNeuralNetwork asset's weights into this critic.
-	* @param NeuralNetworkAsset The asset to load from.
-	*/
+	 * Use a ULearningAgentsNeuralNetwork asset directly for this critic rather than making a copy. If used 
+	 * during training then this asset's weights will be updated as training progresses.
+	 * @param NeuralNetworkAsset The asset to use.
+	 */
+	UFUNCTION(BlueprintCallable, Category = "LearningAgents")
+	void UseCriticFromAsset(ULearningAgentsNeuralNetwork* NeuralNetworkAsset);
+
+	/**
+	 * Load a ULearningAgentsNeuralNetwork asset's weights into this critic.
+	 * @param NeuralNetworkAsset The asset to load from.
+	 */
 	UFUNCTION(BlueprintCallable, Category = "LearningAgents")
-	void LoadCriticFromAsset(const ULearningAgentsNeuralNetwork* NeuralNetworkAsset);
+	void LoadCriticFromAsset(ULearningAgentsNeuralNetwork* NeuralNetworkAsset);
 
 	/**
-	* Save this critic's weights to a ULearningAgentsNeuralNetwork asset.
-	* @param NeuralNetworkAsset The asset to save to.
-	*/
+	 * Save this critic's weights to a ULearningAgentsNeuralNetwork asset.
+	 * @param NeuralNetworkAsset The asset to save to.
+	 */
 	UFUNCTION(BlueprintCallable, Category = "LearningAgents", Meta=(DevelopmentOnly))
-	void SaveCriticToAsset(ULearningAgentsNeuralNetwork* NeuralNetworkAsset) const;
+	void SaveCriticToAsset(ULearningAgentsNeuralNetwork* NeuralNetworkAsset);
 
 // ----- Evaluation -----
 public:
 
 	/**
-	* Calling this function will run the underlying neural network on the previously buffered observations to populate
-	* the output value buffer. This should be called after the corresponding agent interactor's EncodeObservations.
-	*/
+	 * Calling this function will run the underlying neural network on the previously buffered observations to populate
+	 * the output value buffer. This should be called after the corresponding agent interactor's EncodeObservations.
+	 */
 	UFUNCTION(BlueprintCallable, Category = "LearningAgents")
 	void EvaluateCritic();
 
 	/**
-	* Gets an estimate of the average discounted return expected by an agent according to the critic. I.E. the total
-	* sum of future rewards, scaled by the discount factor that was used during training. This value can be useful if 
-	* you want to make some decision based on how well the agent thinks they are doing at achieving their task. This 
-	* should be called only after EvaluateCritic.
-	* 
-	* @param AgentId	The AgentId to look-up the estimated discounted return for
-	* @returns			The estimated average discounted return according to the critic
-	*/
+	 * Gets an estimate of the average discounted return expected by an agent according to the critic. I.E. the total
+	 * sum of future rewards, scaled by the discount factor that was used during training. This value can be useful if 
+	 * you want to make some decision based on how well the agent thinks they are doing at achieving their task. This 
+	 * should be called only after EvaluateCritic.
+	 * 
+	 * @param AgentId	The AgentId to look-up the estimated discounted return for
+	 * @returns			The estimated average discounted return according to the critic
+	 */
 	UFUNCTION(BlueprintPure, Category = "LearningAgents")
 	float GetEstimatedDiscountedReturn(const int32 AgentId) const;
 
 // ----- Non-blueprint public interface -----
 public:
 
+	/** Gets the current Network Asset being used */
+	ULearningAgentsNeuralNetwork* GetNetworkAsset();
+
 	/** Get a reference to this critic's neural network. */
 	UE::Learning::FNeuralNetwork& GetCriticNetwork();
 	
@@ -134,15 +160,25 @@ private:
 	UPROPERTY(VisibleAnywhere, Transient, Category = "LearningAgents")
 	TObjectPtr<ULearningAgentsNeuralNetwork> Network;
 
+	/** Internal Critic Function Object */
+	TSharedPtr<UE::Learning::FNeuralNetworkCriticFunction> CriticObject;
+	
+// ----- Private Iteration Checks ----- 
 private:
 
-	TSharedPtr<UE::Learning::FNeuralNetworkCriticFunction> CriticObject;
+	/** Number of times critic has been evaluated for all agents */
+	TLearningArray<1, uint64, TInlineAllocator<32>> DiscountedReturnAgentIteration;
+
+	/** Temp buffers used to record the set of agents that are valid for evaluation */
+	TArray<int32> ValidAgentIds;
+	UE::Learning::FIndexSet ValidAgentSet;
 
+private:
 #if UE_LEARNING_AGENTS_ENABLE_VISUAL_LOG
 	/** Color used to draw this action in the visual log */
 	FLinearColor VisualLogColor = FColor::Orange;
 
 	/** Describes this critic to the visual logger for debugging purposes. */
-	void VisualLog(const UE::Learning::FIndexSet Instances) const;
+	void VisualLog(const UE::Learning::FIndexSet AgentSet) const;
 #endif
 };
diff --git a/Engine/Plugins/Experimental/LearningAgents/Source/LearningAgents/Public/LearningAgentsHelpers.h b/Engine/Plugins/Experimental/LearningAgents/Source/LearningAgents/Public/LearningAgentsHelpers.h
--- a/Engine/Plugins/Experimental/LearningAgents/Source/LearningAgents/Public/LearningAgentsHelpers.h
+++ b/Engine/Plugins/Experimental/LearningAgents/Source/LearningAgents/Public/LearningAgentsHelpers.h
@@ -16,9 +16,9 @@ class UMeshComponent;
 //------------------------------------------------------------------
 
 /**
-* The base class for all helpers. Helpers are additional objects that can be used in getting or setting observations,
-* actions, rewards, and completions.
-*/
+ * The base class for all helpers. Helpers are additional objects that can be used in getting or setting observations,
+ * actions, rewards, and completions.
+ */
 UCLASS(Abstract, BlueprintType)
 class LEARNINGAGENTS_API ULearningAgentsHelper : public UObject
 {
@@ -47,90 +47,90 @@ class LEARNINGAGENTS_API USplineComponentHelper : public ULearningAgentsHelper
 public:
 
 	/**
-	* Adds a new spline component helper to the given manager component.
-	* @param InManagerComponent The manager component to add this helper to (such as an Interactor or Trainer).
-	* @param Name The name of this new helper. Used for debugging.
-	* @return The newly created helper.
-	*/
-	UFUNCTION(BlueprintCallable, Category = "LearningAgents")
+	 * Adds a new spline component helper to the given manager component.
+	 * @param InManagerComponent The manager component to add this helper to (such as an Interactor or Trainer).
+	 * @param Name The name of this new helper. Used for debugging.
+	 * @return The newly created helper.
+	 */
+	UFUNCTION(BlueprintCallable, Category = "LearningAgents", meta = (DefaultToSelf = "InManagerComponent"))
 	static USplineComponentHelper* AddSplineComponentHelper(ULearningAgentsManagerComponent* InManagerComponent, const FName Name = NAME_None);
 
 	/**
-	* Gets the position on a spline closest to the provided position.
-	* @param AgentId The agent id to run the helper for.
-	* @param SplineComponent The spline component to use.
-	* @param Position The position to find the closest position to.
-	* @param CoordinateSpace The coordinate space to use for the spline.
-	* @return The closest position on the spline.
-	*/
+	 * Gets the position on a spline closest to the provided position.
+	 * @param AgentId The agent id to run the helper for.
+	 * @param SplineComponent The spline component to use.
+	 * @param Position The position to find the closest position to.
+	 * @param CoordinateSpace The coordinate space to use for the spline.
+	 * @return The closest position on the spline.
+	 */
 	UFUNCTION(BlueprintPure = false, Category = "LearningAgents", meta = (AgentId = "-1"))
 	FVector GetNearestPositionOnSpline(const int32 AgentId, const USplineComponent* SplineComponent, const FVector Position, const ESplineCoordinateSpace::Type CoordinateSpace = ESplineCoordinateSpace::World) const;
 
 	/**
-	* Gets the distance along a spline closest to the provided position.
-	* @param AgentId The agent id to run the helper for.
-	* @param SplineComponent The spline component to use.
-	* @param Position The position to find the closest position to.
-	* @param CoordinateSpace The coordinate space to use for the spline.
-	* @return The distance along the spline of the closest position.
-	*/
+	 * Gets the distance along a spline closest to the provided position.
+	 * @param AgentId The agent id to run the helper for.
+	 * @param SplineComponent The spline component to use.
+	 * @param Position The position to find the closest position to.
+	 * @param CoordinateSpace The coordinate space to use for the spline.
+	 * @return The distance along the spline of the closest position.
+	 */
 	UFUNCTION(BlueprintPure = false, Category = "LearningAgents", meta = (AgentId = "-1"))
 	float GetDistanceAlongSplineAtPosition(const int32 AgentId, const USplineComponent* SplineComponent, const FVector Position, const ESplineCoordinateSpace::Type CoordinateSpace = ESplineCoordinateSpace::World) const;
 
 	/**
-	* Gets the position along a spline at the given distance.
-	* @param AgentId The agent id to run the helper for.
-	* @param SplineComponent The spline component to use.
-	* @param DistanceAlongSpline The distance along the spline to use.
-	* @param CoordinateSpace The coordinate space to use for the spline.
-	* @return The position at the given distance along the spline.
-	*/
+	 * Gets the position along a spline at the given distance.
+	 * @param AgentId The agent id to run the helper for.
+	 * @param SplineComponent The spline component to use.
+	 * @param DistanceAlongSpline The distance along the spline to use.
+	 * @param CoordinateSpace The coordinate space to use for the spline.
+	 * @return The position at the given distance along the spline.
+	 */
 	UFUNCTION(BlueprintPure = false, Category = "LearningAgents", meta = (AgentId = "-1"))
 	FVector GetPositionAtDistanceAlongSpline(const int32 AgentId, const USplineComponent* SplineComponent, const float DistanceAlongSpline, const ESplineCoordinateSpace::Type CoordinateSpace = ESplineCoordinateSpace::World) const;
 
 	/**
-	* Gets the direction along a spline at the given distance.
-	* @param AgentId The agent id to run the helper for.
-	* @param SplineComponent The spline component to use.
-	* @param DistanceAlongSpline The distance along the spline to use.
-	* @param CoordinateSpace The coordinate space to use for the spline.
-	* @return The direction at the given distance along the spline.
-	*/
+	 * Gets the direction along a spline at the given distance.
+	 * @param AgentId The agent id to run the helper for.
+	 * @param SplineComponent The spline component to use.
+	 * @param DistanceAlongSpline The distance along the spline to use.
+	 * @param CoordinateSpace The coordinate space to use for the spline.
+	 * @return The direction at the given distance along the spline.
+	 */
 	UFUNCTION(BlueprintPure = false, Category = "LearningAgents", meta = (AgentId = "-1"))
 	FVector GetDirectionAtDistanceAlongSpline(const int32 AgentId, const USplineComponent* SplineComponent, const float DistanceAlongSpline, const ESplineCoordinateSpace::Type CoordinateSpace = ESplineCoordinateSpace::World) const;
 
 	/**
-	* Gets the proportion along a spline in the range 0-1 for a given distance.
-	* @param AgentId The agent id to run the helper for.
-	* @param SplineComponent The spline component to use.
-	* @param DistanceAlongSpline The distance along the spline to use.
-	* @return The proportion along the spline in the range 0-1.
-	*/
+	 * Gets the proportion along a spline in the range 0-1 for a given distance.
+	 * @param AgentId The agent id to run the helper for.
+	 * @param SplineComponent The spline component to use.
+	 * @param DistanceAlongSpline The distance along the spline to use.
+	 * @return The proportion along the spline in the range 0-1.
+	 */
 	UFUNCTION(BlueprintPure = false, Category = "LearningAgents", meta = (AgentId = "-1"))
 	float GetProportionAlongSpline(const int32 AgentId, const USplineComponent* SplineComponent, const float DistanceAlongSpline) const;
 
 	/**
-	* Gets the proportion along a spline encoded as an angle between -180 and 180 degrees for a given distance. 
-	* This should be used for looped splines.
-	* @param AgentId The agent id to run the helper for.
-	* @param SplineComponent The spline component to use.
-	* @param DistanceAlongSpline The distance along the spline to use.
-	* @return The proportion along the spline encoded as an angle between -180 and 180 degrees
-	*/
+	 * Gets the proportion along a spline encoded as an angle between -180 and 180 degrees for a given distance. 
+	 * This should be used for looped splines.
+	 * @param AgentId The agent id to run the helper for.
+	 * @param SplineComponent The spline component to use.
+	 * @param DistanceAlongSpline The distance along the spline to use.
+	 * @return The proportion along the spline encoded as an angle between -180 and 180 degrees
+	 */
 	UFUNCTION(BlueprintPure = false, Category = "LearningAgents", meta = (AgentId = "-1"))
 	float GetProportionAlongSplineAsAngle(const int32 AgentId, const USplineComponent* SplineComponent, const float DistanceAlongSpline) const;
 
 	/**
-	* Gets an array of positions sampled along a spline between some starting and stopping distance. Deals properly
-	* for splines which are looped.
-	* @param OutPositions Output array of spline positions.
-	* @param AgentId The agent id to run the helper for.
-	* @param SplineComponent The spline component to use.
-	* @param PositionNum The number of positions to sample along the spline.
-	* @param StartDistanceAlongSpline The starting distance along the spline.
-	* @param StopDistanceAlongSpline The stopping distance along the spline.
-	* @param CoordinateSpace The coordinate space to use for the spline.
-	*/
+	 * Gets an array of positions sampled along a spline between some starting and stopping distance. Deals properly
+	 * for splines which are looped.
+	 * @param OutPositions Output array of spline positions.
+	 * @param AgentId The agent id to run the helper for.
+	 * @param SplineComponent The spline component to use.
+	 * @param PositionNum The number of positions to sample along the spline.
+	 * @param StartDistanceAlongSpline The starting distance along the spline.
+	 * @param StopDistanceAlongSpline The stopping distance along the spline.
+	 * @param CoordinateSpace The coordinate space to use for the spline.
+	 */
 	UFUNCTION(BlueprintPure = false, Category = "LearningAgents", meta = (AgentId = "-1"))
 	void GetPositionsAlongSpline(
 		TArray<FVector>& OutPositions,
@@ -142,15 +142,15 @@ public:
 		const ESplineCoordinateSpace::Type CoordinateSpace = ESplineCoordinateSpace::World) const;
 
 	/**
-	* Gets the velocity a point is traveling along a spline.
-	* @param AgentId The agent id to run the helper for.
-	* @param SplineComponent The spline component to use.
-	* @param Position The position of the point.
-	* @param Velocity The velocity of the point.
-	* @param FiniteDifferenceDelta The finite difference delta used to estimate the velocity along the spline.
-	* @param CoordinateSpace The coordinate space to use for the spline.
-	* @return The scalar velocity the point is traveling along the spline.
-	*/
+	 * Gets the velocity a point is traveling along a spline.
+	 * @param AgentId The agent id to run the helper for.
+	 * @param SplineComponent The spline component to use.
+	 * @param Position The position of the point.
+	 * @param Velocity The velocity of the point.
+	 * @param FiniteDifferenceDelta The finite difference delta used to estimate the velocity along the spline.
+	 * @param CoordinateSpace The coordinate space to use for the spline.
+	 * @return The scalar velocity the point is traveling along the spline.
+	 */
 	UFUNCTION(BlueprintPure = false, Category = "LearningAgents", meta = (AgentId = "-1"))
 	float GetVelocityAlongSpline(
 		const int32 AgentId,
@@ -172,33 +172,33 @@ class LEARNINGAGENTS_API UProjectionHelper : public ULearningAgentsHelper
 public:
 
 	/**
-	* Adds a new projection helper to the given manager component.
-	* @param InManagerComponent The manager component to add this helper to (such as an Interactor or Trainer).
-	* @param Name The name of this new helper. Used for debugging.
-	* @return The newly created helper.
-	*/
-	UFUNCTION(BlueprintCallable, Category = "LearningAgents")
+	 * Adds a new projection helper to the given manager component.
+	 * @param InManagerComponent The manager component to add this helper to (such as an Interactor or Trainer).
+	 * @param Name The name of this new helper. Used for debugging.
+	 * @return The newly created helper.
+	 */
+	UFUNCTION(BlueprintCallable, Category = "LearningAgents", meta = (DefaultToSelf = "InManagerComponent"))
 	static UProjectionHelper* AddProjectionHelper(ULearningAgentsManagerComponent* InManagerComponent, const FName Name = NAME_None);
 
 	/**
-	* Projects a transform onto the zero height ground plane resulting in translation only in XY and rotation only around the Z.
-	* @param AgentId The agent id to run the helper for.
-	* @param Transform The transform to project.
-	* @param LocalForwardVector The local forward direction.
-	* @return The transform projected onto the ground plane.
-	*/
+	 * Projects a transform onto the zero height ground plane resulting in translation only in XY and rotation only around the Z.
+	 * @param AgentId The agent id to run the helper for.
+	 * @param Transform The transform to project.
+	 * @param LocalForwardVector The local forward direction.
+	 * @return The transform projected onto the ground plane.
+	 */
 	UFUNCTION(BlueprintPure = false, Category = "LearningAgents", meta = (AgentId = "-1"))
 	FTransform ProjectTransformOntoGroundPlane(const int32 AgentId, const FTransform Transform, const FVector LocalForwardVector = FVector::ForwardVector) const;
 
 	/**
-	* Projects a position and rotation onto the zero height ground plane resulting in translation only in XY and rotation only around the Z.
-	* @param OutPosition The output projected position.
-	* @param OutRotation The output projected rotation.
-	* @param AgentId The agent id to run the helper for.
-	* @param InPosition The input position.
-	* @param InRotation The input rotation.
-	* @param LocalForwardVector The local forward direction.
-	*/
+	 * Projects a position and rotation onto the zero height ground plane resulting in translation only in XY and rotation only around the Z.
+	 * @param OutPosition The output projected position.
+	 * @param OutRotation The output projected rotation.
+	 * @param AgentId The agent id to run the helper for.
+	 * @param InPosition The input position.
+	 * @param InRotation The input rotation.
+	 * @param LocalForwardVector The local forward direction.
+	 */
 	UFUNCTION(BlueprintPure = false, Category = "LearningAgents", meta = (AgentId = "-1"))
 	void ProjectPositionRotationOntoGroundPlane(
 		FVector& OutPosition, 
@@ -220,21 +220,21 @@ class LEARNINGAGENTS_API UMeshComponentHelper : public ULearningAgentsHelper
 public:
 
 	/**
-	* Adds a new mesh component helper to the given manager component.
-	* @param InManagerComponent The manager component to add this helper to (such as an Interactor or Trainer).
-	* @param Name The name of this new helper. Used for debugging.
-	* @return The newly created helper.
-	*/
-	UFUNCTION(BlueprintCallable, Category = "LearningAgents")
+	 * Adds a new mesh component helper to the given manager component.
+	 * @param InManagerComponent The manager component to add this helper to (such as an Interactor or Trainer).
+	 * @param Name The name of this new helper. Used for debugging.
+	 * @return The newly created helper.
+	 */
+	UFUNCTION(BlueprintCallable, Category = "LearningAgents", meta = (DefaultToSelf = "InManagerComponent"))
 	static UMeshComponentHelper* AddMeshComponentHelper(ULearningAgentsManagerComponent* InManagerComponent, const FName Name = NAME_None);
 
 	/**
-	* Gets the bone positions for a set of bones of the mesh component in world space.
-	* @param OutBonePositions The output array of bone positions.
-	* @param AgentId The agent id to run the helper for.
-	* @param MeshComponent The mesh component to use.
-	* @param BoneNames The names of the bones to use.
-	*/
+	 * Gets the bone positions for a set of bones of the mesh component in world space.
+	 * @param OutBonePositions The output array of bone positions.
+	 * @param AgentId The agent id to run the helper for.
+	 * @param MeshComponent The mesh component to use.
+	 * @param BoneNames The names of the bones to use.
+	 */
 	UFUNCTION(BlueprintPure = false, Category = "LearningAgents", meta = (AgentId = "-1"))
 	void GetMeshBonePositions(TArray<FVector>& OutBonePositions, const int32 AgentId, const UMeshComponent* MeshComponent, const TArray<FName>& BoneNames) const;
 };
@@ -250,28 +250,28 @@ class LEARNINGAGENTS_API URayCastHelper : public ULearningAgentsHelper
 public:
 
 	/**
-	* Adds a new ray cast helper to the given manager component.
-	* @param InManagerComponent The manager component to add this helper to (such as an Interactor or Trainer).
-	* @param Name The name of this new helper. Used for debugging.
-	* @return The newly created helper.
-	*/
-	UFUNCTION(BlueprintCallable, Category = "LearningAgents")
+	 * Adds a new ray cast helper to the given manager component.
+	 * @param InManagerComponent The manager component to add this helper to (such as an Interactor or Trainer).
+	 * @param Name The name of this new helper. Used for debugging.
+	 * @return The newly created helper.
+	 */
+	UFUNCTION(BlueprintCallable, Category = "LearningAgents", meta = (DefaultToSelf = "InManagerComponent"))
 	static URayCastHelper* AddRayCastHelper(ULearningAgentsManagerComponent* InManagerComponent, const FName Name = NAME_None);
 
 	/**
-	* Samples a grid of heights from above.
-	* @param OutHeights The output array of heights for each ray in the grid.
-	* @param AgentId The agent id to run the helper for.
-	* @param Position The central position of the grid.
-	* @param Rotation The orientation of the grid.
-	* @param RowNum The number of rows in the grid.
-	* @param ColNum The number of columns in the grid.
-	* @param RowWidth The width of grid rows.
-	* @param ColWidth The width of grid columns.
-	* @param MaxHeight The height at which to start ray casts from.
-	* @param MinHeight The height at which to send ray casts to.
-	* @param CollisionChannel The collision channel to use for the ray-casts.
-	*/
+	 * Samples a grid of heights from above.
+	 * @param OutHeights The output array of heights for each ray in the grid.
+	 * @param AgentId The agent id to run the helper for.
+	 * @param Position The central position of the grid.
+	 * @param Rotation The orientation of the grid.
+	 * @param RowNum The number of rows in the grid.
+	 * @param ColNum The number of columns in the grid.
+	 * @param RowWidth The width of grid rows.
+	 * @param ColWidth The width of grid columns.
+	 * @param MaxHeight The height at which to start ray casts from.
+	 * @param MinHeight The height at which to send ray casts to.
+	 * @param CollisionChannel The collision channel to use for the ray-casts.
+	 */
 	UFUNCTION(BlueprintPure=false, Category = "LearningAgents", meta = (AgentId = "-1"))
 	void RayCastGridHeights(
 		TArray<float>& OutHeights,
diff --git a/Engine/Plugins/Experimental/LearningAgents/Source/LearningAgents/Public/LearningAgentsInteractor.h b/Engine/Plugins/Experimental/LearningAgents/Source/LearningAgents/Public/LearningAgentsInteractor.h
--- a/Engine/Plugins/Experimental/LearningAgents/Source/LearningAgents/Public/LearningAgentsInteractor.h
+++ b/Engine/Plugins/Experimental/LearningAgents/Source/LearningAgents/Public/LearningAgentsInteractor.h
@@ -26,14 +26,14 @@ class ULearningAgentsAction;
 class ULearningAgentsObservation;
 
 /**
-* ULearningAgentsInteractor defines how agents interact with the environment through their observations and actions.
-* Additionally, it provides methods to be called during the inference process of those agents.
-*
-* To use this class, you need to implement the SetupObservations and SetupActions functions (as well as their
-* corresponding SetObservations and SetActions functions), which will define the size of inputs and outputs to your
-* policy. Before you can do inference, you need to call SetupInteractor, which will initialize the underlying data
-* structure, and you need to call AddAgent for each object you want controlled by this agent interactor.
-*/
+ * ULearningAgentsInteractor defines how agents interact with the environment through their observations and actions.
+ * Additionally, it provides methods to be called during the inference process of those agents.
+ *
+ * To use this class, you need to implement the SetupObservations and SetupActions functions (as well as their
+ * corresponding SetObservations and SetActions functions), which will define the size of inputs and outputs to your
+ * policy. Before you can do inference, you need to call SetupInteractor, which will initialize the underlying data
+ * structure, and you need to call AddAgent for each object you want controlled by this agent interactor.
+ */
 UCLASS(Abstract, BlueprintType, Blueprintable)
 class LEARNINGAGENTS_API ULearningAgentsInteractor : public ULearningAgentsManagerComponent
 {
@@ -42,38 +42,45 @@ class LEARNINGAGENTS_API ULearningAgentsInteractor : public ULearningAgentsManag
 public:
 
 	// These constructors/destructors are needed to make forward declarations happy
-	ULearningAgentsInteractor();
+	ULearningAgentsInteractor(const FObjectInitializer& ObjectInitializer = FObjectInitializer::Get());
 	ULearningAgentsInteractor(FVTableHelper& Helper);
 	virtual ~ULearningAgentsInteractor();
 
 	/** Initializes this object and runs the setup events for observations and actions. */
 	UFUNCTION(BlueprintCallable, Category = "LearningAgents")
-	void SetupInteractor(ALearningAgentsManager* InAgentManager);
+	void SetupInteractor();
+
+public:
+
+	//~ Begin ULearningAgentsManagerComponent Interface
+	virtual void OnAgentsAdded(const TArray<int32>& AgentIds) override;
+	virtual void OnAgentsRemoved(const TArray<int32>& AgentIds) override;
+	virtual void OnAgentsReset(const TArray<int32>& AgentIds) override;
+	//~ End ULearningAgentsManagerComponent Interface
 
 // ----- Observations -----
 public:
 
 	/**
 	 * During this event, all observations should be added to this agent interactor.
-	 * @param Interactor This agent interactor. This is a convenience for blueprints so you don't have to find the Self pin.
 	 * @see LearningAgentsObservations.h for the list of available observations.
 	 */
 	UFUNCTION(BlueprintNativeEvent, Category = "LearningAgents")
-	void SetupObservations(ULearningAgentsInteractor* Interactor);
+	void SetupObservations();
 
 	/**
-	* During this event, all observations should be set for each agent.
-	* @param AgentIds The list of agent ids to set observations for.
-	* @see LearningAgentsObservations.h for the list of available observations.
-	* @see GetAgent to get the agent corresponding to each id.
-	*/
+	 * During this event, all observations should be set for each agent.
+	 * @param AgentIds The list of agent ids to set observations for.
+	 * @see LearningAgentsObservations.h for the list of available observations.
+	 * @see GetAgent to get the agent corresponding to each id.
+	 */
 	UFUNCTION(BlueprintNativeEvent, Category = "LearningAgents")
 	void SetObservations(const TArray<int32>& AgentIds);
 
 	/**
-	* Used by objects derived from ULearningAgentsObservation to add themselves to this agent interactor during their
-	* creation. You shouldn't need to call this directly.
-	*/
+	 * Used by objects derived from ULearningAgentsObservation to add themselves to this agent interactor during their
+	 * creation. You shouldn't need to call this directly.
+	 */
 	void AddObservation(TObjectPtr<ULearningAgentsObservation> Object, const TSharedRef<UE::Learning::FFeatureObject>& Feature);
 
 // ----- Actions -----
@@ -81,61 +88,60 @@ public:
 
 	/**
 	 * During this event, all actions should be added to this agent interactor.
-	 * @param Interactor This Interactor. This is a convenience for blueprints so you don't have to find the Self pin.
 	 * @see LearningAgentsActions.h for the list of available actions.
 	 */
 	UFUNCTION(BlueprintNativeEvent, Category = "LearningAgents")
-	void SetupActions(ULearningAgentsInteractor* Interactor);
+	void SetupActions();
 
 	/**
-	* During this event, you should retrieve the actions and apply them to your agents.
-	* @param AgentIds The list of agent ids to get actions for.
-	* @see LearningAgentsActions.h for the list of available actions.
-	* @see GetAgent to get the agent corresponding to each id.
-	*/
+	 * During this event, you should retrieve the actions and apply them to your agents.
+	 * @param AgentIds The list of agent ids to get actions for.
+	 * @see LearningAgentsActions.h for the list of available actions.
+	 * @see GetAgent to get the agent corresponding to each id.
+	 */
 	UFUNCTION(BlueprintNativeEvent, Category = "LearningAgents")
 	void GetActions(const TArray<int32>& AgentIds);
 
 	/**
-	* Used by objects derived from ULearningAgentsAction to add themselves to this agent interactor during their creation.
-	* You shouldn't need to call this directly.
-	*/
+	 * Used by objects derived from ULearningAgentsAction to add themselves to this agent interactor during their creation.
+	 * You shouldn't need to call this directly.
+	 */
 	void AddAction(TObjectPtr<ULearningAgentsAction> Object, const TSharedRef<UE::Learning::FFeatureObject>& Feature);
 
 // ----- Encoding / Decoding -----
 public:
 
 	/**
-	* Call this function when it is time to gather all the observations for your agents. This can be done each frame or
-	* you can consider wiring it up to some kind of meaningful event, e.g. a hypothetical "OnAITurnStarted" if you have
-	* a turn-based game. This will call this agent interactor's SetObservations event.
-	*/
+	 * Call this function when it is time to gather all the observations for your agents. This can be done each frame or
+	 * you can consider wiring it up to some kind of meaningful event, e.g. a hypothetical "OnAITurnStarted" if you have
+	 * a turn-based game. This will call this agent interactor's SetObservations event.
+	 */
 	UFUNCTION(BlueprintCallable, Category = "LearningAgents")
 	void EncodeObservations();
 
 	/**
-	* Call this function when it is time for your agents to take their actions. You most likely want to call this after
-	* your policy's EvaluatePolicy function to ensure that you are receiving the latest actions. This will call this
-	* agent interactor's GetActions event.
-	*/
+	 * Call this function when it is time for your agents to take their actions. You most likely want to call this after
+	 * your policy's EvaluatePolicy function to ensure that you are receiving the latest actions. This will call this
+	 * agent interactor's GetActions event.
+	 */
 	UFUNCTION(BlueprintCallable, Category = "LearningAgents")
 	void DecodeActions();
 
 	/**
-	* Gets the observation vector used by a given agent. Should be called only after EncodeObservations.
-	*
-	* @param AgentId				The AgentId to look-up the observation vector for
-	* @param OutObservationVector	The observation vector for the given agent
-	*/
+	 * Gets the observation vector used by a given agent. Should be called only after EncodeObservations.
+	 *
+	 * @param AgentId				The AgentId to look-up the observation vector for
+	 * @param OutObservationVector	The observation vector for the given agent
+	 */
 	UFUNCTION(BlueprintPure, Category = "LearningAgents", meta = (AgentId = "-1"))
 	void GetObservationVector(const int32 AgentId, TArray<float>& OutObservationVector) const;
 
 	/**
-	* Gets the action vector used by a given agent. Should be called only after EncodeActions or EvaluatePolicy.
-	*
-	* @param AgentId				The AgentId to look-up the action vector for
-	* @param OutActionVector		The action vector for the given agent
-	*/
+	 * Gets the action vector used by a given agent. Should be called only after EncodeActions or EvaluatePolicy.
+	 *
+	 * @param AgentId				The AgentId to look-up the action vector for
+	 * @param OutActionVector		The action vector for the given agent
+	 */
 	UFUNCTION(BlueprintPure, Category = "LearningAgents", meta = (AgentId = "-1"))
 	void GetActionVector(const int32 AgentId, TArray<float>& OutActionVector) const;
 
@@ -154,6 +160,18 @@ public:
 	/** Get a const array view of this agent interactor's action objects. */
 	TConstArrayView<ULearningAgentsAction*> GetActionObjects() const;
 
+	/** Get the number of times observations have been encoded for all agents */
+	TLearningArrayView<1, uint64> GetObservationEncodingAgentIteration();
+
+	/** Get the number of times actions have been encoded for all agents */
+	TLearningArrayView<1, uint64> GetActionEncodingAgentIteration();
+
+	/** Encode Observations for a specific set of agents */
+	void EncodeObservations(const UE::Learning::FIndexSet AgentSet);
+
+	/** Decode Actions for a specific set of agents */
+	void DecodeActions(const UE::Learning::FIndexSet AgentSet);
+
 private:
 
 	/** The list of current observation objects. */
@@ -172,4 +190,25 @@ private:
 
 	TSharedPtr<UE::Learning::FConcatenateFeature> Observations;
 	TSharedPtr<UE::Learning::FConcatenateFeature> Actions;
+
+// ----- Private Iteration Checks ----- 
+private:
+
+	/**
+	 * Sets the iteration number for the given set of agents for all observations, actions, as well as for the encoding 
+	 * of observations and actions
+	 */
+	void SetAgentIteration(const UE::Learning::FIndexSet AgentIds, const uint64 Value);
+
+	/** Number of times observations have been encoded for all agents */
+	TLearningArray<1, uint64, TInlineAllocator<32>> ObservationEncodingAgentIteration;
+
+	/** Number of times actions have been encoded for all agents */
+	TLearningArray<1, uint64, TInlineAllocator<32>> ActionEncodingAgentIteration;
+
+	/** Temp buffers used to record the set of agents that are valid for encoding/decoding */
+	TArray<int32> ValidAgentIds;
+	UE::Learning::FIndexSet ValidAgentSet;
+	TBitArray<TInlineAllocator<32>> ValidAgentStatus;
+
 };
