LearningAgents: Added iteration checks for verifying order of operations. Fixed issue with inference for multiple agents. Fixed live recompilation issues.

diff --git a/Engine/Plugins/Experimental/LearningAgents/Source/LearningAgents/Public/LearningAgentsManager.h b/Engine/Plugins/Experimental/LearningAgents/Source/LearningAgents/Public/LearningAgentsManager.h
--- a/Engine/Plugins/Experimental/LearningAgents/Source/LearningAgents/Public/LearningAgentsManager.h
+++ b/Engine/Plugins/Experimental/LearningAgents/Source/LearningAgents/Public/LearningAgentsManager.h
@@ -22,12 +22,12 @@ namespace UE::Learning
 class ULearningAgentsManagerComponent;
 
 /**
-* The agent manager is responsible for tracking which game objects are agents. It's the central class around which
-* most of Learning Agents is built.
-*
-* If you have multiple different types of objects you want controlled by Learning Agents, you should consider creating
-* one agent manager per object type, rather than trying to share an agent manager.
-*/
+ * The agent manager is responsible for tracking which game objects are agents. It's the central class around which
+ * most of Learning Agents is built.
+ *
+ * If you have multiple different types of objects you want controlled by Learning Agents, you should consider creating
+ * one agent manager per object type, rather than trying to share an agent manager.
+ */
 UCLASS(Abstract, BlueprintType, Blueprintable)
 class LEARNINGAGENTS_API ALearningAgentsManager : public AActor
 {
@@ -36,33 +36,23 @@ class LEARNINGAGENTS_API ALearningAgentsManager : public AActor
 public:
 
 	// These constructors/destructors are needed to make forward declarations happy
-	ALearningAgentsManager();
+	ALearningAgentsManager(const FObjectInitializer& ObjectInitializer = FObjectInitializer::Get());
 	ALearningAgentsManager(FVTableHelper& Helper);
 	virtual ~ALearningAgentsManager();
 
-	/** Sets up the agent ids so that agents can be added prior to calling SetupManager. */
 	virtual void PostInitProperties() override;
 
-// ----- Setup -----
+// ----- Agent Management -----
 public:
 
-	/** Initializes this object and runs the setup events for observations and actions. */
-	UFUNCTION(BlueprintCallable, Category = "LearningAgents")
-	void SetupManager();
-
-	/** Returns true if SetupManager has been run successfully; Otherwise, false. */
-	UFUNCTION(BlueprintPure, Category = "LearningAgents")
-	bool IsManagerSetup() const;
-
 	/** Returns the maximum number of agents that this manager is configured to handle. */
 	UFUNCTION(BlueprintPure, Category = "LearningAgents")
-	int32 GetMaxInstanceNum() const;
+	int32 GetMaxAgentNum() const;
 
-// ----- Agent Management -----
 public:
 
 	/**
-	 * Adds the given object as an agent to this manager. This can be called before or after SetupManager.
+	 * Adds the given object as an agent to the manager.
 	 * @param Agent The object to be added.
 	 * @return The agent's newly assigned id.
 	 */
@@ -70,36 +60,54 @@ public:
 	int32 AddAgent(UObject* Agent);
 
 	/**
-	 * Called whenever a new agent is added to this manager. By default, this will add the agent to each of this
-	 * manager's agent components, but you can override this event if you have custom logic. For example, maybe some of
-	 * the agents are inference agents only and should not be added to the trainer component. 
-	 * @param AgentId The agent's newly assigned id.
+	 * Adds the given objects as an agents to the manager.
+	 * @param OutAgentIds The output newly assigned agent ids
+	 * @param InAgents The objects to be added.
 	 */
-	UFUNCTION(BlueprintNativeEvent, Category = "LearningAgents")
-	void OnAgentAdded(const int32 AgentId);
-	
+	UFUNCTION(BlueprintCallable, Category = "LearningAgents")
+	void AddAgents(TArray<int32>& OutAgentIds, const TArray<UObject*>& InAgents);
+
+public:
+
 	/**
-	 * Removes the agent with the given id from this manager.
+	 * Removes the agent with the given id from the manager.
 	 * @param AgentId The id of the agent to remove.
 	 */
 	UFUNCTION(BlueprintCallable, Category = "LearningAgents", meta = (AgentId = "-1"))
-	void RemoveAgentById(const int32 AgentId);
+	void RemoveAgent(const int32 AgentId);
 
 	/**
-	* Removes the given agent from this manager. Use RemoveAgentById if you have the id available as this function
-	* must do a linear search to find the agent.
-	* @param Agent The agent to be removed.
-	*/
+	 * Removes the agents with the given ids from the manager.
+	 * @param AgentIds The ids of the agents to remove.
+	 */
+	UFUNCTION(BlueprintCallable, Category = "LearningAgents")
+	void RemoveAgents(const TArray<int32>& AgentIds);
+
+	/** Removes all agents from the manager. */
 	UFUNCTION(BlueprintCallable, Category = "LearningAgents")
-	void RemoveAgent(UObject* Agent);
+	void RemoveAllAgents();
+
+public:
+
+	/**
+	 * Resets the agent with the given id on the manager. Used to tell components to reset any state associated with this agent.
+	 * @param AgentId The id of the agent to reset.
+	 */
+	UFUNCTION(BlueprintCallable, Category = "LearningAgents", meta = (AgentId = "-1"))
+	void ResetAgent(const int32 AgentId);
 
 	/**
-	 * Called whenever an agent is removed from this manager. By default, this will remove the agent from each of
-	 * this manager's agent components, but you can override this event if you have additional logic. 
-	 * @param AgentId The removed agent's id.
+	 * Resets the agents with the given ids on the manager. Used to tell components to reset any state associated with this agent.
+	 * @param AgentIds The ids of the agents to reset.
 	 */
-	UFUNCTION(BlueprintNativeEvent, Category = "LearningAgents")
-	void OnAgentRemoved(const int32 AgentId);
+	UFUNCTION(BlueprintCallable, Category = "LearningAgents")
+	void ResetAgents(const TArray<int32>& AgentIds);
+
+	/** Resets all the agents on the manager. */
+	UFUNCTION(BlueprintCallable, Category = "LearningAgents")
+	void ResetAllAgents();
+
+public:
 
 	/**
 	 * Gets the agent with the given id. Calling this from blueprint with the appropriate AgentClass will automatically
@@ -119,31 +127,74 @@ public:
 	 * @param OutAgents The output array of agent objects.
 	 */
 	UFUNCTION(BlueprintPure = false, Category = "LearningAgents", meta = (DeterminesOutputType = "AgentClass", DynamicOutputParam = "OutAgents"))
-	void GetAgents(const TArray<int32>& AgentIds, const TSubclassOf<UObject> AgentClass, TArray<UObject*>& OutAgents) const;
+	void GetAgents(TArray<UObject*>& OutAgents, const TArray<int32>& AgentIds, const TSubclassOf<UObject> AgentClass) const;
 
 	/**
-	 * Gets the current added agents. Calling this from blueprint with the appropriate AgentClass will automatically
+	 * Gets all added agents. Calling this from blueprint with the appropriate AgentClass will automatically
 	 * cast the object to the given type.
 	 * @param AgentClass The class to cast the agent objects to (in blueprint).
 	 * @param OutAgents The output array of agent objects.
 	 */
 	UFUNCTION(BlueprintPure = false, Category = "LearningAgents", meta = (DeterminesOutputType = "AgentClass", DynamicOutputParam = "OutAgents"))
-	void GetAddedAgents(const TSubclassOf<UObject> AgentClass, TArray<UObject*>& OutAgents) const;
+	void GetAllAgents(TArray<UObject*>& OutAgents, const TSubclassOf<UObject> AgentClass) const;
 
 	/**
-	* Returns true if the given object is an agent in this manager; Otherwise, false.
-	* Use HasAgentById if you have the id available as this function must do a linear search to find the agent.
-	*/
+	 * Gets the agent id associated with a given agent.
+	 * @param Agent The agent object.
+	 * @return The agent id.
+	 */
 	UFUNCTION(BlueprintPure, Category = "LearningAgents")
-	bool HasAgent(UObject* Agent) const;
+	int32 GetAgentId(UObject* Agent) const;
 
-	/** Returns true if the given id is used by an agent in this manager; Otherwise, false. */
-	UFUNCTION(BlueprintPure, Category = "LearningAgents", meta = (AgentId = "-1"))
-	bool HasAgentById(const int32 AgentId) const;
+	/**
+	 * Gets the agent ids associated with a set of agents.
+	 * @param OutAgentIds The ids of the agents.
+	 * @param InAgents The agent objects.
+	 */
+	UFUNCTION(BlueprintPure = false, Category = "LearningAgents")
+	void GetAgentIds(TArray<int32>& OutAgentIds, const TArray<UObject*>& InAgents) const;
 
-	/** Get the current added agent ids. */
+	/**
+	 * Gets all added agent ids.
+	 * @param OutAgentIds The output array of agent ids.
+	 */
 	UFUNCTION(BlueprintPure = false, Category = "LearningAgents")
-	void GetAgentIds(TArray<int32>& OutAgentIds) const;
+	void GetAllAgentIds(TArray<int32>& OutAgentIds) const;
+
+	/**
+	 * Gets the number of agents added
+	 * @return The number of agents added.
+	 */
+	UFUNCTION(BlueprintPure, Category = "LearningAgents")
+	int32 GetAgentNum() const;
+
+public:
+
+	/** Returns true if the given object is an agent used by the manager; Otherwise, false. */
+	UFUNCTION(BlueprintPure, Category = "LearningAgents")
+	bool HasAgentObject(UObject* Agent) const;
+
+	/** Returns true if the given id is an agent used by the manager; Otherwise, false. */
+	UFUNCTION(BlueprintPure, Category = "LearningAgents", meta = (AgentId = "-1"))
+	bool HasAgent(const int32 AgentId) const;
+
+// ----- Tick Prerequisites -----
+public:
+
+	/**
+	 * Adds this manager as a tick prerequisite of the given actor objects.
+	 * @param InAgents The actor object which this manager will be added as a tick prerequisite to.
+	 */
+	UFUNCTION(BlueprintCallable, Category = "LearningAgents")
+	void AddManagerAsTickPrerequisiteOfAgents(const TArray<AActor*>& InAgents);
+
+	/**
+	 * Adds the given actor objects as tick prerequisite of this manager.
+	 * @param InAgents The actor object which will be added as tick prerequisites to this manager.
+	 */
+	UFUNCTION(BlueprintCallable, Category = "LearningAgents")
+	void AddAgentsAsTickPrerequisiteOfManager(const TArray<AActor*>& InAgents);
+
 
 // ----- Non-blueprint public interface -----
 public:
@@ -154,6 +205,12 @@ public:
 	/** Gets the agent corresponding to the given id. */
 	UObject* GetAgent(const int32 AgentId);
 
+	/** Gets the set of agent ids currently added */
+	const TArray<int32>& GetAllAgentIds() const;
+
+	/** Gets the set of agent ids currently added as an FIndexSet */
+	UE::Learning::FIndexSet GetAllAgentSet() const;
+
 	/** Get a const reference to this manager's underlying instance data. */
 	const TSharedPtr<UE::Learning::FArrayMap>& GetInstanceData() const;
 
@@ -164,27 +221,28 @@ private:
 
 	/** Maximum number of agents. Used to preallocate internal buffers. Setting this higher will allow more agents but use up more memory. */
 	UPROPERTY(EditDefaultsOnly, Category = "LearningAgents", meta = (ClampMin = "1", UIMin = "1"))
-	int32 MaxInstanceNum = 1;
-
-	/** True if SetupManager has been performed; Otherwise, false. */
-	UPROPERTY(VisibleAnywhere, Transient, Category = "LearningAgents")
-	bool bIsSetup = false;
+	int32 MaxAgentNum = 1;
 
 	/** The list of current agents. */
 	UPROPERTY(VisibleAnywhere, Transient, Category = "LearningAgents")
 	TArray<TObjectPtr<UObject>> Agents;
 
+private:
+
+	/** Object containing all of the instance data required by all agents */
+	TSharedPtr<UE::Learning::FArrayMap> InstanceData;
+
 	/** Update the agent sets to keep them in sync with the id lists. */
 	void UpdateAgentSets();
 
+	/** Array of agent ids to be passed to events such as ULearningAgentsManagerComponent::OnAgentAdded. */
+	TArray<int32> OnEventAgentIds;
+
+	/** Array of agent ids currently in use and associated with each agent object. */
 	TArray<int32> OccupiedAgentIds;
-	TArray<int32> VacantAgentIds;
 	UE::Learning::FIndexSet OccupiedAgentSet;
-	UE::Learning::FIndexSet VacantAgentSet;
-
-	TSharedPtr<UE::Learning::FArrayMap> InstanceData;
 
-	/** The manager components that were found during setup. */
-	UPROPERTY(VisibleAnywhere, Transient, Category = "LearningAgents")
-	TArray<TObjectPtr<ULearningAgentsManagerComponent>> CachedManagerComponents;
+	/* Array of agent ids currently free and available for assignment. */
+	TArray<int32> VacantAgentIds;
+	UE::Learning::FIndexSet VacantAgentSet;
 };
diff --git a/Engine/Plugins/Experimental/LearningAgents/Source/LearningAgents/Public/LearningAgentsManagerComponent.h b/Engine/Plugins/Experimental/LearningAgents/Source/LearningAgents/Public/LearningAgentsManagerComponent.h
--- a/Engine/Plugins/Experimental/LearningAgents/Source/LearningAgents/Public/LearningAgentsManagerComponent.h
+++ b/Engine/Plugins/Experimental/LearningAgents/Source/LearningAgents/Public/LearningAgentsManagerComponent.h
@@ -25,29 +25,32 @@ class LEARNINGAGENTS_API ULearningAgentsManagerComponent : public UActorComponen
 public:
 
 	// These constructors/destructors are needed to make forward declarations happy
-	ULearningAgentsManagerComponent();
+	ULearningAgentsManagerComponent(const FObjectInitializer& ObjectInitializer = FObjectInitializer::Get());
 	ULearningAgentsManagerComponent(FVTableHelper& Helper);
 	virtual ~ULearningAgentsManagerComponent();
 
+	virtual void PostInitProperties() override;
+
 	/**
-	 * Adds an agent to this component.
-	 * @param AgentId The id of the agent to be added.
-	 * @return True if the agent was added successfully. Otherwise, false.
+	 * Called whenever agents are added to the parent ALearningAgentsManager object.
+	 * @param AgentIds Array of agent ids which have been added
 	 */
-	UFUNCTION(BlueprintCallable, Category = "LearningAgents", meta = (AgentId = "-1"))
-	virtual bool AddAgent(const int32 AgentId);
+	UFUNCTION(BlueprintCallable, Category = "LearningAgents")
+	virtual void OnAgentsAdded(const TArray<int32>& AgentIds);
 
 	/**
-	* Removes an agent from this component.
-	* @param AgentId The id of the agent to be removed.
-	* @return True if the agent was removed successfully. Otherwise, false.
-	*/
-	UFUNCTION(BlueprintCallable, Category = "LearningAgents", meta = (AgentId = "-1"))
-	virtual bool RemoveAgent(const int32 AgentId);
-
-	/** Returns true if the given id has been previously added to this component; Otherwise, false. */
-	UFUNCTION(BlueprintPure, Category = "LearningAgents", meta = (AgentId = "-1"))
-	bool HasAgent(const int32 AgentId) const;
+	 * Called whenever agents are removed from the parent ALearningAgentsManager object.
+	 * @param AgentIds Array of agent ids which have been removed
+	 */
+	UFUNCTION(BlueprintCallable, Category = "LearningAgents")
+	virtual void OnAgentsRemoved(const TArray<int32>& AgentIds);
+
+	/**
+	 * Called whenever agents are reset on the parent ALearningAgentsManager object.
+	 * @param AgentIds Array of agent ids which have been reset
+	 */
+	UFUNCTION(BlueprintCallable, Category = "LearningAgents")
+	virtual void OnAgentsReset(const TArray<int32>& AgentIds);
 
 	/** Returns true if this component has been setup. Otherwise, false. */
 	UFUNCTION(BlueprintPure, Category = "LearningAgents")
@@ -57,20 +60,31 @@ public:
 protected:
 
 	/**
-	* Gets the agent with the given id from the manager. Calling this from blueprint with the appropriate AgentClass
-	* will automatically cast the object to the given type. If not in a blueprint, you should call the manager's
-	* GetAgent methods directly.
-	* @param AgentId The id of the agent to get.
-	* @param AgentClass The class to cast the agent object to (in blueprint).
-	* @return The agent object.
-	*/
+	 * Gets the agent with the given id from the manager. Calling this from blueprint with the appropriate AgentClass
+	 * will automatically cast the object to the given type. If not in a blueprint, you should call the manager's
+	 * GetAgent methods directly.
+	 * @param AgentId The id of the agent to get.
+	 * @param AgentClass The class to cast the agent object to (in blueprint).
+	 * @return The agent object.
+	 */
 	UFUNCTION(BlueprintPure, Category = "LearningAgents", meta = (AgentId = "-1", DeterminesOutputType = "AgentClass"))
 	UObject* GetAgent(const int32 AgentId, const TSubclassOf<UObject> AgentClass) const;
 
 	/**
-	* Gets the agent manager associated with this component.
-	* @param AgentClass The class to cast the agent manager to (in blueprint).
-	*/
+	 * Gets the agents associated with a set of ids from the manager. Calling this from blueprint with the appropriate 
+	 * AgentClass will automatically cast the object to the given type. If not in a blueprint, you should call the 
+	 * manager's GetAgents method directly.
+	 * @param AgentIds The ids of the agents to get.
+	 * @param AgentClass The class to cast the agent objects to (in blueprint).
+	 * @param OutAgents The output array of agent objects.
+	 */
+	UFUNCTION(BlueprintPure = false, Category = "LearningAgents", meta = (DeterminesOutputType = "AgentClass", DynamicOutputParam = "OutAgents"))
+	void GetAgents(const TArray<int32>& AgentIds, const TSubclassOf<UObject> AgentClass, TArray<UObject*>& OutAgents) const;
+
+	/**
+	 * Gets the agent manager associated with this component.
+	 * @param AgentClass The class to cast the agent manager to (in blueprint).
+	 */
 	UFUNCTION(BlueprintPure, Category = "LearningAgents", meta = (DeterminesOutputType = "AgentManagerClass"))
 	ALearningAgentsManager* GetAgentManager(const TSubclassOf<ALearningAgentsManager> AgentManagerClass) const;
 
@@ -83,15 +97,25 @@ public:
 	/** Gets the agent corresponding to the given id. */
 	UObject* GetAgent(const int32 AgentId);
 
-	ALearningAgentsManager* GetAgentManager() const;
+	/** Checks if the component has the given agent id */
+	bool HasAgent(const int32 AgentId) const;
+
+	/** Gets the associated agent manager assuming it exists. */
+	const ALearningAgentsManager* GetAgentManager() const;
+
+	/** Gets the associated agent manager assuming it exists. */
+	ALearningAgentsManager* GetAgentManager();
+
+	/** Checks if the component has an agent manager. */
+	bool HasAgentManager() const;
 
 // ----- Helpers -----
 public:
 
 	/**
-	* Used by objects derived from ULearningAgentsHelper to add themselves to this component during their creation.
-	* You shouldn't need to call this directly.
-	*/
+	 * Used by objects derived from ULearningAgentsHelper to add themselves to this component during their creation.
+	 * You shouldn't need to call this directly.
+	 */
 	void AddHelper(TObjectPtr<ULearningAgentsHelper> Object);
 
 protected:
@@ -100,15 +124,9 @@ protected:
 	UPROPERTY(VisibleAnywhere, Transient, Category = "LearningAgents")
 	bool bIsSetup = false;
 
-	/** The agent manager associated with this component. */
-	UPROPERTY(VisibleAnywhere, Transient, Category = "LearningAgents")
-	TObjectPtr<ALearningAgentsManager> AgentManager;
-
-	/** The agent ids added to this component. */
+	/** The associated manager this component is attached to. */
 	UPROPERTY(VisibleAnywhere, Transient, Category = "LearningAgents")
-	TArray<int32> AddedAgentIds;
-
-	UE::Learning::FIndexSet AddedAgentSet;
+	TObjectPtr<ALearningAgentsManager> Manager;
 
 	/** The list of current helper objects. */
 	UPROPERTY(VisibleAnywhere, Transient, Category = "LearningAgents")
diff --git a/Engine/Plugins/Experimental/LearningAgents/Source/LearningAgents/Public/LearningAgentsNeuralNetwork.h b/Engine/Plugins/Experimental/LearningAgents/Source/LearningAgents/Public/LearningAgentsNeuralNetwork.h
--- a/Engine/Plugins/Experimental/LearningAgents/Source/LearningAgents/Public/LearningAgentsNeuralNetwork.h
+++ b/Engine/Plugins/Experimental/LearningAgents/Source/LearningAgents/Public/LearningAgentsNeuralNetwork.h
@@ -51,20 +51,40 @@ public:
 	virtual void Serialize(FArchive& Ar) override;
 
 	/**
-	* Load this network from a snapshot.
-	* @param File The snapshot file.
-	*/
+	 * Load this network from a snapshot.
+	 * @param File The snapshot file.
+	 */
 	UFUNCTION(BlueprintCallable, Category = "LearningAgents", meta = (RelativePath))
 	void LoadNetworkFromSnapshot(const FFilePath& File);
 
 	/**
-	* Save this network into a snapshot.
-	* @param File The snapshot file.
-	*/
+	 * Save this network into a snapshot.
+	 * @param File The snapshot file.
+	 */
 	UFUNCTION(BlueprintCallable, Category = "LearningAgents", meta = (RelativePath))
-	void SaveNetworkToSnapshot(const FFilePath& File) const;
+	void SaveNetworkToSnapshot(const FFilePath& File);
+
+	/**
+	 * Copy another asset's network weights into this network. Network sizes must match.
+	 * @param NeuralNetworkAsset The asset to load from.
+	 */
+	UFUNCTION(BlueprintCallable, Category = "LearningAgents")
+	void LoadNetworkFromAsset(ULearningAgentsNeuralNetwork* NeuralNetworkAsset);
+
+	/**
+	 * Copy this networks's weights to another asset. Network sizes must match.
+	 * @param NeuralNetworkAsset The asset to save to.
+	 */
+	UFUNCTION(BlueprintCallable, Category = "LearningAgents", Meta = (DevelopmentOnly))
+	void SaveNetworkToAsset(ULearningAgentsNeuralNetwork* NeuralNetworkAsset);
+
+public:
+
+	/** Marks this asset as modified even during PIE */
+	void ForceMarkDirty();
 
 public:
 
+	/** Pointer to the internal Neural Network Data */
 	TSharedPtr<UE::Learning::FNeuralNetwork> NeuralNetwork;
 };
