LearningAgents: Added iteration checks for verifying order of operations. Fixed issue with inference for multiple agents. Fixed live recompilation issues.

diff --git a/Engine/Plugins/Experimental/LearningAgents/Source/LearningAgents/Public/LearningAgentsObservations.h b/Engine/Plugins/Experimental/LearningAgents/Source/LearningAgents/Public/LearningAgentsObservations.h
--- a/Engine/Plugins/Experimental/LearningAgents/Source/LearningAgents/Public/LearningAgentsObservations.h
+++ b/Engine/Plugins/Experimental/LearningAgents/Source/LearningAgents/Public/LearningAgentsObservations.h
@@ -48,7 +48,14 @@ public:
 	UPROPERTY(VisibleAnywhere, Transient, Category = "LearningAgents")
 	TObjectPtr<ULearningAgentsInteractor> Interactor;
 
+public:
+
+	/** Number of times this observation has been set for all agents */
+	TLearningArray<1, uint64, TInlineAllocator<32>> AgentIteration;
+
+public:
 #if UE_LEARNING_AGENTS_ENABLE_VISUAL_LOG
+
 	/** Color used to draw this observation in the visual log */
 	FLinearColor VisualLogColor = FColor::Red;
 
@@ -68,20 +75,20 @@ class LEARNINGAGENTS_API UFloatObservation : public ULearningAgentsObservation
 public:
 
 	/**
-	* Adds a new float observation to the given agent interactor. Call during ULearningAgentsInteractor::SetupObservations event.
-	* @param InInteractor The agent interactor to add this observation to.
-	* @param Name The name of this new observation. Used for debugging.
-	* @param Scale Used to normalize the data for the observation.
-	* @return The newly created observation.
-	*/
-	UFUNCTION(BlueprintCallable, Category = "LearningAgents")
+	 * Adds a new float observation to the given agent interactor. Call during ULearningAgentsInteractor::SetupObservations event.
+	 * @param InInteractor The agent interactor to add this observation to.
+	 * @param Name The name of this new observation. Used for debugging.
+	 * @param Scale Used to normalize the data for the observation.
+	 * @return The newly created observation.
+	 */
+	UFUNCTION(BlueprintCallable, Category = "LearningAgents", meta=(DefaultToSelf = "InInteractor"))
 	static UFloatObservation* AddFloatObservation(ULearningAgentsInteractor* InInteractor, const FName Name = NAME_None, const float Scale = 1.0f);
 
 	/**
-	* Sets the data for this observation. Call during ULearningAgentsInteractor::SetObservations event.
-	* @param AgentId The agent id this data corresponds to.
-	* @param Observation The value currently being observed.
-	*/
+	 * Sets the data for this observation. Call during ULearningAgentsInteractor::SetObservations event.
+	 * @param AgentId The agent id this data corresponds to.
+	 * @param Observation The value currently being observed.
+	 */
 	UFUNCTION(BlueprintCallable, Category = "LearningAgents", meta = (AgentId = "-1"))
 	void SetFloatObservation(const int32 AgentId, const float Observation);
 
@@ -102,21 +109,21 @@ class LEARNINGAGENTS_API UFloatArrayObservation : public ULearningAgentsObservat
 public:
 
 	/**
-	* Adds a new float array observation to the given agent interactor. Call during ULearningAgentsInteractor::SetupObservations event.
-	* @param InInteractor The agent interactor to add this observation to.
-	* @param Name The name of this new observation. Used for debugging.
-	* @param Num The number of floats in the array
-	* @param Scale Used to normalize the data for the observation.
-	* @return The newly created observation.
-	*/
-	UFUNCTION(BlueprintCallable, Category = "LearningAgents")
+	 * Adds a new float array observation to the given agent interactor. Call during ULearningAgentsInteractor::SetupObservations event.
+	 * @param InInteractor The agent interactor to add this observation to.
+	 * @param Name The name of this new observation. Used for debugging.
+	 * @param Num The number of floats in the array
+	 * @param Scale Used to normalize the data for the observation.
+	 * @return The newly created observation.
+	 */
+	UFUNCTION(BlueprintCallable, Category = "LearningAgents", meta = (DefaultToSelf = "InInteractor"))
 	static UFloatArrayObservation* AddFloatArrayObservation(ULearningAgentsInteractor* InInteractor, const FName Name = NAME_None, const int32 Num = 1, const float Scale = 1.0f);
 
 	/**
-	* Sets the data for this observation. Call during ULearningAgentsInteractor::SetObservations event.
-	* @param AgentId The agent id this data corresponds to.
-	* @param Observation The value currently being observed.
-	*/
+	 * Sets the data for this observation. Call during ULearningAgentsInteractor::SetObservations event.
+	 * @param AgentId The agent id this data corresponds to.
+	 * @param Observation The value currently being observed.
+	 */
 	UFUNCTION(BlueprintCallable, Category = "LearningAgents", meta = (AgentId = "-1"))
 	void SetFloatArrayObservation(const int32 AgentId, const TArray<float>& Observation);
 
@@ -139,20 +146,20 @@ class LEARNINGAGENTS_API UVectorObservation : public ULearningAgentsObservation
 public:
 
 	/**
-	* Adds a new vector observation to the given agent interactor. Call during ULearningAgentsInteractor::SetupObservations event.
-	* @param InInteractor The agent interactor to add this observation to.
-	* @param Name The name of this new observation. Used for debugging.
-	* @param Scale Used to normalize the data for the observation.
-	* @return The newly created observation.
-	*/
-	UFUNCTION(BlueprintCallable, Category = "LearningAgents")
+	 * Adds a new vector observation to the given agent interactor. Call during ULearningAgentsInteractor::SetupObservations event.
+	 * @param InInteractor The agent interactor to add this observation to.
+	 * @param Name The name of this new observation. Used for debugging.
+	 * @param Scale Used to normalize the data for the observation.
+	 * @return The newly created observation.
+	 */
+	UFUNCTION(BlueprintCallable, Category = "LearningAgents", meta = (DefaultToSelf = "InInteractor"))
 	static UVectorObservation* AddVectorObservation(ULearningAgentsInteractor* InInteractor, const FName Name = NAME_None, const float Scale = 1.0f);
 
 	/**
-	* Sets the data for this observation. Call during ULearningAgentsInteractor::SetObservations event.
-	* @param AgentId The agent id this data corresponds to.
-	* @param Observation The values currently being observed.
-	*/
+	 * Sets the data for this observation. Call during ULearningAgentsInteractor::SetObservations event.
+	 * @param AgentId The agent id this data corresponds to.
+	 * @param Observation The values currently being observed.
+	 */
 	UFUNCTION(BlueprintCallable, Category = "LearningAgents", meta = (AgentId = "-1"))
 	void SetVectorObservation(const int32 AgentId, const FVector Observation);
 
@@ -173,21 +180,21 @@ class LEARNINGAGENTS_API UVectorArrayObservation : public ULearningAgentsObserva
 public:
 
 	/**
-	* Adds a new vector array observation to the given agent interactor. Call during ULearningAgentsInteractor::SetupObservations event.
-	* @param InInteractor The agent interactor to add this observation to.
-	* @param Name The name of this new observation. Used for debugging.
-	* @param Num The number of vectors in the array
-	* @param Scale Used to normalize the data for the observation.
-	* @return The newly created observation.
-	*/
-	UFUNCTION(BlueprintCallable, Category = "LearningAgents")
+	 * Adds a new vector array observation to the given agent interactor. Call during ULearningAgentsInteractor::SetupObservations event.
+	 * @param InInteractor The agent interactor to add this observation to.
+	 * @param Name The name of this new observation. Used for debugging.
+	 * @param Num The number of vectors in the array
+	 * @param Scale Used to normalize the data for the observation.
+	 * @return The newly created observation.
+	 */
+	UFUNCTION(BlueprintCallable, Category = "LearningAgents", meta = (DefaultToSelf = "InInteractor"))
 	static UVectorArrayObservation* AddVectorArrayObservation(ULearningAgentsInteractor* InInteractor, const FName Name = NAME_None, const int32 Num = 1, const float Scale = 1.0f);
 
 	/**
-	* Sets the data for this observation. Call during ULearningAgentsInteractor::SetObservations event.
-	* @param AgentId The agent id this data corresponds to.
-	* @param Observation The values currently being observed.
-	*/
+	 * Sets the data for this observation. Call during ULearningAgentsInteractor::SetObservations event.
+	 * @param AgentId The agent id this data corresponds to.
+	 * @param Observation The values currently being observed.
+	 */
 	UFUNCTION(BlueprintCallable, Category = "LearningAgents", meta = (AgentId = "-1"))
 	void SetVectorArrayObservation(const int32 AgentId, const TArray<FVector>& Observation);
 
@@ -210,20 +217,20 @@ class LEARNINGAGENTS_API UEnumObservation : public ULearningAgentsObservation
 public:
 
 	/**
-	* Adds a new enum observation to the given agent interactor. Call during ULearningAgentsInteractor::SetupObservations event.
-	* @param InInteractor The agent interactor to add this observation to.
-	* @param EnumType The type of enum to use
-	* @param Name The name of this new observation. Used for debugging.
-	* @return The newly created observation.
-	*/
-	UFUNCTION(BlueprintCallable, Category = "LearningAgents")
+	 * Adds a new enum observation to the given agent interactor. Call during ULearningAgentsInteractor::SetupObservations event.
+	 * @param InInteractor The agent interactor to add this observation to.
+	 * @param EnumType The type of enum to use
+	 * @param Name The name of this new observation. Used for debugging.
+	 * @return The newly created observation.
+	 */
+	UFUNCTION(BlueprintCallable, Category = "LearningAgents", meta = (DefaultToSelf = "InInteractor"))
 	static UEnumObservation* AddEnumObservation(ULearningAgentsInteractor* InInteractor, const UEnum* EnumType, const FName Name = NAME_None);
 
 	/**
-	* Sets the data for this observation. Call during ULearningAgentsInteractor::SetObservations event.
-	* @param AgentId The agent id this data corresponds to.
-	* @param Value The enum value currently being observed.
-	*/
+	 * Sets the data for this observation. Call during ULearningAgentsInteractor::SetObservations event.
+	 * @param AgentId The agent id this data corresponds to.
+	 * @param Value The enum value currently being observed.
+	 */
 	UFUNCTION(BlueprintCallable, Category = "LearningAgents", meta = (AgentId = "-1"))
 	void SetEnumObservation(const int32 AgentId, const uint8 Value);
 
@@ -245,21 +252,21 @@ class LEARNINGAGENTS_API UEnumArrayObservation : public ULearningAgentsObservati
 public:
 
 	/**
-	* Adds a new enum array observation to the given agent interactor. Call during ULearningAgentsInteractor::SetupObservations event.
-	* @param InInteractor The agent interactor to add this observation to.
-	* @param EnumType The type of enum to use
-	* @param Name The name of this new observation. Used for debugging.
-	* @param EnumNum The number of enum observations in the array
-	* @return The newly created observation.
-	*/
-	UFUNCTION(BlueprintCallable, Category = "LearningAgents")
+	 * Adds a new enum array observation to the given agent interactor. Call during ULearningAgentsInteractor::SetupObservations event.
+	 * @param InInteractor The agent interactor to add this observation to.
+	 * @param EnumType The type of enum to use
+	 * @param Name The name of this new observation. Used for debugging.
+	 * @param EnumNum The number of enum observations in the array
+	 * @return The newly created observation.
+	 */
+	UFUNCTION(BlueprintCallable, Category = "LearningAgents", meta = (DefaultToSelf = "InInteractor"))
 	static UEnumArrayObservation* AddEnumArrayObservation(ULearningAgentsInteractor* InInteractor, const UEnum* EnumType, const FName Name = NAME_None, const int32 EnumNum = 1);
 
 	/**
-	* Sets the data for this observation. Call during ULearningAgentsInteractor::SetObservations event.
-	* @param AgentId The agent id this data corresponds to.
-	* @param Values The enum values currently being observed.
-	*/
+	 * Sets the data for this observation. Call during ULearningAgentsInteractor::SetObservations event.
+	 * @param AgentId The agent id this data corresponds to.
+	 * @param Values The enum values currently being observed.
+	 */
 	UFUNCTION(BlueprintCallable, Category = "LearningAgents", meta = (AgentId = "-1"))
 	void SetEnumArrayObservation(const int32 AgentId, const TArray<uint8>& Values);
 
@@ -283,21 +290,21 @@ class LEARNINGAGENTS_API UTimeObservation : public ULearningAgentsObservation
 public:
 
 	/**
-	* Adds a new time observation to the given agent interactor. Call during ULearningAgentsInteractor::SetupObservations event.
-	* @param InInteractor The agent interactor to add this observation to.
-	* @param Name The name of this new observation. Used for debugging.
-	* @param Scale Used to normalize the data for the observation.
-	* @return The newly created observation.
-	*/
-	UFUNCTION(BlueprintCallable, Category = "LearningAgents")
+	 * Adds a new time observation to the given agent interactor. Call during ULearningAgentsInteractor::SetupObservations event.
+	 * @param InInteractor The agent interactor to add this observation to.
+	 * @param Name The name of this new observation. Used for debugging.
+	 * @param Scale Used to normalize the data for the observation.
+	 * @return The newly created observation.
+	 */
+	UFUNCTION(BlueprintCallable, Category = "LearningAgents", meta = (DefaultToSelf = "InInteractor"))
 	static UTimeObservation* AddTimeObservation(ULearningAgentsInteractor* InInteractor, const FName Name = NAME_None, const float Scale = 1.0f);
 
 	/**
-	* Sets the data for this observation. Call during ULearningAgentsInteractor::SetObservations event.
-	* @param AgentId The agent id this data corresponds to.
-	* @param Time The time currently being observed.
-	* @param RelativeTime The time the provided time should be encoded relative to.
-	*/
+	 * Sets the data for this observation. Call during ULearningAgentsInteractor::SetObservations event.
+	 * @param AgentId The agent id this data corresponds to.
+	 * @param Time The time currently being observed.
+	 * @param RelativeTime The time the provided time should be encoded relative to.
+	 */
 	UFUNCTION(BlueprintCallable, Category = "LearningAgents", meta = (AgentId = "-1"))
 	void SetTimeObservation(const int32 AgentId, const float Time, const float RelativeTime = 0.0f);
 
@@ -318,22 +325,22 @@ class LEARNINGAGENTS_API UTimeArrayObservation : public ULearningAgentsObservati
 public:
 
 	/**
-	* Adds a new angle observation to the given agent interactor. Call during ULearningAgentsInteractor::SetupObservations event.
-	* @param InInteractor The agent interactor to add this observation to.
-	* @param Name The name of this new observation. Used for debugging.
-	* @param TimeNum The number of times in the array
-	* @param Scale Used to normalize the data for the observation.
-	* @return The newly created observation.
-	*/
-	UFUNCTION(BlueprintCallable, Category = "LearningAgents")
+	 * Adds a new angle observation to the given agent interactor. Call during ULearningAgentsInteractor::SetupObservations event.
+	 * @param InInteractor The agent interactor to add this observation to.
+	 * @param Name The name of this new observation. Used for debugging.
+	 * @param TimeNum The number of times in the array
+	 * @param Scale Used to normalize the data for the observation.
+	 * @return The newly created observation.
+	 */
+	UFUNCTION(BlueprintCallable, Category = "LearningAgents", meta = (DefaultToSelf = "InInteractor"))
 	static UTimeArrayObservation* AddTimeArrayObservation(ULearningAgentsInteractor* InInteractor, const FName Name = NAME_None, const int32 TimeNum = 1, const float Scale = 1.0f);
 
 	/**
-	* Sets the data for this observation. Call during ULearningAgentsInteractor::SetObservations event.
-	* @param AgentId The agent id this data corresponds to.
-	* @param Time The times currently being observed.
-	* @param RelativeTime The time the provided time should be encoded relative to.
-	*/
+	 * Sets the data for this observation. Call during ULearningAgentsInteractor::SetObservations event.
+	 * @param AgentId The agent id this data corresponds to.
+	 * @param Time The times currently being observed.
+	 * @param RelativeTime The time the provided time should be encoded relative to.
+	 */
 	UFUNCTION(BlueprintCallable, Category = "LearningAgents", meta = (AgentId = "-1"))
 	void SetTimeArrayObservation(const int32 AgentId, const TArray<float>& Times, const float RelativeTime = 0.0f);
 
@@ -356,21 +363,21 @@ class LEARNINGAGENTS_API UAngleObservation : public ULearningAgentsObservation
 public:
 
 	/**
-	* Adds a new angle observation to the given agent interactor. Call during ULearningAgentsInteractor::SetupObservations event.
-	* @param InInteractor The agent interactor to add this observation to.
-	* @param Name The name of this new observation. Used for debugging.
-	* @param Scale Used to normalize the data for the observation. Angle observations are encoded as directions. 
-	* @return The newly created observation.
-	*/
-	UFUNCTION(BlueprintCallable, Category = "LearningAgents")
+	 * Adds a new angle observation to the given agent interactor. Call during ULearningAgentsInteractor::SetupObservations event.
+	 * @param InInteractor The agent interactor to add this observation to.
+	 * @param Name The name of this new observation. Used for debugging.
+	 * @param Scale Used to normalize the data for the observation. Angle observations are encoded as directions. 
+	 * @return The newly created observation.
+	 */
+	UFUNCTION(BlueprintCallable, Category = "LearningAgents", meta = (DefaultToSelf = "InInteractor"))
 	static UAngleObservation* AddAngleObservation(ULearningAgentsInteractor* InInteractor, const FName Name = NAME_None, const float Scale = 1.0f);
 
 	/**
-	* Sets the data for this observation. Call during ULearningAgentsInteractor::SetObservations event.
-	* @param AgentId The agent id this data corresponds to.
-	* @param Angle The angle currently being observed.
-	* @param RelativeAngle The frame of reference angle.
-	*/
+	 * Sets the data for this observation. Call during ULearningAgentsInteractor::SetObservations event.
+	 * @param AgentId The agent id this data corresponds to.
+	 * @param Angle The angle currently being observed.
+	 * @param RelativeAngle The frame of reference angle.
+	 */
 	UFUNCTION(BlueprintCallable, Category = "LearningAgents", meta = (AgentId = "-1"))
 	void SetAngleObservation(const int32 AgentId, const float Angle, const float RelativeAngle = 0.0f);
 
@@ -391,22 +398,22 @@ class LEARNINGAGENTS_API UAngleArrayObservation : public ULearningAgentsObservat
 public:
 
 	/**
-	* Adds a new angle observation to the given agent interactor. Call during ULearningAgentsInteractor::SetupObservations event.
-	* @param InInteractor The agent interactor to add this observation to.
-	* @param Name The name of this new observation. Used for debugging.
-	* @param AngleNum The number of angles in the array
-	* @param Scale Used to normalize the data for the observation. Angle observations are encoded as directions. 
-	* @return The newly created observation.
-	*/
-	UFUNCTION(BlueprintCallable, Category = "LearningAgents")
+	 * Adds a new angle observation to the given agent interactor. Call during ULearningAgentsInteractor::SetupObservations event.
+	 * @param InInteractor The agent interactor to add this observation to.
+	 * @param Name The name of this new observation. Used for debugging.
+	 * @param AngleNum The number of angles in the array
+	 * @param Scale Used to normalize the data for the observation. Angle observations are encoded as directions. 
+	 * @return The newly created observation.
+	 */
+	UFUNCTION(BlueprintCallable, Category = "LearningAgents", meta = (DefaultToSelf = "InInteractor"))
 	static UAngleArrayObservation* AddAngleArrayObservation(ULearningAgentsInteractor* InInteractor, const FName Name = NAME_None, const int32 AngleNum = 1, const float Scale = 1.0f);
 
 	/**
-	* Sets the data for this observation. Call during ULearningAgentsInteractor::SetObservations event.
-	* @param AgentId The agent id this data corresponds to.
-	* @param Angles The angles currently being observed.
-	* @param RelativeAngle The frame of reference angle.
-	*/
+	 * Sets the data for this observation. Call during ULearningAgentsInteractor::SetObservations event.
+	 * @param AgentId The agent id this data corresponds to.
+	 * @param Angles The angles currently being observed.
+	 * @param RelativeAngle The frame of reference angle.
+	 */
 	UFUNCTION(BlueprintCallable, Category = "LearningAgents", meta = (AgentId = "-1"))
 	void SetAngleArrayObservation(const int32 AgentId, const TArray<float>& Angles, const float RelativeAngle = 0.0f);
 
@@ -427,30 +434,30 @@ class LEARNINGAGENTS_API URotationObservation : public ULearningAgentsObservatio
 public:
 
 	/**
-	* Adds a new rotation observation to the given agent interactor. Call during ULearningAgentsInteractor::SetupObservations event.
-	* @param InInteractor The agent interactor to add this observation to.
-	* @param Name The name of this new observation. Used for debugging.
-	* @param Scale Used to normalize the data for the observation. Rotation observations are encoded as directions. 
-	* @return The newly created observation.
-	*/
-	UFUNCTION(BlueprintCallable, Category = "LearningAgents")
+	 * Adds a new rotation observation to the given agent interactor. Call during ULearningAgentsInteractor::SetupObservations event.
+	 * @param InInteractor The agent interactor to add this observation to.
+	 * @param Name The name of this new observation. Used for debugging.
+	 * @param Scale Used to normalize the data for the observation. Rotation observations are encoded as directions. 
+	 * @return The newly created observation.
+	 */
+	UFUNCTION(BlueprintCallable, Category = "LearningAgents", meta = (DefaultToSelf = "InInteractor"))
 	static URotationObservation* AddRotationObservation(ULearningAgentsInteractor* InInteractor, const FName Name = NAME_None, const float Scale = 1.0f);
 
 	/**
-	* Sets the data for this observation. Call during ULearningAgentsInteractor::SetObservations event.
-	* @param AgentId The agent id this data corresponds to.
-	* @param Rotation The rotation currently being observed.
-	* @param RelativeRotation The frame of reference rotation.
-	*/
+	 * Sets the data for this observation. Call during ULearningAgentsInteractor::SetObservations event.
+	 * @param AgentId The agent id this data corresponds to.
+	 * @param Rotation The rotation currently being observed.
+	 * @param RelativeRotation The frame of reference rotation.
+	 */
 	UFUNCTION(BlueprintCallable, Category = "LearningAgents", meta = (AgentId = "-1"))
 	void SetRotationObservation(const int32 AgentId, const FRotator Rotation, const FRotator RelativeRotation = FRotator::ZeroRotator);
 
 	/**
-	* Sets the data for this observation. Call during ULearningAgentsInteractor::SetObservations event.
-	* @param AgentId The agent id this data corresponds to.
-	* @param Rotation The rotation currently being observed.
-	* @param RelativeRotation The frame of reference rotation.
-	*/
+	 * Sets the data for this observation. Call during ULearningAgentsInteractor::SetObservations event.
+	 * @param AgentId The agent id this data corresponds to.
+	 * @param Rotation The rotation currently being observed.
+	 * @param RelativeRotation The frame of reference rotation.
+	 */
 	UFUNCTION(BlueprintCallable, Category = "LearningAgents", meta = (AgentId = "-1"))
 	void SetRotationObservationFromQuat(const int32 AgentId, const FQuat Rotation, const FQuat RelativeRotation);
 
@@ -471,31 +478,31 @@ class LEARNINGAGENTS_API URotationArrayObservation : public ULearningAgentsObser
 public:
 
 	/**
-	* Adds a new rotation observation to the given agent interactor. Call during ULearningAgentsInteractor::SetupObservations event.
-	* @param InInteractor The agent interactor to add this observation to.
-	* @param Name The name of this new observation. Used for debugging.
-	* @param RotationNum The number of rotations in the array
-	* @param Scale Used to normalize the data for the observation. Rotation observations are encoded as directions. 
-	* @return The newly created observation.
-	*/
-	UFUNCTION(BlueprintCallable, Category = "LearningAgents")
+	 * Adds a new rotation observation to the given agent interactor. Call during ULearningAgentsInteractor::SetupObservations event.
+	 * @param InInteractor The agent interactor to add this observation to.
+	 * @param Name The name of this new observation. Used for debugging.
+	 * @param RotationNum The number of rotations in the array
+	 * @param Scale Used to normalize the data for the observation. Rotation observations are encoded as directions. 
+	 * @return The newly created observation.
+	 */
+	UFUNCTION(BlueprintCallable, Category = "LearningAgents", meta = (DefaultToSelf = "InInteractor"))
 	static URotationArrayObservation* AddRotationArrayObservation(ULearningAgentsInteractor* InInteractor, const FName Name = NAME_None, const int32 RotationNum = 1, const float Scale = 1.0f);
 
 	/**
-	* Sets the data for this observation. Call during ULearningAgentsInteractor::SetObservations event.
-	* @param AgentId The agent id this data corresponds to.
-	* @param Rotations The rotations currently being observed.
-	* @param RelativeRotation The frame of reference rotation.
-	*/
+	 * Sets the data for this observation. Call during ULearningAgentsInteractor::SetObservations event.
+	 * @param AgentId The agent id this data corresponds to.
+	 * @param Rotations The rotations currently being observed.
+	 * @param RelativeRotation The frame of reference rotation.
+	 */
 	UFUNCTION(BlueprintCallable, Category = "LearningAgents", meta = (AgentId = "-1"))
 	void SetRotationArrayObservation(const int32 AgentId, const TArray<FRotator>& Rotations, const FRotator RelativeRotation = FRotator::ZeroRotator);
 
 	/**
-	* Sets the data for this observation. Call during ULearningAgentsInteractor::SetObservations event.
-	* @param AgentId The agent id this data corresponds to.
-	* @param Rotations The rotations currently being observed.
-	* @param RelativeRotation The frame of reference rotation.
-	*/
+	 * Sets the data for this observation. Call during ULearningAgentsInteractor::SetObservations event.
+	 * @param AgentId The agent id this data corresponds to.
+	 * @param Rotations The rotations currently being observed.
+	 * @param RelativeRotation The frame of reference rotation.
+	 */
 	UFUNCTION(BlueprintCallable, Category = "LearningAgents", meta = (AgentId = "-1"))
 	void SetRotationArrayObservationFromQuats(const int32 AgentId, const TArray<FQuat>& Rotations, const FQuat RelativeRotation);
 
@@ -518,22 +525,22 @@ class LEARNINGAGENTS_API UDirectionObservation : public ULearningAgentsObservati
 public:
 
 	/**
-	* Adds a new direction observation to the given agent interactor. Call during ULearningAgentsInteractor::SetupObservations event.
-	* @param InInteractor The agent interactor to add this observation to.
-	* @param Name The name of this new observation. Used for debugging.
-	* @param Scale Used to normalize the data for the observation.
-	* @return The newly created observation.
-	*/
-	UFUNCTION(BlueprintCallable, Category = "LearningAgents")
+	 * Adds a new direction observation to the given agent interactor. Call during ULearningAgentsInteractor::SetupObservations event.
+	 * @param InInteractor The agent interactor to add this observation to.
+	 * @param Name The name of this new observation. Used for debugging.
+	 * @param Scale Used to normalize the data for the observation.
+	 * @return The newly created observation.
+	 */
+	UFUNCTION(BlueprintCallable, Category = "LearningAgents", meta = (DefaultToSelf = "InInteractor"))
 	static UDirectionObservation* AddDirectionObservation(ULearningAgentsInteractor* InInteractor, const FName Name = NAME_None, const float Scale = 1.0f);
 
 	/**
-	* Sets the data for this observation. The relative rotation can be used to make this observation relative to the
-	* agent's perspective, e.g. by passing the agent's forward rotation. Call during ULearningAgentsInteractor::SetObservations event.
-	* @param AgentId The agent id this data corresponds to.
-	* @param Direction The direction currently being observed.
-	* @param RelativeRotation The frame of reference rotation.
-	*/
+	 * Sets the data for this observation. The relative rotation can be used to make this observation relative to the
+	 * agent's perspective, e.g. by passing the agent's forward rotation. Call during ULearningAgentsInteractor::SetObservations event.
+	 * @param AgentId The agent id this data corresponds to.
+	 * @param Direction The direction currently being observed.
+	 * @param RelativeRotation The frame of reference rotation.
+	 */
 	UFUNCTION(BlueprintCallable, Category = "LearningAgents", meta = (AgentId = "-1"))
 	void SetDirectionObservation(const int32 AgentId, const FVector Direction, const FRotator RelativeRotation = FRotator::ZeroRotator);
 
@@ -554,23 +561,23 @@ class LEARNINGAGENTS_API UDirectionArrayObservation : public ULearningAgentsObse
 public:
 
 	/**
-	* Adds a new direction array observation to the given agent interactor. Call during ULearningAgentsInteractor::SetupObservations event.
-	* @param InInteractor The agent interactor to add this observation to.
-	* @param Name The name of this new observation. Used for debugging.
-	* @param DirectionNum The number of directions in the array
-	* @param Scale Used to normalize the data for the observation.
-	* @return The newly created observation.
-	*/
-	UFUNCTION(BlueprintCallable, Category = "LearningAgents")
+	 * Adds a new direction array observation to the given agent interactor. Call during ULearningAgentsInteractor::SetupObservations event.
+	 * @param InInteractor The agent interactor to add this observation to.
+	 * @param Name The name of this new observation. Used for debugging.
+	 * @param DirectionNum The number of directions in the array
+	 * @param Scale Used to normalize the data for the observation.
+	 * @return The newly created observation.
+	 */
+	UFUNCTION(BlueprintCallable, Category = "LearningAgents", meta = (DefaultToSelf = "InInteractor"))
 	static UDirectionArrayObservation* AddDirectionArrayObservation(ULearningAgentsInteractor* InInteractor, const FName Name = NAME_None, const int32 DirectionNum = 1, const float Scale = 1.0f);
 
 	/**
-	* Sets the data for this observation. The relative rotation can be used to make this observation relative to the
-	* agent's perspective, e.g. by passing the agent's forward rotation. Call during ULearningAgentsInteractor::SetObservations event.
-	* @param AgentId The agent id this data corresponds to.
-	* @param Directions The directions currently being observed.
-	* @param RelativeRotation The frame of reference rotation.
-	*/
+	 * Sets the data for this observation. The relative rotation can be used to make this observation relative to the
+	 * agent's perspective, e.g. by passing the agent's forward rotation. Call during ULearningAgentsInteractor::SetObservations event.
+	 * @param AgentId The agent id this data corresponds to.
+	 * @param Directions The directions currently being observed.
+	 * @param RelativeRotation The frame of reference rotation.
+	 */
 	UFUNCTION(BlueprintCallable, Category = "LearningAgents", meta = (AgentId = "-1"))
 	void SetDirectionArrayObservation(const int32 AgentId, const TArray<FVector>& Directions, const FRotator RelativeRotation = FRotator::ZeroRotator);
 
@@ -591,25 +598,25 @@ class LEARNINGAGENTS_API UPlanarDirectionObservation : public ULearningAgentsObs
 public:
 
 	/**
-	* Adds a new planar direction observation to the given agent interactor. The axis parameters define the plane. Call
-	* during ULearningAgentsInteractor::SetupObservations event.
-	* @param InInteractor The agent interactor to add this observation to.
-	* @param Name The name of this new observation. Used for debugging.
-	* @param Scale Used to normalize the data for the observation.
-	* @param Axis0 The forward axis of the plane.
-	* @param Axis1 The right axis of the plane.
-	* @return The newly created observation.
-	*/
-	UFUNCTION(BlueprintCallable, Category = "LearningAgents")
+	 * Adds a new planar direction observation to the given agent interactor. The axis parameters define the plane. Call
+	 * during ULearningAgentsInteractor::SetupObservations event.
+	 * @param InInteractor The agent interactor to add this observation to.
+	 * @param Name The name of this new observation. Used for debugging.
+	 * @param Scale Used to normalize the data for the observation.
+	 * @param Axis0 The forward axis of the plane.
+	 * @param Axis1 The right axis of the plane.
+	 * @return The newly created observation.
+	 */
+	UFUNCTION(BlueprintCallable, Category = "LearningAgents", meta = (DefaultToSelf = "InInteractor"))
 	static UPlanarDirectionObservation* AddPlanarDirectionObservation(ULearningAgentsInteractor* InInteractor, const FName Name = NAME_None, const float Scale = 1.0f, const FVector Axis0 = FVector::ForwardVector, const FVector Axis1 = FVector::RightVector);
 
 	/**
-	* Sets the data for this observation. The relative rotation can be used to make this observation relative to the
-	* agent's perspective, e.g. by passing the agent's forward rotation. Call during ULearningAgentsInteractor::SetObservations event.
-	* @param AgentId The agent id this data corresponds to.
-	* @param Direction The direction currently being observed.
-	* @param RelativeRotation The frame of reference rotation.
-	*/
+	 * Sets the data for this observation. The relative rotation can be used to make this observation relative to the
+	 * agent's perspective, e.g. by passing the agent's forward rotation. Call during ULearningAgentsInteractor::SetObservations event.
+	 * @param AgentId The agent id this data corresponds to.
+	 * @param Direction The direction currently being observed.
+	 * @param RelativeRotation The frame of reference rotation.
+	 */
 	UFUNCTION(BlueprintCallable, Category = "LearningAgents", meta = (AgentId = "-1"))
 	void SetPlanarDirectionObservation(const int32 AgentId, const FVector Direction, const FRotator RelativeRotation = FRotator::ZeroRotator);
 
@@ -630,26 +637,26 @@ class LEARNINGAGENTS_API UPlanarDirectionArrayObservation : public ULearningAgen
 public:
 
 	/**
-	* Adds a new planar direction observation to the given agent interactor. The axis parameters define the plane. Call
-	* during ULearningAgentsInteractor::SetupObservations event.
-	* @param InInteractor The agent interactor to add this observation to.
-	* @param Name The name of this new observation. Used for debugging.
-	* @param DirectionNum The number of directions in the array
-	* @param Scale Used to normalize the data for the observation.
-	* @param Axis0 The forward axis of the plane.
-	* @param Axis1 The right axis of the plane.
-	* @return The newly created observation.
-	*/
-	UFUNCTION(BlueprintCallable, Category = "LearningAgents")
+	 * Adds a new planar direction observation to the given agent interactor. The axis parameters define the plane. Call
+	 * during ULearningAgentsInteractor::SetupObservations event.
+	 * @param InInteractor The agent interactor to add this observation to.
+	 * @param Name The name of this new observation. Used for debugging.
+	 * @param DirectionNum The number of directions in the array
+	 * @param Scale Used to normalize the data for the observation.
+	 * @param Axis0 The forward axis of the plane.
+	 * @param Axis1 The right axis of the plane.
+	 * @return The newly created observation.
+	 */
+	UFUNCTION(BlueprintCallable, Category = "LearningAgents", meta = (DefaultToSelf = "InInteractor"))
 	static UPlanarDirectionArrayObservation* AddPlanarDirectionArrayObservation(ULearningAgentsInteractor* InInteractor, const FName Name = NAME_None, const int32 DirectionNum = 1, const float Scale = 1.0f, const FVector Axis0 = FVector::ForwardVector, const FVector Axis1 = FVector::RightVector);
 
 	/**
-	* Sets the data for this observation. The relative rotation can be used to make this observation relative to the
-	* agent's perspective, e.g. by passing the agent's forward rotation. Call during ULearningAgentsInteractor::SetObservations event.
-	* @param AgentId The agent id this data corresponds to.
-	* @param Directions The directions currently being observed.
-	* @param RelativeRotation The frame of reference rotation.
-	*/
+	 * Sets the data for this observation. The relative rotation can be used to make this observation relative to the
+	 * agent's perspective, e.g. by passing the agent's forward rotation. Call during ULearningAgentsInteractor::SetObservations event.
+	 * @param AgentId The agent id this data corresponds to.
+	 * @param Directions The directions currently being observed.
+	 * @param RelativeRotation The frame of reference rotation.
+	 */
 	UFUNCTION(BlueprintCallable, Category = "LearningAgents", meta = (AgentId = "-1"))
 	void SetPlanarDirectionArrayObservation(const int32 AgentId, const TArray<FVector>& Directions, const FRotator RelativeRotation = FRotator::ZeroRotator);
 
@@ -672,24 +679,24 @@ class LEARNINGAGENTS_API UPositionObservation : public ULearningAgentsObservatio
 public:
 
 	/**
-	* Adds a new position observation to the given agent interactor. Call during ULearningAgentsInteractor::SetupObservations event.
-	* @param InInteractor The agent interactor to add this observation to.
-	* @param Name The name of this new observation. Used for debugging.
-	* @param Scale Used to normalize the data for the observation.
-	* @return The newly created observation.
-	*/
-	UFUNCTION(BlueprintCallable, Category = "LearningAgents")
+	 * Adds a new position observation to the given agent interactor. Call during ULearningAgentsInteractor::SetupObservations event.
+	 * @param InInteractor The agent interactor to add this observation to.
+	 * @param Name The name of this new observation. Used for debugging.
+	 * @param Scale Used to normalize the data for the observation.
+	 * @return The newly created observation.
+	 */
+	UFUNCTION(BlueprintCallable, Category = "LearningAgents", meta = (DefaultToSelf = "InInteractor"))
 	static UPositionObservation* AddPositionObservation(ULearningAgentsInteractor* InInteractor, const FName Name = NAME_None, const float Scale = 100.0f);
 
 	/**
-	* Sets the data for this observation. The relative position & rotation can be used to make this observation
-	* relative to the agent's perspective, e.g. by passing the agent's position & forward rotation. Call during
-	* ULearningAgentsInteractor::SetObservations event.
-	* @param AgentId The agent id this data corresponds to.
-	* @param Position The position currently being observed.
-	* @param RelativePosition The vector Position will be offset from.
-	* @param RelativeRotation The frame of reference rotation.
-	*/
+	 * Sets the data for this observation. The relative position & rotation can be used to make this observation
+	 * relative to the agent's perspective, e.g. by passing the agent's position & forward rotation. Call during
+	 * ULearningAgentsInteractor::SetObservations event.
+	 * @param AgentId The agent id this data corresponds to.
+	 * @param Position The position currently being observed.
+	 * @param RelativePosition The vector Position will be offset from.
+	 * @param RelativeRotation The frame of reference rotation.
+	 */
 	UFUNCTION(BlueprintCallable, Category = "LearningAgents", meta = (AgentId = "-1"))
 	void SetPositionObservation(const int32 AgentId, const FVector Position, const FVector RelativePosition = FVector::ZeroVector, const FRotator RelativeRotation = FRotator::ZeroRotator);
 
@@ -710,25 +717,25 @@ class LEARNINGAGENTS_API UPositionArrayObservation : public ULearningAgentsObser
 public:
 
 	/**
-	* Adds a new position array observation to the given agent interactor. Call during ULearningAgentsInteractor::SetupObservations event.
-	* @param InInteractor The agent interactor to add this observation to.
-	* @param Name The name of this new observation. Used for debugging.
-	* @param PositionNum The number of positions in the array.
-	* @param Scale Used to normalize the data for the observation.
-	* @return The newly created observation.
-	*/
-	UFUNCTION(BlueprintCallable, Category = "LearningAgents")
+	 * Adds a new position array observation to the given agent interactor. Call during ULearningAgentsInteractor::SetupObservations event.
+	 * @param InInteractor The agent interactor to add this observation to.
+	 * @param Name The name of this new observation. Used for debugging.
+	 * @param PositionNum The number of positions in the array.
+	 * @param Scale Used to normalize the data for the observation.
+	 * @return The newly created observation.
+	 */
+	UFUNCTION(BlueprintCallable, Category = "LearningAgents", meta = (DefaultToSelf = "InInteractor"))
 	static UPositionArrayObservation* AddPositionArrayObservation(ULearningAgentsInteractor* InInteractor, const FName Name = NAME_None, const int32 PositionNum = 1, const float Scale = 100.0f);
 
 	/**
-	* Sets the data for this observation. The relative position & rotation can be used to make this observation
-	* relative to the agent's perspective, e.g. by passing the agent's position & forward rotation. Call during
-	* ULearningAgentsInteractor::SetObservations event.
-	* @param AgentId The agent id this data corresponds to.
-	* @param Positions The positions currently being observed.
-	* @param RelativePosition The vector Positions will be offset from.
-	* @param RelativeRotation The frame of reference rotation.
-	*/
+	 * Sets the data for this observation. The relative position & rotation can be used to make this observation
+	 * relative to the agent's perspective, e.g. by passing the agent's position & forward rotation. Call during
+	 * ULearningAgentsInteractor::SetObservations event.
+	 * @param AgentId The agent id this data corresponds to.
+	 * @param Positions The positions currently being observed.
+	 * @param RelativePosition The vector Positions will be offset from.
+	 * @param RelativeRotation The frame of reference rotation.
+	 */
 	UFUNCTION(BlueprintCallable, Category = "LearningAgents", meta = (AgentId = "-1"))
 	void SetPositionArrayObservation(int32 AgentId, const TArray<FVector>& Positions, const FVector RelativePosition = FVector::ZeroVector, const FRotator RelativeRotation = FRotator::ZeroRotator);
 
@@ -749,36 +756,36 @@ class LEARNINGAGENTS_API UScalarPositionObservation : public ULearningAgentsObse
 public:
 
 	/**
-	* Adds a new scalar position observation to the given agent interactor. Call during 
-	* ULearningAgentsInteractor::SetupObservations event.
-	* @param InInteractor The agent interactor to add this observation to.
-	* @param Name The name of this new observation. Used for debugging.
-	* @param Scale Used to normalize the data for the observation.
-	* @return The newly created observation.
-	*/
-	UFUNCTION(BlueprintCallable, Category = "LearningAgents")
+	 * Adds a new scalar position observation to the given agent interactor. Call during 
+	 * ULearningAgentsInteractor::SetupObservations event.
+	 * @param InInteractor The agent interactor to add this observation to.
+	 * @param Name The name of this new observation. Used for debugging.
+	 * @param Scale Used to normalize the data for the observation.
+	 * @return The newly created observation.
+	 */
+	UFUNCTION(BlueprintCallable, Category = "LearningAgents", meta = (DefaultToSelf = "InInteractor"))
 	static UScalarPositionObservation* AddScalarPositionObservation(ULearningAgentsInteractor* InInteractor, const FName Name = NAME_None, const float Scale = 100.0f);
 
 	/**
-	* Sets the data for this observation. The relative position can be used to make this observation
-	* relative to the agent's perspective, e.g. by passing the agent's position. Call during 
-	* ULearningAgentsInteractor::SetObservations event.
-	* @param AgentId The agent id this data corresponds to.
-	* @param Position The position currently being observed.
-	* @param RelativePosition The vector Position will be offset from.
-	*/
+	 * Sets the data for this observation. The relative position can be used to make this observation
+	 * relative to the agent's perspective, e.g. by passing the agent's position. Call during 
+	 * ULearningAgentsInteractor::SetObservations event.
+	 * @param AgentId The agent id this data corresponds to.
+	 * @param Position The position currently being observed.
+	 * @param RelativePosition The vector Position will be offset from.
+	 */
 	UFUNCTION(BlueprintCallable, Category = "LearningAgents", meta = (AgentId = "-1"))
 	void SetScalarPositionObservation(const int32 AgentId, const float Position, const float RelativePosition = 0.0f);
 
 	/**
-	* Sets the data for this observation. The relative position can be used to make this observation
-	* relative to the agent's perspective, e.g. by passing the agent's position. Call during
-	* ULearningAgentsInteractor::SetObservations event.
-	* @param AgentId The agent id this data corresponds to.
-	* @param Position The position currently being observed.
-	* @param RelativePosition The vector Position will be offset from.
-	* @param Axis The axis along which to encode the positions
-	*/
+	 * Sets the data for this observation. The relative position can be used to make this observation
+	 * relative to the agent's perspective, e.g. by passing the agent's position. Call during
+	 * ULearningAgentsInteractor::SetObservations event.
+	 * @param AgentId The agent id this data corresponds to.
+	 * @param Position The position currently being observed.
+	 * @param RelativePosition The vector Position will be offset from.
+	 * @param Axis The axis along which to encode the positions
+	 */
 	UFUNCTION(BlueprintCallable, Category = "LearningAgents", meta = (AgentId = "-1"))
 	void SetScalarPositionObservationWithAxis(const int32 AgentId, const FVector Position, const FVector RelativePosition = FVector::ZeroVector, const FVector Axis = FVector::UpVector);
 
@@ -799,37 +806,37 @@ class LEARNINGAGENTS_API UScalarPositionArrayObservation : public ULearningAgent
 public:
 
 	/**
-	* Adds a new scalar position array observation to the given agent interactor. Call during 
-	* ULearningAgentsInteractor::SetupObservations event.
-	* @param InInteractor The agent interactor to add this observation to.
-	* @param Name The name of this new observation. Used for debugging.
-	* @param PositionNum The number of positions in the array.
-	* @param Scale Used to normalize the data for the observation.
-	* @return The newly created observation.
-	*/
-	UFUNCTION(BlueprintCallable, Category = "LearningAgents")
+	 * Adds a new scalar position array observation to the given agent interactor. Call during 
+	 * ULearningAgentsInteractor::SetupObservations event.
+	 * @param InInteractor The agent interactor to add this observation to.
+	 * @param Name The name of this new observation. Used for debugging.
+	 * @param PositionNum The number of positions in the array.
+	 * @param Scale Used to normalize the data for the observation.
+	 * @return The newly created observation.
+	 */
+	UFUNCTION(BlueprintCallable, Category = "LearningAgents", meta = (DefaultToSelf = "InInteractor"))
 	static UScalarPositionArrayObservation* AddScalarPositionArrayObservation(ULearningAgentsInteractor* InInteractor, const FName Name = NAME_None, const int32 PositionNum = 1, const float Scale = 100.0f);
 
 	/**
-	* Sets the data for this observation. The relative position can be used to make this observation
-	* relative to the agent's perspective, e.g. by passing the agent's position. Call during
-	* ULearningAgentsInteractor::SetObservations event.
-	* @param AgentId The agent id this data corresponds to.
-	* @param Positions The positions currently being observed.
-	* @param RelativePosition The vector Positions will be offset from.
-	*/
+	 * Sets the data for this observation. The relative position can be used to make this observation
+	 * relative to the agent's perspective, e.g. by passing the agent's position. Call during
+	 * ULearningAgentsInteractor::SetObservations event.
+	 * @param AgentId The agent id this data corresponds to.
+	 * @param Positions The positions currently being observed.
+	 * @param RelativePosition The vector Positions will be offset from.
+	 */
 	UFUNCTION(BlueprintCallable, Category = "LearningAgents", meta = (AgentId = "-1"))
 	void SetScalarPositionArrayObservation(const int32 AgentId, const TArray<float>& Positions, const float RelativePosition = 0.0f);
 
 	/**
-	* Sets the data for this observation. The relative position can be used to make this observation
-	* relative to the agent's perspective, e.g. by passing the agent's position. Call during
-	* ULearningAgentsInteractor::SetObservations event.
-	* @param AgentId The agent id this data corresponds to.
-	* @param Positions The positions currently being observed.
-	* @param RelativePosition The vector Positions will be offset from.
-	* @param Axis The axis along which to encode the positions
-	*/
+	 * Sets the data for this observation. The relative position can be used to make this observation
+	 * relative to the agent's perspective, e.g. by passing the agent's position. Call during
+	 * ULearningAgentsInteractor::SetObservations event.
+	 * @param AgentId The agent id this data corresponds to.
+	 * @param Positions The positions currently being observed.
+	 * @param RelativePosition The vector Positions will be offset from.
+	 * @param Axis The axis along which to encode the positions
+	 */
 	UFUNCTION(BlueprintCallable, Category = "LearningAgents", meta = (AgentId = "-1"))
 	void SetScalarPositionArrayObservationWithAxis(const int32 AgentId, const TArray<FVector>& Positions, const FVector RelativePosition = FVector::ZeroVector, const FVector Axis = FVector::UpVector);
 
@@ -850,27 +857,27 @@ class LEARNINGAGENTS_API UPlanarPositionObservation : public ULearningAgentsObse
 public:
 
 	/**
-	* Adds a new planar position observation to the given agent interactor. The axis parameters define the plane. Call
-	* during ULearningAgentsInteractor::SetupObservations event.
-	* @param InInteractor The agent interactor to add this observation to.
-	* @param Name The name of this new observation. Used for debugging.
-	* @param Scale Used to normalize the data for the observation.
-	* @param Axis0 The forward axis of the plane.
-	* @param Axis1 The right axis of the plane.
-	* @return The newly created observation.
-	*/
-	UFUNCTION(BlueprintCallable, Category = "LearningAgents")
+	 * Adds a new planar position observation to the given agent interactor. The axis parameters define the plane. Call
+	 * during ULearningAgentsInteractor::SetupObservations event.
+	 * @param InInteractor The agent interactor to add this observation to.
+	 * @param Name The name of this new observation. Used for debugging.
+	 * @param Scale Used to normalize the data for the observation.
+	 * @param Axis0 The forward axis of the plane.
+	 * @param Axis1 The right axis of the plane.
+	 * @return The newly created observation.
+	 */
+	UFUNCTION(BlueprintCallable, Category = "LearningAgents", meta = (DefaultToSelf = "InInteractor"))
 	static UPlanarPositionObservation* AddPlanarPositionObservation(ULearningAgentsInteractor* InInteractor, const FName Name = NAME_None, const float Scale = 100.0f, const FVector Axis0 = FVector::ForwardVector, const FVector Axis1 = FVector::RightVector);
 
 	/**
-	* Sets the data for this observation. The relative position & rotation can be used to make this observation
-	* relative to the agent's perspective, e.g. by passing the agent's position & forward rotation. Call during
-	* ULearningAgentsInteractor::SetObservations event.
-	* @param AgentId The agent id this data corresponds to.
-	* @param Position The position currently being observed.
-	* @param RelativePosition The vector Position will be offset from.
-	* @param RelativeRotation The frame of reference rotation.
-	*/
+	 * Sets the data for this observation. The relative position & rotation can be used to make this observation
+	 * relative to the agent's perspective, e.g. by passing the agent's position & forward rotation. Call during
+	 * ULearningAgentsInteractor::SetObservations event.
+	 * @param AgentId The agent id this data corresponds to.
+	 * @param Position The position currently being observed.
+	 * @param RelativePosition The vector Position will be offset from.
+	 * @param RelativeRotation The frame of reference rotation.
+	 */
 	UFUNCTION(BlueprintCallable, Category = "LearningAgents", meta = (AgentId = "-1"))
 	void SetPlanarPositionObservation(const int32 AgentId, const FVector Position, const FVector RelativePosition = FVector::ZeroVector, const FRotator RelativeRotation = FRotator::ZeroRotator);
 
@@ -891,28 +898,28 @@ class LEARNINGAGENTS_API UPlanarPositionArrayObservation : public ULearningAgent
 public:
 
 	/**
-	* Adds a new planar position array observation to the given agent interactor. The axis parameters define the plane.
-	* Call during ULearningAgentsInteractor::SetupObservations event.
-	* @param InInteractor The agent interactor to add this observation to.
-	* @param Name The name of this new observation. Used for debugging.
-	* @param PositionNum The number of positions in the array.
-	* @param Scale Used to normalize the data for the observation.
-	* @param Axis0 The forward axis of the plane.
-	* @param Axis1 The right axis of the plane.
-	* @return The newly created observation.
-	*/
-	UFUNCTION(BlueprintCallable, Category = "LearningAgents")
+	 * Adds a new planar position array observation to the given agent interactor. The axis parameters define the plane.
+	 * Call during ULearningAgentsInteractor::SetupObservations event.
+	 * @param InInteractor The agent interactor to add this observation to.
+	 * @param Name The name of this new observation. Used for debugging.
+	 * @param PositionNum The number of positions in the array.
+	 * @param Scale Used to normalize the data for the observation.
+	 * @param Axis0 The forward axis of the plane.
+	 * @param Axis1 The right axis of the plane.
+	 * @return The newly created observation.
+	 */
+	UFUNCTION(BlueprintCallable, Category = "LearningAgents", meta = (DefaultToSelf = "InInteractor"))
 	static UPlanarPositionArrayObservation* AddPlanarPositionArrayObservation(ULearningAgentsInteractor* InInteractor, const FName Name = NAME_None, const int32 PositionNum = 1, const float Scale = 100.0f, const FVector Axis0 = FVector::ForwardVector, const FVector Axis1 = FVector::RightVector);
 
 	/**
-	* Sets the data for this observation. The relative position & rotation can be used to make this observation
-	* relative to the agent's perspective, e.g. by passing the agent's position & forward rotation. Call during
-	* ULearningAgentsInteractor::SetObservations event.
-	* @param AgentId The agent id this data corresponds to.
-	* @param Positions The positions currently being observed.
-	* @param RelativePosition The vector Positions will be offset from.
-	* @param RelativeRotation The frame of reference rotation.
-	*/
+	 * Sets the data for this observation. The relative position & rotation can be used to make this observation
+	 * relative to the agent's perspective, e.g. by passing the agent's position & forward rotation. Call during
+	 * ULearningAgentsInteractor::SetObservations event.
+	 * @param AgentId The agent id this data corresponds to.
+	 * @param Positions The positions currently being observed.
+	 * @param RelativePosition The vector Positions will be offset from.
+	 * @param RelativeRotation The frame of reference rotation.
+	 */
 	UFUNCTION(BlueprintCallable, Category = "LearningAgents", meta = (AgentId = "-1"))
 	void SetPlanarPositionArrayObservation(const int32 AgentId, const TArray<FVector>& Positions, const FVector RelativePosition = FVector::ZeroVector, const FRotator RelativeRotation = FRotator::ZeroRotator);
 
@@ -935,22 +942,22 @@ class LEARNINGAGENTS_API UVelocityObservation : public ULearningAgentsObservatio
 public:
 
 	/**
-	* Adds a new velocity observation to the given agent interactor. Call during ULearningAgentsInteractor::SetupObservations event.
-	* @param InInteractor The agent interactor to add this observation to.
-	* @param Name The name of this new observation. Used for debugging.
-	* @param Scale Used to normalize the data for the observation.
-	* @return The newly created observation.
-	*/
-	UFUNCTION(BlueprintCallable, Category = "LearningAgents")
+	 * Adds a new velocity observation to the given agent interactor. Call during ULearningAgentsInteractor::SetupObservations event.
+	 * @param InInteractor The agent interactor to add this observation to.
+	 * @param Name The name of this new observation. Used for debugging.
+	 * @param Scale Used to normalize the data for the observation.
+	 * @return The newly created observation.
+	 */
+	UFUNCTION(BlueprintCallable, Category = "LearningAgents", meta = (DefaultToSelf = "InInteractor"))
 	static UVelocityObservation* AddVelocityObservation(ULearningAgentsInteractor* InInteractor, const FName Name = NAME_None, const float Scale = 200.0f);
 
 	/**
-	* Sets the data for this observation. The relative rotation can be used to make this observation relative to the
-	* agent's perspective, e.g. by passing the agent's forward rotation. Call during ULearningAgentsInteractor::SetObservations event.
-	* @param AgentId The agent id this data corresponds to.
-	* @param Velocity The velocity currently being observed.
-	* @param RelativeRotation The frame of reference rotation.
-	*/
+	 * Sets the data for this observation. The relative rotation can be used to make this observation relative to the
+	 * agent's perspective, e.g. by passing the agent's forward rotation. Call during ULearningAgentsInteractor::SetObservations event.
+	 * @param AgentId The agent id this data corresponds to.
+	 * @param Velocity The velocity currently being observed.
+	 * @param RelativeRotation The frame of reference rotation.
+	 */
 	UFUNCTION(BlueprintCallable, Category = "LearningAgents", meta = (AgentId = "-1"))
 	void SetVelocityObservation(const int32 AgentId, const FVector Velocity, const FRotator RelativeRotation = FRotator::ZeroRotator);
 
@@ -971,23 +978,23 @@ class LEARNINGAGENTS_API UVelocityArrayObservation : public ULearningAgentsObser
 public:
 
 	/**
-	* Adds a new velocity observation to the given agent interactor. Call during ULearningAgentsInteractor::SetupObservations event.
-	* @param InInteractor The agent interactor to add this observation to.
-	* @param Name The name of this new observation. Used for debugging.
-	* @param VelocityNum The number of velocities in the array.
-	* @param Scale Used to normalize the data for the observation.
-	* @return The newly created observation.
-	*/
-	UFUNCTION(BlueprintCallable, Category = "LearningAgents")
+	 * Adds a new velocity observation to the given agent interactor. Call during ULearningAgentsInteractor::SetupObservations event.
+	 * @param InInteractor The agent interactor to add this observation to.
+	 * @param Name The name of this new observation. Used for debugging.
+	 * @param VelocityNum The number of velocities in the array.
+	 * @param Scale Used to normalize the data for the observation.
+	 * @return The newly created observation.
+	 */
+	UFUNCTION(BlueprintCallable, Category = "LearningAgents", meta = (DefaultToSelf = "InInteractor"))
 	static UVelocityArrayObservation* AddVelocityArrayObservation(ULearningAgentsInteractor* InInteractor, const FName Name = NAME_None, const int32 VelocityNum = 1, const float Scale = 200.0f);
 
 	/**
-	* Sets the data for this observation. The relative rotation can be used to make this observation relative to the
-	* agent's perspective, e.g. by passing the agent's forward rotation. Call during ULearningAgentsInteractor::SetObservations event.
-	* @param AgentId The agent id this data corresponds to.
-	* @param Velocities The velocities currently being observed.
-	* @param RelativeRotation The frame of reference rotation.
-	*/
+	 * Sets the data for this observation. The relative rotation can be used to make this observation relative to the
+	 * agent's perspective, e.g. by passing the agent's forward rotation. Call during ULearningAgentsInteractor::SetObservations event.
+	 * @param AgentId The agent id this data corresponds to.
+	 * @param Velocities The velocities currently being observed.
+	 * @param RelativeRotation The frame of reference rotation.
+	 */
 	UFUNCTION(BlueprintCallable, Category = "LearningAgents", meta = (AgentId = "-1"))
 	void SetVelocityArrayObservation(const int32 AgentId, const TArray<FVector>& Velocities, const FRotator RelativeRotation = FRotator::ZeroRotator);
 
@@ -1008,30 +1015,30 @@ class LEARNINGAGENTS_API UScalarVelocityObservation : public ULearningAgentsObse
 public:
 
 	/**
-	* Adds a new scalar velocity observation to the given agent interactor. 
-	* Call during ULearningAgentsInteractor::SetupObservations event.
-	* @param InInteractor The agent interactor to add this observation to.
-	* @param Name The name of this new observation. Used for debugging.
-	* @param Scale Used to normalize the data for the observation.
-	* @return The newly created observation.
-	*/
-	UFUNCTION(BlueprintCallable, Category = "LearningAgents")
+	 * Adds a new scalar velocity observation to the given agent interactor. 
+	 * Call during ULearningAgentsInteractor::SetupObservations event.
+	 * @param InInteractor The agent interactor to add this observation to.
+	 * @param Name The name of this new observation. Used for debugging.
+	 * @param Scale Used to normalize the data for the observation.
+	 * @return The newly created observation.
+	 */
+	UFUNCTION(BlueprintCallable, Category = "LearningAgents", meta = (DefaultToSelf = "InInteractor"))
 	static UScalarVelocityObservation* AddScalarVelocityObservation(ULearningAgentsInteractor* InInteractor, const FName Name = NAME_None, const float Scale = 200.0f);
 
 	/**
-	* Sets the data for this observation. Call during ULearningAgentsInteractor::SetObservations event.
-	* @param AgentId The agent id this data corresponds to.
-	* @param Velocity The velocity currently being observed.
-	*/
+	 * Sets the data for this observation. Call during ULearningAgentsInteractor::SetObservations event.
+	 * @param AgentId The agent id this data corresponds to.
+	 * @param Velocity The velocity currently being observed.
+	 */
 	UFUNCTION(BlueprintCallable, Category = "LearningAgents", meta = (AgentId = "-1"))
 	void SetScalarVelocityObservation(const int32 AgentId, const float Velocity);
 
 	/**
-	* Sets the data for this observation. Call during ULearningAgentsInteractor::SetObservations event.
-	* @param AgentId The agent id this data corresponds to.
-	* @param Velocity The velocity currently being observed.
-	* @param Axis The axis to encode the velocity along
-	*/
+	 * Sets the data for this observation. Call during ULearningAgentsInteractor::SetObservations event.
+	 * @param AgentId The agent id this data corresponds to.
+	 * @param Velocity The velocity currently being observed.
+	 * @param Axis The axis to encode the velocity along
+	 */
 	UFUNCTION(BlueprintCallable, Category = "LearningAgents", meta = (AgentId = "-1"))
 	void SetScalarVelocityObservationWithAxis(const int32 AgentId, const FVector Velocity, const FVector Axis = FVector::UpVector);
 
@@ -1052,31 +1059,31 @@ class LEARNINGAGENTS_API UScalarVelocityArrayObservation : public ULearningAgent
 public:
 
 	/**
-	* Adds a new scalar velocity observation to the given agent interactor.
-	* Call during ULearningAgentsInteractor::SetupObservations event.
-	* @param InInteractor The agent interactor to add this observation to.
-	* @param Name The name of this new observation. Used for debugging.
-	* @param VelocityNum The number of velocities in the array.
-	* @param Scale Used to normalize the data for the observation.
-	* @return The newly created observation.
-	*/
-	UFUNCTION(BlueprintCallable, Category = "LearningAgents")
+	 * Adds a new scalar velocity observation to the given agent interactor.
+	 * Call during ULearningAgentsInteractor::SetupObservations event.
+	 * @param InInteractor The agent interactor to add this observation to.
+	 * @param Name The name of this new observation. Used for debugging.
+	 * @param VelocityNum The number of velocities in the array.
+	 * @param Scale Used to normalize the data for the observation.
+	 * @return The newly created observation.
+	 */
+	UFUNCTION(BlueprintCallable, Category = "LearningAgents", meta = (DefaultToSelf = "InInteractor"))
 	static UScalarVelocityArrayObservation* AddScalarVelocityArrayObservation(ULearningAgentsInteractor* InInteractor, const FName Name = NAME_None, const int32 VelocityNum = 1, const float Scale = 200.0f);
 
 	/**
-	* Sets the data for this observation. Call during ULearningAgentsInteractor::SetObservations event.
-	* @param AgentId The agent id this data corresponds to.
-	* @param Velocities The velocities currently being observed.
-	*/
+	 * Sets the data for this observation. Call during ULearningAgentsInteractor::SetObservations event.
+	 * @param AgentId The agent id this data corresponds to.
+	 * @param Velocities The velocities currently being observed.
+	 */
 	UFUNCTION(BlueprintCallable, Category = "LearningAgents", meta = (AgentId = "-1"))
 	void SetScalarVelocityArrayObservation(const int32 AgentId, const TArray<float>& Velocities);
 
 	/**
-	* Sets the data for this observation. Call during ULearningAgentsInteractor::SetObservations event.
-	* @param AgentId The agent id this data corresponds to.
-	* @param Velocities The velocities currently being observed.
-	* @param Axis The axis to encode the velocity along
-	*/
+	 * Sets the data for this observation. Call during ULearningAgentsInteractor::SetObservations event.
+	 * @param AgentId The agent id this data corresponds to.
+	 * @param Velocities The velocities currently being observed.
+	 * @param Axis The axis to encode the velocity along
+	 */
 	UFUNCTION(BlueprintCallable, Category = "LearningAgents", meta = (AgentId = "-1"))
 	void SetScalarVelocityArrayObservationWithAxis(const int32 AgentId, const TArray<FVector>& Velocities, const FVector Axis = FVector::UpVector);
 
@@ -1097,25 +1104,25 @@ class LEARNINGAGENTS_API UPlanarVelocityObservation : public ULearningAgentsObse
 public:
 
 	/**
-	* Adds a new planar velocity observation to the given agent interactor. The axis parameters define the plane.
-	* Call during ULearningAgentsInteractor::SetupObservations event.
-	* @param InInteractor The agent interactor to add this observation to.
-	* @param Name The name of this new observation. Used for debugging.
-	* @param Scale Used to normalize the data for the observation.
-	* @param Axis0 The forward axis of the plane.
-	* @param Axis1 The right axis of the plane.
-	* @return The newly created observation.
-	*/
-	UFUNCTION(BlueprintCallable, Category = "LearningAgents")
+	 * Adds a new planar velocity observation to the given agent interactor. The axis parameters define the plane.
+	 * Call during ULearningAgentsInteractor::SetupObservations event.
+	 * @param InInteractor The agent interactor to add this observation to.
+	 * @param Name The name of this new observation. Used for debugging.
+	 * @param Scale Used to normalize the data for the observation.
+	 * @param Axis0 The forward axis of the plane.
+	 * @param Axis1 The right axis of the plane.
+	 * @return The newly created observation.
+	 */
+	UFUNCTION(BlueprintCallable, Category = "LearningAgents", meta = (DefaultToSelf = "InInteractor"))
 	static UPlanarVelocityObservation* AddPlanarVelocityObservation(ULearningAgentsInteractor* InInteractor, const FName Name = NAME_None, const float Scale = 200.0f, const FVector Axis0 = FVector::ForwardVector, const FVector Axis1 = FVector::RightVector);
 
 	/**
-	* Sets the data for this observation. The relative rotation can be used to make this observation relative to the
-	* agent's perspective, e.g. by passing the agent's forward rotation. Call during ULearningAgentsInteractor::SetObservations event.
-	* @param AgentId The agent id this data corresponds to.
-	* @param Velocity The velocity currently being observed.
-	* @param RelativeRotation The frame of reference rotation.
-	*/
+	 * Sets the data for this observation. The relative rotation can be used to make this observation relative to the
+	 * agent's perspective, e.g. by passing the agent's forward rotation. Call during ULearningAgentsInteractor::SetObservations event.
+	 * @param AgentId The agent id this data corresponds to.
+	 * @param Velocity The velocity currently being observed.
+	 * @param RelativeRotation The frame of reference rotation.
+	 */
 	UFUNCTION(BlueprintCallable, Category = "LearningAgents", meta = (AgentId = "-1"))
 	void SetPlanarVelocityObservation(const int32 AgentId, const FVector Velocity, const FRotator RelativeRotation = FRotator::ZeroRotator);
 
@@ -1136,26 +1143,26 @@ class LEARNINGAGENTS_API UPlanarVelocityArrayObservation : public ULearningAgent
 public:
 
 	/**
-	* Adds a new planar velocity observation to the given agent interactor. The axis parameters define the plane.
-	* Call during ULearningAgentsInteractor::SetupObservations event.
-	* @param InInteractor The agent interactor to add this observation to.
-	* @param Name The name of this new observation. Used for debugging.
-	* @param VelocityNum The number of velocities in the array.
-	* @param Scale Used to normalize the data for the observation.
-	* @param Axis0 The forward axis of the plane.
-	* @param Axis1 The right axis of the plane.
-	* @return The newly created observation.
-	*/
-	UFUNCTION(BlueprintCallable, Category = "LearningAgents")
+	 * Adds a new planar velocity observation to the given agent interactor. The axis parameters define the plane.
+	 * Call during ULearningAgentsInteractor::SetupObservations event.
+	 * @param InInteractor The agent interactor to add this observation to.
+	 * @param Name The name of this new observation. Used for debugging.
+	 * @param VelocityNum The number of velocities in the array.
+	 * @param Scale Used to normalize the data for the observation.
+	 * @param Axis0 The forward axis of the plane.
+	 * @param Axis1 The right axis of the plane.
+	 * @return The newly created observation.
+	 */
+	UFUNCTION(BlueprintCallable, Category = "LearningAgents", meta = (DefaultToSelf = "InInteractor"))
 	static UPlanarVelocityArrayObservation* AddPlanarVelocityArrayObservation(ULearningAgentsInteractor* InInteractor, const FName Name = NAME_None, const int32 VelocityNum = 1, const float Scale = 200.0f, const FVector Axis0 = FVector::ForwardVector, const FVector Axis1 = FVector::RightVector);
 
 	/**
-	* Sets the data for this observation. The relative rotation can be used to make this observation relative to the
-	* agent's perspective, e.g. by passing the agent's forward rotation. Call during ULearningAgentsInteractor::SetObservations event.
-	* @param AgentId The agent id this data corresponds to.
-	* @param Velocities The velocities currently being observed.
-	* @param RelativeRotation The frame of reference rotation.
-	*/
+	 * Sets the data for this observation. The relative rotation can be used to make this observation relative to the
+	 * agent's perspective, e.g. by passing the agent's forward rotation. Call during ULearningAgentsInteractor::SetObservations event.
+	 * @param AgentId The agent id this data corresponds to.
+	 * @param Velocities The velocities currently being observed.
+	 * @param RelativeRotation The frame of reference rotation.
+	 */
 	UFUNCTION(BlueprintCallable, Category = "LearningAgents", meta = (AgentId = "-1"))
 	void SetPlanarVelocityArrayObservation(const int32 AgentId, const TArray<FVector>& Velocities, const FRotator RelativeRotation = FRotator::ZeroRotator);
 
@@ -1178,22 +1185,22 @@ class LEARNINGAGENTS_API UAngularVelocityObservation : public ULearningAgentsObs
 public:
 
 	/**
-	* Adds a new angular velocity observation to the given agent interactor. Call during ULearningAgentsInteractor::SetupObservations event.
-	* @param InInteractor The agent interactor to add this observation to.
-	* @param Name The name of this new observation. Used for debugging.
-	* @param Scale Used to normalize the data for the observation.
-	* @return The newly created observation.
-	*/
-	UFUNCTION(BlueprintCallable, Category = "LearningAgents")
+	 * Adds a new angular velocity observation to the given agent interactor. Call during ULearningAgentsInteractor::SetupObservations event.
+	 * @param InInteractor The agent interactor to add this observation to.
+	 * @param Name The name of this new observation. Used for debugging.
+	 * @param Scale Used to normalize the data for the observation.
+	 * @return The newly created observation.
+	 */
+	UFUNCTION(BlueprintCallable, Category = "LearningAgents", meta = (DefaultToSelf = "InInteractor"))
 	static UAngularVelocityObservation* AddAngularVelocityObservation(ULearningAgentsInteractor* InInteractor, const FName Name = NAME_None, const float Scale = 180.0f);
 
 	/**
-	* Sets the data for this observation. The relative rotation can be used to make this observation relative to the
-	* agent's perspective, e.g. by passing the agent's forward rotation. Call during ULearningAgentsInteractor::SetObservations event.
-	* @param AgentId The agent id this data corresponds to.
-	* @param AngularVelocity The angular velocity currently being observed.
-	* @param RelativeRotation The frame of reference rotation.
-	*/
+	 * Sets the data for this observation. The relative rotation can be used to make this observation relative to the
+	 * agent's perspective, e.g. by passing the agent's forward rotation. Call during ULearningAgentsInteractor::SetObservations event.
+	 * @param AgentId The agent id this data corresponds to.
+	 * @param AngularVelocity The angular velocity currently being observed.
+	 * @param RelativeRotation The frame of reference rotation.
+	 */
 	UFUNCTION(BlueprintCallable, Category = "LearningAgents", meta = (AgentId = "-1"))
 	void SetAngularVelocityObservation(const int32 AgentId, const FVector AngularVelocity, const FRotator RelativeRotation = FRotator::ZeroRotator);
 
@@ -1214,23 +1221,23 @@ class LEARNINGAGENTS_API UAngularVelocityArrayObservation : public ULearningAgen
 public:
 
 	/**
-	* Adds a new angular velocity array observation to the given agent interactor. Call during ULearningAgentsInteractor::SetupObservations event.
-	* @param InInteractor The agent interactor to add this observation to.
-	* @param Name The name of this new observation. Used for debugging.
-	* @param AngularVelocityNum The number of angular velocities in the array.
-	* @param Scale Used to normalize the data for the observation.
-	* @return The newly created observation.
-	*/
-	UFUNCTION(BlueprintCallable, Category = "LearningAgents")
+	 * Adds a new angular velocity array observation to the given agent interactor. Call during ULearningAgentsInteractor::SetupObservations event.
+	 * @param InInteractor The agent interactor to add this observation to.
+	 * @param Name The name of this new observation. Used for debugging.
+	 * @param AngularVelocityNum The number of angular velocities in the array.
+	 * @param Scale Used to normalize the data for the observation.
+	 * @return The newly created observation.
+	 */
+	UFUNCTION(BlueprintCallable, Category = "LearningAgents", meta = (DefaultToSelf = "InInteractor"))
 	static UAngularVelocityArrayObservation* AddAngularVelocityArrayObservation(ULearningAgentsInteractor* InInteractor, const FName Name = NAME_None, const int32 AngularVelocityNum = 1, const float Scale = 180.0f);
 
 	/**
-	* Sets the data for this observation. The relative rotation can be used to make this observation relative to the
-	* agent's perspective, e.g. by passing the agent's forward rotation. Call during ULearningAgentsInteractor::SetObservations event.
-	* @param AgentId The agent id this data corresponds to.
-	* @param Velocities The angular velocities currently being observed.
-	* @param RelativeRotation The frame of reference rotation.
-	*/
+	 * Sets the data for this observation. The relative rotation can be used to make this observation relative to the
+	 * agent's perspective, e.g. by passing the agent's forward rotation. Call during ULearningAgentsInteractor::SetObservations event.
+	 * @param AgentId The agent id this data corresponds to.
+	 * @param Velocities The angular velocities currently being observed.
+	 * @param RelativeRotation The frame of reference rotation.
+	 */
 	UFUNCTION(BlueprintCallable, Category = "LearningAgents", meta = (AgentId = "-1"))
 	void SetAngularVelocityArrayObservation(const int32 AgentId, const TArray<FVector>& AngularVelocities, const FRotator RelativeRotation = FRotator::ZeroRotator);
 
@@ -1251,30 +1258,30 @@ class LEARNINGAGENTS_API UScalarAngularVelocityObservation : public ULearningAge
 public:
 
 	/**
-	* Adds a new scalar angular velocity observation to the given agent interactor.
-	* Call during ULearningAgentsInteractor::SetupObservations event.
-	* @param InInteractor The agent interactor to add this observation to.
-	* @param Name The name of this new observation. Used for debugging.
-	* @param Scale Used to normalize the data for the observation.
-	* @return The newly created observation.
-	*/
-	UFUNCTION(BlueprintCallable, Category = "LearningAgents")
+	 * Adds a new scalar angular velocity observation to the given agent interactor.
+	 * Call during ULearningAgentsInteractor::SetupObservations event.
+	 * @param InInteractor The agent interactor to add this observation to.
+	 * @param Name The name of this new observation. Used for debugging.
+	 * @param Scale Used to normalize the data for the observation.
+	 * @return The newly created observation.
+	 */
+	UFUNCTION(BlueprintCallable, Category = "LearningAgents", meta = (DefaultToSelf = "InInteractor"))
 	static UScalarAngularVelocityObservation* AddScalarAngularVelocityObservation(ULearningAgentsInteractor* InInteractor, const FName Name = NAME_None, const float Scale = 180.0f);
 
 	/**
-	* Sets the data for this observation. Call during ULearningAgentsInteractor::SetObservations event.
-	* @param AgentId The agent id this data corresponds to.
-	* @param AngularVelocity The angular velocity currently being observed.
-	*/
+	 * Sets the data for this observation. Call during ULearningAgentsInteractor::SetObservations event.
+	 * @param AgentId The agent id this data corresponds to.
+	 * @param AngularVelocity The angular velocity currently being observed.
+	 */
 	UFUNCTION(BlueprintCallable, Category = "LearningAgents", meta = (AgentId = "-1"))
 	void SetScalarAngularVelocityObservation(const int32 AgentId, const float AngularVelocity);
 
 	/**
-	* Sets the data for this observation. Call during ULearningAgentsInteractor::SetObservations event.
-	* @param AgentId The agent id this data corresponds to.
-	* @param AngularVelocity The angular velocity currently being observed.
-	* @param Axis The axis to encode the angular velocity around
-	*/
+	 * Sets the data for this observation. Call during ULearningAgentsInteractor::SetObservations event.
+	 * @param AgentId The agent id this data corresponds to.
+	 * @param AngularVelocity The angular velocity currently being observed.
+	 * @param Axis The axis to encode the angular velocity around
+	 */
 	UFUNCTION(BlueprintCallable, Category = "LearningAgents", meta = (AgentId = "-1"))
 	void SetScalarAngularVelocityObservationWithAxis(const int32 AgentId, const FVector AngularVelocity, const FVector Axis = FVector::UpVector);
 
@@ -1295,31 +1302,31 @@ class LEARNINGAGENTS_API UScalarAngularVelocityArrayObservation : public ULearni
 public:
 
 	/**
-	* Adds a new scalar angular velocity array observation to the given agent interactor.
-	* Call during ULearningAgentsInteractor::SetupObservations event.
-	* @param InInteractor The agent interactor to add this observation to.
-	* @param Name The name of this new observation. Used for debugging.
-	* @param AngularVelocityNum The number of angular velocities in the array.
-	* @param Scale Used to normalize the data for the observation.
-	* @return The newly created observation.
-	*/
-	UFUNCTION(BlueprintCallable, Category = "LearningAgents")
+	 * Adds a new scalar angular velocity array observation to the given agent interactor.
+	 * Call during ULearningAgentsInteractor::SetupObservations event.
+	 * @param InInteractor The agent interactor to add this observation to.
+	 * @param Name The name of this new observation. Used for debugging.
+	 * @param AngularVelocityNum The number of angular velocities in the array.
+	 * @param Scale Used to normalize the data for the observation.
+	 * @return The newly created observation.
+	 */
+	UFUNCTION(BlueprintCallable, Category = "LearningAgents", meta = (DefaultToSelf = "InInteractor"))
 	static UScalarAngularVelocityArrayObservation* AddScalarAngularVelocityArrayObservation(ULearningAgentsInteractor* InInteractor, const FName Name = NAME_None, const int32 AngularVelocityNum = 1, const float Scale = 180.0f);
 
 	/**
-	* Sets the data for this observation. Call during ULearningAgentsInteractor::SetObservations event.
-	* @param AgentId The agent id this data corresponds to.
-	* @param AngularVelocities The angular velocities currently being observed.
-	*/
+	 * Sets the data for this observation. Call during ULearningAgentsInteractor::SetObservations event.
+	 * @param AgentId The agent id this data corresponds to.
+	 * @param AngularVelocities The angular velocities currently being observed.
+	 */
 	UFUNCTION(BlueprintCallable, Category = "LearningAgents", meta = (AgentId = "-1"))
 	void SetScalarAngularVelocityArrayObservation(const int32 AgentId, const TArray<float>& AngularVelocities);
 
 	/**
-	* Sets the data for this observation. Call during ULearningAgentsInteractor::SetObservations event.
-	* @param AgentId The agent id this data corresponds to.
-	* @param AngularVelocities The angular velocities currently being observed.
-	* @param Axis The axis to encode the angular velocity around
-	*/
+	 * Sets the data for this observation. Call during ULearningAgentsInteractor::SetObservations event.
+	 * @param AgentId The agent id this data corresponds to.
+	 * @param AngularVelocities The angular velocities currently being observed.
+	 * @param Axis The axis to encode the angular velocity around
+	 */
 	UFUNCTION(BlueprintCallable, Category = "LearningAgents", meta = (AgentId = "-1"))
 	void SetScalarAngularVelocityArrayObservationWithAxis(const int32 AgentId, const TArray<FVector>& AngularVelocities, const FVector Axis = FVector::UpVector);
 
diff --git a/Engine/Plugins/Experimental/LearningAgents/Source/LearningAgents/Public/LearningAgentsPolicy.h b/Engine/Plugins/Experimental/LearningAgents/Source/LearningAgents/Public/LearningAgentsPolicy.h
--- a/Engine/Plugins/Experimental/LearningAgents/Source/LearningAgents/Public/LearningAgentsPolicy.h
+++ b/Engine/Plugins/Experimental/LearningAgents/Source/LearningAgents/Public/LearningAgentsPolicy.h
@@ -40,9 +40,9 @@ public:
 	UPROPERTY(EditAnywhere, Category = "LearningAgents", meta = (ClampMin = "0.0", UIMin = "0.0"))
 	float ActionNoiseMax = 0.25f;
 
-	/** Initial scale of the action noise used by the policy. Should be 1.0 for agents participating in training. */
+	/** Scale of the action noise used by the policy should be 1.0 during training. */
 	UPROPERTY(EditAnywhere, Category = "LearningAgents", meta = (ClampMin = "0.0", UIMin = "0.0"))
-	float InitialActionNoiseScale = 1.0f;
+	float ActionNoiseScale = 1.0f;
 
 	/** Total layers for policy network including input, hidden, and output layers */
 	UPROPERTY(EditAnywhere, Category = "LearningAgents", meta = (ClampMin = "2", UIMin = "2"))
@@ -66,100 +66,93 @@ class LEARNINGAGENTS_API ULearningAgentsPolicy : public ULearningAgentsManagerCo
 public:
 
 	// These constructors/destructors are needed to make forward declarations happy
-	ULearningAgentsPolicy();
+	ULearningAgentsPolicy(const FObjectInitializer& ObjectInitializer = FObjectInitializer::Get());
 	ULearningAgentsPolicy(FVTableHelper& Helper);
 	virtual ~ULearningAgentsPolicy();
 
-	/** Initializes this object to be used with the given agent interactor and policy settings. */
+	/**
+	 * Initializes this object to be used with the given agent interactor and policy settings.
+	 * @param InInteractor The input Interactor component
+	 * @param PolicySettings The policy settings to use
+	 * @param NeuralNetworkAsset Optional Network Asset to use. If provided must match the given PolicySettings. If not
+	 * provided or asset is empty then a new neural network object will be created according to the given 
+	 * PolicySettings and used.
+	 */
 	UFUNCTION(BlueprintCallable, Category = "LearningAgents")
-	void SetupPolicy(ALearningAgentsManager* InAgentManager, ULearningAgentsInteractor* InInteractor, const FLearningAgentsPolicySettings& PolicySettings = FLearningAgentsPolicySettings());
-
-public:
-
-	//~ Begin ULearningAgentsManagerComponent Interface
-	virtual bool AddAgent(const int32 AgentId) override;
-	//~ End ULearningAgentsManagerComponent Interface
+	void SetupPolicy(ULearningAgentsInteractor* InInteractor, 
+		const FLearningAgentsPolicySettings& PolicySettings = FLearningAgentsPolicySettings(),
+		ULearningAgentsNeuralNetwork* NeuralNetworkAsset = nullptr);
 
 // ----- Load / Save -----
 public:
 
 	/**
-	* Load a snapshot's weights into this policy.
-	* @param File The snapshot file.
-	*/
+	 * Load a snapshot's weights into this policy.
+	 * @param File The snapshot file.
+	 */
 	UFUNCTION(BlueprintCallable, Category = "LearningAgents", meta = (RelativePath))
 	void LoadPolicyFromSnapshot(const FFilePath& File);
 
 	/**
-	* Save this policy's weights into a snapshot.
-	* @param File The snapshot file.
-	*/
+	 * Save this policy's weights into a snapshot.
+	 * @param File The snapshot file.
+	 */
 	UFUNCTION(BlueprintCallable, Category = "LearningAgents", meta = (RelativePath))
 	void SavePolicyToSnapshot(const FFilePath& File) const;
 
 	/**
-	* Load a ULearningAgentsNeuralNetwork asset's weights into this policy.
-	* @param NeuralNetworkAsset The asset to load from.
-	*/
+	 * Use a ULearningAgentsNeuralNetwork asset directly for this critic rather than making a copy. If used
+	 * during training then this asset's weights will be updated as training progresses.
+	 * @param NeuralNetworkAsset The asset to use.
+	 */
 	UFUNCTION(BlueprintCallable, Category = "LearningAgents")
-	void LoadPolicyFromAsset(const ULearningAgentsNeuralNetwork* NeuralNetworkAsset);
+	void UsePolicyFromAsset(ULearningAgentsNeuralNetwork* NeuralNetworkAsset);
 
 	/**
-	* Save this policy's weights to a ULearningAgentsNeuralNetwork asset.
-	* @param NeuralNetworkAsset The asset to save to.
-	*/
+	 * Load a ULearningAgentsNeuralNetwork asset's weights into this policy.
+	 * @param NeuralNetworkAsset The asset to load from.
+	 */
+	UFUNCTION(BlueprintCallable, Category = "LearningAgents")
+	void LoadPolicyFromAsset(ULearningAgentsNeuralNetwork* NeuralNetworkAsset);
+
+	/**
+	 * Save this policy's weights to a ULearningAgentsNeuralNetwork asset.
+	 * @param NeuralNetworkAsset The asset to save to.
+	 */
 	UFUNCTION(BlueprintCallable, Category = "LearningAgents", Meta = (DevelopmentOnly))
-	void SavePolicyToAsset(ULearningAgentsNeuralNetwork* NeuralNetworkAsset) const;
+	void SavePolicyToAsset(ULearningAgentsNeuralNetwork* NeuralNetworkAsset);
 
 // ----- Evaluation -----
 public:
 
 	/**
-	* Calling this function will run the underlying neural network on the previously buffered observations to populate
-	* the output action buffer. This should be called after the associated agent interactor's EncodeObservations and
-	* before its DecodeActions.
-	*/
+	 * Calling this function will run the underlying neural network on the previously buffered observations to populate
+	 * the output action buffer. This should be called after the associated agent interactor's EncodeObservations and
+	 * before its DecodeActions.
+	 */
 	UFUNCTION(BlueprintCallable, Category = "LearningAgents")
 	void EvaluatePolicy();
 
 	/**
-	* Calls EncodeObservations, followed by EvaluatePolicy, followed by DecodeActions
-	*/
+	 * Calls EncodeObservations, followed by EvaluatePolicy, followed by DecodeActions
+	 */
 	UFUNCTION(BlueprintCallable, Category = "LearningAgents")
 	void RunInference();
 
-	/**
-	* Get the action noise scale used by an agent.
-	*
-	* @param AgentId			The AgentId to get the action noise scale for.
-	* @return					Action noise scale for that agent.
-	*/
-	UFUNCTION(BlueprintPure, Category = "LearningAgents", meta = (AgentId = "-1"))
-	float GetAgentActionNoiseScale(const int32 AgentId) const;
-
-	/**
-	* Set the action noise scale used by an agent. This can be useful if you have certain
-	* agents that are participating in training (and so should have an action noise scale of 1.0)
-	* and certain agents which you are testing the inference for (and so will want action noise
-	* scale of 0.0)
-	*
-	* @param AgentId			The AgentId to set the action noise scale for.
-	* @param ActionNoiseScale	Action noise scale for that agent.
-	*/
-	UFUNCTION(BlueprintCallable, Category = "LearningAgents", meta = (AgentId = "-1"))
-	void SetAgentActionNoiseScale(const int32 AgentId, const float ActionNoiseScale);
+	/** Get the action noise scale used. */
+	UFUNCTION(BlueprintPure, Category = "LearningAgents")
+	float GetActionNoiseScale() const;
 
-	/**
-	* Set the action noise scale used by all agents.
-	*
-	* @param ActionNoiseScale	Action noise scale to use for all agents
-	*/
+	/** Set the action noise scale used. */
 	UFUNCTION(BlueprintCallable, Category = "LearningAgents")
-	void SetAllAgentsActionNoiseScale(const float ActionNoiseScale);
+	void SetActionNoiseScale(const float ActionNoiseScale);
 
 // ----- Non-blueprint public interface -----
 public:
 
+	/** Gets the current Network Asset being used */
+	ULearningAgentsNeuralNetwork* GetNetworkAsset();
+
 	/** Get a reference to this policy's neural network. */
 	UE::Learning::FNeuralNetwork& GetPolicyNetwork();
 
@@ -177,17 +170,22 @@ private:
 	UPROPERTY(VisibleAnywhere, Transient, Category = "LearningAgents")
 	TObjectPtr<ULearningAgentsNeuralNetwork> Network;
 
+	/** Internal Policy Function Object */
+	TSharedPtr<UE::Learning::FNeuralNetworkPolicyFunction> PolicyObject;
+
+// ----- Private Iteration Checks ----- 
 private:
 
-	TSharedPtr<UE::Learning::FNeuralNetworkPolicyFunction> PolicyObject;
-	
-	float InitialActionNoiseScale = 1.0f;
+	/** Temp buffers used to record the set of agents that are valid for evaluation */
+	TArray<int32> ValidAgentIds;
+	UE::Learning::FIndexSet ValidAgentSet;
 
+private:
 #if UE_LEARNING_AGENTS_ENABLE_VISUAL_LOG
 	/** Color used to draw this action in the visual log */
 	FLinearColor VisualLogColor = FColor::Purple;
 
 	/** Describes this policy to the visual logger for debugging purposes. */
-	void VisualLog(const UE::Learning::FIndexSet Instances) const;
+	void VisualLog(const UE::Learning::FIndexSet AgentSet) const;
 #endif
 };
diff --git a/Engine/Plugins/Experimental/LearningAgents/Source/LearningAgentsTraining/Private/LearningAgentsCompletions.cpp b/Engine/Plugins/Experimental/LearningAgents/Source/LearningAgentsTraining/Private/LearningAgentsCompletions.cpp
--- a/Engine/Plugins/Experimental/LearningAgents/Source/LearningAgentsTraining/Private/LearningAgentsCompletions.cpp
+++ b/Engine/Plugins/Experimental/LearningAgents/Source/LearningAgentsTraining/Private/LearningAgentsCompletions.cpp
@@ -23,15 +23,24 @@ namespace UE::Learning::Agents::Completions::Private
 			return nullptr;
 		}
 
+		if (!InAgentTrainer->HasAgentManager())
+		{
+			UE_LOG(LogLearning, Error, TEXT("%s: Must be attached to a LearningAgentsManager Actor."), *InAgentTrainer->GetName());
+			return nullptr;
+		}
+
 		CompletionUObject* Completion = NewObject<CompletionUObject>(InAgentTrainer, Name);
 
 		Completion->AgentTrainer = InAgentTrainer;
 		Completion->CompletionObject = MakeShared<CompletionFObject>(
 			Completion->GetFName(),
 			InAgentTrainer->GetAgentManager()->GetInstanceData().ToSharedRef(),
-			InAgentTrainer->GetAgentManager()->GetMaxInstanceNum(),
+			InAgentTrainer->GetAgentManager()->GetMaxAgentNum(),
 			Forward<InArgTypes>(Args)...);
 
+		Completion->AgentIteration.SetNumUninitialized({ InAgentTrainer->GetAgentManager()->GetMaxAgentNum() });
+		UE::Learning::Array::Set<1, uint64>(Completion->AgentIteration, INDEX_NONE);
+
 		InAgentTrainer->AddCompletion(Completion, Completion->CompletionObject.ToSharedRef());
 
 		return Completion;
@@ -54,6 +63,7 @@ void UConditionalCompletion::SetConditionalCompletion(const int32 AgentId, const
 	}
 
 	CompletionObject->InstanceData->View(CompletionObject->ConditionHandle)[AgentId] = bIsComplete;
+	AgentIteration[AgentId]++;
 }
 
 #if UE_LEARNING_AGENTS_ENABLE_VISUAL_LOG
@@ -96,6 +106,7 @@ void UTimeElapsedCompletion::SetTimeElapsedCompletion(const int32 AgentId, const
 	}
 
 	CompletionObject->InstanceData->View(CompletionObject->TimeHandle)[AgentId] = Time;
+	AgentIteration[AgentId]++;
 }
 
 #if UE_LEARNING_AGENTS_ENABLE_VISUAL_LOG
@@ -155,6 +166,7 @@ void UPlanarPositionDifferenceCompletion::SetPlanarPositionDifferenceCompletion(
 
 	CompletionObject->InstanceData->View(CompletionObject->Position0Handle)[AgentId][0] = Position0;
 	CompletionObject->InstanceData->View(CompletionObject->Position1Handle)[AgentId][0] = Position1;
+	AgentIteration[AgentId]++;
 }
 
 #if UE_LEARNING_AGENTS_ENABLE_VISUAL_LOG
@@ -254,6 +266,7 @@ void UPlanarPositionSimilarityCompletion::SetPlanarPositionSimilarityCompletion(
 
 	CompletionObject->InstanceData->View(CompletionObject->Position0Handle)[AgentId][0] = Position0;
 	CompletionObject->InstanceData->View(CompletionObject->Position1Handle)[AgentId][0] = Position1;
+	AgentIteration[AgentId]++;
 }
 
 #if UE_LEARNING_AGENTS_ENABLE_VISUAL_LOG
