LearningAgents: Added iteration checks for verifying order of operations. Fixed issue with inference for multiple agents. Fixed live recompilation issues.

diff --git a/Engine/Plugins/Experimental/LearningAgents/Source/LearningAgentsTraining/Private/LearningAgentsImitationTrainer.cpp b/Engine/Plugins/Experimental/LearningAgents/Source/LearningAgentsTraining/Private/LearningAgentsImitationTrainer.cpp
--- a/Engine/Plugins/Experimental/LearningAgents/Source/LearningAgentsTraining/Private/LearningAgentsImitationTrainer.cpp
+++ b/Engine/Plugins/Experimental/LearningAgents/Source/LearningAgentsTraining/Private/LearningAgentsImitationTrainer.cpp
@@ -16,7 +16,7 @@
 #include "LearningNeuralNetworkObject.h"
 #include "Misc/Paths.h"
 
-ULearningAgentsImitationTrainer::ULearningAgentsImitationTrainer() : UActorComponent() {}
+ULearningAgentsImitationTrainer::ULearningAgentsImitationTrainer(const FObjectInitializer& ObjectInitializer) : UActorComponent(ObjectInitializer) {}
 ULearningAgentsImitationTrainer::ULearningAgentsImitationTrainer(FVTableHelper& Helper) : ULearningAgentsImitationTrainer() {}
 ULearningAgentsImitationTrainer::~ULearningAgentsImitationTrainer() {}
 
@@ -33,6 +33,7 @@ void ULearningAgentsImitationTrainer::EndPlay(const EEndPlayReason::Type EndPlay
 void ULearningAgentsImitationTrainer::BeginTraining(
 	ULearningAgentsPolicy* InPolicy, 
 	const ULearningAgentsRecording* Recording,
+	const FLearningAgentsImitationTrainerSettings& ImitationTrainerSettings,
 	const FLearningAgentsImitationTrainerTrainingSettings& ImitationTrainerTrainingSettings,
 	const FLearningAgentsTrainerPathSettings& ImitationTrainerPathSettings,
 	const bool bReinitializePolicyNetwork)
@@ -69,6 +70,10 @@ void ULearningAgentsImitationTrainer::BeginTraining(
 
 	Policy = InPolicy;
 
+	// Record Timeout Setting
+
+	TrainerTimeout = ImitationTrainerSettings.TrainerCommunicationTimeout;
+
 	// Check Paths
 
 	const FString PythonExecutablePath = UE::Learning::Trainer::GetPythonExecutablePath(ImitationTrainerPathSettings.GetEditorEnginePath());
@@ -149,9 +154,6 @@ void ULearningAgentsImitationTrainer::BeginTraining(
 
 	UE_LOG(LogLearning, Display, TEXT("%s: Imitation Training Started"), *GetName());
 
-	bIsTraining = true;
-	bIsTrainingComplete = false;
-	bRequestImitationTrainingStop = false;
 
 	UE::Learning::FImitationTrainerTrainingSettings ImitationTrainingSettings;
 	ImitationTrainingSettings.IterationNum = ImitationTrainerTrainingSettings.NumberOfIterations;
@@ -179,43 +181,128 @@ void ULearningAgentsImitationTrainer::BeginTraining(
 		RecordedActions.Num<1>(),
 		ImitationTrainingSettings);
 
-	ImitationTrainingTask = UE::Tasks::Launch(UE_SOURCE_LOCATION, [this, TrainerFlags]()
+	UE_LOG(LogLearning, Display, TEXT("%s: Sending / Receiving initial policy..."), *GetName());
+
+	UE::Learning::ETrainerResponse Response = UE::Learning::ETrainerResponse::Success;
+
+	if (bReinitializePolicyNetwork)
+	{
+		Response = ImitationTrainer->RecvPolicy(Policy->GetPolicyNetwork(), TrainerTimeout);
+		Policy->GetNetworkAsset()->ForceMarkDirty();
+	}
+	else
+	{
+		Response = ImitationTrainer->SendPolicy(Policy->GetPolicyNetwork(), TrainerTimeout);
+	}
+
+	if (Response != UE::Learning::ETrainerResponse::Success)
+	{
+		UE_LOG(LogLearning, Error, TEXT("%s: Error sending or receiving policy from trainer: %s. Check log for errors."), *GetName(), UE::Learning::Trainer::GetResponseString(Response));
+		ImitationTrainer->Terminate();
+		return;
+	}
+
+	UE_LOG(LogLearning, Display, TEXT("%s: Sending Experience..."), *GetName());
+
+	// Send Experience
+
+	Response = ImitationTrainer->SendExperience(RecordedObservations, RecordedActions, TrainerTimeout);
+
+	if (Response != UE::Learning::ETrainerResponse::Success)
+	{
+		UE_LOG(LogLearning, Error, TEXT("%s: Error sending experience to trainer: %s. Check log for errors."), *GetName(), UE::Learning::Trainer::GetResponseString(Response));
+		ImitationTrainer->Terminate();
+		return;
+	}
+
+	bIsTraining = true;
+}
+
+void ULearningAgentsImitationTrainer::DoneTraining()
+{
+	if (IsTraining())
 	{
-		UE::Learning::ImitationTrainer::Train(
-			*ImitationTrainer,
-			Policy->GetPolicyNetwork(),
-			RecordedObservations,
-			RecordedActions,
-			TrainerFlags,
-			&bRequestImitationTrainingStop,
-			&NetworkLock);
+		// Wait for Trainer to finish
+		ImitationTrainer->Wait(1.0f);
+
+		// If not finished in time, terminate
+		ImitationTrainer->Terminate();
 
-			bIsTrainingComplete = true;
-	});
+		bIsTraining = false;
+	}
 }
 
 void ULearningAgentsImitationTrainer::EndTraining()
 {
+	if (IsTraining())
+	{
+		UE_LOG(LogLearning, Display, TEXT("%s: Stopping training..."), *GetName());
+		ImitationTrainer->SendStop();
+		DoneTraining();
+	}
+}
+
+void ULearningAgentsImitationTrainer::IterateTraining()
+{
+	UE_LEARNING_TRACE_CPUPROFILER_EVENT_SCOPE(ULearningAgentsImitationTrainer::IterateTraining);
+
 	if (!IsTraining())
 	{
-		UE_LOG(LogLearning, Warning, TEXT("%s: Cannot end training as we are not training!"), *GetName());
+		UE_LOG(LogLearning, Error, TEXT("%s: Training not running."), *GetName());
 		return;
 	}
 
-	UE_LOG(LogLearning, Display, TEXT("%s: Imitation Training Ended."), *GetName());
-
-	bRequestImitationTrainingStop = true;
-	ImitationTrainingTask.Wait(FTimespan(0, 0, 30));
+	if (ImitationTrainer->HasPolicyOrCompleted())
+	{
+		UE::Learning::ETrainerResponse Response = ImitationTrainer->RecvPolicy(Policy->GetPolicyNetwork(), TrainerTimeout);
+		Policy->GetNetworkAsset()->ForceMarkDirty();
 
-	bIsTraining = false;
+		if (Response == UE::Learning::ETrainerResponse::Completed)
+		{
+			UE_LOG(LogLearning, Display, TEXT("%s: Trainer completed training."), *GetName());
+			DoneTraining();
+			return;
+		}
+		else if (Response != UE::Learning::ETrainerResponse::Success)
+		{
+			UE_LOG(LogLearning, Error, TEXT("%s: Error receiving policy from trainer. Check log for errors."), *GetName());
+			EndTraining();
+			return;
+		}
+	}
 }
 
-bool ULearningAgentsImitationTrainer::IsTraining() const
+void ULearningAgentsImitationTrainer::RunTraining(
+	ULearningAgentsPolicy* InPolicy,
+	const ULearningAgentsRecording* Recording,
+	const FLearningAgentsImitationTrainerSettings& ImitationTrainerSettings,
+	const FLearningAgentsImitationTrainerTrainingSettings& ImitationTrainerTrainingSettings,
+	const FLearningAgentsTrainerPathSettings& ImitationTrainerPathSettings,
+	const bool bReinitializePolicyNetwork)
 {
-	return bIsTraining;
+	// If we aren't training yet, then start training and do the first inference step.
+	if (!IsTraining())
+	{
+		BeginTraining(
+			InPolicy,
+			Recording,
+			ImitationTrainerSettings,
+			ImitationTrainerTrainingSettings,
+			ImitationTrainerPathSettings,
+			bReinitializePolicyNetwork);
+
+		if (!IsTraining())
+		{
+			// If IsTraining is false, then BeginTraining must have failed and we can't continue.
+			return;
+		}
+	}
+
+	// Otherwise, do the regular training process.
+	IterateTraining();
 }
 
-bool ULearningAgentsImitationTrainer::IsTrainingComplete() const
+bool ULearningAgentsImitationTrainer::IsTraining() const
 {
-	return bIsTrainingComplete;
+	return bIsTraining;
 }
diff --git a/Engine/Plugins/Experimental/LearningAgents/Source/LearningAgentsTraining/Private/LearningAgentsRecorder.cpp b/Engine/Plugins/Experimental/LearningAgents/Source/LearningAgentsTraining/Private/LearningAgentsRecorder.cpp
--- a/Engine/Plugins/Experimental/LearningAgents/Source/LearningAgentsTraining/Private/LearningAgentsRecorder.cpp
+++ b/Engine/Plugins/Experimental/LearningAgents/Source/LearningAgentsTraining/Private/LearningAgentsRecorder.cpp
@@ -82,7 +82,7 @@ void ULearningAgentsRecorder::FAgentRecordBuffer::CopyToRecord(FLearningAgentsRe
 	}
 }
 
-ULearningAgentsRecorder::ULearningAgentsRecorder() : ULearningAgentsManagerComponent() {}
+ULearningAgentsRecorder::ULearningAgentsRecorder(const FObjectInitializer& ObjectInitializer) : ULearningAgentsManagerComponent(ObjectInitializer) {}
 ULearningAgentsRecorder::ULearningAgentsRecorder(FVTableHelper& Helper) : ULearningAgentsRecorder() {}
 ULearningAgentsRecorder::~ULearningAgentsRecorder() {}
 
@@ -97,9 +97,9 @@ void ULearningAgentsRecorder::EndPlay(const EEndPlayReason::Type EndPlayReason)
 }
 
 void ULearningAgentsRecorder::SetupRecorder(
-	ALearningAgentsManager* InAgentManager, 
 	ULearningAgentsInteractor* InInteractor,
-	const FLearningAgentsRecorderPathSettings& RecorderPathSettings)
+	const FLearningAgentsRecorderPathSettings& RecorderPathSettings,
+	ULearningAgentsRecording* RecordingAsset)
 {
 	if (IsSetup())
 	{
@@ -107,20 +107,12 @@ void ULearningAgentsRecorder::SetupRecorder(
 		return;
 	}
 
-	if (!InAgentManager)
+	if (!Manager)
 	{
-		UE_LOG(LogLearning, Error, TEXT("%s: InAgentManager is nullptr."), *GetName());
+		UE_LOG(LogLearning, Error, TEXT("%s: Must be attached to a LearningAgentsManager Actor."), *GetName());
 		return;
 	}
 
-	if (!InAgentManager->IsManagerSetup())
-	{
-		UE_LOG(LogLearning, Error, TEXT("%s: %s's SetupManager must be run before it can be used."), *GetName(), *InAgentManager->GetName());
-		return;
-	}
-
-	AgentManager = InAgentManager;
-
 	if (!InInteractor)
 	{
 		UE_LOG(LogLearning, Error, TEXT("%s: InInteractor is nullptr."), *GetName());
@@ -135,29 +127,51 @@ void ULearningAgentsRecorder::SetupRecorder(
 
 	Interactor = InInteractor;
 
-	Recording = NewObject<ULearningAgentsRecording>(this, TEXT("Recording"));
+	if (RecordingAsset)
+	{
+		Recording = RecordingAsset;
+	}
+	else
+	{
+		Recording = NewObject<ULearningAgentsRecording>(this, TEXT("Recording"));
+	}
+
 	RecordingDirectory = RecorderPathSettings.IntermediateRelativePath.Path / TEXT("LearningAgents") / RecorderPathSettings.RecordingsSubdirectory;
 
 	RecordBuffers.Empty();
-	RecordBuffers.SetNum(AgentManager->GetMaxInstanceNum());
+	RecordBuffers.SetNum(Manager->GetMaxAgentNum());
 
 	bIsSetup = true;
 }
 
-bool ULearningAgentsRecorder::RemoveAgent(const int32 AgentId)
+void ULearningAgentsRecorder::OnAgentsRemoved(const TArray<int32>& AgentIds)
 {
-	bool bSuccess = Super::RemoveAgent(AgentId);
-
-	if (bSuccess)
+	if (IsSetup())
 	{
-		if (!RecordBuffers[AgentId].IsEmpty())
+		for (const int32 AgentId : AgentIds)
 		{
-			RecordBuffers[AgentId].CopyToRecord(Recording->Records.AddDefaulted_GetRef());
-			RecordBuffers[AgentId].Empty();
+			if (!RecordBuffers[AgentId].IsEmpty())
+			{
+				RecordBuffers[AgentId].CopyToRecord(Recording->Records.AddDefaulted_GetRef());
+				RecordBuffers[AgentId].Empty();
+			}
 		}
 	}
+}
 
-	return bSuccess;
+void ULearningAgentsRecorder::OnAgentsReset(const TArray<int32>& AgentIds)
+{
+	if (IsSetup())
+	{
+		for (const int32 AgentId : AgentIds)
+		{
+			if (!RecordBuffers[AgentId].IsEmpty())
+			{
+				RecordBuffers[AgentId].CopyToRecord(Recording->Records.AddDefaulted_GetRef());
+				RecordBuffers[AgentId].Empty();
+			}
+		}
+	}
 }
 
 void ULearningAgentsRecorder::AddExperience()
@@ -174,10 +188,24 @@ void ULearningAgentsRecorder::AddExperience()
 		return;
 	}
 
-	for (const int32 AgentId : AddedAgentSet)
+	for (const int32 AgentId : Manager->GetAllAgentSet())
 	{
+		if (Interactor->GetObservationEncodingAgentIteration()[AgentId] == 0 ||
+			Interactor->GetActionEncodingAgentIteration()[AgentId] == 0)
+		{
+			UE_LOG(LogLearning, Warning, TEXT("%s: Agent with id %i has not made observations and taken actions so experience will not be recorded for it."), *GetName(), AgentId);
+			continue;
+		}
+
+		if (Interactor->GetObservationEncodingAgentIteration()[AgentId] !=
+			Interactor->GetActionEncodingAgentIteration()[AgentId])
+		{
+			UE_LOG(LogLearning, Warning, TEXT("%s: Agent with id %i does not have matching iteration numbers for observations and actions so experience will not be recorded for it."), *GetName(), AgentId);
+			continue;
+		}
+
 		RecordBuffers[AgentId].Push(
-			Interactor->GetObservationFeature().FeatureBuffer()[AgentId], 
+			Interactor->GetObservationFeature().FeatureBuffer()[AgentId],
 			Interactor->GetActionFeature().FeatureBuffer()[AgentId]);
 	}
 }
@@ -203,7 +231,7 @@ void ULearningAgentsRecorder::EndRecording()
 
 	// Write to Recording
 
-	for (const int32 AgentId : AddedAgentSet)
+	for (const int32 AgentId : Manager->GetAllAgentSet())
 	{
 		if (!RecordBuffers[AgentId].IsEmpty())
 		{
@@ -243,7 +271,7 @@ void ULearningAgentsRecorder::SaveRecordingToFile(const FFilePath& File) const
 	Recording->SaveRecordingToFile(File);
 }
 
-void ULearningAgentsRecorder::LoadRecordingFromAsset(const ULearningAgentsRecording* Asset)
+void ULearningAgentsRecorder::AppendRecordingFromFile(const FFilePath& File)
 {
 	if (!IsSetup())
 	{
@@ -251,16 +279,10 @@ void ULearningAgentsRecorder::LoadRecordingFromAsset(const ULearningAgentsRecord
 		return;
 	}
 
-	if (!Asset)
-	{
-		UE_LOG(LogLearning, Error, TEXT("%s: Asset is nullptr."), *GetName());
-		return;
-	}
-
-	Recording->Records = Asset->Records;
+	Recording->AppendRecordingFromFile(File);
 }
 
-void ULearningAgentsRecorder::SaveRecordingToAsset(ULearningAgentsRecording* Asset) const
+void ULearningAgentsRecorder::UseRecordingAsset(ULearningAgentsRecording* RecordingAsset)
 {
 	if (!IsSetup())
 	{
@@ -268,31 +290,22 @@ void ULearningAgentsRecorder::SaveRecordingToAsset(ULearningAgentsRecording* Ass
 		return;
 	}
 
-	if (!Asset)
+	if (!RecordingAsset)
 	{
 		UE_LOG(LogLearning, Error, TEXT("%s: Asset is nullptr."), *GetName());
 		return;
 	}
 
-	Asset->Records = Recording->Records;
-
-	// Manually mark the package as dirty since just using `Modify` prevents 
-	// marking packages as dirty during PIE which is most likely when this
-	// is being used.
-	if (UPackage* Package = Asset->GetPackage())
+	if (RecordingAsset == Recording)
 	{
-		const bool bIsDirty = Package->IsDirty();
-
-		if (!bIsDirty)
-		{
-			Package->SetDirtyFlag(true);
-		}
-
-		Package->PackageMarkedDirtyEvent.Broadcast(Package, bIsDirty);
+		UE_LOG(LogLearning, Error, TEXT("%s: Asset is same as the current recording."), *GetName());
+		return;
 	}
+
+	Recording = RecordingAsset;
 }
 
-void ULearningAgentsRecorder::AppendRecordingToAsset(ULearningAgentsRecording* Asset) const
+void ULearningAgentsRecorder::LoadRecordingFromAsset(ULearningAgentsRecording* RecordingAsset)
 {
 	if (!IsSetup())
 	{
@@ -300,28 +313,29 @@ void ULearningAgentsRecorder::AppendRecordingToAsset(ULearningAgentsRecording* A
 		return;
 	}
 
-	if (!Asset)
+	Recording->LoadRecordingFromAsset(RecordingAsset);
+}
+
+void ULearningAgentsRecorder::SaveRecordingToAsset(ULearningAgentsRecording* RecordingAsset)
+{
+	if (!IsSetup())
 	{
-		UE_LOG(LogLearning, Error, TEXT("%s: Asset is invalid."), *GetName());
+		UE_LOG(LogLearning, Error, TEXT("%s: Setup not complete."), *GetName());
 		return;
 	}
 
-	Asset->Records.Append(Recording->Records);
+	Recording->SaveRecordingToAsset(RecordingAsset);
+}
 
-	// Manually mark the package as dirty since just using `Modify` prevents 
-	// marking packages as dirty during PIE which is most likely when this
-	// is being used.
-	if (UPackage* Package = Asset->GetPackage())
+void ULearningAgentsRecorder::AppendRecordingToAsset(ULearningAgentsRecording* RecordingAsset)
+{
+	if (!IsSetup())
 	{
-		const bool bIsDirty = Package->IsDirty();
-
-		if (!bIsDirty)
-		{
-			Package->SetDirtyFlag(true);
-		}
-
-		Package->PackageMarkedDirtyEvent.Broadcast(Package, bIsDirty);
+		UE_LOG(LogLearning, Error, TEXT("%s: Setup not complete."), *GetName());
+		return;
 	}
+
+	Recording->AppendRecordingToAsset(RecordingAsset);
 }
 
 void ULearningAgentsRecorder::BeginRecording(bool bReinitializeRecording)
@@ -341,9 +355,10 @@ void ULearningAgentsRecorder::BeginRecording(bool bReinitializeRecording)
 	if (bReinitializeRecording)
 	{
 		Recording->Records.Empty();
+		Recording->ForceMarkDirty();
 	}
 
-	for (const int32 AgentId : AddedAgentSet)
+	for (const int32 AgentId : Manager->GetAllAgentSet())
 	{
 		RecordBuffers[AgentId].Empty();
 	}
diff --git a/Engine/Plugins/Experimental/LearningAgents/Source/LearningAgentsTraining/Private/LearningAgentsRecording.cpp b/Engine/Plugins/Experimental/LearningAgents/Source/LearningAgentsTraining/Private/LearningAgentsRecording.cpp
--- a/Engine/Plugins/Experimental/LearningAgents/Source/LearningAgentsTraining/Private/LearningAgentsRecording.cpp
+++ b/Engine/Plugins/Experimental/LearningAgents/Source/LearningAgentsTraining/Private/LearningAgentsRecording.cpp
@@ -74,26 +74,13 @@ void ULearningAgentsRecording::LoadRecordingFromFile(const FFilePath& File)
 		}
 
 		UE_LEARNING_CHECK(Offset == RecordingData.Num());
+
+		ForceMarkDirty();
 	}
 	else
 	{
 		UE_LOG(LogLearning, Error, TEXT("%s: Failed to load recording. File not found: \"%s\""), *GetName(), *File.FilePath);
 	}
-
-	// Manually mark the package as dirty since just using `Modify` prevents 
-	// marking packages as dirty during PIE which is most likely when this
-	// is being used.
-	if (UPackage* Package = GetPackage())
-	{
-		const bool bIsDirty = Package->IsDirty();
-
-		if (!bIsDirty)
-		{
-			Package->SetDirtyFlag(true);
-		}
-
-		Package->PackageMarkedDirtyEvent.Broadcast(Package, bIsDirty);
-	}
 }
 
 void ULearningAgentsRecording::SaveRecordingToFile(const FFilePath& File) const
@@ -137,4 +124,85 @@ void ULearningAgentsRecording::SaveRecordingToFile(const FFilePath& File) const
 	{
 		UE_LOG(LogLearning, Error, TEXT("%s: Failed to save recording to file: \"%s\""), *GetName(), *File.FilePath);
 	}
+}
+
+void ULearningAgentsRecording::AppendRecordingFromFile(const FFilePath& File)
+{
+	ULearningAgentsRecording* TempRecording = NewObject<ULearningAgentsRecording>(this, TEXT("TempRecording"));
+	TempRecording->LoadRecordingFromFile(File);
+
+	Records.Append(TempRecording->Records);
+	ForceMarkDirty();
+}
+
+void ULearningAgentsRecording::LoadRecordingFromAsset(ULearningAgentsRecording* RecordingAsset)
+{
+	if (!RecordingAsset)
+	{
+		UE_LOG(LogLearning, Error, TEXT("%s: Asset is nullptr."), *GetName());
+		return;
+	}
+
+	if (RecordingAsset == this)
+	{
+		UE_LOG(LogLearning, Error, TEXT("%s: Asset is same as the current recording."), *GetName());
+		return;
+	}
+
+	Records = RecordingAsset->Records;
+	ForceMarkDirty();
+}
+
+void ULearningAgentsRecording::SaveRecordingToAsset(ULearningAgentsRecording* RecordingAsset)
+{
+	if (!RecordingAsset)
+	{
+		UE_LOG(LogLearning, Error, TEXT("%s: Asset is nullptr."), *GetName());
+		return;
+	}
+
+	if (RecordingAsset == this)
+	{
+		UE_LOG(LogLearning, Error, TEXT("%s: Asset is same as the current recording."), *GetName());
+		return;
+	}
+
+	RecordingAsset->Records = Records;
+	RecordingAsset->ForceMarkDirty();
+}
+
+void ULearningAgentsRecording::AppendRecordingToAsset(ULearningAgentsRecording* RecordingAsset)
+{
+	if (!RecordingAsset)
+	{
+		UE_LOG(LogLearning, Error, TEXT("%s: Asset is invalid."), *GetName());
+		return;
+	}
+
+	if (RecordingAsset == this)
+	{
+		UE_LOG(LogLearning, Error, TEXT("%s: Asset is same as the current recording."), *GetName());
+		return;
+	}
+
+	RecordingAsset->Records.Append(Records);
+	RecordingAsset->ForceMarkDirty();
+}
+
+void ULearningAgentsRecording::ForceMarkDirty()
+{
+	// Manually mark the package as dirty since just using `Modify` prevents 
+	// marking packages as dirty during PIE which is most likely when this
+	// is being used.
+	if (UPackage* Package = GetPackage())
+	{
+		const bool bIsDirty = Package->IsDirty();
+
+		if (!bIsDirty)
+		{
+			Package->SetDirtyFlag(true);
+		}
+
+		Package->PackageMarkedDirtyEvent.Broadcast(Package, bIsDirty);
+	}
 }
\ No newline at end of file
diff --git a/Engine/Plugins/Experimental/LearningAgents/Source/LearningAgentsTraining/Private/LearningAgentsRewards.cpp b/Engine/Plugins/Experimental/LearningAgents/Source/LearningAgentsTraining/Private/LearningAgentsRewards.cpp
--- a/Engine/Plugins/Experimental/LearningAgents/Source/LearningAgentsTraining/Private/LearningAgentsRewards.cpp
+++ b/Engine/Plugins/Experimental/LearningAgents/Source/LearningAgentsTraining/Private/LearningAgentsRewards.cpp
@@ -22,15 +22,24 @@ namespace UE::Learning::Agents::Rewards::Private
 			return nullptr;
 		}
 
+		if (!InAgentTrainer->HasAgentManager())
+		{
+			UE_LOG(LogLearning, Error, TEXT("%s: Must be attached to a LearningAgentsManager Actor."), *InAgentTrainer->GetName());
+			return nullptr;
+		}
+
 		RewardUObject* Reward = NewObject<RewardUObject>(InAgentTrainer, Name);
 
 		Reward->AgentTrainer = InAgentTrainer;
 		Reward->RewardObject = MakeShared<RewardFObject>(
 			Reward->GetFName(),
 			InAgentTrainer->GetAgentManager()->GetInstanceData().ToSharedRef(),
-			InAgentTrainer->GetAgentManager()->GetMaxInstanceNum(),
+			InAgentTrainer->GetAgentManager()->GetMaxAgentNum(),
 			Forward<InArgTypes>(Args)...);
 
+		Reward->AgentIteration.SetNumUninitialized({ InAgentTrainer->GetAgentManager()->GetMaxAgentNum() });
+		UE::Learning::Array::Set<1, uint64>(Reward->AgentIteration, INDEX_NONE);
+
 		InAgentTrainer->AddReward(Reward, Reward->RewardObject.ToSharedRef());
 
 		return Reward;
@@ -53,6 +62,7 @@ void UFloatReward::SetFloatReward(const int32 AgentId, const float Reward)
 	}
 
 	RewardObject->InstanceData->View(RewardObject->ValueHandle)[AgentId] = Reward;
+	AgentIteration[AgentId]++;
 }
 
 #if UE_LEARNING_AGENTS_ENABLE_VISUAL_LOG
@@ -97,6 +107,7 @@ void UScalarVelocityReward::SetScalarVelocityReward(int32 AgentId, float Velocit
 	}
 
 	RewardObject->InstanceData->View(RewardObject->VelocityHandle)[AgentId] = Velocity;
+	AgentIteration[AgentId]++;
 }
 
 #if UE_LEARNING_AGENTS_ENABLE_VISUAL_LOG
@@ -142,6 +153,7 @@ void ULocalDirectionalVelocityReward::SetLocalDirectionalVelocityReward(const in
 
 	RewardObject->InstanceData->View(RewardObject->VelocityHandle)[AgentId] = Velocity;
 	RewardObject->InstanceData->View(RewardObject->RelativeRotationHandle)[AgentId] = RelativeRotation.Quaternion();
+	AgentIteration[AgentId]++;
 }
 
 #if UE_LEARNING_AGENTS_ENABLE_VISUAL_LOG
@@ -237,6 +249,7 @@ void UPlanarPositionDifferencePenalty::SetPlanarPositionDifferencePenalty(const
 
 	RewardObject->InstanceData->View(RewardObject->Position0Handle)[AgentId] = Position0;
 	RewardObject->InstanceData->View(RewardObject->Position1Handle)[AgentId] = Position1;
+	AgentIteration[AgentId]++;
 }
 
 #if UE_LEARNING_AGENTS_ENABLE_VISUAL_LOG
@@ -362,6 +375,7 @@ void UPositionArraySimilarityReward::SetPositionArraySimilarityReward(
 	RelativeRotation1View[AgentId] = RelativeRotation1.Quaternion();
 	UE::Learning::Array::Copy<1, FVector>(Positions0View[AgentId], Positions0);
 	UE::Learning::Array::Copy<1, FVector>(Positions1View[AgentId], Positions1);
+	AgentIteration[AgentId]++;
 }
 
 #if UE_LEARNING_AGENTS_ENABLE_VISUAL_LOG
diff --git a/Engine/Plugins/Experimental/LearningAgents/Source/LearningAgentsTraining/Private/LearningAgentsTrainer.cpp b/Engine/Plugins/Experimental/LearningAgents/Source/LearningAgentsTraining/Private/LearningAgentsTrainer.cpp
--- a/Engine/Plugins/Experimental/LearningAgents/Source/LearningAgentsTraining/Private/LearningAgentsTrainer.cpp
+++ b/Engine/Plugins/Experimental/LearningAgents/Source/LearningAgentsTraining/Private/LearningAgentsTrainer.cpp
@@ -2,8 +2,6 @@
 
 #include "LearningAgentsTrainer.h"
 
-#include "GenericPlatform/GenericPlatformMisc.h"
-#include "HAL/FileManager.h"
 #include "LearningAgentsInteractor.h"
 #include "LearningAgentsManager.h"
 #include "LearningAgentsRewards.h"
@@ -22,6 +20,10 @@
 #include "LearningRewardObject.h"
 #include "LearningCompletion.h"
 #include "LearningCompletionObject.h"
+
+#include "GenericPlatform/GenericPlatformMisc.h"
+#include "HAL/FileManager.h"
+#include "UObject/Package.h"
 #include "Misc/App.h"
 #include "Misc/Paths.h"
 #include "GameFramework/GameUserSettings.h"
@@ -101,12 +103,11 @@ FString FLearningAgentsTrainerPathSettings::GetIntermediatePath() const
 	return IntermediateRelativePath.Path;
 }
 
-ULearningAgentsTrainer::ULearningAgentsTrainer() : ULearningAgentsManagerComponent() {}
+ULearningAgentsTrainer::ULearningAgentsTrainer(const FObjectInitializer& ObjectInitializer) : ULearningAgentsManagerComponent(ObjectInitializer) {}
 ULearningAgentsTrainer::ULearningAgentsTrainer(FVTableHelper& Helper) : ULearningAgentsTrainer() {}
 ULearningAgentsTrainer::~ULearningAgentsTrainer() {}
 
 void ULearningAgentsTrainer::SetupTrainer(
-	ALearningAgentsManager* InAgentManager,
 	ULearningAgentsInteractor* InInteractor,
 	ULearningAgentsPolicy* InPolicy,
 	ULearningAgentsCritic* InCritic,
@@ -118,20 +119,12 @@ void ULearningAgentsTrainer::SetupTrainer(
 		return;
 	}
 
-	if (!InAgentManager)
-	{
-		UE_LOG(LogLearning, Error, TEXT("%s: InAgentManager is nullptr."), *GetName());
-		return;
-	}
-
-	if (!InAgentManager->IsManagerSetup())
+	if (!Manager)
 	{
-		UE_LOG(LogLearning, Error, TEXT("%s: %s's SetupManager must be run before it can be used."), *GetName(), *InAgentManager->GetName());
+		UE_LOG(LogLearning, Error, TEXT("%s: Must be attached to a LearningAgentsManager Actor."), *GetName());
 		return;
 	}
 
-	AgentManager = InAgentManager;
-
 	if (!InInteractor)
 	{
 		UE_LOG(LogLearning, Error, TEXT("%s: InInteractor is nullptr."), *GetName());
@@ -144,8 +137,6 @@ void ULearningAgentsTrainer::SetupTrainer(
 		return;
 	}
 
-	Interactor = InInteractor;
-
 	if (!InPolicy)
 	{
 		UE_LOG(LogLearning, Error, TEXT("%s: InPolicy is nullptr."), *GetName());
@@ -158,28 +149,21 @@ void ULearningAgentsTrainer::SetupTrainer(
 		return;
 	}
 
-	Policy = InPolicy;
-
 	// The critic is optional unlike the other components
-	if (InCritic)
+	if (InCritic && !InCritic->IsSetup())
 	{
-		if (!InCritic->IsSetup())
-		{
-			UE_LOG(LogLearning, Error, TEXT("%s: %s's Setup must be run before it can be used."), *GetName(), *InCritic->GetName());
-			return;
-		}
-
-		Critic = InCritic;
+		UE_LOG(LogLearning, Error, TEXT("%s: %s's Setup must be run before it can be used."), *GetName(), *InCritic->GetName());
+		return;
 	}
 
+	Interactor = InInteractor;
+	Policy = InPolicy;
+	Critic = InCritic;
+
 	// Setup Rewards
 	RewardObjects.Empty();
 	RewardFeatures.Empty();
-	SetupRewards(this);
-	Rewards = MakeShared<UE::Learning::FSumReward>(TEXT("Rewards"),
-		TLearningArrayView<1, const TSharedRef<UE::Learning::FRewardObject>>(RewardFeatures),
-		AgentManager->GetInstanceData().ToSharedRef(),
-		AgentManager->GetMaxInstanceNum());
+	SetupRewards();
 
 	if (RewardObjects.Num() == 0)
 	{
@@ -187,25 +171,31 @@ void ULearningAgentsTrainer::SetupTrainer(
 		return;
 	}
 
+	Rewards = MakeShared<UE::Learning::FSumReward>(TEXT("Rewards"),
+		TLearningArrayView<1, const TSharedRef<UE::Learning::FRewardObject>>(RewardFeatures),
+		Manager->GetInstanceData().ToSharedRef(),
+		Manager->GetMaxAgentNum());
+
 	// Setup Completions
 	CompletionObjects.Empty();
 	CompletionFeatures.Empty();
-	SetupCompletions(this);
-	Completions = MakeShared<UE::Learning::FAnyCompletion>(TEXT("Completions"),
-		TLearningArrayView<1, const TSharedRef<UE::Learning::FCompletionObject>>(CompletionFeatures),
-		AgentManager->GetInstanceData().ToSharedRef(),
-		AgentManager->GetMaxInstanceNum());
+	SetupCompletions();
 
-	if (RewardObjects.Num() == 0)
+	if (CompletionObjects.Num() == 0)
 	{
 		// Not an error or warning because it's fine to run training without any completions.
 		UE_LOG(LogLearning, Display, TEXT("%s: No completions added to Trainer during SetupCompletions."), *GetName());
 	}
 
+	Completions = MakeShared<UE::Learning::FAnyCompletion>(TEXT("Completions"),
+		TLearningArrayView<1, const TSharedRef<UE::Learning::FCompletionObject>>(CompletionFeatures),
+		Manager->GetInstanceData().ToSharedRef(),
+		Manager->GetMaxAgentNum());
+
 	// Create Episode Buffer
 	EpisodeBuffer = MakeUnique<UE::Learning::FEpisodeBuffer>();
 	EpisodeBuffer->Resize(
-		AgentManager->GetMaxInstanceNum(),
+		Manager->GetMaxAgentNum(),
 		TrainerSettings.MaxStepNum,
 		Interactor->GetObservationFeature().DimNum(),
 		Interactor->GetActionFeature().DimNum());
@@ -222,27 +212,74 @@ void ULearningAgentsTrainer::SetupTrainer(
 
 	// Create Reset Buffer
 	ResetBuffer = MakeUnique<UE::Learning::FResetInstanceBuffer>();
-	ResetBuffer->Resize(AgentManager->GetMaxInstanceNum());
+	ResetBuffer->Resize(Manager->GetMaxAgentNum());
 
 	// Record Timeout Setting
 	TrainerTimeout = TrainerSettings.TrainerCommunicationTimeout;
 
+	// Reset Agent iteration
+	RewardEvaluatedAgentIteration.SetNumUninitialized({ Manager->GetMaxAgentNum() });
+	CompletionEvaluatedAgentIteration.SetNumUninitialized({ Manager->GetMaxAgentNum() });
+	UE::Learning::Array::Set<1, uint64>(RewardEvaluatedAgentIteration, INDEX_NONE);
+	UE::Learning::Array::Set<1, uint64>(CompletionEvaluatedAgentIteration, INDEX_NONE);
+	SetAgentIteration(Manager->GetAllAgentSet(), 0);
+
 	bIsSetup = true;
 }
 
-bool ULearningAgentsTrainer::AddAgent(const int32 AgentId)
+void ULearningAgentsTrainer::OnAgentsAdded(const TArray<int32>& AgentIds)
+{
+	if (IsSetup())
+	{
+		UE::Learning::FIndexSet AgentsSet = AgentIds;
+		AgentsSet.TryMakeSlice();
+
+		ResetEpisodes(AgentIds);
+		EpisodeBuffer->Reset(AgentsSet);
+
+		SetAgentIteration(AgentsSet, 0);
+	}
+}
+
+void ULearningAgentsTrainer::OnAgentsRemoved(const TArray<int32>& AgentIds)
+{
+	if (IsSetup())
+	{
+		UE::Learning::FIndexSet AgentsSet = AgentIds;
+		AgentsSet.TryMakeSlice();
+
+		SetAgentIteration(AgentsSet, INDEX_NONE);
+	}
+}
+
+void ULearningAgentsTrainer::OnAgentsReset(const TArray<int32>& AgentIds)
+{
+	if (IsSetup())
+	{
+		UE::Learning::FIndexSet AgentsSet = AgentIds;
+		AgentsSet.TryMakeSlice();
+
+		ResetEpisodes(AgentIds);
+		EpisodeBuffer->Reset(AgentsSet);
+
+		SetAgentIteration(AgentsSet, 0);
+	}
+}
+
+void ULearningAgentsTrainer::SetAgentIteration(const UE::Learning::FIndexSet AgentIds, const uint64 Value)
 {
-	bool bSuccess = Super::AddAgent(AgentId);
+	UE::Learning::Array::Set<1, uint64>(RewardEvaluatedAgentIteration, Value, AgentIds);
+	UE::Learning::Array::Set<1, uint64>(CompletionEvaluatedAgentIteration, Value, AgentIds);
 
-	if (bSuccess && IsTraining())
+	for (ULearningAgentsReward* RewardObject : RewardObjects)
 	{
-		// Reset the instance and the buffer in case we have stale 
-		// data from another agent that was using this id previously
-		ResetEpisodes({ AgentId });
-		EpisodeBuffer->Reset(AgentId);
+		UE::Learning::Array::Set<1, uint64>(RewardObject->AgentIteration, Value, AgentIds);
 	}
 
-	return bSuccess;
+	for (ULearningAgentsCompletion* CompletionObject : CompletionObjects)
+	{
+		UE::Learning::Array::Set<1, uint64>(CompletionObject->AgentIteration, Value, AgentIds);
+	}
 }
 
 void ULearningAgentsTrainer::EndPlay(const EEndPlayReason::Type EndPlayReason)
@@ -255,7 +292,7 @@ void ULearningAgentsTrainer::EndPlay(const EEndPlayReason::Type EndPlayReason)
 	Super::EndPlay(EndPlayReason);
 }
 
-void ULearningAgentsTrainer::SetupRewards_Implementation(ULearningAgentsTrainer* AgentTrainer)
+void ULearningAgentsTrainer::SetupRewards_Implementation()
 {
 	// Can be overridden to setup rewards without blueprints
 }
@@ -268,11 +305,12 @@ void ULearningAgentsTrainer::SetRewards_Implementation(const TArray<int32>& Agen
 void ULearningAgentsTrainer::AddReward(TObjectPtr<ULearningAgentsReward> Object, const TSharedRef<UE::Learning::FRewardObject>& Reward)
 {
 	UE_LEARNING_CHECK(!IsSetup());
+	UE_LEARNING_CHECK(Object);
 	RewardObjects.Add(Object);
 	RewardFeatures.Add(Reward);
 }
 
-void ULearningAgentsTrainer::SetupCompletions_Implementation(ULearningAgentsTrainer* AgentTrainer)
+void ULearningAgentsTrainer::SetupCompletions_Implementation()
 {
 	// Can be overridden to setup completions without blueprints
 }
@@ -285,6 +323,7 @@ void ULearningAgentsTrainer::SetCompletions_Implementation(const TArray<int32>&
 void ULearningAgentsTrainer::AddCompletion(TObjectPtr<ULearningAgentsCompletion> Object, const TSharedRef<UE::Learning::FCompletionObject>& Completion)
 {
 	UE_LEARNING_CHECK(!IsSetup());
+	UE_LEARNING_CHECK(Object);
 	CompletionObjects.Add(Object);
 	CompletionFeatures.Add(Completion);
 }
@@ -300,7 +339,8 @@ void ULearningAgentsTrainer::BeginTraining(
 	const FLearningAgentsTrainerPathSettings& TrainerPathSettings,
 	const FLearningAgentsCriticSettings& CriticSettings,
 	const bool bReinitializePolicyNetwork,
-	const bool bReinitializeCriticNetwork)
+	const bool bReinitializeCriticNetwork,
+	const bool bResetAgentsOnBegin)
 {
 	if (!IsSetup())
 	{
@@ -465,7 +505,7 @@ void ULearningAgentsTrainer::BeginTraining(
 		PPONetworkSettings,
 		TrainerFlags);
 
-	UE_LOG(LogLearning, Display, TEXT("%s: Receiving initial policy..."), *GetName());
+	UE_LOG(LogLearning, Display, TEXT("%s: Sending / Receiving initial policy..."), *GetName());
 
 	UE::Learning::ETrainerResponse Response = UE::Learning::ETrainerResponse::Success;
 
@@ -476,6 +516,7 @@ void ULearningAgentsTrainer::BeginTraining(
 	else
 	{
 		Response = Trainer->RecvPolicy(Policy->GetPolicyNetwork(), TrainerTimeout);
+		Policy->GetNetworkAsset()->ForceMarkDirty();
 	}
 
 	if (Response != UE::Learning::ETrainerResponse::Success)
@@ -494,6 +535,7 @@ void ULearningAgentsTrainer::BeginTraining(
 		else if ((bool)(TrainerFlags & UE::Learning::EPPOTrainerFlags::SynchronizeCriticNetwork))
 		{
 			Response = Trainer->RecvCritic(Critic->GetCriticNetwork(), TrainerTimeout);
+			Critic->GetNetworkAsset()->ForceMarkDirty();
 		}
 		else
 		{
@@ -508,9 +550,13 @@ void ULearningAgentsTrainer::BeginTraining(
 		}
 	}
 
-	// Reset Agents, Episode Buffer, and Replay Buffer
-	ResetEpisodes(AddedAgentIds);
-	EpisodeBuffer->Reset(AddedAgentSet);
+	// Reset Agents, Replay Buffer
+
+	if (bResetAgentsOnBegin)
+	{
+		Manager->ResetAllAgents();
+	}
+
 	ReplayBuffer->Reset();
 
 	bIsTraining = true;
@@ -572,17 +618,88 @@ void ULearningAgentsTrainer::EvaluateRewards()
 		return;
 	}
 
-	SetRewards(AddedAgentIds);
+	// Check agents have actually make observations and taken actions.
 
-	Rewards->Evaluate(AddedAgentSet);
+	ValidAgentIds.Empty(Manager->GetAgentNum());
 
-#if ENABLE_VISUAL_LOG
-	for (const ULearningAgentsReward* RewardObject : RewardObjects)
+	for (const int32 AgentId : Manager->GetAllAgentSet())
 	{
-		if (RewardObject)
+		if (Interactor->GetObservationEncodingAgentIteration()[AgentId] == 0 || 
+			Interactor->GetActionEncodingAgentIteration()[AgentId] == 0)
 		{
-			RewardObject->VisualLog(AddedAgentSet);
+			UE_LOG(LogLearning, Display, TEXT("%s: Agent with id %i has not made observations and taken actions so rewards will not be evaluated for it."), *GetName(), AgentId);
+			continue;
 		}
+
+		ValidAgentIds.Add(AgentId);
+	}
+
+	ValidAgentSet = ValidAgentIds;
+	ValidAgentSet.TryMakeSlice();
+
+	// Set Rewards
+
+	SetRewards(ValidAgentIds);
+
+	// Check that all rewards have had their setter run
+
+	ValidAgentStatus.SetNumUninitialized(Manager->GetMaxAgentNum());
+	ValidAgentStatus.SetRange(0, Manager->GetMaxAgentNum(), true);
+
+	for (ULearningAgentsReward* RewardObject : RewardObjects)
+	{
+		for (const int32 AgentId : ValidAgentSet)
+		{
+			if (RewardObject->AgentIteration[AgentId] == RewardEvaluatedAgentIteration[AgentId])
+			{
+				UE_LOG(LogLearning, Error, TEXT("%s: Reward %s for agent with id %i has not been set (got iteration %i, expected iteration %i) and so agent will not have rewards evaluated."), *GetName(), *RewardObject->GetName(), AgentId, RewardObject->AgentIteration[AgentId], RewardEvaluatedAgentIteration[AgentId] + 1);
+				ValidAgentStatus[AgentId] = false;
+				continue;
+			}
+
+			if (RewardObject->AgentIteration[AgentId] > RewardEvaluatedAgentIteration[AgentId] + 1)
+			{
+				UE_LOG(LogLearning, Error, TEXT("%s: Reward %s for agent with id %i appears to have been set multiple times (got iteration %i, expected iteration %i) and so agent will not have rewards evaluated."), *GetName(), *RewardObject->GetName(), AgentId, RewardObject->AgentIteration[AgentId], RewardEvaluatedAgentIteration[AgentId] + 1);
+				ValidAgentStatus[AgentId] = false;
+				continue;
+			}
+
+			if (RewardObject->AgentIteration[AgentId] != RewardEvaluatedAgentIteration[AgentId] + 1)
+			{
+				UE_LOG(LogLearning, Error, TEXT("%s: Reward %s for agent with id %i does not have a matching iteration number (got iteration %i, expected iteration %i) and so agent will not have rewards evaluated."), *GetName(), *RewardObject->GetName(), AgentId, RewardObject->AgentIteration[AgentId], RewardEvaluatedAgentIteration[AgentId] + 1);
+				ValidAgentStatus[AgentId] = false;
+				continue;
+			}
+		}
+	}
+
+	FinalValidAgentIds.Empty(ValidAgentSet.Num());
+
+	for (const int32 AgentId : ValidAgentSet)
+	{
+		if (ValidAgentStatus[AgentId]) { FinalValidAgentIds.Add(AgentId); }
+	}
+
+	FinalValidAgentSet = FinalValidAgentIds;
+	FinalValidAgentSet.TryMakeSlice();
+
+	// Evaluate Rewards
+
+	Rewards->Evaluate(FinalValidAgentSet);
+
+	// Increment Reward Evaluation Iteration
+
+	for (const int32 AgentId : FinalValidAgentSet)
+	{
+		RewardEvaluatedAgentIteration[AgentId]++;
+	}
+
+	// Visual Logger
+
+#if UE_LEARNING_AGENTS_ENABLE_VISUAL_LOG
+	for (const ULearningAgentsReward* RewardObject : RewardObjects)
+	{
+		RewardObject->VisualLog(FinalValidAgentSet);
 	}
 #endif
 }
@@ -597,23 +714,96 @@ void ULearningAgentsTrainer::EvaluateCompletions()
 		return;
 	}
 
-	SetCompletions(AddedAgentIds);
+	// Check agents have actually make observations and taken actions.
 
-	Completions->Evaluate(AddedAgentSet);
-	
-#if ENABLE_VISUAL_LOG
-	for (const ULearningAgentsCompletion* CompletionObject : CompletionObjects)
+	ValidAgentIds.Empty(Manager->GetAgentNum());
+
+	for (const int32 AgentId : Manager->GetAllAgentSet())
 	{
-		if (CompletionObject)
+		if (Interactor->GetObservationEncodingAgentIteration()[AgentId] == 0 || 
+			Interactor->GetActionEncodingAgentIteration()[AgentId] == 0)
 		{
-			CompletionObject->VisualLog(AddedAgentSet);
+			UE_LOG(LogLearning, Display, TEXT("%s: Agent with id %i has not made observations and taken actions so completions will not be evaluated for it."), *GetName(), AgentId);
+			continue;
 		}
+
+		ValidAgentIds.Add(AgentId);
+	}
+		
+	ValidAgentSet = ValidAgentIds;
+	ValidAgentSet.TryMakeSlice();
+
+	// Set Completions
+
+	SetCompletions(ValidAgentIds);
+
+	// Check that all completions have had their setter run
+
+	ValidAgentStatus.SetNumUninitialized(Manager->GetMaxAgentNum());
+	ValidAgentStatus.SetRange(0, Manager->GetMaxAgentNum(), true);
+
+	for (ULearningAgentsCompletion* CompletionObject : CompletionObjects)
+	{
+		for (const int32 AgentId : ValidAgentSet)
+		{
+			if (CompletionObject->AgentIteration[AgentId] == CompletionEvaluatedAgentIteration[AgentId])
+			{
+				UE_LOG(LogLearning, Error, TEXT("%s: Completion %s for agent with id %i has not been set (got iteration %i, expected iteration %i) and so agent will not have completions evaluated."), *GetName(), *CompletionObject->GetName(), AgentId, CompletionObject->AgentIteration[AgentId], CompletionEvaluatedAgentIteration[AgentId] + 1);
+				ValidAgentStatus[AgentId] = false;
+				continue;
+			}
+
+			if (CompletionObject->AgentIteration[AgentId] > CompletionEvaluatedAgentIteration[AgentId] + 1)
+			{
+				UE_LOG(LogLearning, Error, TEXT("%s: Completion %s for agent with id %i appears to have been set multiple times (got iteration %i, expected iteration %i) and so agent will not have completions evaluated."), *GetName(), *CompletionObject->GetName(), AgentId, CompletionObject->AgentIteration[AgentId], CompletionEvaluatedAgentIteration[AgentId] + 1);
+				ValidAgentStatus[AgentId] = false;
+				continue;
+			}
+
+			if (CompletionObject->AgentIteration[AgentId] != CompletionEvaluatedAgentIteration[AgentId] + 1)
+			{
+				UE_LOG(LogLearning, Error, TEXT("%s: Completion %s for agent with id %i does not have a matching iteration number (got iteration %i, expected iteration %i) and so agent will not have completions evaluated."), *GetName(), *CompletionObject->GetName(), AgentId, CompletionObject->AgentIteration[AgentId], CompletionEvaluatedAgentIteration[AgentId] + 1);
+				ValidAgentStatus[AgentId] = false;
+				continue;
+			}
+		}
+	}
+
+	FinalValidAgentIds.Empty(ValidAgentSet.Num());
+
+	for (const int32 AgentId : ValidAgentSet)
+	{
+		if (ValidAgentStatus[AgentId]) { FinalValidAgentIds.Add(AgentId); }
+	}
+
+	FinalValidAgentSet = FinalValidAgentIds;
+	FinalValidAgentSet.TryMakeSlice();
+
+	// Evaluate Completions
+
+	Completions->Evaluate(FinalValidAgentSet);
+	
+	// Increment Completion Evaluation Iteration
+
+	for (const int32 AgentId : FinalValidAgentSet)
+	{
+		CompletionEvaluatedAgentIteration[AgentId]++;
+	}
+
+	// Visual Logger
+
+#if UE_LEARNING_AGENTS_ENABLE_VISUAL_LOG
+	for (const ULearningAgentsCompletion* CompletionObject : CompletionObjects)
+	{
+		CompletionObject->VisualLog(FinalValidAgentSet);
 	}
 #endif
 }
 
 void ULearningAgentsTrainer::ProcessExperience()
 {
+	UE_LEARNING_TRACE_CPUPROFILER_EVENT_SCOPE(ULearningAgentsTrainer::ProcessExperience);
+
 	if (!IsSetup())
 	{
 		UE_LOG(LogLearning, Error, TEXT("%s: Setup not complete."), *GetName());
@@ -626,6 +816,56 @@ void ULearningAgentsTrainer::ProcessExperience()
 		return;
 	}
 
+	if (Policy->GetActionNoiseScale() != 1.0f)
+	{
+		UE_LOG(LogLearning, Warning, TEXT("%s: Policy Action Noise Scale should be set to 1.0 during training."), *GetName());
+	}
+
+	// Check Observations, Actions, Rewards, and Completions have been completed and have matching iteration number
+
+	ValidAgentIds.Empty(Manager->GetAgentNum());
+
+	for (const int32 AgentId : Manager->GetAllAgentSet())
+	{
+		if (Interactor->GetObservationEncodingAgentIteration()[AgentId] == 0 ||
+			Interactor->GetActionEncodingAgentIteration()[AgentId] == 0 ||
+			RewardEvaluatedAgentIteration[AgentId] == 0 ||
+			CompletionEvaluatedAgentIteration[AgentId] == 0)
+		{
+			UE_LOG(LogLearning, Display, TEXT("%s: Agent with id %i has not completed a full step of observations, actions, rewards, completions and so experience will not be processed for it."), *GetName(), AgentId);
+			continue;
+		}
+
+		if (Interactor->GetObservationEncodingAgentIteration()[AgentId] != Interactor->GetActionEncodingAgentIteration()[AgentId] ||
+			Interactor->GetObservationEncodingAgentIteration()[AgentId] != RewardEvaluatedAgentIteration[AgentId] ||
+			Interactor->GetObservationEncodingAgentIteration()[AgentId] != CompletionEvaluatedAgentIteration[AgentId])
+		{
+			UE_LOG(LogLearning, Warning, TEXT("%s: Agent with id %i has non-matching iteration numbers (observation: %i, action: %i, reward: %i, completion: %i). Experience will not be processed for it."), *GetName(), AgentId,
+				Interactor->GetObservationEncodingAgentIteration()[AgentId],
+				Interactor->GetActionEncodingAgentIteration()[AgentId],
+				RewardEvaluatedAgentIteration[AgentId],
+				CompletionEvaluatedAgentIteration[AgentId]);
+			continue;
+		}
+
+		ValidAgentIds.Add(AgentId);
+	}
+
+	ValidAgentSet = ValidAgentIds;
+
+	// Check for episodes that have been immediately completed
+
+	for (const int32 AgentId : ValidAgentSet)
+	{
+		if (Completions->CompletionBuffer()[AgentId] != UE::Learning::ECompletionMode::Running && 
+			EpisodeBuffer->GetEpisodeStepNums()[AgentId] == 0)
+		{
+			UE_LOG(LogLearning, Warning, TEXT("%s: Agent with id %i has completed episode and will be reset but has not generated any experience."), *GetName(), AgentId);
+		}
+	}
+
+	// Get Feature Buffers
+
 	UE::Learning::FFeatureObject& Observations = Interactor->GetObservationFeature();
 	UE::Learning::FFeatureObject& Actions = Interactor->GetActionFeature();
 
@@ -634,36 +874,25 @@ void ULearningAgentsTrainer::ProcessExperience()
 		Observations.FeatureBuffer(),
 		Actions.FeatureBuffer(),
 		Rewards->RewardBuffer(),
-		AddedAgentSet);
+		ValidAgentSet);
 
 	// Check for completion based on reaching the maximum episode length
 	UE::Learning::Completion::EvaluateEndOfEpisodeCompletions(
 		Completions->CompletionBuffer(),
 		EpisodeBuffer->GetEpisodeStepNums(),
 		EpisodeBuffer->GetMaxStepNum(),
-		MaxStepsCompletion == ELearningAgentsCompletion::Truncation 
-			? UE::Learning::ECompletionMode::Truncated 
-			: UE::Learning::ECompletionMode::Terminated,
-		AddedAgentSet);
+		MaxStepsCompletion == ELearningAgentsCompletion::Truncation
+		? UE::Learning::ECompletionMode::Truncated
+		: UE::Learning::ECompletionMode::Terminated,
+		ValidAgentSet);
 
 	// Find the set of Instances that need to be reset
-	ResetBuffer->SetResetInstancesFromCompletions(Completions->CompletionBuffer(), AddedAgentSet);
+	ResetBuffer->SetResetInstancesFromCompletions(Completions->CompletionBuffer(), ValidAgentSet);
 
 	if (ResetBuffer->GetResetInstanceNum() > 0)
 	{
 		// Encode Observations for completed Instances
-		Interactor->SetObservations(ResetBuffer->GetResetInstances().ToArray());
-		Observations.Encode(ResetBuffer->GetResetInstances());
-		
-#if ENABLE_VISUAL_LOG
-		for (const ULearningAgentsObservation* ObservationObject : Interactor->GetObservationObjects())
-		{
-			if (ObservationObject)
-			{
-				ObservationObject->VisualLog(ResetBuffer->GetResetInstances());
-			}
-		}
-#endif
+		Interactor->EncodeObservations(ResetBuffer->GetResetInstances());
 
 		const bool bReplayBufferFull = ReplayBuffer->AddEpisodes(
 			Completions->CompletionBuffer(),
@@ -686,6 +915,7 @@ void ULearningAgentsTrainer::ProcessExperience()
 
 			// Get Updated Policy
 			Response = Trainer->RecvPolicy(Policy->GetPolicyNetwork(), TrainerTimeout);
+			Policy->GetNetworkAsset()->ForceMarkDirty();
 
 			if (Response == UE::Learning::ETrainerResponse::Completed)
 			{
@@ -704,6 +934,7 @@ void ULearningAgentsTrainer::ProcessExperience()
 			if (Critic)
 			{
 				Response = Trainer->RecvCritic(Critic->GetCriticNetwork(), TrainerTimeout);
+				Critic->GetNetworkAsset()->ForceMarkDirty();
 
 				if (Response != UE::Learning::ETrainerResponse::Success)
 				{
@@ -713,38 +944,22 @@ void ULearningAgentsTrainer::ProcessExperience()
 				}
 			}
 
-			// Mark all instances for reset since we have a new policy
-			ResetBuffer->SetResetInstances(AddedAgentSet);
+			// Mark all agents for reset since we have a new policy
+			ResetBuffer->SetResetInstances(Manager->GetAllAgentSet());
 		}
 
-		ResetEpisodes(ResetBuffer->GetResetInstances().ToArray());
-
-		EpisodeBuffer->Reset(ResetBuffer->GetResetInstances());
+		Manager->ResetAgents(ResetBuffer->GetResetInstances().ToArray());
 	}
 }
 
-void ULearningAgentsTrainer::ResetAllEpisodes()
-{
-	UE_LEARNING_TRACE_CPUPROFILER_EVENT_SCOPE(ULearningAgentsTrainer::ResetAllEpisodes);
-
-	if (!IsSetup())
-	{
-		UE_LOG(LogLearning, Error, TEXT("%s: Setup not complete."), *GetName());
-		return;
-	}
-
-	ResetEpisodes(AddedAgentIds);
-
-	EpisodeBuffer->Reset(AddedAgentSet);
-}
-
 void ULearningAgentsTrainer::RunTraining(
 	const FLearningAgentsTrainerTrainingSettings& TrainerTrainingSettings,
 	const FLearningAgentsTrainerGameSettings& TrainerGameSettings,
 	const FLearningAgentsTrainerPathSettings& TrainerPathSettings,
 	const FLearningAgentsCriticSettings& CriticSettings,
 	const bool bReinitializePolicyNetwork,
-	const bool bReinitializeCriticNetwork)
+	const bool bReinitializeCriticNetwork,
+	const bool bResetAgentsOnBegin)
 {
 	if (!IsSetup())
 	{
@@ -761,7 +976,8 @@ void ULearningAgentsTrainer::RunTraining(
 			TrainerPathSettings, 
 			CriticSettings, 
 			bReinitializePolicyNetwork, 
-			bReinitializeCriticNetwork);
+			bReinitializeCriticNetwork,
+			bResetAgentsOnBegin);
 
 		if (!IsTraining())
 		{
@@ -793,6 +1009,12 @@ float ULearningAgentsTrainer::GetReward(const int32 AgentId) const
 		return 0.0f;
 	}
 
+	if (RewardEvaluatedAgentIteration[AgentId] == 0)
+	{
+		UE_LOG(LogLearning, Error, TEXT("%s: Agent with id %d has not evaluated rewards. Did you run EvaluateRewards?"), *GetName(), AgentId);
+		return 0.0f;
+	}
+
 	return Rewards->RewardBuffer()[AgentId];
 }
 
@@ -812,6 +1034,13 @@ bool ULearningAgentsTrainer::IsCompleted(const int32 AgentId, ELearningAgentsCom
 		return false;
 	}
 
+	if (CompletionEvaluatedAgentIteration[AgentId] == 0)
+	{
+		UE_LOG(LogLearning, Error, TEXT("%s: Agent with id %d has not evaluated completions. Did you run EvaluateCompletions?"), *GetName(), AgentId);
+		OutCompletion = ELearningAgentsCompletion::Termination;
+		return false;
+	}
+
 	const UE::Learning::ECompletionMode CompletionMode = Completions->CompletionBuffer()[AgentId];
 
 	if (CompletionMode == UE::Learning::ECompletionMode::Running)
diff --git a/Engine/Plugins/Experimental/LearningAgents/Source/LearningAgentsTraining/Public/LearningAgentsCompletions.h b/Engine/Plugins/Experimental/LearningAgents/Source/LearningAgentsTraining/Public/LearningAgentsCompletions.h
--- a/Engine/Plugins/Experimental/LearningAgents/Source/LearningAgentsTraining/Public/LearningAgentsCompletions.h
+++ b/Engine/Plugins/Experimental/LearningAgents/Source/LearningAgentsTraining/Public/LearningAgentsCompletions.h
@@ -47,6 +47,12 @@ public:
 	UPROPERTY(VisibleAnywhere, Transient, Category = "LearningAgents")
 	TObjectPtr<ULearningAgentsTrainer> AgentTrainer;
 
+public:
+
+	/** Number of times this completion has been set for all agents */
+	TLearningArray<1, uint64, TInlineAllocator<32>> AgentIteration;
+
+public:
 #if UE_LEARNING_AGENTS_ENABLE_VISUAL_LOG
 	/** Color used to draw this completion in the visual log */
 	FLinearColor VisualLogColor = FColor::Yellow;
@@ -67,20 +73,20 @@ class LEARNINGAGENTSTRAINING_API UConditionalCompletion : public ULearningAgents
 public:
 
 	/**
-	* Adds a new conditional completion to the given trainer. Call during ULearningAgentsTrainer::SetupCompletions event.
-	* @param InAgentTrainer The trainer to add this completion to.
-	* @param Name The name of this new completion. Used for debugging.
-	* @param InCompletionMode The completion mode.
-	* @return The newly created completion.
-	*/
-	UFUNCTION(BlueprintCallable, Category = "LearningAgents")
+	 * Adds a new conditional completion to the given trainer. Call during ULearningAgentsTrainer::SetupCompletions event.
+	 * @param InAgentTrainer The trainer to add this completion to.
+	 * @param Name The name of this new completion. Used for debugging.
+	 * @param InCompletionMode The completion mode.
+	 * @return The newly created completion.
+	 */
+	UFUNCTION(BlueprintCallable, Category = "LearningAgents", meta = (DefaultToSelf = "InAgentTrainer"))
 	static UConditionalCompletion* AddConditionalCompletion(ULearningAgentsTrainer* InAgentTrainer, const FName Name = NAME_None, const ELearningAgentsCompletion InCompletionMode = ELearningAgentsCompletion::Termination);
 
 	/**
-	* Sets the data for this completion. Call during ULearningAgentsTrainer::SetCompletions event.
-	* @param AgentId The agent id this data corresponds to.
-	* @param bIsCompleted Pass in true if condition is met. Otherwise, false.
-	*/
+	 * Sets the data for this completion. Call during ULearningAgentsTrainer::SetCompletions event.
+	 * @param AgentId The agent id this data corresponds to.
+	 * @param bIsCompleted Pass in true if condition is met. Otherwise, false.
+	 */
 	UFUNCTION(BlueprintCallable, Category = "LearningAgents", meta = (AgentId = "-1"))
 	void SetConditionalCompletion(const int32 AgentId, const bool bIsCompleted);
 
@@ -103,21 +109,21 @@ class LEARNINGAGENTSTRAINING_API UTimeElapsedCompletion : public ULearningAgents
 public:
 
 	/**
-	* Adds a new time elapsed completion to the given trainer. Call during ULearningAgentsTrainer::SetupCompletions event.
-	* @param InAgentTrainer The trainer to add this completion to.
-	* @param Name The name of this new completion. Used for debugging.
-	* @param Threshold How much time should be elapsed for the completion to trigger.
-	* @param InCompletionMode The completion mode.
-	* @return The newly created completion.
-	*/
-	UFUNCTION(BlueprintCallable, Category = "LearningAgents")
+	 * Adds a new time elapsed completion to the given trainer. Call during ULearningAgentsTrainer::SetupCompletions event.
+	 * @param InAgentTrainer The trainer to add this completion to.
+	 * @param Name The name of this new completion. Used for debugging.
+	 * @param Threshold How much time should be elapsed for the completion to trigger.
+	 * @param InCompletionMode The completion mode.
+	 * @return The newly created completion.
+	 */
+	UFUNCTION(BlueprintCallable, Category = "LearningAgents", meta = (DefaultToSelf = "InAgentTrainer"))
 	static UTimeElapsedCompletion* AddTimeElapsedCompletion(ULearningAgentsTrainer* InAgentTrainer, const FName Name = NAME_None, const float Threshold = 10.0f, const ELearningAgentsCompletion InCompletionMode = ELearningAgentsCompletion::Termination);
 
 	/**
-	* Sets the data for this completion. Call during ULearningAgentsTrainer::SetCompletions event.
-	* @param AgentId The agent id this data corresponds to.
-	* @param Time The amount of time that has passed
-	*/
+	 * Sets the data for this completion. Call during ULearningAgentsTrainer::SetCompletions event.
+	 * @param AgentId The agent id this data corresponds to.
+	 * @param Time The amount of time that has passed
+	 */
 	UFUNCTION(BlueprintCallable, Category = "LearningAgents", meta = (AgentId = "-1"))
 	void SetTimeElapsedCompletion(const int32 AgentId, const float Time);
 
@@ -132,9 +138,9 @@ public:
 //------------------------------------------------------------------
 
 /**
-* A completion for if two positions differ by some threshold in a plane, e.g. if the agent gets too far from a
-* starting position.
-*/
+ * A completion for if two positions differ by some threshold in a plane, e.g. if the agent gets too far from a
+ * starting position.
+ */
 UCLASS()
 class LEARNINGAGENTSTRAINING_API UPlanarPositionDifferenceCompletion : public ULearningAgentsCompletion
 {
@@ -143,17 +149,17 @@ class LEARNINGAGENTSTRAINING_API UPlanarPositionDifferenceCompletion : public UL
 public:
 
 	/**
-	* Adds a new planar position difference completion to the given trainer. The axis parameters define the plane.
-	* Call during ULearningAgentsTrainer::SetupCompletions event.
-	* @param InAgentTrainer The trainer to add this completion to.
-	* @param Name The name of this new completion. Used for debugging.
-	* @param Threshold If the distance becomes greater than this threshold, then the episode will complete.
-	* @param InCompletionMode The completion mode.
-	* @param Axis0 The forward axis of the plane.
-	* @param Axis1 The right axis of the plane.
-	* @return The newly created completion.
-	*/
-	UFUNCTION(BlueprintCallable, Category = "LearningAgents")
+	 * Adds a new planar position difference completion to the given trainer. The axis parameters define the plane.
+	 * Call during ULearningAgentsTrainer::SetupCompletions event.
+	 * @param InAgentTrainer The trainer to add this completion to.
+	 * @param Name The name of this new completion. Used for debugging.
+	 * @param Threshold If the distance becomes greater than this threshold, then the episode will complete.
+	 * @param InCompletionMode The completion mode.
+	 * @param Axis0 The forward axis of the plane.
+	 * @param Axis1 The right axis of the plane.
+	 * @return The newly created completion.
+	 */
+	UFUNCTION(BlueprintCallable, Category = "LearningAgents", meta = (DefaultToSelf = "InAgentTrainer"))
 	static UPlanarPositionDifferenceCompletion* AddPlanarPositionDifferenceCompletion(
 		ULearningAgentsTrainer* InAgentTrainer,
 		const FName Name = NAME_None,
@@ -163,11 +169,11 @@ public:
 		const FVector Axis1 = FVector::RightVector);
 
 	/**
-	* Sets the data for this completion. Call during ULearningAgentsTrainer::SetCompletions event.
-	* @param AgentId The agent id this data corresponds to.
-	* @param Position0 The first position.
-	* @param Position1 The second position.
-	*/
+	 * Sets the data for this completion. Call during ULearningAgentsTrainer::SetCompletions event.
+	 * @param AgentId The agent id this data corresponds to.
+	 * @param Position0 The first position.
+	 * @param Position1 The second position.
+	 */
 	UFUNCTION(BlueprintCallable, Category = "LearningAgents", meta = (AgentId = "-1"))
 	void SetPlanarPositionDifferenceCompletion(const int32 AgentId, const FVector Position0, const FVector Position1);
 
@@ -180,8 +186,8 @@ public:
 };
 
 /**
-* A completion for if two positions are near by some threshold in a plane, e.g. if the agent gets close to a position.
-*/
+ * A completion for if two positions are near by some threshold in a plane, e.g. if the agent gets close to a position.
+ */
 UCLASS()
 class LEARNINGAGENTSTRAINING_API UPlanarPositionSimilarityCompletion : public ULearningAgentsCompletion
 {
@@ -190,17 +196,17 @@ class LEARNINGAGENTSTRAINING_API UPlanarPositionSimilarityCompletion : public UL
 public:
 
 	/**
-	* Adds a new planar position difference completion to the given trainer. The axis parameters define the plane.
-	* Call during ULearningAgentsTrainer::SetupCompletions event.
-	* @param InAgentTrainer The trainer to add this completion to.
-	* @param Name The name of this new completion. Used for debugging.
-	* @param Threshold If the distance becomes greater than this threshold, then the episode will complete.
-	* @param InCompletionMode The completion mode.
-	* @param Axis0 The forward axis of the plane.
-	* @param Axis1 The right axis of the plane.
-	* @return The newly created completion.
-	*/
-	UFUNCTION(BlueprintCallable, Category = "LearningAgents")
+	 * Adds a new planar position difference completion to the given trainer. The axis parameters define the plane.
+	 * Call during ULearningAgentsTrainer::SetupCompletions event.
+	 * @param InAgentTrainer The trainer to add this completion to.
+	 * @param Name The name of this new completion. Used for debugging.
+	 * @param Threshold If the distance becomes greater than this threshold, then the episode will complete.
+	 * @param InCompletionMode The completion mode.
+	 * @param Axis0 The forward axis of the plane.
+	 * @param Axis1 The right axis of the plane.
+	 * @return The newly created completion.
+	 */
+	UFUNCTION(BlueprintCallable, Category = "LearningAgents", meta = (DefaultToSelf = "InAgentTrainer"))
 	static UPlanarPositionSimilarityCompletion* AddPlanarPositionSimilarityCompletion(
 		ULearningAgentsTrainer* InAgentTrainer,
 		const FName Name = NAME_None,
@@ -210,11 +216,11 @@ public:
 		const FVector Axis1 = FVector::RightVector);
 
 	/**
-	* Sets the data for this completion. Call during ULearningAgentsTrainer::SetCompletions event.
-	* @param AgentId The agent id this data corresponds to.
-	* @param Position0 The first position.
-	* @param Position1 The second position.
-	*/
+	 * Sets the data for this completion. Call during ULearningAgentsTrainer::SetCompletions event.
+	 * @param AgentId The agent id this data corresponds to.
+	 * @param Position0 The first position.
+	 * @param Position1 The second position.
+	 */
 	UFUNCTION(BlueprintCallable, Category = "LearningAgents", meta = (AgentId = "-1"))
 	void SetPlanarPositionSimilarityCompletion(const int32 AgentId, const FVector Position0, const FVector Position1);
 
diff --git a/Engine/Plugins/Experimental/LearningAgents/Source/LearningAgentsTraining/Public/LearningAgentsImitationTrainer.h b/Engine/Plugins/Experimental/LearningAgents/Source/LearningAgentsTraining/Public/LearningAgentsImitationTrainer.h
--- a/Engine/Plugins/Experimental/LearningAgents/Source/LearningAgentsTraining/Public/LearningAgentsImitationTrainer.h
+++ b/Engine/Plugins/Experimental/LearningAgents/Source/LearningAgentsTraining/Public/LearningAgentsImitationTrainer.h
@@ -7,7 +7,6 @@
 
 #include "Templates/SharedPointer.h"
 #include "UObject/ObjectPtr.h"
-#include "Tasks/Task.h"
 
 #include "LearningAgentsImitationTrainer.generated.h"
 
@@ -16,10 +15,22 @@ namespace UE::Learning
 	struct FSharedMemoryImitationTrainer;
 }
 
-class ULearningAgentsInteractor;
 class ULearningAgentsPolicy;
 class ULearningAgentsRecording;
 
+/** The configurable settings for a ULearningAgentsImitationTrainer. */
+USTRUCT(BlueprintType, Category = "LearningAgents")
+struct FLearningAgentsImitationTrainerSettings
+{
+	GENERATED_BODY()
+
+public:
+
+	/** Time in seconds to wait for the training process before timing out. */
+	UPROPERTY(EditAnywhere, Category = "LearningAgents", meta = (ClampMin = "0.0", UIMin = "0.0"))
+	float TrainerCommunicationTimeout = 20.0f;
+};
+
 /** The configurable settings for the training process. */
 USTRUCT(BlueprintType, Category = "LearningAgents")
 struct FLearningAgentsImitationTrainerTrainingSettings
@@ -41,16 +52,16 @@ public:
 	float LearningRateDecay = 0.99f;
 
 	/**
-	* Amount of weight decay to apply to the network. Larger values encourage network weights to be smaller but too
-	* large a value can cause the network weights to collapse to all zeros.
-	*/
+	 * Amount of weight decay to apply to the network. Larger values encourage network weights to be smaller but too
+	 * large a value can cause the network weights to collapse to all zeros.
+	 */
 	UPROPERTY(EditAnywhere, Category = "LearningAgents", meta = (ClampMin = "0.0", ClampMax = "1.0", UIMin = "0.0", UIMax = "1.0"))
 	float WeightDecay = 0.001f;
 
 	/**
-	* Batch size to use for training. Smaller values tend to produce better results at the cost of slowing down
-	* training. Large batch sizes are much more computationally efficient when training on the GPU.
-	*/
+	 * Batch size to use for training. Smaller values tend to produce better results at the cost of slowing down
+	 * training. Large batch sizes are much more computationally efficient when training on the GPU.
+	 */
 	UPROPERTY(EditAnywhere, Category = "LearningAgents", meta = (ClampMin = "0", ClampMax = "4096", UIMin = "0", UIMax = "4096"))
 	uint32 BatchSize = 128;
 
@@ -68,14 +79,14 @@ public:
 };
 
 /**
-* The ULearningAgentsImitationTrainer enable imitation learning, i.e. learning from human/AI demonstrations.
-* Imitation training is typically much faster than reinforcement learning, but requires gathering large amounts of
-* data in order to generalize. This can be used to initialize a reinforcement learning policy to speed up initial
-* exploration.
-* @see ULearningAgentsInteractor to understand how observations and actions work.
-* @see ULearningAgentsRecorder to understand how to make new recordings.
-* @see ULearningAgentsDataStorage to understand how to retrieve previous recordings.
-*/
+ * The ULearningAgentsImitationTrainer enable imitation learning, i.e. learning from human/AI demonstrations.
+ * Imitation training is typically much faster than reinforcement learning, but requires gathering large amounts of
+ * data in order to generalize. This can be used to initialize a reinforcement learning policy to speed up initial
+ * exploration.
+ * @see ULearningAgentsInteractor to understand how observations and actions work.
+ * @see ULearningAgentsController to understand how we can manually perform actions via a human or AI.
+ * @see ULearningAgentsRecorder to understand how to make new recordings.
+ */
 UCLASS(BlueprintType, Blueprintable)
 class LEARNINGAGENTSTRAINING_API ULearningAgentsImitationTrainer : public UActorComponent
 {
@@ -85,25 +96,31 @@ class LEARNINGAGENTSTRAINING_API ULearningAgentsImitationTrainer : public UActor
 public:
 
 	// These constructors/destructors are needed to make forward declarations happy
-	ULearningAgentsImitationTrainer();
+	ULearningAgentsImitationTrainer(const FObjectInitializer& ObjectInitializer = FObjectInitializer::Get());
 	ULearningAgentsImitationTrainer(FVTableHelper& Helper);
 	virtual ~ULearningAgentsImitationTrainer();
 
 	/** Will automatically call EndTraining if training is still in-progress when play is ending. */
 	virtual void EndPlay(const EEndPlayReason::Type EndPlayReason) override;
 
+	/** Returns true if the trainer is currently training; Otherwise, false. */
+	UFUNCTION(BlueprintPure, Category = "LearningAgents")
+	bool IsTraining() const;
+
 	/**
-	* Begins the training process with the provided settings.
-	* @param InPolicy The policy to train.
-	* @param Recording The data to train on.
-	* @param ImitationTrainerTrainingSettings The settings for this training run.
-	* @param ImitationTrainerPathSettings The path settings used by the imitation trainer.
-	* @param bReinitializePolicyNetwork If true, reinitialize the policy. Set this to false if your policy is pre-trained.
-	*/
+	 * Begins the training process with the provided settings.
+	 * @param InPolicy The policy to train.
+	 * @param Recording The data to train on.
+	 * @param ImitationTrainerSettings The settings for this trainer.
+	 * @param ImitationTrainerTrainingSettings The training settings for this network.
+	 * @param ImitationTrainerPathSettings The path settings used by the imitation trainer.
+	 * @param bReinitializePolicyNetwork If true, reinitialize the policy. Set this to false if your policy is pre-trained.
+	 */
 	UFUNCTION(BlueprintCallable, Category = "LearningAgents")
 	void BeginTraining(
 		ULearningAgentsPolicy* InPolicy,
 		const ULearningAgentsRecording* Recording,
+		const FLearningAgentsImitationTrainerSettings& ImitationTrainerSettings = FLearningAgentsImitationTrainerSettings(),
 		const FLearningAgentsImitationTrainerTrainingSettings& ImitationTrainerTrainingSettings = FLearningAgentsImitationTrainerTrainingSettings(),
 		const FLearningAgentsTrainerPathSettings& ImitationTrainerPathSettings = FLearningAgentsTrainerPathSettings(),
 		const bool bReinitializePolicyNetwork = true);
@@ -112,13 +129,28 @@ public:
 	UFUNCTION(BlueprintCallable, Category = "LearningAgents")
 	void EndTraining();
 
-	/** Returns true if the trainer is currently training; Otherwise, false. */
-	UFUNCTION(BlueprintPure, Category = "LearningAgents")
-	bool IsTraining() const;
+	/** Iterates the training process and gets the updated policy network. */
+	UFUNCTION(BlueprintCallable, Category = "LearningAgents")
+	void IterateTraining();
 
-	/** Returns true if the previously launched training has completed; Otherwise, false. */
-	UFUNCTION(BlueprintPure, Category = "LearningAgents")
-	bool IsTrainingComplete() const;
+	/**
+	 * Convenience function that runs a basic training loop. If training has not been started, it will start it. On 
+	 * each following call to this function, it will call IterateTraining.
+	 * @param InPolicy The policy to train.
+	 * @param Recording The data to train on.
+	 * @param ImitationTrainerSettings The settings for this trainer.
+	 * @param ImitationTrainerTrainingSettings The training settings for this network.
+	 * @param ImitationTrainerPathSettings The path settings used by the imitation trainer.
+	 * @param bReinitializePolicyNetwork If true, reinitialize the policy. Set this to false if your policy is pre-trained.
+	 */
+	UFUNCTION(BlueprintCallable, Category = "LearningAgents")
+	void RunTraining(
+		ULearningAgentsPolicy* InPolicy,
+		const ULearningAgentsRecording* Recording,
+		const FLearningAgentsImitationTrainerSettings& ImitationTrainerSettings = FLearningAgentsImitationTrainerSettings(),
+		const FLearningAgentsImitationTrainerTrainingSettings& ImitationTrainerTrainingSettings = FLearningAgentsImitationTrainerTrainingSettings(),
+		const FLearningAgentsTrainerPathSettings& ImitationTrainerPathSettings = FLearningAgentsTrainerPathSettings(),
+		const bool bReinitializePolicyNetwork = true);
 
 // ----- Private Data -----
 private:
@@ -131,16 +163,12 @@ private:
 	UPROPERTY(VisibleAnywhere, Transient, Category = "LearningAgents")
 	bool bIsTraining = false;
 
-	/** True if training is completed. Otherwise, false. */
-	UPROPERTY(VisibleAnywhere, Transient, Category = "LearningAgents")
-	bool bIsTrainingComplete = false;
+	float TrainerTimeout = 10.0f;
+
+	void DoneTraining();
 
 	TLearningArray<2, float> RecordedObservations;
 	TLearningArray<2, float> RecordedActions;
 
 	TUniquePtr<UE::Learning::FSharedMemoryImitationTrainer> ImitationTrainer;
-	UE::Tasks::FTask ImitationTrainingTask;
-	FRWLock NetworkLock;
-
-	TAtomic<bool> bRequestImitationTrainingStop = false;
 };
diff --git a/Engine/Plugins/Experimental/LearningAgents/Source/LearningAgentsTraining/Public/LearningAgentsRecorder.h b/Engine/Plugins/Experimental/LearningAgents/Source/LearningAgentsTraining/Public/LearningAgentsRecorder.h
--- a/Engine/Plugins/Experimental/LearningAgents/Source/LearningAgentsTraining/Public/LearningAgentsRecorder.h
+++ b/Engine/Plugins/Experimental/LearningAgents/Source/LearningAgentsTraining/Public/LearningAgentsRecorder.h
@@ -42,7 +42,7 @@ class LEARNINGAGENTSTRAINING_API ULearningAgentsRecorder : public ULearningAgent
 public:
 
 	// These constructors/destructors are needed to make forward declarations happy
-	ULearningAgentsRecorder();
+	ULearningAgentsRecorder(const FObjectInitializer& ObjectInitializer = FObjectInitializer::Get());
 	ULearningAgentsRecorder(FVTableHelper& Helper);
 	virtual ~ULearningAgentsRecorder();
 
@@ -50,36 +50,38 @@ public:
 	virtual void EndPlay(const EEndPlayReason::Type EndPlayReason) override;
 
 	/**
-	* Initializes this object and runs the setup functions for the underlying data storage.
-	* @param InAgentManager The agent manager we are associated with.
-	* @param InInteractor The agent interactor we are recording with.
-	* @param RecorderPathSettings The path settings used by the recorder.
-	*/
+	 * Initializes this object and runs the setup functions for the underlying data storage.
+	 * @param InInteractor The agent interactor we are recording with.
+	 * @param RecorderPathSettings The path settings used by the recorder.
+	 * @param RecordingAsset Optional recording asset to use. If not provided then a new recording object will be 
+	 * created.
+	 */
 	UFUNCTION(BlueprintCallable, Category = "LearningAgents")
 	void SetupRecorder(
-		ALearningAgentsManager* InAgentManager, 
 		ULearningAgentsInteractor* InInteractor,
-		const FLearningAgentsRecorderPathSettings& RecorderPathSettings = FLearningAgentsRecorderPathSettings());
+		const FLearningAgentsRecorderPathSettings& RecorderPathSettings = FLearningAgentsRecorderPathSettings(),
+		ULearningAgentsRecording* RecordingAsset = nullptr);
 
 public:
 
 	//~ Begin ULearningAgentsManagerComponent Interface
-	virtual bool RemoveAgent(const int32 AgentId) override;
+	virtual void OnAgentsRemoved(const TArray<int32>& AgentIds) override;
+	virtual void OnAgentsReset(const TArray<int32>& AgentIds) override;
 	//~ End ULearningAgentsManagerComponent Interface
 
 // ----- Recording Process -----
 public:
 
 	/**
-	* Begins the recording of the observations and actions of each added agent.
-	* @param bReinitializeRecording If to reinitialize the current recording object.
-	*/
+	 * Begins the recording of the observations and actions of each added agent.
+	 * @param bReinitializeRecording If to clear all records from the recording object in use at the start of recording.
+	 */
 	UFUNCTION(BlueprintCallable, Category = "LearningAgents")
 	void BeginRecording(bool bReinitializeRecording = true);
 
 	/**
-	* Ends the recording of the observations and actions of each agent and stores them in the current recording object.
-	*/
+	 * Ends the recording of the observations and actions of each agent and stores them in the current recording object.
+	 */
 	UFUNCTION(BlueprintCallable, Category = "LearningAgents")
 	void EndRecording();
 
@@ -88,10 +90,10 @@ public:
 	bool IsRecording() const;
 
 	/**
-	* While recording, adds the current observations and actions of the added agents to the internal buffer. Call this 
-	* after ULearningAgentsInteractor::EncodeObservations and either ULearningAgentsController::EncodeActions (if 
-	* recording a human/AI demonstration) or ULearningAgentsInteractor::DecodeActions (if recording another policy).
-	*/
+	 * While recording, adds the current observations and actions of the added agents to the internal buffer. Call this 
+	 * after ULearningAgentsInteractor::EncodeObservations and either ULearningAgentsController::EncodeActions (if 
+	 * recording a human/AI demonstration) or ULearningAgentsInteractor::DecodeActions (if recording another policy).
+	 */
 	UFUNCTION(BlueprintCallable, Category = "LearningAgents")
 	void AddExperience();
 
@@ -110,17 +112,25 @@ public:
 	UFUNCTION(BlueprintCallable, Category = "LearningAgents", meta = (RelativePath))
 	void SaveRecordingToFile(const FFilePath& File) const;
 
+	/** Append to the current recording object from a file. */
+	UFUNCTION(BlueprintCallable, Category = "LearningAgents", meta = (RelativePath))
+	void AppendRecordingFromFile(const FFilePath& File);
+
+	/** Uses the given recording asset. New recordings will be appended to this. */
+	UFUNCTION(BlueprintCallable, Category = "LearningAgents")
+	void UseRecordingAsset(ULearningAgentsRecording* RecordingAsset);
+
 	/** Loads the current recording object from the given recording asset */
 	UFUNCTION(BlueprintCallable, Category = "LearningAgents")
-	void LoadRecordingFromAsset(const ULearningAgentsRecording* Asset);
+	void LoadRecordingFromAsset(ULearningAgentsRecording* RecordingAsset);
 
 	/** Saves the current recording object to the given recording asset */
 	UFUNCTION(BlueprintCallable, Category = "LearningAgents", meta = (DevelopmentOnly))
-	void SaveRecordingToAsset(ULearningAgentsRecording* Asset) const;
+	void SaveRecordingToAsset(ULearningAgentsRecording* RecordingAsset);
 
 	/** Appends the current recording object to the given recording asset */
 	UFUNCTION(BlueprintCallable, Category = "LearningAgents", meta = (DevelopmentOnly))
-	void AppendRecordingToAsset(ULearningAgentsRecording* Asset) const;
+	void AppendRecordingToAsset(ULearningAgentsRecording* RecordingAsset);
 
 
 // ----- Private Data ----- 
@@ -141,6 +151,7 @@ private:
 // ----- Private Data ----- 
 private:
 
+	/** Directory to store recordings in */
 	FString RecordingDirectory;
 
 	/** Basic structure used to buffer the observations and actions of each agent. */
@@ -167,5 +178,6 @@ private:
 		void CopyToRecord(FLearningAgentsRecord& Record) const;
 	};
 
+	/** Array of recording buffers for each agent */
 	TArray<FAgentRecordBuffer, TInlineAllocator<32>> RecordBuffers;
 };
diff --git a/Engine/Plugins/Experimental/LearningAgents/Source/LearningAgentsTraining/Public/LearningAgentsRecording.h b/Engine/Plugins/Experimental/LearningAgents/Source/LearningAgentsTraining/Public/LearningAgentsRecording.h
--- a/Engine/Plugins/Experimental/LearningAgents/Source/LearningAgentsTraining/Public/LearningAgentsRecording.h
+++ b/Engine/Plugins/Experimental/LearningAgents/Source/LearningAgentsTraining/Public/LearningAgentsRecording.h
@@ -67,6 +67,27 @@ public:
 	UFUNCTION(BlueprintCallable, Category = "LearningAgents", meta = (RelativePath))
 	void SaveRecordingToFile(const FFilePath& File) const;
 
+	/** Append to this recording from a file. */
+	UFUNCTION(BlueprintCallable, Category = "LearningAgents", meta = (RelativePath))
+	void AppendRecordingFromFile(const FFilePath& File);
+
+	/** Loads this recording from the given recording asset */
+	UFUNCTION(BlueprintCallable, Category = "LearningAgents")
+	void LoadRecordingFromAsset(ULearningAgentsRecording* RecordingAsset);
+
+	/** Saves this recording to the given recording asset */
+	UFUNCTION(BlueprintCallable, Category = "LearningAgents", meta = (DevelopmentOnly))
+	void SaveRecordingToAsset(ULearningAgentsRecording* RecordingAsset);
+
+	/** Appends this recording to the given recording asset */
+	UFUNCTION(BlueprintCallable, Category = "LearningAgents", meta = (DevelopmentOnly))
+	void AppendRecordingToAsset(ULearningAgentsRecording* RecordingAsset);
+
+public:
+
+	/** Marks this asset as modified even during PIE */
+	void ForceMarkDirty();
+
 public:
 
 	/** Set of records. */
diff --git a/Engine/Plugins/Experimental/LearningAgents/Source/LearningAgentsTraining/Public/LearningAgentsRewards.h b/Engine/Plugins/Experimental/LearningAgents/Source/LearningAgentsTraining/Public/LearningAgentsRewards.h
--- a/Engine/Plugins/Experimental/LearningAgents/Source/LearningAgentsTraining/Public/LearningAgentsRewards.h
+++ b/Engine/Plugins/Experimental/LearningAgents/Source/LearningAgentsTraining/Public/LearningAgentsRewards.h
@@ -43,6 +43,12 @@ public:
 	UPROPERTY(VisibleAnywhere, Transient, Category = "LearningAgents")
 	TObjectPtr<ULearningAgentsTrainer> AgentTrainer;
 
+public:
+
+	/** Number of times this reward has been set for all agents */
+	TLearningArray<1, uint64, TInlineAllocator<32>> AgentIteration;
+
+public:
 #if UE_LEARNING_AGENTS_ENABLE_VISUAL_LOG
 	/** Color used to draw this reward in the visual log */
 	FLinearColor VisualLogColor = FColor::Green;
@@ -63,23 +69,23 @@ class LEARNINGAGENTSTRAINING_API UFloatReward : public ULearningAgentsReward
 public:
 
 	/**
-	* Adds a new float reward to the given trainer. Call during ULearningAgentsTrainer::SetupRewards event.
-	* @param InAgentTrainer The trainer to add this reward to.
-	* @param Name The name of this new reward. Used for debugging.
-	* @param Weight Multiplier for this reward when being summed up for the total reward.
-	* @return The newly created reward.
-	*/
-	UFUNCTION(BlueprintCallable, Category = "LearningAgents")
+	 * Adds a new float reward to the given trainer. Call during ULearningAgentsTrainer::SetupRewards event.
+	 * @param InAgentTrainer The trainer to add this reward to.
+	 * @param Name The name of this new reward. Used for debugging.
+	 * @param Weight Multiplier for this reward when being summed up for the total reward.
+	 * @return The newly created reward.
+	 */
+	UFUNCTION(BlueprintCallable, Category = "LearningAgents", meta = (DefaultToSelf = "InAgentTrainer"))
 	static UFloatReward* AddFloatReward(
 		ULearningAgentsTrainer* InAgentTrainer,
 		const FName Name = NAME_None,
 		const float Weight = 1.0f);
 
 	/**
-	* Sets the data for this reward. Call during ULearningAgentsTrainer::SetRewards event.
-	* @param AgentId The agent id this data corresponds to.
-	* @param Reward The value currently being rewarded.
-	*/
+	 * Sets the data for this reward. Call during ULearningAgentsTrainer::SetRewards event.
+	 * @param AgentId The agent id this data corresponds to.
+	 * @param Reward The value currently being rewarded.
+	 */
 	UFUNCTION(BlueprintCallable, Category = "LearningAgents", meta = (AgentId = "-1"))
 	void SetFloatReward(const int32 AgentId, const float Reward);
 
@@ -102,14 +108,14 @@ class LEARNINGAGENTSTRAINING_API UScalarVelocityReward : public ULearningAgentsR
 public:
 
 	/**
-	* Adds a new scalar velocity reward to the given trainer. Call during ULearningAgentsTrainer::SetupRewards event.
-	* @param InAgentTrainer The trainer to add this reward to.
-	* @param Name The name of this new reward. Used for debugging.
-	* @param Weight Multiplier for this reward when being summed up for the total reward.
-	* @param Scale Used to normalize the data for the reward.
-	* @return The newly created reward.
-	*/
-	UFUNCTION(BlueprintCallable, Category = "LearningAgents")
+	 * Adds a new scalar velocity reward to the given trainer. Call during ULearningAgentsTrainer::SetupRewards event.
+	 * @param InAgentTrainer The trainer to add this reward to.
+	 * @param Name The name of this new reward. Used for debugging.
+	 * @param Weight Multiplier for this reward when being summed up for the total reward.
+	 * @param Scale Used to normalize the data for the reward.
+	 * @return The newly created reward.
+	 */
+	UFUNCTION(BlueprintCallable, Category = "LearningAgents", meta = (DefaultToSelf = "InAgentTrainer"))
 	static UScalarVelocityReward* AddScalarVelocityReward(
 		ULearningAgentsTrainer* InAgentTrainer,
 		const FName Name = NAME_None,
@@ -117,10 +123,10 @@ public:
 		const float Scale = 200.0f);
 
 	/**
-	* Sets the data for this reward. Call during ULearningAgentsTrainer::SetRewards event.
-	* @param AgentId The agent id this data corresponds to.
-	* @param Velocity The current scalar velocity.
-	*/
+	 * Sets the data for this reward. Call during ULearningAgentsTrainer::SetRewards event.
+	 * @param AgentId The agent id this data corresponds to.
+	 * @param Velocity The current scalar velocity.
+	 */
 	UFUNCTION(BlueprintCallable, Category = "LearningAgents", meta = (AgentId = "-1"))
 	void SetScalarVelocityReward(const int32 AgentId, const float Velocity);
 
@@ -143,15 +149,15 @@ class LEARNINGAGENTSTRAINING_API ULocalDirectionalVelocityReward : public ULearn
 public:
 
 	/**
-	* Adds a new directional velocity reward to the given trainer. Call during ULearningAgentsTrainer::SetupRewards event.
-	* @param InAgentTrainer The trainer to add this reward to.
-	* @param Name The name of this new reward. Used for debugging.
-	* @param Weight Multiplier for this reward when being summed up for the total reward.
-	* @param Scale Used to normalize the data for the reward.
-	* @param Axis The local direction we want to maximize velocity in.
-	* @return The newly created reward.
-	*/
-	UFUNCTION(BlueprintCallable, Category = "LearningAgents")
+	 * Adds a new directional velocity reward to the given trainer. Call during ULearningAgentsTrainer::SetupRewards event.
+	 * @param InAgentTrainer The trainer to add this reward to.
+	 * @param Name The name of this new reward. Used for debugging.
+	 * @param Weight Multiplier for this reward when being summed up for the total reward.
+	 * @param Scale Used to normalize the data for the reward.
+	 * @param Axis The local direction we want to maximize velocity in.
+	 * @return The newly created reward.
+	 */
+	UFUNCTION(BlueprintCallable, Category = "LearningAgents", meta = (DefaultToSelf = "InAgentTrainer"))
 	static ULocalDirectionalVelocityReward* AddLocalDirectionalVelocityReward(
 		ULearningAgentsTrainer* InAgentTrainer,
 		const FName Name = NAME_None,
@@ -160,11 +166,11 @@ public:
 		const FVector Axis = FVector::ForwardVector);
 
 	/**
-	* Sets the data for this reward. Call during ULearningAgentsTrainer::SetRewards event.
-	* @param AgentId The agent id this data corresponds to.
-	* @param Velocity The current velocity.
-	* @param RelativeRotation The frame of reference rotation.
-	*/
+	 * Sets the data for this reward. Call during ULearningAgentsTrainer::SetRewards event.
+	 * @param AgentId The agent id this data corresponds to.
+	 * @param Velocity The current velocity.
+	 * @param RelativeRotation The frame of reference rotation.
+	 */
 	UFUNCTION(BlueprintCallable, Category = "LearningAgents", meta = (AgentId = "-1"))
 	void SetLocalDirectionalVelocityReward(
 		const int32 AgentId,
@@ -190,18 +196,18 @@ class LEARNINGAGENTSTRAINING_API UPlanarPositionDifferencePenalty : public ULear
 public:
 
 	/**
-	* Adds a new planar difference penalty to the given trainer. The axis parameters define the plane.
-	* Call during ULearningAgentsTrainer::SetupRewards event.
-	* @param InAgentTrainer The trainer to add this penalty to.
-	* @param Name The name of this new penalty. Used for debugging.
-	* @param Weight Multiplier for this penalty when being summed up for the total reward.
-	* @param Scale Used to normalize the data for the penalty.
-	* @param Threshold Minimal distance to apply this penalty.
-	* @param Axis0 The forward axis of the plane.
-	* @param Axis1 The right axis of the plane.
-	* @return The newly created reward.
-	*/
-	UFUNCTION(BlueprintCallable, Category = "LearningAgents")
+	 * Adds a new planar difference penalty to the given trainer. The axis parameters define the plane.
+	 * Call during ULearningAgentsTrainer::SetupRewards event.
+	 * @param InAgentTrainer The trainer to add this penalty to.
+	 * @param Name The name of this new penalty. Used for debugging.
+	 * @param Weight Multiplier for this penalty when being summed up for the total reward.
+	 * @param Scale Used to normalize the data for the penalty.
+	 * @param Threshold Minimal distance to apply this penalty.
+	 * @param Axis0 The forward axis of the plane.
+	 * @param Axis1 The right axis of the plane.
+	 * @return The newly created reward.
+	 */
+	UFUNCTION(BlueprintCallable, Category = "LearningAgents", meta = (DefaultToSelf = "InAgentTrainer"))
 	static UPlanarPositionDifferencePenalty* AddPlanarPositionDifferencePenalty(
 		ULearningAgentsTrainer* InAgentTrainer,
 		const FName Name = NAME_None,
@@ -212,11 +218,11 @@ public:
 		const FVector Axis1 = FVector::RightVector);
 
 	/**
-	* Sets the data for this penalty. Call during ULearningAgentsTrainer::SetRewards event.
-	* @param AgentId The agent id this data corresponds to.
-	* @param Position0 The current position.
-	* @param Position1 The goal position.
-	*/
+	 * Sets the data for this penalty. Call during ULearningAgentsTrainer::SetRewards event.
+	 * @param AgentId The agent id this data corresponds to.
+	 * @param Position0 The current position.
+	 * @param Position1 The goal position.
+	 */
 	UFUNCTION(BlueprintCallable, Category = "LearningAgents", meta = (AgentId = "-1"))
 	void SetPlanarPositionDifferencePenalty(const int32 AgentId, const FVector Position0, const FVector Position1);
 
@@ -237,15 +243,15 @@ class LEARNINGAGENTSTRAINING_API UPositionArraySimilarityReward : public ULearni
 public:
 
 	/**
-	* Adds a new position array similarity reward to the given trainer. Call during ULearningAgentsTrainer::SetupRewards event.
-	* @param InAgentTrainer The trainer to add this reward to.
-	* @param Name The name of this new reward. Used for debugging.
-	* @param PositionNum The number of positions in the array.
-	* @param Scale Used to normalize the data for the reward.
-	* @param Weight Multiplier for this reward when being summed up for the total reward.
-	* @return The newly created reward.
-	*/
-	UFUNCTION(BlueprintCallable, Category = "LearningAgents")
+	 * Adds a new position array similarity reward to the given trainer. Call during ULearningAgentsTrainer::SetupRewards event.
+	 * @param InAgentTrainer The trainer to add this reward to.
+	 * @param Name The name of this new reward. Used for debugging.
+	 * @param PositionNum The number of positions in the array.
+	 * @param Scale Used to normalize the data for the reward.
+	 * @param Weight Multiplier for this reward when being summed up for the total reward.
+	 * @return The newly created reward.
+	 */
+	UFUNCTION(BlueprintCallable, Category = "LearningAgents", meta = (DefaultToSelf = "InAgentTrainer"))
 	static UPositionArraySimilarityReward* AddPositionArraySimilarityReward(
 		ULearningAgentsTrainer* InAgentTrainer,
 		const FName Name = NAME_None,
@@ -254,15 +260,15 @@ public:
 		const float Weight = 1.0f);
 
 	/**
-	* Sets the data for this reward. Call during ULearningAgentsTrainer::SetRewards event.
-	* @param AgentId The agent id this data corresponds to.
-	* @param Positions0 The current positions.
-	* @param Positions1 The goal positions.
-	* @param RelativePosition0 The vector Positions0 will be offset from.
-	* @param RelativePosition1 The vector Positions1 will be offset from.
-	* @param RelativeRotation0 The frame of reference rotation for Positions0.
-	* @param RelativeRotation1 The frame of reference rotation for Positions1.
-	*/
+	 * Sets the data for this reward. Call during ULearningAgentsTrainer::SetRewards event.
+	 * @param AgentId The agent id this data corresponds to.
+	 * @param Positions0 The current positions.
+	 * @param Positions1 The goal positions.
+	 * @param RelativePosition0 The vector Positions0 will be offset from.
+	 * @param RelativePosition1 The vector Positions1 will be offset from.
+	 * @param RelativeRotation0 The frame of reference rotation for Positions0.
+	 * @param RelativeRotation1 The frame of reference rotation for Positions1.
+	 */
 	UFUNCTION(BlueprintCallable, Category = "LearningAgents", meta = (AgentId = "-1"))
 	void SetPositionArraySimilarityReward(
 		const int32 AgentId,
