LearningAgents: Added iteration checks for verifying order of operations. Fixed issue with inference for multiple agents. Fixed live recompilation issues.

diff --git a/Engine/Plugins/Experimental/LearningAgents/Source/LearningAgentsTraining/Public/LearningAgentsTrainer.h b/Engine/Plugins/Experimental/LearningAgents/Source/LearningAgentsTraining/Public/LearningAgentsTrainer.h
--- a/Engine/Plugins/Experimental/LearningAgents/Source/LearningAgentsTraining/Public/LearningAgentsTrainer.h
+++ b/Engine/Plugins/Experimental/LearningAgents/Source/LearningAgentsTraining/Public/LearningAgentsTrainer.h
@@ -83,9 +83,9 @@ public:
 };
 
 /**
-* The configurable game settings for a ULearningAgentsTrainer. These allow the timestep and physics tick to be fixed
-* during training, which can enable ticking faster than real-time.
-*/
+ * The configurable game settings for a ULearningAgentsTrainer. These allow the timestep and physics tick to be fixed
+ * during training, which can enable ticking faster than real-time.
+ */
 USTRUCT(BlueprintType, Category="LearningAgents")
 struct FLearningAgentsTrainerGameSettings
 {
@@ -94,18 +94,18 @@ struct FLearningAgentsTrainerGameSettings
 public:
 
 	/**
-	* If true, the game will run in fixed time step mode (i.e the frame's delta times will always be the same
-	* regardless of how much wall time has passed). This can enable faster than real-time training if your game runs
-	* quickly. If false, the time steps will match real wall time.
-	*/
+	 * If true, the game will run in fixed time step mode (i.e the frame's delta times will always be the same
+	 * regardless of how much wall time has passed). This can enable faster than real-time training if your game runs
+	 * quickly. If false, the time steps will match real wall time.
+	 */
 	UPROPERTY(EditAnywhere, Category = "LearningAgents")
 	bool bUseFixedTimeStep = true;
 
 	/**
-	* Determines the amount of time for each frame when bUseFixedTimeStep is true; Ignored if false. You want this
-	* time step to match as closely as possible to the expected inference time steps, otherwise your training results
-	* may not generalize to your game.
-	*/
+	 * Determines the amount of time for each frame when bUseFixedTimeStep is true; Ignored if false. You want this
+	 * time step to match as closely as possible to the expected inference time steps, otherwise your training results
+	 * may not generalize to your game.
+	 */
 	UPROPERTY(EditAnywhere, Category = "LearningAgents", meta = (DisplayName = "Fixed Time Step Frequency (Hz)"), meta = (ClampMin = "0.0", UIMin = "0.0"))
 	float FixedTimeStepFrequency = 60.0f;
 
@@ -156,9 +156,9 @@ public:
 	float LearningRatePolicy = 0.0001f;
 
 	/**
-	* Learning rate of the critic network. To avoid instability generally the critic should have a larger learning 
-	* rate than the policy. Typically this can be set to 10x the rate of the policy.
-	*/
+	 * Learning rate of the critic network. To avoid instability generally the critic should have a larger learning 
+	 * rate than the policy. Typically this can be set to 10x the rate of the policy.
+	 */
 	UPROPERTY(EditAnywhere, Category = "LearningAgents", meta = (ClampMin = "0.0", ClampMax = "1.0", UIMin = "0.0", UIMax = "1.0"))
 	float LearningRateCritic = 0.001f;
 
@@ -167,59 +167,59 @@ public:
 	float LearningRateDecay = 0.99f;
 
 	/**
-	* Amount of weight decay to apply to the network. Larger values encourage network weights to be smaller but too 
-	* large a value can cause the network weights to collapse to all zeros.
-	*/
+	 * Amount of weight decay to apply to the network. Larger values encourage network weights to be smaller but too 
+	 * large a value can cause the network weights to collapse to all zeros.
+	 */
 	UPROPERTY(EditAnywhere, Category = "LearningAgents", meta = (ClampMin = "0.0", ClampMax = "1.0", UIMin = "0.0", UIMax = "1.0"))
 	float WeightDecay = 0.001f;
 
 	/**
-	* The initial scaling for the weights of the output layer of the neural network. Typically, you would use this to
-	* scale down the initial actions as it can stabilize the training and speed up convergence.
-	*/
+	 * The initial scaling for the weights of the output layer of the neural network. Typically, you would use this to
+	 * scale down the initial actions as it can stabilize the training and speed up convergence.
+	 */
 	UPROPERTY(EditAnywhere, Category = "LearningAgents", meta = (ClampMin = "0.0", ClampMax = "10.0", UIMin = "0.0", UIMax = "10.0"))
 	float InitialActionScale = 0.1f;
 
 	/**
-	* Batch size to use for training. Smaller values tend to produce better results at the cost of slowing down 
-	* training. Large batch sizes are much more computationally efficient when training on the GPU.
-	*/
+	 * Batch size to use for training. Smaller values tend to produce better results at the cost of slowing down 
+	 * training. Large batch sizes are much more computationally efficient when training on the GPU.
+	 */
 	UPROPERTY(EditAnywhere, Category = "LearningAgents", meta = (ClampMin = "0", ClampMax = "4096", UIMin = "0", UIMax = "4096"))
 	int32 BatchSize = 128;
 
 	/**
-	* Clipping ratio to apply to policy updates. Keeps the training "on-policy". Larger values may speed up training at 
-	* the cost of stability. Conversely, too small values will keep the policy from being able to learn an 
-	* optimal policy.
-	*/
+	 * Clipping ratio to apply to policy updates. Keeps the training "on-policy". Larger values may speed up training at 
+	 * the cost of stability. Conversely, too small values will keep the policy from being able to learn an 
+	 * optimal policy.
+	 */
 	UPROPERTY(EditAnywhere, Category = "LearningAgents", meta = (ClampMin = "0.0", ClampMax = "1.0", UIMin = "0.0", UIMax = "1.0"))
 	float EpsilonClip = 0.2f;
 
 	/**
-	* Weight used to regularize actions. Larger values will encourage smaller actions but too large will cause actions 
-	* to become always zero.
-	*/
+	 * Weight used to regularize actions. Larger values will encourage smaller actions but too large will cause actions 
+	 * to become always zero.
+	 */
 	UPROPERTY(EditAnywhere, Category = "LearningAgents", meta = (ClampMin = "0.0", ClampMax = "1.0", UIMin = "0.0", UIMax = "1.0"))
 	float ActionRegularizationWeight = 0.001f;
 
 	/**
-	* Weighting used for the entropy bonus. Larger values encourage larger action noise and therefore greater 
-	* exploration but can make actions very noisy.
-	*/
+	 * Weighting used for the entropy bonus. Larger values encourage larger action noise and therefore greater 
+	 * exploration but can make actions very noisy.
+	 */
 	UPROPERTY(EditAnywhere, Category = "LearningAgents", meta = (ClampMin = "0.0", ClampMax = "1.0", UIMin = "0.0", UIMax = "1.0"))
 	float EntropyWeight = 0.01f;
 
 	/**
-	* This is used in the Generalized Advantage Estimation as what is essentially an exponential smoothing/decay. 
-	* Typical values should be between 0.9 and 1.0.
-	*/
+	 * This is used in the Generalized Advantage Estimation as what is essentially an exponential smoothing/decay. 
+	 * Typical values should be between 0.9 and 1.0.
+	 */
 	UPROPERTY(EditAnywhere, Category = "LearningAgents", meta = (ClampMin = "0.0", ClampMax = "1.0", UIMin = "0.0", UIMax = "1.0"))
 	float GaeLambda = 0.9f;
 
 	/**
-	* When true, very large or small advantages will be clipped. This has few downsides and helps with numerical 
-	* stability.
-	*/
+	 * When true, very large or small advantages will be clipped. This has few downsides and helps with numerical 
+	 * stability.
+	 */
 	UPROPERTY(EditAnywhere, Category = "LearningAgents")
 	bool bClipAdvantages = true;
 
@@ -228,16 +228,16 @@ public:
 	bool bAdvantageNormalization = true;
 
 	/**
-	* The number of steps to trim from the start of the episode, e.g. can be useful if some things are still getting
-	* setup at the start of the episode and you don't want them used for training.
-	*/
+	 * The number of steps to trim from the start of the episode, e.g. can be useful if some things are still getting
+	 * setup at the start of the episode and you don't want them used for training.
+	 */
 	UPROPERTY(EditAnywhere, Category = "LearningAgents", meta = (ClampMin = "0", UIMin = "0"))
 	int32 NumberOfStepsToTrimAtStartOfEpisode = 0;
 
 	/**
-	* The number of steps to trim from the end of the episode. Can be useful if the end of the episode contains
-	* irrelevant data.
-	*/
+	 * The number of steps to trim from the end of the episode. Can be useful if the end of the episode contains
+	 * irrelevant data.
+	 */
 	UPROPERTY(EditAnywhere, Category = "LearningAgents", meta = (ClampMin = "0", UIMin = "0"))
 	int32 NumberOfStepsToTrimAtEndOfEpisode = 0;
 
@@ -246,9 +246,9 @@ public:
 	int32 RandomSeed = 1234;
 
 	/**
-	* The discount factor to use during training. This affects how much the agent cares about future rewards vs
-	* near-term rewards. Should typically be a value less than but near 1.0. 
-	*/
+	 * The discount factor to use during training. This affects how much the agent cares about future rewards vs
+	 * near-term rewards. Should typically be a value less than but near 1.0. 
+	 */
 	UPROPERTY(EditAnywhere, Category = "LearningAgents", meta = (ClampMin = "0.0", ClampMax = "1.0", UIMin = "0.0", UIMax = "1.0"))
 	float DiscountFactor = 0.99f;
 
@@ -276,23 +276,23 @@ public:
 	FDirectoryPath EditorEngineRelativePath;
 
 	/**
-	* The relative path to the editor engine folder for non-editor builds.
-	* 
-	* If we want to run training in cooked, non-editor builds, then by default we wont have access to python and the
-	* LearningAgents training scripts - these are editor-only things and are stripped during the cooking process.
-	*
-	* However, running training in non-editor builds can be very important - we probably want to disable rendering
-	* and sound while we are training to make experience gathering as fast as possible - and for any non-trivial game
-	* is simply may not be realistic to run it for a long time in play-in-editor.
-	*
-	* For this reason even in non-editor builds we let you provide the path where all of these editor-only things can 
-	* be found. This allows you to run training when these things actually exist somewhere accessible to the executable, 
-	* which will usually be the case on a normal development machine or cloud machine if it is set up that way.
-	*
-	* Since non-editor builds can be produced in a number of different ways, this is not set by default and cannot 
-	* use a directory picker since it is relative to the final location of where your cooked, non-editor executable 
-	* will exist rather than the current with-editor executable.
-	*/
+	 * The relative path to the editor engine folder for non-editor builds.
+	 * 
+	 * If we want to run training in cooked, non-editor builds, then by default we wont have access to python and the
+	 * LearningAgents training scripts - these are editor-only things and are stripped during the cooking process.
+	 *
+	 * However, running training in non-editor builds can be very important - we probably want to disable rendering
+	 * and sound while we are training to make experience gathering as fast as possible - and for any non-trivial game
+	 * is simply may not be realistic to run it for a long time in play-in-editor.
+	 *
+	 * For this reason even in non-editor builds we let you provide the path where all of these editor-only things can 
+	 * be found. This allows you to run training when these things actually exist somewhere accessible to the executable, 
+	 * which will usually be the case on a normal development machine or cloud machine if it is set up that way.
+	 *
+	 * Since non-editor builds can be produced in a number of different ways, this is not set by default and cannot 
+	 * use a directory picker since it is relative to the final location of where your cooked, non-editor executable 
+	 * will exist rather than the current with-editor executable.
+	 */
 	UPROPERTY(EditAnywhere, Category = "LearningAgents")
 	FString NonEditorEngineRelativePath;
 
@@ -310,19 +310,19 @@ public:
 };
 
 /**
-* The ULearningAgentsTrainer is the core class for reinforcement learning training. It has a few responsibilities:
-*   1) It keeps track of which agents are gathering training data.
-*   2) It defines how those agents' rewards, completions, and resets are implemented.
-*   3) It provides methods for orchestrating the training process.
-*
-* To use this class, you need to implement the SetupRewards and SetupCompletions functions (as well as their
-* corresponding SetRewards and SetCompletions functions), which will define the rewards and penalties the agent
-* receives and what conditions cause an episode to end. Before you can begin training, you need to call SetupTrainer,
-* which will initialize the underlying data structures, and you need to call AddAgent for each agent you want to gather
-* training data from.
-*
-* @see ULearningAgentsInteractor to understand how observations and actions work.
-*/
+ * The ULearningAgentsTrainer is the core class for reinforcement learning training. It has a few responsibilities:
+ *   1) It keeps track of which agents are gathering training data.
+ *   2) It defines how those agents' rewards, completions, and resets are implemented.
+ *   3) It provides methods for orchestrating the training process.
+ *
+ * To use this class, you need to implement the SetupRewards and SetupCompletions functions (as well as their
+ * corresponding SetRewards and SetCompletions functions), which will define the rewards and penalties the agent
+ * receives and what conditions cause an episode to end. Before you can begin training, you need to call SetupTrainer,
+ * which will initialize the underlying data structures, and you need to call AddAgent for each agent you want to gather
+ * training data from.
+ *
+ * @see ULearningAgentsInteractor to understand how observations and actions work.
+ */
 UCLASS(Abstract, BlueprintType, Blueprintable)
 class LEARNINGAGENTSTRAINING_API ULearningAgentsTrainer : public ULearningAgentsManagerComponent
 {
@@ -332,7 +332,7 @@ class LEARNINGAGENTSTRAINING_API ULearningAgentsTrainer : public ULearningAgents
 public:
 
 	// These constructors/destructors are needed to make forward declarations happy
-	ULearningAgentsTrainer();
+	ULearningAgentsTrainer(const FObjectInitializer& ObjectInitializer = FObjectInitializer::Get());
 	ULearningAgentsTrainer(FVTableHelper& Helper);
 	virtual ~ULearningAgentsTrainer();
 
@@ -340,16 +340,14 @@ public:
 	virtual void EndPlay(const EEndPlayReason::Type EndPlayReason) override;
 
 	/**
-	* Initializes this object and runs the setup functions for rewards and completions.
-	* @param InAgentManager The agent manager we are associated with.
-	* @param InInteractor The agent interactor we are training with.
-	* @param InPolicy The policy to be trained.
-	* @param InCritic Optional - only needs to be provided if we want the critic to be accessible at runtime.
-	* @param TrainerSettings The trainer settings to use.
-	*/
+	 * Initializes this object and runs the setup functions for rewards and completions.
+	 * @param InInteractor The agent interactor we are training with.
+	 * @param InPolicy The policy to be trained.
+	 * @param InCritic Optional - only needs to be provided if we want the critic to be accessible at runtime.
+	 * @param TrainerSettings The trainer settings to use.
+	 */
 	UFUNCTION(BlueprintCallable, Category = "LearningAgents")
 	void SetupTrainer(
-		ALearningAgentsManager* InAgentManager,
 		ULearningAgentsInteractor* InInteractor,
 		ULearningAgentsPolicy* InPolicy,
 		ULearningAgentsCritic* InCritic = nullptr,
@@ -358,7 +356,9 @@ public:
 public: 
 
 	//~ Begin ULearningAgentsManagerComponent Interface
-	virtual bool AddAgent(const int32 AgentId) override;
+	virtual void OnAgentsAdded(const TArray<int32>& AgentIds) override;
+	virtual void OnAgentsRemoved(const TArray<int32>& AgentIds) override;
+	virtual void OnAgentsReset(const TArray<int32>& AgentIds) override;
 	//~ End ULearningAgentsManagerComponent Interface
 
 // ----- Rewards -----
@@ -366,25 +366,24 @@ public:
 	
 	/**
 	 * During this event, all rewards/penalties should be added to this trainer.
-	 * @param AgentTrainer This trainer. This is a convenience for blueprints so you don't have to find the Self pin.
 	 * @see LearningAgentsRewards.h for the list of available rewards/penalties.
 	 */
 	UFUNCTION(BlueprintNativeEvent, Category = "LearningAgents")
-	void SetupRewards(ULearningAgentsTrainer* AgentTrainer);
+	void SetupRewards();
 
 	/**
-	* During this event, all rewards/penalties should be set for each agent.
-	* @param AgentIds The list of agent ids to set rewards/penalties for.
-	* @see LearningAgentsRewards.h for the list of available rewards/penalties.
-	* @see GetAgent to get the agent corresponding to each id.
-	*/
+	 * During this event, all rewards/penalties should be set for each agent.
+	 * @param AgentIds The list of agent ids to set rewards/penalties for.
+	 * @see LearningAgentsRewards.h for the list of available rewards/penalties.
+	 * @see GetAgent to get the agent corresponding to each id.
+	 */
 	UFUNCTION(BlueprintNativeEvent, Category = "LearningAgents")
 	void SetRewards(const TArray<int32>& AgentIds);
 
 	/**
-	* Used by objects derived from ULearningAgentsReward to add themselves to this trainer during their creation.
-	* You shouldn't need to call this directly.
-	*/
+	 * Used by objects derived from ULearningAgentsReward to add themselves to this trainer during their creation.
+	 * You shouldn't need to call this directly.
+	 */
 	void AddReward(TObjectPtr<ULearningAgentsReward> Object, const TSharedRef<UE::Learning::FRewardObject>& Reward);
 
 // ----- Completions ----- 
@@ -392,25 +391,24 @@ public:
 
 	/**
 	 * During this event, all completions should be added to this trainer.
-	 * @param AgentTrainer This trainer. This is a convenience for blueprints so you don't have to find the Self pin.
 	 * @see LearningAgentsCompletions.h for the list of available completions.
 	 */
 	UFUNCTION(BlueprintNativeEvent, Category = "LearningAgents")
-	void SetupCompletions(ULearningAgentsTrainer* AgentTrainer);
+	void SetupCompletions();
 
 	/**
-	* During this event, all completions should be set for each agent.
-	* @param AgentIds The list of agent ids to set completions for.
-	* @see LearningAgentsCompletions.h for the list of available completions.
-	* @see GetAgent to get the agent corresponding to each id.
-	*/
+	 * During this event, all completions should be set for each agent.
+	 * @param AgentIds The list of agent ids to set completions for.
+	 * @see LearningAgentsCompletions.h for the list of available completions.
+	 * @see GetAgent to get the agent corresponding to each id.
+	 */
 	UFUNCTION(BlueprintNativeEvent, Category = "LearningAgents")
 	void SetCompletions(const TArray<int32>& AgentIds);
 
 	/**
-	* Used by objects derived from ULearningAgentsCompletion to add themselves to this trainer during their creation.
-	* You shouldn't need to call this directly.
-	*/
+	 * Used by objects derived from ULearningAgentsCompletion to add themselves to this trainer during their creation.
+	 * You shouldn't need to call this directly.
+	 */
 	void AddCompletion(TObjectPtr<ULearningAgentsCompletion> Object, const TSharedRef<UE::Learning::FCompletionObject>& Completion);
 
 // ----- Resets ----- 
@@ -432,14 +430,15 @@ public:
 	const bool IsTraining() const;
 
 	/**
-	* Begins the training process with the provided settings.
-	* @param TrainerTrainingSettings The settings for this training run.
-	* @param TrainerGameSettings The settings that will affect the game's simulation.
-	* @param TrainerPathSettings The path settings used by the trainer.
-	* @param CriticSettings The settings for the critic (if we are using one).
-	* @param bReinitializePolicyNetwork If true, reinitialize the policy. Set this to false if your policy is pre-trained, e.g. with imitation learning.
-	* @param bReinitializeCriticNetwork If true, reinitialize the critic. Set this to false if your critic is pre-trained.
-	*/
+	 * Begins the training process with the provided settings.
+	 * @param TrainerTrainingSettings The settings for this training run.
+	 * @param TrainerGameSettings The settings that will affect the game's simulation.
+	 * @param TrainerPathSettings The path settings used by the trainer.
+	 * @param CriticSettings The settings for the critic (if we are using one).
+	 * @param bReinitializePolicyNetwork If true, reinitialize the policy. Set this to false if your policy is pre-trained, e.g. with imitation learning.
+	 * @param bReinitializeCriticNetwork If true, reinitialize the critic. Set this to false if your critic is pre-trained.
+	 * @param bResetAgentsOnBegin If true, reset all agents at the beginning of training.
+	 */
 	UFUNCTION(BlueprintCallable, Category = "LearningAgents")
 	void BeginTraining(
 		const FLearningAgentsTrainerTrainingSettings& TrainerTrainingSettings = FLearningAgentsTrainerTrainingSettings(),
@@ -447,52 +446,50 @@ public:
 		const FLearningAgentsTrainerPathSettings& TrainerPathSettings = FLearningAgentsTrainerPathSettings(),
 		const FLearningAgentsCriticSettings& CriticSettings = FLearningAgentsCriticSettings(),
 		const bool bReinitializePolicyNetwork = true,
-		const bool bReinitializeCriticNetwork = true);
+		const bool bReinitializeCriticNetwork = true,
+		const bool bResetAgentsOnBegin = true);
 
 	/** Stops the training process. */
 	UFUNCTION(BlueprintCallable, Category = "LearningAgents")
 	void EndTraining();
 
 	/**
-	* Call this function when it is time to evaluate the rewards for your agents. This should be done at the beginning
-	* of each iteration of your training loop after the initial step, i.e. after taking an action, you want to get into
-	* the next state before evaluating the rewards.
-	*/
+	 * Call this function when it is time to evaluate the rewards for your agents. This should be done at the beginning
+	 * of each iteration of your training loop after the initial step, i.e. after taking an action, you want to get into
+	 * the next state before evaluating the rewards.
+	 */
 	UFUNCTION(BlueprintCallable, Category = "LearningAgents")
 	void EvaluateRewards();
 
 	/**
-	* Call this function when it is time to evaluate the completions for your agents. This should be done at the beginning
-	* of each iteration of your training loop after the initial step, i.e. after taking an action, you want to get into
-	* the next state before evaluating the completions.
-	*/
+	 * Call this function when it is time to evaluate the completions for your agents. This should be done at the beginning
+	 * of each iteration of your training loop after the initial step, i.e. after taking an action, you want to get into
+	 * the next state before evaluating the completions.
+	 */
 	UFUNCTION(BlueprintCallable, Category = "LearningAgents")
 	void EvaluateCompletions();
 
 	/**
-	* Call this function at the end of each step of your training loop. This takes the current observations/actions/
-	* rewards and moves them into the episode experience buffer. All agents with full episode buffers or those which
-	* have been signaled complete will be reset. If enough experience is gathered, it will be sent to the training 
-	* process and an iteration of training will be run and the updated policy will be synced back.
-	*/
+	 * Call this function at the end of each step of your training loop. This takes the current observations/actions/
+	 * rewards and moves them into the episode experience buffer. All agents with full episode buffers or those which
+	 * have been signaled complete will be reset. If enough experience is gathered, it will be sent to the training 
+	 * process and an iteration of training will be run and the updated policy will be synced back.
+	 */
 	UFUNCTION(BlueprintCallable, Category = "LearningAgents")
 	void ProcessExperience();
 
-	/** Manually reset all agents. Does not record the experience gathered up to this point by each agent. */
-	UFUNCTION(BlueprintCallable, Category = "LearningAgents")
-	void ResetAllEpisodes();
-
 	/**
-	* Convenience function that runs a basic training loop. If training has not been started, it will start it, and 
-	* then call RunInference. On each following call to this function, it will call EvaluateRewards, 
-	* EvaluateCompletions, and ProcessExperience, followed by RunInference.
-	* @param TrainerTrainingSettings The settings for this training run.
-	* @param TrainerGameSettings The settings that will affect the game's simulation.
-	* @param TrainerPathSettings The path settings used by the trainer.
-	* @param CriticSettings The settings for the critic (if we are using one).
-	* @param bReinitializePolicyNetwork If true, reinitialize the policy. Set this to false if your policy is pre-trained, e.g. with imitation learning.
-	* @param bReinitializeCriticNetwork If true, reinitialize the critic. Set this to false if your critic is pre-trained.
-	*/
+	 * Convenience function that runs a basic training loop. If training has not been started, it will start it, and 
+	 * then call RunInference. On each following call to this function, it will call EvaluateRewards, 
+	 * EvaluateCompletions, and ProcessExperience, followed by RunInference.
+	 * @param TrainerTrainingSettings The settings for this training run.
+	 * @param TrainerGameSettings The settings that will affect the game's simulation.
+	 * @param TrainerPathSettings The path settings used by the trainer.
+	 * @param CriticSettings The settings for the critic (if we are using one).
+	 * @param bReinitializePolicyNetwork If true, reinitialize the policy. Set this to false if your policy is pre-trained, e.g. with imitation learning.
+	 * @param bReinitializeCriticNetwork If true, reinitialize the critic. Set this to false if your critic is pre-trained.
+	 * @param bResetAgentsOnBegin If true, reset all agents at the beginning of training.
+	 */
 	UFUNCTION(BlueprintCallable, Category = "LearningAgents")
 	void RunTraining(
 		const FLearningAgentsTrainerTrainingSettings& TrainerTrainingSettings = FLearningAgentsTrainerTrainingSettings(),
@@ -500,25 +497,26 @@ public:
 		const FLearningAgentsTrainerPathSettings& TrainerPathSettings = FLearningAgentsTrainerPathSettings(),
 		const FLearningAgentsCriticSettings& CriticSettings = FLearningAgentsCriticSettings(),
 		const bool bReinitializePolicyNetwork = true,
-		const bool bReinitializeCriticNetwork = true);
+		const bool bReinitializeCriticNetwork = true,
+		const bool bResetAgentsOnBegin = true);
 
 	/**
-	* Gets the current reward for an agent according to the critic. Should be called only after EvaluateRewards.
-	*
-	* @param AgentId	The AgentId to look-up the reward for
-	* @returns			The reward
-	*/
+	 * Gets the current reward for an agent according to the critic. Should be called only after EvaluateRewards.
+	 *
+	 * @param AgentId	The AgentId to look-up the reward for
+	 * @returns			The reward
+	 */
 	UFUNCTION(BlueprintPure, Category = "LearningAgents", meta = (AgentId = "-1"))
 	float GetReward(const int32 AgentId) const;
 
 	/**
-	* Gets if the agent will complete the episode or not according to the given set of completions. Should be called 
-	* only after EvaluateCompletions.
-	*
-	* @param AgentId		The AgentId to look-up the completion for
-	* @param OutCompletion	The completion type if the agent will complete the episode
-	* @returns				If the agent will complete the episode
-	*/
+	 * Gets if the agent will complete the episode or not according to the given set of completions. Should be called 
+	 * only after EvaluateCompletions.
+	 *
+	 * @param AgentId		The AgentId to look-up the completion for
+	 * @param OutCompletion	The completion type if the agent will complete the episode
+	 * @returns				If the agent will complete the episode
+	 */
 	UFUNCTION(BlueprintPure, Category = "LearningAgents", meta = (AgentId = "-1"))
 	bool IsCompleted(const int32 AgentId, ELearningAgentsCompletion& OutCompletion) const;
 	
@@ -559,13 +557,35 @@ private:
 	TUniquePtr<UE::Learning::FReplayBuffer> ReplayBuffer;
 	TUniquePtr<UE::Learning::FResetInstanceBuffer> ResetBuffer;
 	TUniquePtr<UE::Learning::FSharedMemoryPPOTrainer> Trainer;
-
+	
 	ELearningAgentsCompletion MaxStepsCompletion = ELearningAgentsCompletion::Truncation;
 
 	float TrainerTimeout = 10.0f;
 
 	void DoneTraining();
 
+// ----- Private Iteration Checks ----- 
+private:
+
+	/**
+	 * Sets the iteration number for the given set of agents for all observations, actions, as well as for the encoding
+	 * of observations and actions
+	 */
+	void SetAgentIteration(const UE::Learning::FIndexSet AgentIds, const uint64 Value);
+
+	/** Number of times rewards have been evaluated for all agents */
+	TLearningArray<1, uint64, TInlineAllocator<32>> RewardEvaluatedAgentIteration;
+
+	/** Number of times completions have been evaluated for all agents */
+	TLearningArray<1, uint64, TInlineAllocator<32>> CompletionEvaluatedAgentIteration;
+
+	/** Temp buffers used to record the set of agents that are valid for training */
+	TArray<int32> ValidAgentIds;
+	TArray<int32> FinalValidAgentIds;
+	UE::Learning::FIndexSet ValidAgentSet;
+	UE::Learning::FIndexSet FinalValidAgentSet;
+	TBitArray<TInlineAllocator<32>> ValidAgentStatus;
+
 // ----- Private Recording of GameSettings ----- 
 private:
 
diff --git a/Engine/Plugins/Experimental/LearningAgents/Source/LearningTraining/Private/LearningImitationTrainer.cpp b/Engine/Plugins/Experimental/LearningAgents/Source/LearningTraining/Private/LearningImitationTrainer.cpp
--- a/Engine/Plugins/Experimental/LearningAgents/Source/LearningTraining/Private/LearningImitationTrainer.cpp
+++ b/Engine/Plugins/Experimental/LearningAgents/Source/LearningTraining/Private/LearningImitationTrainer.cpp
@@ -230,9 +230,9 @@ namespace UE::Learning
 		return SharedMemoryTraining::SendStop(Controls.View);
 	}
 
-	ETrainerResponse FSharedMemoryImitationTrainer::SendContinue(const float Timeout)
+	bool FSharedMemoryImitationTrainer::HasPolicyOrCompleted()
 	{
-		return SharedMemoryTraining::SendContinue(Controls.View);
+		return SharedMemoryTraining::HasPolicyOrCompleted(Controls.View);
 	}
 
 	ETrainerResponse FSharedMemoryImitationTrainer::RecvPolicy(
@@ -574,9 +574,9 @@ namespace UE::Learning
 		return SocketTraining::SendStop(*Socket, Timeout);
 	}
 
-	ETrainerResponse FSocketImitationTrainer::SendContinue(const float Timeout)
+	bool FSocketImitationTrainer::HasPolicyOrCompleted()
 	{
-		return SocketTraining::SendContinue(*Socket, Timeout);
+		return SocketTraining::HasPolicyOrCompleted(*Socket);
 	}
 
 	ETrainerResponse FSocketImitationTrainer::RecvPolicy(
@@ -709,39 +709,27 @@ namespace UE::Learning
 
 					break;
 				}
-				else
+				
+				if (Trainer.HasPolicyOrCompleted())
 				{
-					Response = Trainer.SendContinue();
+					Response = Trainer.RecvPolicy(Network, 10.0f, NetworkLock);
 
-					if (Response != ETrainerResponse::Success)
+					if (Response == ETrainerResponse::Completed)
 					{
 						if (LogSettings != ELogSetting::Silent)
 						{
-							UE_LOG(LogLearning, Error, TEXT("Error sending continue signal to trainer: %s. Check log for errors."), Trainer::GetResponseString(Response));
+							UE_LOG(LogLearning, Display, TEXT("Trainer completed training."));
 						}
-
-						Trainer.Terminate();
-						return Response;
+						break;
 					}
-				}
-
-				Response = Trainer.RecvPolicy(Network, 10.0f, NetworkLock);
-
-				if (Response == ETrainerResponse::Completed)
-				{
-					if (LogSettings != ELogSetting::Silent)
+					else if (Response != ETrainerResponse::Success)
 					{
-						UE_LOG(LogLearning, Display, TEXT("Trainer completed training."));
-					}
-					break;
-				}
-				else if (Response != ETrainerResponse::Success)
-				{
-					if (LogSettings != ELogSetting::Silent)
-					{
-						UE_LOG(LogLearning, Error, TEXT("Error receiving policy from trainer: %s. Check log for errors."), Trainer::GetResponseString(Response));
+						if (LogSettings != ELogSetting::Silent)
+						{
+							UE_LOG(LogLearning, Error, TEXT("Error receiving policy from trainer: %s. Check log for errors."), Trainer::GetResponseString(Response));
+						}
+						break;
 					}
-					break;
 				}
 
 				if (bNetworkUpdatedSignal)
diff --git a/Engine/Plugins/Experimental/LearningAgents/Source/LearningTraining/Private/LearningSharedMemoryTraining.cpp b/Engine/Plugins/Experimental/LearningAgents/Source/LearningTraining/Private/LearningSharedMemoryTraining.cpp
--- a/Engine/Plugins/Experimental/LearningAgents/Source/LearningTraining/Private/LearningSharedMemoryTraining.cpp
+++ b/Engine/Plugins/Experimental/LearningAgents/Source/LearningTraining/Private/LearningSharedMemoryTraining.cpp
@@ -20,7 +20,6 @@ namespace UE::Learning::SharedMemoryTraining
 		CriticSignal = 4,
 		CompleteSignal = 5,
 		StopSignal = 6,
-		ContinueSignal = 7,
 
 		ControlNum = 8,
 	};
@@ -36,10 +35,9 @@ namespace UE::Learning::SharedMemoryTraining
 		return ETrainerResponse::Success;
 	}
 
-	ETrainerResponse SendContinue(TLearningArrayView<1, volatile int32> Controls)
+	bool HasPolicyOrCompleted(TLearningArrayView<1, volatile int32> Controls)
 	{
-		Controls[(uint8)EControls::ContinueSignal] = true;
-		return ETrainerResponse::Success;
+		return Controls[(uint8)EControls::PolicySignal] || Controls[(uint8)EControls::CompleteSignal];
 	}
 
 	ETrainerResponse RecvPolicy(
diff --git a/Engine/Plugins/Experimental/LearningAgents/Source/LearningTraining/Private/LearningSharedMemoryTraining.h b/Engine/Plugins/Experimental/LearningAgents/Source/LearningTraining/Private/LearningSharedMemoryTraining.h
--- a/Engine/Plugins/Experimental/LearningAgents/Source/LearningTraining/Private/LearningSharedMemoryTraining.h
+++ b/Engine/Plugins/Experimental/LearningAgents/Source/LearningTraining/Private/LearningSharedMemoryTraining.h
@@ -33,8 +33,7 @@ namespace UE::Learning
 		LEARNINGTRAINING_API ETrainerResponse SendStop(
 			TLearningArrayView<1, volatile int32> Controls);
 
-		LEARNINGTRAINING_API ETrainerResponse SendContinue(
-			TLearningArrayView<1, volatile int32> Controls);
+		LEARNINGTRAINING_API bool HasPolicyOrCompleted(TLearningArrayView<1, volatile int32> Controls);
 
 		LEARNINGTRAINING_API ETrainerResponse SendPolicy(
 			TLearningArrayView<1, volatile int32> Controls,
diff --git a/Engine/Plugins/Experimental/LearningAgents/Source/LearningTraining/Private/LearningSocketTraining.cpp b/Engine/Plugins/Experimental/LearningAgents/Source/LearningTraining/Private/LearningSocketTraining.cpp
--- a/Engine/Plugins/Experimental/LearningAgents/Source/LearningTraining/Private/LearningSocketTraining.cpp
+++ b/Engine/Plugins/Experimental/LearningAgents/Source/LearningTraining/Private/LearningSocketTraining.cpp
@@ -22,7 +22,6 @@ namespace UE::Learning::SocketTraining
 		SendCritic = 6,
 		RecvComplete = 7,
 		SendStop = 8,
-		SendContinue = 9,
 	};
 
 	ETrainerResponse WaitForConnection(FSocket& Socket, const float Timeout)
@@ -214,12 +213,10 @@ namespace UE::Learning::SocketTraining
 		return SendWithTimeout(Socket, &Signal, 1, Timeout);
 	}
 
-	ETrainerResponse SendContinue(
-		FSocket& Socket,
-		const float Timeout)
+	bool HasPolicyOrCompleted(FSocket& Socket)
 	{
-		const uint8 Signal = (uint8)ESignal::SendContinue;
-		return SendWithTimeout(Socket, &Signal, 1, Timeout);
+		uint32 PendingDataSize;
+		return Socket.HasPendingData(PendingDataSize);
 	}
 
 	ETrainerResponse SendPolicy(
diff --git a/Engine/Plugins/Experimental/LearningAgents/Source/LearningTraining/Private/LearningSocketTraining.h b/Engine/Plugins/Experimental/LearningAgents/Source/LearningTraining/Private/LearningSocketTraining.h
--- a/Engine/Plugins/Experimental/LearningAgents/Source/LearningTraining/Private/LearningSocketTraining.h
+++ b/Engine/Plugins/Experimental/LearningAgents/Source/LearningTraining/Private/LearningSocketTraining.h
@@ -54,9 +54,7 @@ namespace UE::Learning
 			FSocket& Socket, 
 			const float Timeout = Trainer::DefaultTimeout);
 
-		LEARNINGTRAINING_API ETrainerResponse SendContinue(
-			FSocket& Socket,
-			const float Timeout = Trainer::DefaultTimeout);
+		LEARNINGTRAINING_API bool HasPolicyOrCompleted(FSocket& Socket);
 
 		LEARNINGTRAINING_API ETrainerResponse SendPolicy(
 			FSocket& Socket,
diff --git a/Engine/Plugins/Experimental/LearningAgents/Source/LearningTraining/Public/LearningImitationTrainer.h b/Engine/Plugins/Experimental/LearningAgents/Source/LearningTraining/Public/LearningImitationTrainer.h
--- a/Engine/Plugins/Experimental/LearningAgents/Source/LearningTraining/Public/LearningImitationTrainer.h
+++ b/Engine/Plugins/Experimental/LearningAgents/Source/LearningTraining/Public/LearningImitationTrainer.h
@@ -129,12 +129,11 @@ namespace UE::Learning
 		virtual ETrainerResponse SendStop(const float Timeout = Trainer::DefaultTimeout) = 0;
 
 		/*
-		* Signal for the trainer to continue.
+		* Check if a new policy is available or if training is complete.
 		* 
-		* @param Timeout		Timeout to wait in seconds
-		* @returns				Trainer response
+		* @returns				If there is a new policy ready or training is complete.
 		*/
-		virtual ETrainerResponse SendContinue(const float Timeout = Trainer::DefaultTimeout) = 0;
+		virtual bool HasPolicyOrCompleted() = 0;
 
 		/**
 		* Wait for the trainer and pull an updated policy.
@@ -222,11 +221,11 @@ namespace UE::Learning
 
 		virtual void Terminate() override final;
 
-		virtual ETrainerResponse Wait(const float Timeout) override final;
+		virtual ETrainerResponse Wait(const float Timeout = Trainer::DefaultTimeout) override final;
 
-		virtual ETrainerResponse SendStop(const float Timeout) override final;
+		virtual ETrainerResponse SendStop(const float Timeout = Trainer::DefaultTimeout) override final;
 
-		virtual ETrainerResponse SendContinue(const float Timeout) override final;
+		virtual bool HasPolicyOrCompleted() override final;
 
 		virtual ETrainerResponse RecvPolicy(
 			FNeuralNetwork& OutNetwork,
@@ -360,11 +359,11 @@ namespace UE::Learning
 
 		virtual void Terminate() override final;
 
-		virtual ETrainerResponse Wait(const float Timeout) override final;
+		virtual ETrainerResponse Wait(const float Timeout = Trainer::DefaultTimeout) override final;
 
-		virtual ETrainerResponse SendStop(const float Timeout) override final;
+		virtual ETrainerResponse SendStop(const float Timeout = Trainer::DefaultTimeout) override final;
 
-		virtual ETrainerResponse SendContinue(const float Timeout) override final;
+		virtual bool HasPolicyOrCompleted() override final;
 
 		virtual ETrainerResponse RecvPolicy(
 			FNeuralNetwork& OutNetwork,
diff --git a/Engine/Plugins/Experimental/LearningAgents/Source/LearningTraining/Public/LearningPPOTrainer.h b/Engine/Plugins/Experimental/LearningAgents/Source/LearningTraining/Public/LearningPPOTrainer.h
--- a/Engine/Plugins/Experimental/LearningAgents/Source/LearningTraining/Public/LearningPPOTrainer.h
+++ b/Engine/Plugins/Experimental/LearningAgents/Source/LearningTraining/Public/LearningPPOTrainer.h
@@ -340,7 +340,7 @@ namespace UE::Learning
 
 		virtual void Terminate() override final;
 
-		virtual ETrainerResponse Wait(const float Timeout) override final;
+		virtual ETrainerResponse Wait(const float Timeout = Trainer::DefaultTimeout) override final;
 
 		virtual ETrainerResponse RecvPolicy(
 			FNeuralNetwork& OutNetwork,
@@ -449,7 +449,7 @@ namespace UE::Learning
 		* @param Timeout		Timeout to wait in seconds
 		* @returns				true if successful, otherwise false if it times out
 		*/
-		bool Wait(float Timeout);
+		bool Wait(float Timeout = Trainer::DefaultTimeout);
 
 		/**
 		* Terminate the server process
@@ -501,7 +501,7 @@ namespace UE::Learning
 
 		virtual void Terminate() override final;
 
-		virtual ETrainerResponse Wait(const float Timeout) override final;
+		virtual ETrainerResponse Wait(const float Timeout = Trainer::DefaultTimeout) override final;
 
 		virtual ETrainerResponse RecvPolicy(
 			FNeuralNetwork& OutNetwork,
