Fixed mismatched copy constructor/assignment operator for UE::Learning::FIndexSet.

diff --git a/Engine/Plugins/Experimental/LearningAgents/Source/Learning/Public/LearningArray.h b/Engine/Plugins/Experimental/LearningAgents/Source/Learning/Public/LearningArray.h
--- a/Engine/Plugins/Experimental/LearningAgents/Source/Learning/Public/LearningArray.h
+++ b/Engine/Plugins/Experimental/LearningAgents/Source/Learning/Public/LearningArray.h
@@ -189,7 +189,6 @@ namespace UE::Learning
 		FIndexSet(const TLearningArray<1, int32, InAllocatorType>& InIndices) : Indices(InIndices) {}
 
 		FIndexSet& operator=(const int32 InSingleIndex) { SliceStart = InSingleIndex; SliceNum = 1; Indices = TLearningArrayView<1, const int32>(); return *this; }
-		FIndexSet& operator=(const FIndexSet& Rhs) { SliceStart = Rhs.SliceStart; SliceNum = Rhs.SliceNum; Indices = Rhs.Indices; return *this; }
 		FIndexSet& operator=(const TArrayView<const int32> InIndices) { SliceStart = 0; SliceNum = 0; Indices = InIndices; return *this; }
 		FIndexSet& operator=(TLearningArrayView<1, const int32> InIndices) { SliceStart = 0; SliceNum = 0; Indices = InIndices; return *this; }
 
