LearningAgents: make settings classes part of the public API

diff --git a/Engine/Plugins/Experimental/LearningAgents/Source/LearningAgents/Public/LearningAgentsNeuralNetwork.h b/Engine/Plugins/Experimental/LearningAgents/Source/LearningAgents/Public/LearningAgentsNeuralNetwork.h
--- a/Engine/Plugins/Experimental/LearningAgents/Source/LearningAgents/Public/LearningAgentsNeuralNetwork.h
+++ b/Engine/Plugins/Experimental/LearningAgents/Source/LearningAgents/Public/LearningAgentsNeuralNetwork.h
@@ -72,7 +72,7 @@ public:
 	void LoadNetworkFromAsset(ULearningAgentsNeuralNetwork* NeuralNetworkAsset);
 
 	/**
-	 * Copy this networks's weights to another asset. Network sizes must match.
+	 * Copy this network's weights to another asset. Network sizes must match.
 	 * @param NeuralNetworkAsset The asset to save to.
 	 */
 	UFUNCTION(BlueprintCallable, Category = "LearningAgents", Meta = (DevelopmentOnly))
diff --git a/Engine/Plugins/Experimental/LearningAgents/Source/LearningAgentsTraining/Public/LearningAgentsImitationTrainer.h b/Engine/Plugins/Experimental/LearningAgents/Source/LearningAgentsTraining/Public/LearningAgentsImitationTrainer.h
--- a/Engine/Plugins/Experimental/LearningAgents/Source/LearningAgentsTraining/Public/LearningAgentsImitationTrainer.h
+++ b/Engine/Plugins/Experimental/LearningAgents/Source/LearningAgentsTraining/Public/LearningAgentsImitationTrainer.h
@@ -20,7 +20,7 @@ class ULearningAgentsRecording;
 
 /** The configurable settings for a ULearningAgentsImitationTrainer. */
 USTRUCT(BlueprintType, Category = "LearningAgents")
-struct FLearningAgentsImitationTrainerSettings
+struct LEARNINGAGENTSTRAINING_API FLearningAgentsImitationTrainerSettings
 {
 	GENERATED_BODY()
 
@@ -33,7 +33,7 @@ public:
 
 /** The configurable settings for the training process. */
 USTRUCT(BlueprintType, Category = "LearningAgents")
-struct FLearningAgentsImitationTrainerTrainingSettings
+struct LEARNINGAGENTSTRAINING_API FLearningAgentsImitationTrainerTrainingSettings
 {
 	GENERATED_BODY()
 
diff --git a/Engine/Plugins/Experimental/LearningAgents/Source/LearningAgentsTraining/Public/LearningAgentsRecorder.h b/Engine/Plugins/Experimental/LearningAgents/Source/LearningAgentsTraining/Public/LearningAgentsRecorder.h
--- a/Engine/Plugins/Experimental/LearningAgents/Source/LearningAgentsTraining/Public/LearningAgentsRecorder.h
+++ b/Engine/Plugins/Experimental/LearningAgents/Source/LearningAgentsTraining/Public/LearningAgentsRecorder.h
@@ -15,7 +15,7 @@ class ULearningAgentsRecording;
 
 /** The path settings for the recorder. */
 USTRUCT(BlueprintType, Category = "LearningAgents")
-struct FLearningAgentsRecorderPathSettings
+struct LEARNINGAGENTSTRAINING_API FLearningAgentsRecorderPathSettings
 {
 	GENERATED_BODY()
 
diff --git a/Engine/Plugins/Experimental/LearningAgents/Source/LearningAgentsTraining/Public/LearningAgentsTrainer.h b/Engine/Plugins/Experimental/LearningAgents/Source/LearningAgentsTraining/Public/LearningAgentsTrainer.h
--- a/Engine/Plugins/Experimental/LearningAgents/Source/LearningAgentsTraining/Public/LearningAgentsTrainer.h
+++ b/Engine/Plugins/Experimental/LearningAgents/Source/LearningAgentsTraining/Public/LearningAgentsTrainer.h
@@ -55,7 +55,7 @@ namespace UE::Learning::Agents
 
 /** The configurable settings for a ULearningAgentsTrainer. */
 USTRUCT(BlueprintType, Category = "LearningAgents")
-struct FLearningAgentsTrainerSettings
+struct LEARNINGAGENTSTRAINING_API FLearningAgentsTrainerSettings
 {
 	GENERATED_BODY()
 
@@ -87,7 +87,7 @@ public:
  * during training, which can enable ticking faster than real-time.
  */
 USTRUCT(BlueprintType, Category="LearningAgents")
-struct FLearningAgentsTrainerGameSettings
+struct LEARNINGAGENTSTRAINING_API FLearningAgentsTrainerGameSettings
 {
 	GENERATED_BODY()
 
@@ -141,7 +141,7 @@ namespace UE::Learning::Agents
 
 /** The configurable settings for the training process. */
 USTRUCT(BlueprintType, Category = "LearningAgents")
-struct FLearningAgentsTrainerTrainingSettings
+struct LEARNINGAGENTSTRAINING_API FLearningAgentsTrainerTrainingSettings
 {
 	GENERATED_BODY()
 
@@ -263,7 +263,7 @@ public:
 
 /** The path settings for the trainer. */
 USTRUCT(BlueprintType, Category = "LearningAgents")
-struct FLearningAgentsTrainerPathSettings
+struct LEARNINGAGENTSTRAINING_API FLearningAgentsTrainerPathSettings
 {
 	GENERATED_BODY()
 
