LearningAgents - Add ScriptName to ELearningAgentsCompletion to disambiguate python API

#rnx

diff --git a/Engine/Plugins/Experimental/LearningAgents/Source/LearningAgentsTraining/Public/LearningAgentsTrainer.h b/Engine/Plugins/Experimental/LearningAgents/Source/LearningAgentsTraining/Public/LearningAgentsTrainer.h
--- a/Engine/Plugins/Experimental/LearningAgents/Source/LearningAgentsTraining/Public/LearningAgentsTrainer.h
+++ b/Engine/Plugins/Experimental/LearningAgents/Source/LearningAgentsTraining/Public/LearningAgentsTrainer.h
@@ -34,7 +34,7 @@ class ULearningAgentsPolicy;
 class ULearningAgentsCritic;
 
 /** Completion modes for episodes. */
-UENUM(BlueprintType, Category = "LearningAgents")
+UENUM(BlueprintType, Category = "LearningAgents", meta = (ScriptName = "LearningAgentsCompletionEnum"))
 enum class ELearningAgentsCompletion : uint8
 {
 	/** Episode ended early but was still in progress. Critic will be used to estimate final return. */
