LearningAgents: Add a couple BlueprintPure = false

diff --git a/Engine/Plugins/Experimental/LearningAgents/Source/LearningAgents/Public/LearningAgentsCritic.h b/Engine/Plugins/Experimental/LearningAgents/Source/LearningAgents/Public/LearningAgentsCritic.h
--- a/Engine/Plugins/Experimental/LearningAgents/Source/LearningAgents/Public/LearningAgentsCritic.h
+++ b/Engine/Plugins/Experimental/LearningAgents/Source/LearningAgents/Public/LearningAgentsCritic.h
@@ -90,7 +90,7 @@ public:
 	 * Save this critic's weights into a snapshot.
 	 * @param File The snapshot file.
 	 */
-	UFUNCTION(BlueprintCallable, Category = "LearningAgents", meta = (RelativePath))
+	UFUNCTION(BlueprintCallable, BlueprintPure = false, Category = "LearningAgents", meta = (RelativePath))
 	void SaveCriticToSnapshot(const FFilePath& File) const;
 
 	/**
diff --git a/Engine/Plugins/Experimental/LearningAgents/Source/LearningAgents/Public/LearningAgentsPolicy.h b/Engine/Plugins/Experimental/LearningAgents/Source/LearningAgents/Public/LearningAgentsPolicy.h
--- a/Engine/Plugins/Experimental/LearningAgents/Source/LearningAgents/Public/LearningAgentsPolicy.h
+++ b/Engine/Plugins/Experimental/LearningAgents/Source/LearningAgents/Public/LearningAgentsPolicy.h
@@ -97,7 +97,7 @@ public:
 	 * Save this policy's weights into a snapshot.
 	 * @param File The snapshot file.
 	 */
-	UFUNCTION(BlueprintCallable, Category = "LearningAgents", meta = (RelativePath))
+	UFUNCTION(BlueprintCallable, BlueprintPure = false, Category = "LearningAgents", meta = (RelativePath))
 	void SavePolicyToSnapshot(const FFilePath& File) const;
 
 	/**
