LearningAgents: Added ability to reset network asset in editor. Added proper shape validation for loading network from snapshot. Added new radial ray cast helper. Added planar velocity action.

diff --git a/Engine/Plugins/Experimental/LearningAgents/Source/LearningAgents/Private/LearningAgentsActions.cpp b/Engine/Plugins/Experimental/LearningAgents/Source/LearningAgents/Private/LearningAgentsActions.cpp
--- a/Engine/Plugins/Experimental/LearningAgents/Source/LearningAgents/Private/LearningAgentsActions.cpp
+++ b/Engine/Plugins/Experimental/LearningAgents/Source/LearningAgents/Private/LearningAgentsActions.cpp
@@ -364,6 +364,97 @@ void UVectorArrayAction::VisualLog(const UE::Learning::FIndexSet Instances) cons
 
 //------------------------------------------------------------------
 
+
+UPlanarVelocityAction* UPlanarVelocityAction::AddPlanarVelocityAction(ULearningAgentsInteractor* InInteractor, const FName Name, const float Scale, const FVector Axis0, const FVector Axis1)
+{
+	return UE::Learning::Agents::Actions::Private::AddAction<UPlanarVelocityAction, UE::Learning::FPlanarVelocityFeature>(
+		InInteractor, 
+		Name, 
+		TEXT("AddPlanarVelocityAction"), 
+		1, 
+		Scale,
+		Axis0.GetSafeNormal(UE_SMALL_NUMBER, FVector::ForwardVector),
+		Axis1.GetSafeNormal(UE_SMALL_NUMBER, FVector::RightVector));
+}
+
+FVector UPlanarVelocityAction::GetPlanarVelocityAction(const int32 AgentId)
+{
+	if (!Interactor->HasAgent(AgentId))
+	{
+		UE_LOG(LogLearning, Error, TEXT("%s: AgentId %d not found in the agents set."), *GetName(), AgentId);
+		return FVector::ZeroVector;
+	}
+
+	const TLearningArrayView<2, const FVector> View = FeatureObject->InstanceData->ConstView(FeatureObject->VelocityHandle);
+
+	AgentGetIteration[AgentId]++;
+
+	return View[AgentId][0];
+}
+
+void UPlanarVelocityAction::SetPlanarVelocityAction(const int32 AgentId, const FVector Velocity)
+{
+	if (!Interactor->HasAgent(AgentId))
+	{
+		UE_LOG(LogLearning, Error, TEXT("%s: AgentId %d not found in the agents set."), *GetName(), AgentId);
+		return;
+	}
+
+	const TLearningArrayView<2, FVector> View = FeatureObject->InstanceData->View(FeatureObject->VelocityHandle);
+	View[AgentId][0] = Velocity;
+
+	AgentSetIteration[AgentId]++;
+}
+
+#if UE_LEARNING_AGENTS_ENABLE_VISUAL_LOG
+void UPlanarVelocityAction::VisualLog(const UE::Learning::FIndexSet Instances) const
+{
+	UE_LEARNING_TRACE_CPUPROFILER_EVENT_SCOPE(UPlanarVelocityAction::VisualLog);
+
+	const TLearningArrayView<2, const FVector> VelocityView = FeatureObject->InstanceData->ConstView(FeatureObject->VelocityHandle);
+	const TLearningArrayView<2, const float> FeatureView = FeatureObject->InstanceData->ConstView(FeatureObject->FeatureHandle);
+
+	for (const int32 Instance : Instances)
+	{
+		if (const AActor* Actor = Cast<AActor>(Interactor->GetAgent(Instance)))
+		{
+			const FVector Velocity = VelocityView[Instance][0];
+
+			UE_LEARNING_AGENTS_VLOG_ARROW(this, LogLearning, Display,
+				Actor->GetActorLocation(),
+				Actor->GetActorLocation() + Velocity,
+				VisualLogColor.ToFColor(true),
+				TEXT(""));
+
+			UE_LEARNING_AGENTS_VLOG_STRING(this, LogLearning, Display,
+				Actor->GetActorLocation() + Velocity,
+				VisualLogColor.ToFColor(true),
+				TEXT("Velocity: [% 6.3f % 6.3f % 6.3f]"),
+				Velocity.X, Velocity.Y, Velocity.Z);
+
+			UE_LEARNING_AGENTS_VLOG_PLANE(this, LogLearning, Display,
+				Actor->GetActorLocation(),
+				FQuat::Identity,
+				FeatureObject->Axis0,
+				FeatureObject->Axis1,
+				VisualLogColor.ToFColor(true),
+				TEXT(""));
+
+			UE_LEARNING_AGENTS_VLOG_TRANSFORM(this, LogLearning, Display,
+				Actor->GetActorLocation(),
+				FQuat::Identity,
+				VisualLogColor.ToFColor(true),
+				TEXT("Agent %i\nScale: [% 6.2f]\nEncoded: %s"),
+				Instance,
+				FeatureObject->Scale,
+				*UE::Learning::Array::FormatFloat(FeatureView[Instance]));
+		}
+	}
+}
+#endif
+
+//------------------------------------------------------------------
+
 URotationAction* URotationAction::AddRotationAction(ULearningAgentsInteractor* InInteractor, const FName Name, const float Scale)
 {
 	return UE::Learning::Agents::Actions::Private::AddAction<URotationAction, UE::Learning::FRotationVectorFeature>(InInteractor, Name, TEXT("AddRotationAction"), 1, Scale);
diff --git a/Engine/Plugins/Experimental/LearningAgents/Source/LearningAgents/Private/LearningAgentsHelpers.cpp b/Engine/Plugins/Experimental/LearningAgents/Source/LearningAgents/Private/LearningAgentsHelpers.cpp
--- a/Engine/Plugins/Experimental/LearningAgents/Source/LearningAgents/Private/LearningAgentsHelpers.cpp
+++ b/Engine/Plugins/Experimental/LearningAgents/Source/LearningAgents/Private/LearningAgentsHelpers.cpp
@@ -724,3 +724,98 @@ void URayCastHelper::RayCastGridHeights(
 	}
 #endif
 }
+
+void URayCastHelper::RayCastRadial(
+	TArray<float>& OutDistances,
+	const int32 AgentId,
+	const FVector Position,
+	const FRotator Rotation,
+	const int32 RayNum,
+	const float MinAngle,
+	const float MaxAngle,
+	const float MaxRayDist,
+	const FVector LocalForward,
+	const ECollisionChannel CollisionChannel) const
+{
+	if (!ManagerComponent->HasAgent(AgentId))
+	{
+		UE_LOG(LogLearning, Error, TEXT("%s: AgentId %d not found in the agents set."), *GetName(), AgentId);
+		OutDistances.Empty();
+		return;
+	}
+
+	if (RayNum < 1)
+	{
+		UE_LOG(LogLearning, Error, TEXT("%s: Need at least 1 ray."), *GetName());
+		OutDistances.Empty();
+		return;
+	}
+
+	OutDistances.SetNumUninitialized(RayNum);
+
+	TArray<FHitResult, TInlineAllocator<25>> TraceHits;
+	TraceHits.SetNum(RayNum);
+
+
+	for (int32 RayIdx = 0; RayIdx < RayNum; RayIdx++)
+	{
+		const float Alpha = RayNum == 1 ? 0.0f : ((float)RayIdx) / (RayNum - 1);
+		const FRotator YawRotator(0.0f, Rotation.Yaw + Alpha * (MaxAngle - MinAngle) + MinAngle, 0.0f);
+
+		const FVector Start = Position;
+		const FVector End = Position + YawRotator.RotateVector(MaxRayDist * LocalForward.GetSafeNormal(UE_SMALL_NUMBER, FVector::ForwardVector));
+
+		FCollisionObjectQueryParams ObjectQueryParams;
+		ObjectQueryParams.AddObjectTypesToQuery(CollisionChannel);
+
+		const bool bHit = GetWorld()->LineTraceSingleByObjectType(TraceHits[RayIdx], Start, End, ObjectQueryParams);
+
+		if (bHit)
+		{
+			OutDistances[RayIdx] = TraceHits[RayIdx].ImpactPoint.Z;
+		}
+		else
+		{
+			OutDistances[RayIdx] = MaxRayDist;
+		}
+	}
+
+#if UE_LEARNING_AGENTS_ENABLE_VISUAL_LOG
+	{
+		UE_LEARNING_AGENTS_VLOG_TRANSFORM(this, LogLearning, Display,
+			Position,
+			Rotation,
+			VisualLogColor.ToFColor(true),
+			TEXT(""));
+
+		for (int32 RayIdx = 0; RayIdx < RayNum; RayIdx++)
+		{
+			const float Alpha = RayNum == 1 ? 0.0f : ((float)RayIdx) / (RayNum - 1);
+			const FRotator YawRotator(0.0f, Rotation.Yaw + Alpha * (MaxAngle - MinAngle) + MinAngle, 0.0f);
+
+			const FVector Start = Position;
+			const FVector End = Position + YawRotator.RotateVector(MaxRayDist * LocalForward.GetSafeNormal(UE_SMALL_NUMBER, FVector::ForwardVector));
+
+			const FVector ImpactPoint =
+				TraceHits[RayIdx].bBlockingHit ?
+				TraceHits[RayIdx].ImpactPoint : End;
+
+			UE_LEARNING_AGENTS_VLOG_ARROW(this, LogLearning, Display,
+				Start,
+				End,
+				VisualLogColor.ToFColor(true),
+				TEXT(""));
+
+			UE_LEARNING_AGENTS_VLOG_LOCATION(this, LogLearning, Display,
+				ImpactPoint,
+				2.5f,
+				VisualLogColor.ToFColor(true),
+				TEXT("Agent %i\nRay: %i\nImpact Point: [% 6.1f % 6.1f % 6.1f]\nDistance: [% 6.1f]"),
+				AgentId,
+				RayIdx,
+				ImpactPoint.X, ImpactPoint.Y, ImpactPoint.Z,
+				OutDistances[RayIdx]);
+		}
+	}
+#endif
+}
diff --git a/Engine/Plugins/Experimental/LearningAgents/Source/LearningAgents/Private/LearningAgentsNeuralNetwork.cpp b/Engine/Plugins/Experimental/LearningAgents/Source/LearningAgents/Private/LearningAgentsNeuralNetwork.cpp
--- a/Engine/Plugins/Experimental/LearningAgents/Source/LearningAgents/Private/LearningAgentsNeuralNetwork.cpp
+++ b/Engine/Plugins/Experimental/LearningAgents/Source/LearningAgents/Private/LearningAgentsNeuralNetwork.cpp
@@ -95,6 +95,12 @@ void ULearningAgentsNeuralNetwork::Serialize(FArchive& Ar)
 	}
 }
 
+void ULearningAgentsNeuralNetwork::ResetNetwork()
+{
+	NeuralNetwork.Reset();
+	ForceMarkDirty();
+}
+
 void ULearningAgentsNeuralNetwork::LoadNetworkFromSnapshot(const FFilePath& File)
 {
 	TArray<uint8> RecordingData;
@@ -127,14 +133,38 @@ void ULearningAgentsNeuralNetwork::LoadNetworkFromSnapshot(const FFilePath& File
 			return;
 		}
 
-		if (!NeuralNetwork)
+		TSharedPtr<UE::Learning::FNeuralNetwork> TempNeuralNetwork = MakeShared<UE::Learning::FNeuralNetwork>();
+		TempNeuralNetwork->DeserializeFromBytes(Offset, RecordingData);
+
+		if (Offset != RecordingData.Num())
 		{
-			NeuralNetwork = MakeShared<UE::Learning::FNeuralNetwork>();
+			UE_LOG(LogLearning, Error, TEXT("%s: Failed to load network. Unexpected end of file."));
+			return;
 		}
 
-		NeuralNetwork->DeserializeFromBytes(Offset, RecordingData);
+		if (NeuralNetwork)
+		{
+			// If we already have a neural network check settings match
 
-		UE_LEARNING_CHECK(Offset == RecordingData.Num());
+			if (TempNeuralNetwork->GetInputNum() != NeuralNetwork->GetInputNum() ||
+				TempNeuralNetwork->GetOutputNum() != NeuralNetwork->GetOutputNum() ||
+				TempNeuralNetwork->GetLayerNum() != NeuralNetwork->GetLayerNum() ||
+				TempNeuralNetwork->ActivationFunction != NeuralNetwork->ActivationFunction)
+			{
+				UE_LOG(LogLearning, Error, TEXT("%s: Failed to load network from snapshot as settings don't match."), *GetName());
+				return;
+			}
+			else
+			{
+				*NeuralNetwork = *TempNeuralNetwork;
+			}
+		}
+		else
+		{
+			// Otherwise use loaded neural network as-is
+
+			NeuralNetwork = TempNeuralNetwork;
+		}
 
 		ForceMarkDirty();
 	}
diff --git a/Engine/Plugins/Experimental/LearningAgents/Source/LearningAgents/Public/LearningAgentsActions.h b/Engine/Plugins/Experimental/LearningAgents/Source/LearningAgents/Public/LearningAgentsActions.h
--- a/Engine/Plugins/Experimental/LearningAgents/Source/LearningAgents/Public/LearningAgentsActions.h
+++ b/Engine/Plugins/Experimental/LearningAgents/Source/LearningAgents/Public/LearningAgentsActions.h
@@ -14,6 +14,7 @@ namespace UE::Learning
 {
 	struct FFeatureObject;
 	struct FFloatFeature;
+	struct FPlanarVelocityFeature;
 	struct FRotationVectorFeature;
 }
 
@@ -87,7 +88,7 @@ public:
 	float GetFloatAction(const int32 AgentId);
 
 	/**
-	 * Sets the data for this action. Call during ULearningAgentsController::EncodeActions event.
+	 * Sets the data for this action. Call during ULearningAgentsController::SetActions event.
 	 * @param AgentId The agent id to set data for.
 	 * @param Value The current action value.
 	 */	
@@ -130,7 +131,7 @@ public:
 	void GetFloatArrayAction(const int32 AgentId, TArray<float>& OutValues);
 
 	/**
-	 * Sets the data for this action. Call during ULearningAgentsController::EncodeActions event.
+	 * Sets the data for this action. Call during ULearningAgentsController::SetActions event.
 	 * @param AgentId The agent id to set data for.
 	 * @param Values The current action values.
 	 */
@@ -174,7 +175,7 @@ public:
 	FVector GetVectorAction(const int32 AgentId);
 
 	/**
-	 * Sets the data for this action. Call during ULearningAgentsController::EncodeActions event.
+	 * Sets the data for this action. Call during ULearningAgentsController::SetActions event.
 	 * @param AgentId The agent id to set data for.
 	 * @param Value The current action values.
 	 */
@@ -217,7 +218,7 @@ public:
 	void GetVectorArrayAction(const int32 AgentId, TArray<FVector>& OutVectors);
 
 	/**
-	 * Sets the data for this action. Call during ULearningAgentsController::EncodeActions event.
+	 * Sets the data for this action. Call during ULearningAgentsController::SetActions event.
 	 * @param AgentId The agent id to set data for.
 	 * @param Vectors The current action values.
 	 */
@@ -234,6 +235,51 @@ public:
 
 //------------------------------------------------------------------
 
+/** A planar velocity action. */
+UCLASS()
+class LEARNINGAGENTS_API UPlanarVelocityAction : public ULearningAgentsAction
+{
+	GENERATED_BODY()
+
+public:
+
+	/**
+	 * Adds a new planar velocity action to the given agent interactor. The axis parameters define the plane.
+	 * Call during ULearningAgentsInteractor::SetupActions event.
+	 * @param InInteractor The agent interactor to add this action to.
+	 * @param Name The name of this new action. Used for debugging.
+	 * @param Scale Used to normalize the data for the action.
+	 * @param Axis0 The forward axis of the plane.
+	 * @param Axis1 The right axis of the plane.
+	 * @return The newly created action.
+	 */
+	UFUNCTION(BlueprintCallable, Category = "LearningAgents", meta = (DefaultToSelf = "InInteractor"))
+	static UPlanarVelocityAction* AddPlanarVelocityAction(ULearningAgentsInteractor* InInteractor, const FName Name = NAME_None, const float Scale = 200.0f, const FVector Axis0 = FVector::ForwardVector, const FVector Axis1 = FVector::RightVector);
+
+	/**
+	 * Gets the data for this action. Call during ULearningAgentsInteractor::GetActions event.
+	 * @param AgentId The agent id to get data for.
+	 * @return The current action value.
+	 */
+	UFUNCTION(BlueprintCallable, Category = "LearningAgents", meta = (AgentId = "-1"))
+	FVector GetPlanarVelocityAction(const int32 AgentId);
+
+	/**
+	 * Sets the data for this action. Call during ULearningAgentsController::SetActions event.
+	 * @param AgentId The agent id this data corresponds to.
+	 * @param Velocity The velocity currently being observed.
+	 */
+	UFUNCTION(BlueprintCallable, Category = "LearningAgents", meta = (AgentId = "-1"))
+	void SetPlanarVelocityAction(const int32 AgentId, const FVector Velocity);
+
+#if UE_LEARNING_AGENTS_ENABLE_VISUAL_LOG
+	/** Describes this action to the visual logger for debugging purposes. */
+	virtual void VisualLog(const UE::Learning::FIndexSet Instances) const override;
+#endif
+
+	TSharedPtr<UE::Learning::FPlanarVelocityFeature> FeatureObject;
+};
+
 /** A rotation action. */
 UCLASS()
 class LEARNINGAGENTS_API URotationAction : public ULearningAgentsAction
diff --git a/Engine/Plugins/Experimental/LearningAgents/Source/LearningAgents/Public/LearningAgentsHelpers.h b/Engine/Plugins/Experimental/LearningAgents/Source/LearningAgents/Public/LearningAgentsHelpers.h
--- a/Engine/Plugins/Experimental/LearningAgents/Source/LearningAgents/Public/LearningAgentsHelpers.h
+++ b/Engine/Plugins/Experimental/LearningAgents/Source/LearningAgents/Public/LearningAgentsHelpers.h
@@ -285,4 +285,30 @@ public:
 		const float MaxHeight = 10000.0f,
 		const float MinHeight = -10000.0f,
 		const ECollisionChannel CollisionChannel = ECollisionChannel::ECC_WorldStatic) const;
+
+	/**
+	 * Casts rays out in a radius from a given position.
+	 * @param OutDistances The output array of distances for each ray from the start position.
+	 * @param AgentId The agent id to run the helper for.
+	 * @param Position The central position of the ray cast.
+	 * @param Rotation The orientation of the ray cast.
+	 * @param RayNum The number of rays to cast.
+	 * @param MinAngle The minimum angle to spread the rays.
+	 * @param MaxAngle The maximum angle to spread the rays.
+	 * @param MaxRayDist The maximum distance at which to shoot rays.
+	 * @param CollisionChannel The collision channel to use for the ray-casts.
+	 */
+	UFUNCTION(BlueprintPure = false, Category = "LearningAgents", meta = (AgentId = "-1"))
+	void RayCastRadial(
+		TArray<float>& OutDistances,
+		const int32 AgentId,
+		const FVector Position,
+		const FRotator Rotation,
+		const int32 RayNum = 7,
+		const float MinAngle = -90.0f,
+		const float MaxAngle = 90.0f,
+		const float MaxRayDist = 500.0f,
+		const FVector LocalForward = FVector::ForwardVector,
+		const ECollisionChannel CollisionChannel = ECollisionChannel::ECC_WorldStatic) const;
+
 };
diff --git a/Engine/Plugins/Experimental/LearningAgents/Source/LearningAgents/Public/LearningAgentsNeuralNetwork.h b/Engine/Plugins/Experimental/LearningAgents/Source/LearningAgents/Public/LearningAgentsNeuralNetwork.h
--- a/Engine/Plugins/Experimental/LearningAgents/Source/LearningAgents/Public/LearningAgentsNeuralNetwork.h
+++ b/Engine/Plugins/Experimental/LearningAgents/Source/LearningAgents/Public/LearningAgentsNeuralNetwork.h
@@ -50,6 +50,12 @@ public:
 	/** Serialize this neural network to/from the given archive. */
 	virtual void Serialize(FArchive& Ar) override;
 
+	/**
+	 * Resets this network asset to be empty.
+	 */
+	UFUNCTION(CallInEditor, Category = "LearningAgents")
+	void ResetNetwork();
+
 	/**
 	 * Load this network from a snapshot.
 	 * @param File The snapshot file.
