LearningAgents: Exposed AgentsAdded, AgentsRemoved, and AgentsReset events to ULearningAgentsManagerComponent and propagated these to observations, actions, rewards, completions, and helpers.
Made FloatReward take into account the provided weight. Fixed bug in visual log of policy. Made naming of objects unique for visual logger. Added ConditionalReward. Added UCollisionMonitorHelper.

diff --git a/Engine/Plugins/Experimental/LearningAgents/Source/Learning/Private/LearningRewardObject.cpp b/Engine/Plugins/Experimental/LearningAgents/Source/Learning/Private/LearningRewardObject.cpp
--- a/Engine/Plugins/Experimental/LearningAgents/Source/Learning/Private/LearningRewardObject.cpp
+++ b/Engine/Plugins/Experimental/LearningAgents/Source/Learning/Private/LearningRewardObject.cpp
@@ -393,9 +393,13 @@ namespace UE::Learning
 		UE_LEARNING_TRACE_CPUPROFILER_EVENT_SCOPE(FFloatReward::Evaluate);
 
 		const TLearningArrayView<1, const float> Value = InstanceData->ConstView(ValueHandle);
+		const TLearningArrayView<1, const float> Weight = InstanceData->ConstView(WeightHandle);
 		TLearningArrayView<1, float> Reward = InstanceData->View(RewardHandle);
 
-		Array::Copy(Reward, Value, Instances);
+		for (const int32 InstanceIdx : Instances)
+		{
+			Reward[InstanceIdx] = Weight[InstanceIdx] * Value[InstanceIdx];
+		}
 	}
 
 	//------------------------------------------------------------------
diff --git a/Engine/Plugins/Experimental/LearningAgents/Source/LearningAgents/Private/LearningAgentsActions.cpp b/Engine/Plugins/Experimental/LearningAgents/Source/LearningAgents/Private/LearningAgentsActions.cpp
--- a/Engine/Plugins/Experimental/LearningAgents/Source/LearningAgents/Private/LearningAgentsActions.cpp
+++ b/Engine/Plugins/Experimental/LearningAgents/Source/LearningAgents/Private/LearningAgentsActions.cpp
@@ -28,8 +28,10 @@ namespace UE::Learning::Agents::Actions::Private
 			return nullptr;
 		}
 
-		ActionUObject* Action = NewObject<ActionUObject>(InInteractor, Name);
+		const FName UniqueName = MakeUniqueObjectName(InInteractor, ActionUObject::StaticClass(), Name, EUniqueObjectNameOptions::GloballyUnique);
 
+		ActionUObject* Action = NewObject<ActionUObject>(InInteractor, UniqueName);
+		Action->Init(InInteractor->GetAgentManager()->GetMaxAgentNum());
 		Action->Interactor = InInteractor;
 		Action->FeatureObject = MakeShared<ActionFObject>(
 			Action->GetFName(),
@@ -37,11 +39,6 @@ namespace UE::Learning::Agents::Actions::Private
 			InInteractor->GetAgentManager()->GetMaxAgentNum(),
 			Forward<InArgTypes>(Args)...);
 
-		Action->AgentGetIteration.SetNumUninitialized({ InInteractor->GetAgentManager()->GetMaxAgentNum() });
-		Action->AgentSetIteration.SetNumUninitialized({ InInteractor->GetAgentManager()->GetMaxAgentNum() });
-		UE::Learning::Array::Set<1, uint64>(Action->AgentGetIteration, INDEX_NONE);
-		UE::Learning::Array::Set<1, uint64>(Action->AgentSetIteration, INDEX_NONE);
-
 		// We assume all supported action feature objects can be encoded and decoded
 		UE_LEARNING_CHECK(Action->FeatureObject->IsEncodable() && Action->FeatureObject->IsDecodable());
 
@@ -53,6 +50,44 @@ namespace UE::Learning::Agents::Actions::Private
 
 //------------------------------------------------------------------
 
+void ULearningAgentsAction::Init(const int32 MaxAgentNum)
+{
+	AgentGetIteration.SetNumUninitialized({ MaxAgentNum });
+	AgentSetIteration.SetNumUninitialized({ MaxAgentNum });
+	UE::Learning::Array::Set<1, uint64>(AgentGetIteration, INDEX_NONE);
+	UE::Learning::Array::Set<1, uint64>(AgentSetIteration, INDEX_NONE);
+}
+
+void ULearningAgentsAction::OnAgentsAdded(const TArray<int32>& AgentIds)
+{
+	UE::Learning::Array::Set<1, uint64>(AgentGetIteration, 0, AgentIds);
+	UE::Learning::Array::Set<1, uint64>(AgentSetIteration, 0, AgentIds);
+}
+
+void ULearningAgentsAction::OnAgentsRemoved(const TArray<int32>& AgentIds)
+{
+	UE::Learning::Array::Set<1, uint64>(AgentGetIteration, INDEX_NONE, AgentIds);
+	UE::Learning::Array::Set<1, uint64>(AgentSetIteration, INDEX_NONE, AgentIds);
+}
+
+void ULearningAgentsAction::OnAgentsReset(const TArray<int32>& AgentIds)
+{
+	UE::Learning::Array::Set<1, uint64>(AgentGetIteration, 0, AgentIds);
+	UE::Learning::Array::Set<1, uint64>(AgentSetIteration, 0, AgentIds);
+}
+
+uint64 ULearningAgentsAction::GetAgentGetIteration(const int32 AgentId) const
+{
+	return AgentGetIteration[AgentId];
+}
+
+uint64 ULearningAgentsAction::GetAgentSetIteration(const int32 AgentId) const
+{
+	return AgentSetIteration[AgentId];
+}
+
+//------------------------------------------------------------------
+
 UFloatAction* UFloatAction::AddFloatAction(ULearningAgentsInteractor* InInteractor, const FName Name, const float Scale)
 {
 	return UE::Learning::Agents::Actions::Private::AddAction<UFloatAction,UE::Learning::FFloatFeature>(InInteractor, Name, TEXT("AddFloatAction"), 1, Scale);
diff --git a/Engine/Plugins/Experimental/LearningAgents/Source/LearningAgents/Private/LearningAgentsController.cpp b/Engine/Plugins/Experimental/LearningAgents/Source/LearningAgents/Private/LearningAgentsController.cpp
--- a/Engine/Plugins/Experimental/LearningAgents/Source/LearningAgents/Private/LearningAgentsController.cpp
+++ b/Engine/Plugins/Experimental/LearningAgents/Source/LearningAgents/Private/LearningAgentsController.cpp
@@ -46,6 +46,8 @@ void ULearningAgentsController::SetupController(ULearningAgentsInteractor* InInt
 	Interactor = InInteractor;
 
 	bIsSetup = true;
+
+	OnAgentsAdded(Manager->GetAllAgentIds());
 }
 
 void ULearningAgentsController::EncodeActions()
@@ -73,23 +75,23 @@ void ULearningAgentsController::EncodeActions()
 	{
 		for (const int32 AgentId : Manager->GetAllAgentSet())
 		{
-			if (ActionObject->AgentSetIteration[AgentId] == ActionEncodingAgentIteration[AgentId])
+			if (ActionObject->GetAgentSetIteration(AgentId) == ActionEncodingAgentIteration[AgentId])
 			{
-				UE_LOG(LogLearning, Warning, TEXT("%s: Action %s for agent with id %i has not been set (got iteration %i, expected iteration %i) and so agent will not have actions encoded."), *GetName(), *ActionObject->GetName(), AgentId, ActionObject->AgentSetIteration[AgentId], ActionEncodingAgentIteration[AgentId] + 1);
+				UE_LOG(LogLearning, Warning, TEXT("%s: Action %s for agent with id %i has not been set (got iteration %i, expected iteration %i) and so agent will not have actions encoded."), *GetName(), *ActionObject->GetName(), AgentId, ActionObject->GetAgentSetIteration(AgentId), ActionEncodingAgentIteration[AgentId] + 1);
 				ValidAgentStatus[AgentId] = false;
 				continue;
 			}
 
-			if (ActionObject->AgentSetIteration[AgentId] > ActionEncodingAgentIteration[AgentId] + 1)
+			if (ActionObject->GetAgentSetIteration(AgentId) > ActionEncodingAgentIteration[AgentId] + 1)
 			{
-				UE_LOG(LogLearning, Warning, TEXT("%s: Action %s for agent with id %i appears to have been set multiple times (got iteration %i, expected iteration %i) and so agent will not have actions encoded."), *GetName(), *ActionObject->GetName(), AgentId, ActionObject->AgentSetIteration[AgentId], ActionEncodingAgentIteration[AgentId] + 1);
+				UE_LOG(LogLearning, Warning, TEXT("%s: Action %s for agent with id %i appears to have been set multiple times (got iteration %i, expected iteration %i) and so agent will not have actions encoded."), *GetName(), *ActionObject->GetName(), AgentId, ActionObject->GetAgentSetIteration(AgentId), ActionEncodingAgentIteration[AgentId] + 1);
 				ValidAgentStatus[AgentId] = false;
 				continue;
 			}
 
-			if (ActionObject->AgentSetIteration[AgentId] != ActionEncodingAgentIteration[AgentId] + 1)
+			if (ActionObject->GetAgentSetIteration(AgentId) != ActionEncodingAgentIteration[AgentId] + 1)
 			{
-				UE_LOG(LogLearning, Warning, TEXT("%s: Action %s for agent with id %i does not have a matching iteration number (got iteration %i, expected iteration %i) and so agent will not have actions encoded."), *GetName(), *ActionObject->GetName(), AgentId, ActionObject->AgentSetIteration[AgentId], ActionEncodingAgentIteration[AgentId] + 1);
+				UE_LOG(LogLearning, Warning, TEXT("%s: Action %s for agent with id %i does not have a matching iteration number (got iteration %i, expected iteration %i) and so agent will not have actions encoded."), *GetName(), *ActionObject->GetName(), AgentId, ActionObject->GetAgentSetIteration(AgentId), ActionEncodingAgentIteration[AgentId] + 1);
 				ValidAgentStatus[AgentId] = false;
 				continue;
 			}
diff --git a/Engine/Plugins/Experimental/LearningAgents/Source/LearningAgents/Private/LearningAgentsCritic.cpp b/Engine/Plugins/Experimental/LearningAgents/Source/LearningAgents/Private/LearningAgentsCritic.cpp
--- a/Engine/Plugins/Experimental/LearningAgents/Source/LearningAgents/Private/LearningAgentsCritic.cpp
+++ b/Engine/Plugins/Experimental/LearningAgents/Source/LearningAgents/Private/LearningAgentsCritic.cpp
@@ -4,6 +4,7 @@
 
 #include "LearningAgentsManager.h"
 #include "LearningAgentsInteractor.h"
+#include "LearningAgentsHelpers.h"
 #include "LearningFeatureObject.h"
 #include "LearningNeuralNetwork.h"
 #include "LearningNeuralNetworkObject.h"
@@ -87,7 +88,9 @@ void ULearningAgentsCritic::SetupCritic(
 	{
 		// Create New Neural Network Asset
 
-		Network = NewObject<ULearningAgentsNeuralNetwork>(this, TEXT("CriticNetwork"));
+		const FName UniqueName = MakeUniqueObjectName(this, ULearningAgentsNeuralNetwork::StaticClass(), TEXT("CriticNetwork"), EUniqueObjectNameOptions::GloballyUnique);
+
+		Network = NewObject<ULearningAgentsNeuralNetwork>(this, UniqueName);
 		Network->NeuralNetwork = MakeShared<UE::Learning::FNeuralNetwork>();
 		Network->NeuralNetwork->Resize(
 			Interactor->GetObservationFeature().DimNum(),
@@ -108,18 +111,24 @@ void ULearningAgentsCritic::SetupCritic(
 
 	DiscountedReturnAgentIteration.SetNumUninitialized({ Manager->GetMaxAgentNum() });
 	UE::Learning::Array::Set<1, uint64>(DiscountedReturnAgentIteration, INDEX_NONE);
-	UE::Learning::Array::Set<1, uint64>(DiscountedReturnAgentIteration, 0, Manager->GetAllAgentSet());
 
 	bIsSetup = true;
+
+	OnAgentsAdded(Manager->GetAllAgentIds());
 }
 
 void ULearningAgentsCritic::OnAgentsAdded(const TArray<int32>& AgentIds)
 {
 	if (IsSetup())
 	{
-		UE::Learning::FIndexSet AgentSet = AgentIds;
-		AgentSet.TryMakeSlice();
-		UE::Learning::Array::Set<1, uint64>(DiscountedReturnAgentIteration, 0, AgentSet);
+		UE::Learning::Array::Set<1, uint64>(DiscountedReturnAgentIteration, 0, AgentIds);
+
+		for (ULearningAgentsHelper* Helper : HelperObjects)
+		{
+			Helper->OnAgentsAdded(AgentIds);
+		}
+
+		AgentsAdded(AgentIds);
 	}
 }
 
@@ -127,9 +136,14 @@ void ULearningAgentsCritic::OnAgentsRemoved(const TArray<int32>& AgentIds)
 {
 	if (IsSetup())
 	{
-		UE::Learning::FIndexSet AgentSet = AgentIds;
-		AgentSet.TryMakeSlice();
-		UE::Learning::Array::Set<1, uint64>(DiscountedReturnAgentIteration, INDEX_NONE, AgentSet);
+		UE::Learning::Array::Set<1, uint64>(DiscountedReturnAgentIteration, INDEX_NONE, AgentIds);
+
+		for (ULearningAgentsHelper* Helper : HelperObjects)
+		{
+			Helper->OnAgentsRemoved(AgentIds);
+		}
+
+		AgentsRemoved(AgentIds);
 	}
 }
 
@@ -137,9 +151,14 @@ void ULearningAgentsCritic::OnAgentsReset(const TArray<int32>& AgentIds)
 {
 	if (IsSetup())
 	{
-		UE::Learning::FIndexSet AgentSet = AgentIds;
-		AgentSet.TryMakeSlice();
-		UE::Learning::Array::Set<1, uint64>(DiscountedReturnAgentIteration, 0, AgentSet);
+		UE::Learning::Array::Set<1, uint64>(DiscountedReturnAgentIteration, 0, AgentIds);
+
+		for (ULearningAgentsHelper* Helper : HelperObjects)
+		{
+			Helper->OnAgentsReset(AgentIds);
+		}
+
+		AgentsReset(AgentIds);
 	}
 }
 
diff --git a/Engine/Plugins/Experimental/LearningAgents/Source/LearningAgents/Private/LearningAgentsHelpers.cpp b/Engine/Plugins/Experimental/LearningAgents/Source/LearningAgents/Private/LearningAgentsHelpers.cpp
--- a/Engine/Plugins/Experimental/LearningAgents/Source/LearningAgents/Private/LearningAgentsHelpers.cpp
+++ b/Engine/Plugins/Experimental/LearningAgents/Source/LearningAgents/Private/LearningAgentsHelpers.cpp
@@ -12,6 +12,12 @@
 
 //------------------------------------------------------------------
 
+void ULearningAgentsHelper::OnAgentsAdded(const TArray<int32>& AgentIds) {}
+void ULearningAgentsHelper::OnAgentsRemoved(const TArray<int32>& AgentIds) {}
+void ULearningAgentsHelper::OnAgentsReset(const TArray<int32>& AgentIds) {}
+
+//------------------------------------------------------------------
+
 USplineComponentHelper* USplineComponentHelper::AddSplineComponentHelper(ULearningAgentsManagerComponent* InManagerComponent, const FName Name)
 {
 	if (!InManagerComponent)
@@ -20,7 +26,9 @@ USplineComponentHelper* USplineComponentHelper::AddSplineComponentHelper(ULearni
 		return nullptr;
 	}
 
-	USplineComponentHelper* Helper = NewObject<USplineComponentHelper>(InManagerComponent, Name);
+	const FName UniqueName = MakeUniqueObjectName(InManagerComponent, USplineComponentHelper::StaticClass(), Name, EUniqueObjectNameOptions::GloballyUnique);
+
+	USplineComponentHelper* Helper = NewObject<USplineComponentHelper>(InManagerComponent, UniqueName);
 	Helper->ManagerComponent = InManagerComponent;
 	InManagerComponent->AddHelper(Helper);
 	return Helper;
@@ -426,7 +434,9 @@ UProjectionHelper* UProjectionHelper::AddProjectionHelper(ULearningAgentsManager
 		return nullptr;
 	}
 
-	UProjectionHelper* Helper = NewObject<UProjectionHelper>(InManagerComponent, Name);
+	const FName UniqueName = MakeUniqueObjectName(InManagerComponent, UProjectionHelper::StaticClass(), Name, EUniqueObjectNameOptions::GloballyUnique);
+
+	UProjectionHelper* Helper = NewObject<UProjectionHelper>(InManagerComponent, UniqueName);
 	Helper->ManagerComponent = InManagerComponent;
 	InManagerComponent->AddHelper(Helper);
 	return Helper;
@@ -535,7 +545,9 @@ UMeshComponentHelper* UMeshComponentHelper::AddMeshComponentHelper(ULearningAgen
 		return nullptr;
 	}
 
-	UMeshComponentHelper* Helper = NewObject<UMeshComponentHelper>(InManagerComponent, Name);
+	const FName UniqueName = MakeUniqueObjectName(InManagerComponent, UMeshComponentHelper::StaticClass(), Name, EUniqueObjectNameOptions::GloballyUnique);
+
+	UMeshComponentHelper* Helper = NewObject<UMeshComponentHelper>(InManagerComponent, UniqueName);
 	Helper->ManagerComponent = InManagerComponent;
 	InManagerComponent->AddHelper(Helper);
 	return Helper;
@@ -608,7 +620,9 @@ URayCastHelper* URayCastHelper::AddRayCastHelper(ULearningAgentsManagerComponent
 		return nullptr;
 	}
 
-	URayCastHelper* Helper = NewObject<URayCastHelper>(InManagerComponent, Name);
+	const FName UniqueName = MakeUniqueObjectName(InManagerComponent, URayCastHelper::StaticClass(), Name, EUniqueObjectNameOptions::GloballyUnique);
+
+	URayCastHelper* Helper = NewObject<URayCastHelper>(InManagerComponent, UniqueName);
 	Helper->ManagerComponent = InManagerComponent;
 	InManagerComponent->AddHelper(Helper);
 	return Helper;
@@ -772,7 +786,7 @@ void URayCastHelper::RayCastRadial(
 
 		if (bHit)
 		{
-			OutDistances[RayIdx] = TraceHits[RayIdx].ImpactPoint.Z;
+			OutDistances[RayIdx] = TraceHits[RayIdx].Distance;
 		}
 		else
 		{
@@ -819,3 +833,144 @@ void URayCastHelper::RayCastRadial(
 	}
 #endif
 }
+
+
+
+UCollisionMonitorHelper* UCollisionMonitorHelper::AddCollisionMonitorHelper(
+	ULearningAgentsManagerComponent* InManagerComponent,
+	const FName Name)
+{
+	if (!InManagerComponent)
+	{
+		UE_LOG(LogLearning, Error, TEXT("AddCollisionMonitorHelper: InManagerComponent is nullptr."));
+		return nullptr;
+	}
+
+	const FName UniqueName = MakeUniqueObjectName(InManagerComponent, UCollisionMonitorHelper::StaticClass(), Name, EUniqueObjectNameOptions::GloballyUnique);
+
+	UCollisionMonitorHelper* Helper = NewObject<UCollisionMonitorHelper>(InManagerComponent, UniqueName);
+	Helper->Init(InManagerComponent->GetAgentManager()->GetMaxAgentNum());
+	Helper->ManagerComponent = InManagerComponent;
+
+	InManagerComponent->AddHelper(Helper);
+	return Helper;
+}
+
+void UCollisionMonitorHelper::Init(const int32 MaxAgentNum)
+{
+	Components.SetNumUninitialized({ MaxAgentNum });
+	OtherComponentTags.SetNumUninitialized({ MaxAgentNum });
+	CollisionsOccured.SetNumUninitialized({ MaxAgentNum });
+
+	UE::Learning::Array::Set<1, UPrimitiveComponent*>(Components, nullptr);
+	UE::Learning::Array::Set<1, FName>(OtherComponentTags, NAME_None);
+	UE::Learning::Array::Set<1, bool>(CollisionsOccured, false);
+}
+
+void UCollisionMonitorHelper::OnAgentsRemoved(const TArray<int32>& AgentIds)
+{
+	ULearningAgentsHelper::OnAgentsRemoved(AgentIds);
+
+	for (const int32 AgentId : AgentIds)
+	{
+		if (Components[AgentId])
+		{
+			Components[AgentId]->OnComponentHit.RemoveDynamic(this, &UCollisionMonitorHelper::HandleOnHit);
+		}
+	}
+
+	UE::Learning::Array::Set<1, UPrimitiveComponent*>(Components, nullptr);
+	UE::Learning::Array::Set<1, FName>(OtherComponentTags, NAME_None, AgentIds);
+	UE::Learning::Array::Set<1, bool>(CollisionsOccured, false, AgentIds);
+}
+
+void UCollisionMonitorHelper::OnAgentsReset(const TArray<int32>& AgentIds)
+{
+	ULearningAgentsHelper::OnAgentsReset(AgentIds);
+
+	UE::Learning::Array::Set(CollisionsOccured, false, AgentIds);
+}
+
+void UCollisionMonitorHelper::SetComponent(const int32 AgentId, UPrimitiveComponent* Component, const FName OtherComponentTag)
+{
+	if (!ManagerComponent->HasAgent(AgentId))
+	{
+		UE_LOG(LogLearning, Error, TEXT("%s: AgentId %d not found in the agents set."), *GetName(), AgentId);
+		return;
+	}
+
+	if (Components[AgentId] != Component)
+	{
+		if (Components[AgentId])
+		{
+			Components[AgentId]->OnComponentHit.RemoveDynamic(this, &UCollisionMonitorHelper::HandleOnHit);
+		}
+
+		Components[AgentId] = Component;
+		OtherComponentTags[AgentId] = OtherComponentTag;
+		Component->OnComponentHit.AddUniqueDynamic(this, &UCollisionMonitorHelper::HandleOnHit);
+	}
+}
+
+bool UCollisionMonitorHelper::GetCollisionOccurred(const int32 AgentId) const
+{
+	if (!ManagerComponent->HasAgent(AgentId))
+	{
+		UE_LOG(LogLearning, Error, TEXT("%s: AgentId %d not found in the agents set."), *GetName(), AgentId);
+		return false;
+	}
+
+	return CollisionsOccured[AgentId];
+}
+
+bool UCollisionMonitorHelper::GetAndResetCollisionOccurred(const int32 AgentId)
+{
+	if (!ManagerComponent->HasAgent(AgentId))
+	{
+		UE_LOG(LogLearning, Error, TEXT("%s: AgentId %d not found in the agents set."), *GetName(), AgentId);
+		return false;
+	}
+
+	bool bOccurred = CollisionsOccured[AgentId];
+	CollisionsOccured[AgentId] = false;
+	return bOccurred;
+}
+
+void UCollisionMonitorHelper::ResetCollisionOccurred(const int32 AgentId)
+{
+	if (!ManagerComponent->HasAgent(AgentId))
+	{
+		UE_LOG(LogLearning, Error, TEXT("%s: AgentId %d not found in the agents set."), *GetName(), AgentId);
+		return;
+	}
+
+	CollisionsOccured[AgentId] = false;
+}
+
+void UCollisionMonitorHelper::HandleOnHit(UPrimitiveComponent* HitComponent, AActor* OtherActor, UPrimitiveComponent* OtherComp, FVector NormalImpulse, const FHitResult& Hit)
+{
+	for (int32 AgentId = 0; AgentId < Components.Num(); AgentId++)
+	{
+		if (Components[AgentId] == HitComponent)
+		{
+			CollisionsOccured[AgentId] = true;
+
+#if UE_LEARNING_AGENTS_ENABLE_VISUAL_LOG
+			{
+				UE_LEARNING_AGENTS_VLOG_ARROW(this, LogLearning, Display,
+					Hit.ImpactPoint,
+					Hit.ImpactPoint + Hit.ImpactNormal * 100.0f,
+					VisualLogColor.ToFColor(true),
+					TEXT("Agent %i\nCollision \"%s\" and \"%s\"\nTag: \"%s\""),
+					AgentId,
+					*HitComponent->GetName(),
+					*OtherComp->GetName(),
+					*OtherComponentTags[AgentId].ToString());
+			}
+#endif
+			return;
+		}
+	}
+
+	UE_LOG(LogLearning, Error, TEXT("%s: HitComponent not found in the set of agent's Components."), *GetName());
+}
\ No newline at end of file
diff --git a/Engine/Plugins/Experimental/LearningAgents/Source/LearningAgents/Private/LearningAgentsInteractor.cpp b/Engine/Plugins/Experimental/LearningAgents/Source/LearningAgents/Private/LearningAgentsInteractor.cpp
--- a/Engine/Plugins/Experimental/LearningAgents/Source/LearningAgents/Private/LearningAgentsInteractor.cpp
+++ b/Engine/Plugins/Experimental/LearningAgents/Source/LearningAgents/Private/LearningAgentsInteractor.cpp
@@ -5,6 +5,7 @@
 #include "LearningAgentsManager.h"
 #include "LearningAgentsObservations.h"
 #include "LearningAgentsActions.h"
+#include "LearningAgentsHelpers.h"
 
 #include "LearningArray.h"
 #include "LearningArrayMap.h"
@@ -80,18 +81,35 @@ void ULearningAgentsInteractor::SetupInteractor()
 	ActionEncodingAgentIteration.SetNumUninitialized({ Manager->GetMaxAgentNum() });
 	UE::Learning::Array::Set<1, uint64>(ObservationEncodingAgentIteration, INDEX_NONE);
 	UE::Learning::Array::Set<1, uint64>(ActionEncodingAgentIteration, INDEX_NONE);
-	SetAgentIteration(Manager->GetAllAgentSet(), 0);
 
 	bIsSetup = true;
+
+	OnAgentsAdded(Manager->GetAllAgentIds());
 }
 
 void ULearningAgentsInteractor::OnAgentsAdded(const TArray<int32>& AgentIds)
 {
 	if (IsSetup())
 	{
-		UE::Learning::FIndexSet AgentSet = AgentIds;
-		AgentSet.TryMakeSlice();
-		SetAgentIteration(AgentSet, 0);
+		UE::Learning::Array::Set<1, uint64>(ObservationEncodingAgentIteration, 0, AgentIds);
+		UE::Learning::Array::Set<1, uint64>(ActionEncodingAgentIteration, 0, AgentIds);
+
+		for (ULearningAgentsHelper* Helper : HelperObjects)
+		{
+			Helper->OnAgentsAdded(AgentIds);
+		}
+
+		for (ULearningAgentsObservation* ObservationObject : ObservationObjects)
+		{
+			ObservationObject->OnAgentsAdded(AgentIds);
+		}
+
+		for (ULearningAgentsAction* ActionObject : ActionObjects)
+		{
+			ActionObject->OnAgentsAdded(AgentIds);
+		}
+
+		AgentsAdded(AgentIds);
 	}
 }
 
@@ -99,9 +117,25 @@ void ULearningAgentsInteractor::OnAgentsRemoved(const TArray<int32>& AgentIds)
 {
 	if (IsSetup())
 	{
-		UE::Learning::FIndexSet AgentSet = AgentIds;
-		AgentSet.TryMakeSlice();
-		SetAgentIteration(AgentSet, INDEX_NONE);
+		UE::Learning::Array::Set<1, uint64>(ObservationEncodingAgentIteration, INDEX_NONE, AgentIds);
+		UE::Learning::Array::Set<1, uint64>(ActionEncodingAgentIteration, INDEX_NONE, AgentIds);
+
+		for (ULearningAgentsHelper* Helper : HelperObjects)
+		{
+			Helper->OnAgentsRemoved(AgentIds);
+		}
+
+		for (ULearningAgentsObservation* ObservationObject : ObservationObjects)
+		{
+			ObservationObject->OnAgentsRemoved(AgentIds);
+		}
+
+		for (ULearningAgentsAction* ActionObject : ActionObjects)
+		{
+			ActionObject->OnAgentsRemoved(AgentIds);
+		}
+
+		AgentsRemoved(AgentIds);
 	}
 }
 
@@ -109,26 +143,25 @@ void ULearningAgentsInteractor::OnAgentsReset(const TArray<int32>& AgentIds)
 {
 	if (IsSetup())
 	{
-		UE::Learning::FIndexSet AgentSet = AgentIds;
-		AgentSet.TryMakeSlice();
-		SetAgentIteration(AgentSet, 0);
-	}
-}
+		UE::Learning::Array::Set<1, uint64>(ObservationEncodingAgentIteration, 0, AgentIds);
+		UE::Learning::Array::Set<1, uint64>(ActionEncodingAgentIteration, 0, AgentIds);
+		
+		for (ULearningAgentsHelper* Helper : HelperObjects)
+		{
+			Helper->OnAgentsReset(AgentIds);
+		}
 
-void ULearningAgentsInteractor::SetAgentIteration(const UE::Learning::FIndexSet AgentIds, const uint64 Value)
-{
-	UE::Learning::Array::Set<1, uint64>(ObservationEncodingAgentIteration, Value, AgentIds);
-	UE::Learning::Array::Set<1, uint64>(ActionEncodingAgentIteration, Value, AgentIds);
+		for (ULearningAgentsObservation* ObservationObject : ObservationObjects)
+		{
+			ObservationObject->OnAgentsReset(AgentIds);
+		}
 
-	for (ULearningAgentsObservation* ObservationObject : ObservationObjects)
-	{
-		UE::Learning::Array::Set<1, uint64>(ObservationObject->AgentIteration, Value, AgentIds);
-	}
+		for (ULearningAgentsAction* ActionObject : ActionObjects)
+		{
+			ActionObject->OnAgentsReset(AgentIds);
+		}
 
-	for (ULearningAgentsAction* ActionObject : ActionObjects)
-	{
-		UE::Learning::Array::Set<1, uint64>(ActionObject->AgentGetIteration, Value, AgentIds);
-		UE::Learning::Array::Set<1, uint64>(ActionObject->AgentSetIteration, Value, AgentIds);
+		AgentsReset(AgentIds);
 	}
 }
 
@@ -227,23 +260,23 @@ void ULearningAgentsInteractor::EncodeObservations(const UE::Learning::FIndexSet
 	{
 		for (const int32 AgentId : AgentSet)
 		{
-			if (ObservationObject->AgentIteration[AgentId] <= ObservationEncodingAgentIteration[AgentId])
+			if (ObservationObject->GetAgentIteration(AgentId) <= ObservationEncodingAgentIteration[AgentId])
 			{
-				UE_LOG(LogLearning, Warning, TEXT("%s: Observation %s for agent with id %i has not been set (got iteration %i, expected iteration %i) and so agent will not have observations encoded."), *GetName(), *ObservationObject->GetName(), AgentId, ObservationObject->AgentIteration[AgentId], ObservationEncodingAgentIteration[AgentId] + 1);
+				UE_LOG(LogLearning, Warning, TEXT("%s: Observation %s for agent with id %i has not been set (got iteration %i, expected iteration %i) and so agent will not have observations encoded."), *GetName(), *ObservationObject->GetName(), AgentId, ObservationObject->GetAgentIteration(AgentId), ObservationEncodingAgentIteration[AgentId] + 1);
 				ValidAgentStatus[AgentId] = false;
 				continue;
 			}
 
-			if (ObservationObject->AgentIteration[AgentId] > ObservationEncodingAgentIteration[AgentId] + 1)
+			if (ObservationObject->GetAgentIteration(AgentId) > ObservationEncodingAgentIteration[AgentId] + 1)
 			{
-				UE_LOG(LogLearning, Warning, TEXT("%s: Observation %s for agent with id %i appears to have been set multiple times (got iteration %i, expected iteration %i) and so agent will not have observations encoded."), *GetName(), *ObservationObject->GetName(), AgentId, ObservationObject->AgentIteration[AgentId], ObservationEncodingAgentIteration[AgentId] + 1);
+				UE_LOG(LogLearning, Warning, TEXT("%s: Observation %s for agent with id %i appears to have been set multiple times (got iteration %i, expected iteration %i) and so agent will not have observations encoded."), *GetName(), *ObservationObject->GetName(), AgentId, ObservationObject->GetAgentIteration(AgentId), ObservationEncodingAgentIteration[AgentId] + 1);
 				ValidAgentStatus[AgentId] = false;
 				continue;
 			}
 
-			if (ObservationObject->AgentIteration[AgentId] != ObservationEncodingAgentIteration[AgentId] + 1)
+			if (ObservationObject->GetAgentIteration(AgentId) != ObservationEncodingAgentIteration[AgentId] + 1)
 			{
-				UE_LOG(LogLearning, Warning, TEXT("%s: Observation %s for agent with id %i does not have a matching iteration number (got iteration %i, expected iteration %i) and so agent will not have observations encoded."), *GetName(), *ObservationObject->GetName(), AgentId, ObservationObject->AgentIteration[AgentId], ObservationEncodingAgentIteration[AgentId] + 1);
+				UE_LOG(LogLearning, Warning, TEXT("%s: Observation %s for agent with id %i does not have a matching iteration number (got iteration %i, expected iteration %i) and so agent will not have observations encoded."), *GetName(), *ObservationObject->GetName(), AgentId, ObservationObject->GetAgentIteration(AgentId), ObservationEncodingAgentIteration[AgentId] + 1);
 				ValidAgentStatus[AgentId] = false;
 				continue;
 			}
@@ -323,21 +356,21 @@ void ULearningAgentsInteractor::DecodeActions(const UE::Learning::FIndexSet Agen
 	{
 		for (const int32 AgentId : ValidAgentSet)
 		{
-			if (ActionObject->AgentGetIteration[AgentId] == ActionEncodingAgentIteration[AgentId] - 1)
+			if (ActionObject->GetAgentGetIteration(AgentId) == ActionEncodingAgentIteration[AgentId] - 1)
 			{
 				UE_LOG(LogLearning, Warning, TEXT("%s: Action %s for agent with id %i appears to have not been got."), *GetName(), *ActionObject->GetName(), AgentId);
 				continue;
 			}
 				
-			if (ActionObject->AgentGetIteration[AgentId] > ActionEncodingAgentIteration[AgentId])
+			if (ActionObject->GetAgentGetIteration(AgentId) > ActionEncodingAgentIteration[AgentId])
 			{
 				UE_LOG(LogLearning, Warning, TEXT("%s: Action %s for agent with id %i appears to have been got multiple times."), *GetName(), *ActionObject->GetName(), AgentId);
 				continue;
 			}
 				
-			if (ActionObject->AgentGetIteration[AgentId] != ActionEncodingAgentIteration[AgentId])
+			if (ActionObject->GetAgentGetIteration(AgentId) != ActionEncodingAgentIteration[AgentId])
 			{
-				UE_LOG(LogLearning, Warning, TEXT("%s: Action %s for agent with id %i does not have a matching iteration number (got %i, expected %i)."), *GetName(), *ActionObject->GetName(), AgentId, ActionObject->AgentGetIteration[AgentId], ActionEncodingAgentIteration[AgentId]);
+				UE_LOG(LogLearning, Warning, TEXT("%s: Action %s for agent with id %i does not have a matching iteration number (got %i, expected %i)."), *GetName(), *ActionObject->GetName(), AgentId, ActionObject->GetAgentGetIteration(AgentId), ActionEncodingAgentIteration[AgentId]);
 				continue;
 			}
 		}
diff --git a/Engine/Plugins/Experimental/LearningAgents/Source/LearningAgents/Private/LearningAgentsManager.cpp b/Engine/Plugins/Experimental/LearningAgents/Source/LearningAgents/Private/LearningAgentsManager.cpp
--- a/Engine/Plugins/Experimental/LearningAgents/Source/LearningAgents/Private/LearningAgentsManager.cpp
+++ b/Engine/Plugins/Experimental/LearningAgents/Source/LearningAgents/Private/LearningAgentsManager.cpp
@@ -321,9 +321,10 @@ void ALearningAgentsManager::GetAgents(TArray<UObject*>& OutAgents, const TArray
 	}
 }
 
-void ALearningAgentsManager::GetAllAgents(TArray<UObject*>& OutAgents, const TSubclassOf<UObject> AgentClass) const
+void ALearningAgentsManager::GetAllAgents(TArray<UObject*>& OutAgents, TArray<int32>& OutAgentIds, const TSubclassOf<UObject> AgentClass) const
 {
-	GetAgents(OutAgents, OccupiedAgentIds, AgentClass);
+	OutAgentIds = OccupiedAgentIds;
+	GetAgents(OutAgents, OutAgentIds, AgentClass);
 }
 
 int32 ALearningAgentsManager::GetAgentId(UObject* Agent) const
@@ -361,11 +362,6 @@ void ALearningAgentsManager::GetAgentIds(TArray<int32>& OutAgentIds, const TArra
 	}
 }
 
-void ALearningAgentsManager::GetAllAgentIds(TArray<int32>& OutAgentIds) const
-{
-	OutAgentIds = OccupiedAgentIds;
-}
-
 int32 ALearningAgentsManager::GetAgentNum() const
 {
 	return OccupiedAgentIds.Num();
diff --git a/Engine/Plugins/Experimental/LearningAgents/Source/LearningAgents/Private/LearningAgentsManagerComponent.cpp b/Engine/Plugins/Experimental/LearningAgents/Source/LearningAgents/Private/LearningAgentsManagerComponent.cpp
--- a/Engine/Plugins/Experimental/LearningAgents/Source/LearningAgents/Private/LearningAgentsManagerComponent.cpp
+++ b/Engine/Plugins/Experimental/LearningAgents/Source/LearningAgents/Private/LearningAgentsManagerComponent.cpp
@@ -2,6 +2,7 @@
 
 #include "LearningAgentsManagerComponent.h"
 #include "LearningAgentsManager.h"
+#include "LearningAgentsHelpers.h"
 #include "LearningLog.h"
 
 ULearningAgentsManagerComponent::ULearningAgentsManagerComponent(const FObjectInitializer& ObjectInitializer) : UActorComponent(ObjectInitializer) {}
@@ -25,11 +26,59 @@ void ULearningAgentsManagerComponent::PostInitProperties()
 	}
 }
 
-void ULearningAgentsManagerComponent::OnAgentsAdded(const TArray<int32>& AgentIds) { }
+void ULearningAgentsManagerComponent::AgentsAdded_Implementation(const TArray<int32>& AgentIds)
+{
+	// Can be overridden for additional logic when agents are added
+}
+
+void ULearningAgentsManagerComponent::AgentsRemoved_Implementation(const TArray<int32>& AgentIds)
+{
+	// Can be overridden for additional logic when agents are removed
+}
+
+void ULearningAgentsManagerComponent::AgentsReset_Implementation(const TArray<int32>& AgentIds)
+{
+	// Can be overridden for additional logic when agents are reset
+}
+
+void ULearningAgentsManagerComponent::OnAgentsAdded(const TArray<int32>& AgentIds)
+{
+	if (IsSetup())
+	{
+		for (ULearningAgentsHelper* Helper : HelperObjects)
+		{
+			Helper->OnAgentsAdded(AgentIds);
+		}
+
+		AgentsAdded(AgentIds);
+	}
+}
+
+void ULearningAgentsManagerComponent::OnAgentsRemoved(const TArray<int32>& AgentIds)
+{
+	if (IsSetup())
+	{
+		for (ULearningAgentsHelper* Helper : HelperObjects)
+		{
+			Helper->OnAgentsRemoved(AgentIds);
+		}
 
-void ULearningAgentsManagerComponent::OnAgentsRemoved(const TArray<int32>& AgentIds) { }
+		AgentsRemoved(AgentIds);
+	}
+}
 
-void ULearningAgentsManagerComponent::OnAgentsReset(const TArray<int32>& AgentIds) { }
+void ULearningAgentsManagerComponent::OnAgentsReset(const TArray<int32>& AgentIds)
+{
+	if (IsSetup())
+	{
+		for (ULearningAgentsHelper* Helper : HelperObjects)
+		{
+			Helper->OnAgentsReset(AgentIds);
+		}	
+	
+		AgentsReset(AgentIds);
+	}
+}
 
 UObject* ULearningAgentsManagerComponent::GetAgent(const int32 AgentId, const TSubclassOf<UObject> AgentClass) const
 {
@@ -54,6 +103,18 @@ void ULearningAgentsManagerComponent::GetAgents(const TArray<int32>& AgentIds, c
 	Manager->GetAgents(OutAgents, AgentIds, AgentClass);
 }
 
+void ULearningAgentsManagerComponent::GetAllAgents(TArray<UObject*>& OutAgents, TArray<int32>& OutAgentIds, const TSubclassOf<UObject> AgentClass) const
+{
+	if (!Manager)
+	{
+		UE_LOG(LogLearning, Error, TEXT("%s: Must be attached to a LearningAgentsManager Actor."), *GetName());
+		OutAgents.Empty();
+		return;
+	}
+
+	Manager->GetAllAgents(OutAgents, OutAgentIds, AgentClass);
+}
+
 ALearningAgentsManager* ULearningAgentsManagerComponent::GetAgentManager(const TSubclassOf<ALearningAgentsManager> AgentManagerClass) const
 {
 	if (!Manager)
diff --git a/Engine/Plugins/Experimental/LearningAgents/Source/LearningAgents/Private/LearningAgentsObservations.cpp b/Engine/Plugins/Experimental/LearningAgents/Source/LearningAgents/Private/LearningAgentsObservations.cpp
--- a/Engine/Plugins/Experimental/LearningAgents/Source/LearningAgents/Private/LearningAgentsObservations.cpp
+++ b/Engine/Plugins/Experimental/LearningAgents/Source/LearningAgents/Private/LearningAgentsObservations.cpp
@@ -28,8 +28,11 @@ namespace UE::Learning::Agents::Observations::Private
 			return nullptr;
 		}
 
-		ObservationUObject* Observation = NewObject<ObservationUObject>(InInteractor, Name);
+		const FName UniqueName = MakeUniqueObjectName(InInteractor, ObservationUObject::StaticClass(), Name, EUniqueObjectNameOptions::GloballyUnique);
 
+		ObservationUObject* Observation = NewObject<ObservationUObject>(InInteractor, UniqueName);
+
+		Observation->Init(InInteractor->GetAgentManager()->GetMaxAgentNum());
 		Observation->Interactor = InInteractor;
 		Observation->FeatureObject = MakeShared<ObservationFObject>(
 			Observation->GetFName(),
@@ -37,9 +40,6 @@ namespace UE::Learning::Agents::Observations::Private
 			InInteractor->GetAgentManager()->GetMaxAgentNum(),
 			Forward<InArgTypes>(Args)...);
 
-		Observation->AgentIteration.SetNumUninitialized({ InInteractor->GetAgentManager()->GetMaxAgentNum() });
-		UE::Learning::Array::Set<1, uint64>(Observation->AgentIteration, INDEX_NONE);
-
 		// We assume all supported observation feature objects can be encoded
 		UE_LEARNING_CHECK(Observation->FeatureObject->IsEncodable());
 
@@ -51,6 +51,34 @@ namespace UE::Learning::Agents::Observations::Private
 
 //------------------------------------------------------------------
 
+void ULearningAgentsObservation::Init(const int32 MaxAgentNum)
+{
+	AgentIteration.SetNumUninitialized({ MaxAgentNum });
+	UE::Learning::Array::Set<1, uint64>(AgentIteration, INDEX_NONE);
+}
+
+void ULearningAgentsObservation::OnAgentsAdded(const TArray<int32>& AgentIds)
+{
+	UE::Learning::Array::Set<1, uint64>(AgentIteration, 0, AgentIds);
+}
+
+void ULearningAgentsObservation::OnAgentsRemoved(const TArray<int32>& AgentIds)
+{
+	UE::Learning::Array::Set<1, uint64>(AgentIteration, INDEX_NONE, AgentIds);
+}
+
+void ULearningAgentsObservation::OnAgentsReset(const TArray<int32>& AgentIds)
+{
+	UE::Learning::Array::Set<1, uint64>(AgentIteration, 0, AgentIds);
+}
+
+uint64 ULearningAgentsObservation::GetAgentIteration(const int32 AgentId) const
+{
+	return AgentIteration[AgentId];
+}
+
+//------------------------------------------------------------------
+
 UFloatObservation* UFloatObservation::AddFloatObservation(ULearningAgentsInteractor* InInteractor, const FName Name, const float Scale)
 {
 	return UE::Learning::Agents::Observations::Private::AddObservation<UFloatObservation, UE::Learning::FFloatFeature>(InInteractor, Name, TEXT("AddFloatObservation"), 1, Scale);
diff --git a/Engine/Plugins/Experimental/LearningAgents/Source/LearningAgents/Private/LearningAgentsPolicy.cpp b/Engine/Plugins/Experimental/LearningAgents/Source/LearningAgents/Private/LearningAgentsPolicy.cpp
--- a/Engine/Plugins/Experimental/LearningAgents/Source/LearningAgents/Private/LearningAgentsPolicy.cpp
+++ b/Engine/Plugins/Experimental/LearningAgents/Source/LearningAgents/Private/LearningAgentsPolicy.cpp
@@ -89,7 +89,9 @@ void ULearningAgentsPolicy::SetupPolicy(
 	{
 		// Create New Neural Network Asset
 
-		Network = NewObject<ULearningAgentsNeuralNetwork>(this, TEXT("PolicyNetwork"));
+		const FName UniqueName = MakeUniqueObjectName(this, ULearningAgentsNeuralNetwork::StaticClass(), TEXT("PolicyNetwork"), EUniqueObjectNameOptions::GloballyUnique);
+
+		Network = NewObject<ULearningAgentsNeuralNetwork>(this, UniqueName);
 		Network->NeuralNetwork = MakeShared<UE::Learning::FNeuralNetwork>();
 		Network->NeuralNetwork->Resize(
 			Interactor->GetObservationFeature().DimNum(),
@@ -117,6 +119,8 @@ void ULearningAgentsPolicy::SetupPolicy(
 	Manager->GetInstanceData()->Link(PolicyObject->OutputHandle, Interactor->GetActionFeature().FeatureHandle);
 
 	bIsSetup = true;
+
+	OnAgentsAdded(Manager->GetAllAgentIds());
 }
 
 ULearningAgentsNeuralNetwork* ULearningAgentsPolicy::GetNetworkAsset()
@@ -316,10 +320,10 @@ void ULearningAgentsPolicy::VisualLog(const UE::Learning::FIndexSet AgentSet) co
 				AgentId,
 				ActionNoiseScaleView[AgentId],
 				*UE::Learning::Array::FormatFloat(InputView[AgentId]),
-				*UE::Learning::Array::FormatFloat(OutputView[AgentId]),
+				*UE::Learning::Agents::Debug::FloatArrayToStatsString(InputView[AgentId]),
 				*UE::Learning::Array::FormatFloat(OutputMeanView[AgentId]),
 				*UE::Learning::Array::FormatFloat(OutputStdView[AgentId]),
-				*UE::Learning::Agents::Debug::FloatArrayToStatsString(InputView[AgentId]),
+				*UE::Learning::Array::FormatFloat(OutputView[AgentId]),
 				*UE::Learning::Agents::Debug::FloatArrayToStatsString(OutputView[AgentId]));
 
 		}
diff --git a/Engine/Plugins/Experimental/LearningAgents/Source/LearningAgents/Public/LearningAgentsActions.h b/Engine/Plugins/Experimental/LearningAgents/Source/LearningAgents/Public/LearningAgentsActions.h
--- a/Engine/Plugins/Experimental/LearningAgents/Source/LearningAgents/Public/LearningAgentsActions.h
+++ b/Engine/Plugins/Experimental/LearningAgents/Source/LearningAgents/Public/LearningAgentsActions.h
@@ -43,11 +43,32 @@ public:
 
 public:
 
-	/** Number of times this action has been got for all agents */
-	TLearningArray<1, uint64, TInlineAllocator<32>> AgentGetIteration;
+	/** Initialize the internal state for a given maximum number of agents */
+	void Init(const int32 MaxAgentNum);
 
-	/** Number of times this action has been set for all agents */
-	TLearningArray<1, uint64, TInlineAllocator<32>> AgentSetIteration;
+	/**
+	 * Called whenever agents are added to the associated ULearningAgentsInteractor object.
+	 * @param AgentIds Array of agent ids which have been added
+	 */
+	virtual void OnAgentsAdded(const TArray<int32>& AgentIds);
+
+	/**
+	 * Called whenever agents are removed from the associated ULearningAgentsInteractor object.
+	 * @param AgentIds Array of agent ids which have been removed
+	 */
+	virtual void OnAgentsRemoved(const TArray<int32>& AgentIds);
+
+	/**
+	 * Called whenever agents are reset on the associated ULearningAgentsInteractor object.
+	 * @param AgentIds Array of agent ids which have been reset
+	 */
+	virtual void OnAgentsReset(const TArray<int32>& AgentIds);
+
+	/** Get the number of times an action has been got for the given agent id. */
+	uint64 GetAgentGetIteration(const int32 AgentId) const;
+
+	/** Get the number of times an action has been set for the given agent id. */
+	uint64 GetAgentSetIteration(const int32 AgentId) const;
 
 public:
 #if UE_LEARNING_AGENTS_ENABLE_VISUAL_LOG
@@ -57,6 +78,14 @@ public:
 	/** Describes this action to the visual logger for debugging purposes. */
 	virtual void VisualLog(const UE::Learning::FIndexSet Instances) const {}
 #endif
+
+protected:
+
+	/** Number of times this action has been got for all agents */
+	TLearningArray<1, uint64, TInlineAllocator<32>> AgentGetIteration;
+
+	/** Number of times this action has been set for all agents */
+	TLearningArray<1, uint64, TInlineAllocator<32>> AgentSetIteration;
 };
 
 //------------------------------------------------------------------
diff --git a/Engine/Plugins/Experimental/LearningAgents/Source/LearningAgents/Public/LearningAgentsHelpers.h b/Engine/Plugins/Experimental/LearningAgents/Source/LearningAgents/Public/LearningAgentsHelpers.h
--- a/Engine/Plugins/Experimental/LearningAgents/Source/LearningAgents/Public/LearningAgentsHelpers.h
+++ b/Engine/Plugins/Experimental/LearningAgents/Source/LearningAgents/Public/LearningAgentsHelpers.h
@@ -30,6 +30,27 @@ public:
 	UPROPERTY(VisibleAnywhere, Transient, Category = "LearningAgents")
 	TObjectPtr<ULearningAgentsManagerComponent> ManagerComponent;
 
+public:
+
+	/**
+	 * Called whenever agents are added to the associated ULearningAgentsManagerComponent object.
+	 * @param AgentIds Array of agent ids which have been added
+	 */
+	virtual void OnAgentsAdded(const TArray<int32>& AgentIds);
+
+	/**
+	 * Called whenever agents are removed from the associated ULearningAgentsManagerComponent object.
+	 * @param AgentIds Array of agent ids which have been removed
+	 */
+	virtual void OnAgentsRemoved(const TArray<int32>& AgentIds);
+
+	/**
+	 * Called whenever agents are reset on the associated ULearningAgentsManagerComponent object.
+	 * @param AgentIds Array of agent ids which have been reset
+	 */
+	virtual void OnAgentsReset(const TArray<int32>& AgentIds);
+
+public:
 #if UE_LEARNING_AGENTS_ENABLE_VISUAL_LOG
 	/** Color used to draw this helper in the visual log */
 	FLinearColor VisualLogColor = FColor::Magenta;
@@ -312,3 +333,87 @@ public:
 		const ECollisionChannel CollisionChannel = ECollisionChannel::ECC_WorldStatic) const;
 
 };
+
+//------------------------------------------------------------------
+
+/**
+ * A helper for monitoring collisions between components.
+ *
+ * This object works by monitoring for collisions between a given component and any other components with a given tag.
+ * To indicate which component to monitor for collisions `SetComponent` should be called (generally on the 
+ * `ULearningAgentsManagerComponent::AgentsAdded` event) with the corresponding agent id.
+ * 
+ * If any collisions occur `GetCollisionOccurred` will report true until it is reset to false either by calling 
+ * `ResetCollisionOccurred` or `GetAndResetCollisionOccurred`. It will also be automatically reset when the agent is 
+ * reset.
+ */
+UCLASS()
+class LEARNINGAGENTS_API UCollisionMonitorHelper : public ULearningAgentsHelper
+{
+	GENERATED_BODY()
+
+public:
+
+	/**
+	 * Adds a new collision monitor helper to the given manager component.
+	 * @param InManagerComponent The manager component to add this helper to (such as an Interactor or Trainer).
+	 * @param Name The name of this new helper. Used for debugging.
+	 * @return The newly created helper.
+	 */
+	UFUNCTION(BlueprintCallable, Category = "LearningAgents", meta = (DefaultToSelf = "InManagerComponent"))
+	static UCollisionMonitorHelper* AddCollisionMonitorHelper(
+		ULearningAgentsManagerComponent* InManagerComponent, 
+		const FName Name = NAME_None);
+
+	/**
+	 * Sets the component to monitor for collisions. Generally this should be called on the 
+	 * `ULearningAgentsManagerComponent::AgentsAdded` event.
+	 * @param AgentId The agent id to associate this component with.
+	 * @param Component The component to monitor for collisions.
+	 * @param OtherComponentTag The tag other components must have to trigger collisions.
+	 */
+	UFUNCTION(BlueprintCallable, Category = "LearningAgents", meta = (AgentId = "-1"))
+	void SetComponent(const int32 AgentId, UPrimitiveComponent* Component, const FName OtherComponentTag = NAME_None);
+
+	/**
+	 * Gets if a collision has occurred for this agent.
+	 * @param AgentId The agent id to check for collisions for.
+	 */
+	UFUNCTION(BlueprintPure = false, Category = "LearningAgents", meta = (AgentId = "-1"))
+	bool GetCollisionOccurred(const int32 AgentId) const;
+	
+	/**
+	 * Gets if a collision has occurred for this agent and resets the state of if a collision has occurred back to false.
+	 * @param AgentId The agent id to check for collisions for.
+	 */
+	UFUNCTION(BlueprintCallable, Category = "LearningAgents", meta = (AgentId = "-1"))
+	bool GetAndResetCollisionOccurred(const int32 AgentId);
+
+	/**
+	 * Resets the state of if a collision has occurred back to false.
+	 * @param AgentId The agent id to reset the state for.
+	 */
+	UFUNCTION(BlueprintCallable, Category = "LearningAgents", meta = (AgentId = "-1"))
+	void ResetCollisionOccurred(const int32 AgentId);
+
+public:
+
+	/** Initialize the internal state for a given maximum number of agents */
+	void Init(const int32 MaxAgentNum);
+
+	//~ Begin ULearningAgentsHelper Interface
+	virtual void OnAgentsRemoved(const TArray<int32>& AgentIds) override;
+	virtual void OnAgentsReset(const TArray<int32>& AgentIds) override;
+	//~ End ULearningAgentsHelper Interface
+
+private:
+
+	TLearningArray<1, UPrimitiveComponent*, TInlineAllocator<32>> Components;
+	TLearningArray<1, FName, TInlineAllocator<32>> OtherComponentTags;
+	TLearningArray<1, bool, TInlineAllocator<32>> CollisionsOccured;
+
+private:
+
+	UFUNCTION()
+	void HandleOnHit(UPrimitiveComponent* HitComponent, AActor* OtherActor, UPrimitiveComponent* OtherComp, FVector NormalImpulse, const FHitResult& Hit);
+};
diff --git a/Engine/Plugins/Experimental/LearningAgents/Source/LearningAgents/Public/LearningAgentsInteractor.h b/Engine/Plugins/Experimental/LearningAgents/Source/LearningAgents/Public/LearningAgentsInteractor.h
--- a/Engine/Plugins/Experimental/LearningAgents/Source/LearningAgents/Public/LearningAgentsInteractor.h
+++ b/Engine/Plugins/Experimental/LearningAgents/Source/LearningAgents/Public/LearningAgentsInteractor.h
@@ -194,12 +194,6 @@ private:
 // ----- Private Iteration Checks ----- 
 private:
 
-	/**
-	 * Sets the iteration number for the given set of agents for all observations, actions, as well as for the encoding 
-	 * of observations and actions
-	 */
-	void SetAgentIteration(const UE::Learning::FIndexSet AgentIds, const uint64 Value);
-
 	/** Number of times observations have been encoded for all agents */
 	TLearningArray<1, uint64, TInlineAllocator<32>> ObservationEncodingAgentIteration;
 
diff --git a/Engine/Plugins/Experimental/LearningAgents/Source/LearningAgents/Public/LearningAgentsManager.h b/Engine/Plugins/Experimental/LearningAgents/Source/LearningAgents/Public/LearningAgentsManager.h
--- a/Engine/Plugins/Experimental/LearningAgents/Source/LearningAgents/Public/LearningAgentsManager.h
+++ b/Engine/Plugins/Experimental/LearningAgents/Source/LearningAgents/Public/LearningAgentsManager.h
@@ -134,9 +134,10 @@ public:
 	 * cast the object to the given type.
 	 * @param AgentClass The class to cast the agent objects to (in blueprint).
 	 * @param OutAgents The output array of agent objects.
+	 * @param OutAgentIds The output array of agent ids.
 	 */
 	UFUNCTION(BlueprintPure = false, Category = "LearningAgents", meta = (DeterminesOutputType = "AgentClass", DynamicOutputParam = "OutAgents"))
-	void GetAllAgents(TArray<UObject*>& OutAgents, const TSubclassOf<UObject> AgentClass) const;
+	void GetAllAgents(TArray<UObject*>& OutAgents, TArray<int32>& OutAgentIds, const TSubclassOf<UObject> AgentClass) const;
 
 	/**
 	 * Gets the agent id associated with a given agent.
@@ -154,13 +155,6 @@ public:
 	UFUNCTION(BlueprintPure = false, Category = "LearningAgents")
 	void GetAgentIds(TArray<int32>& OutAgentIds, const TArray<UObject*>& InAgents) const;
 
-	/**
-	 * Gets all added agent ids.
-	 * @param OutAgentIds The output array of agent ids.
-	 */
-	UFUNCTION(BlueprintPure = false, Category = "LearningAgents")
-	void GetAllAgentIds(TArray<int32>& OutAgentIds) const;
-
 	/**
 	 * Gets the number of agents added
 	 * @return The number of agents added.
diff --git a/Engine/Plugins/Experimental/LearningAgents/Source/LearningAgents/Public/LearningAgentsManagerComponent.h b/Engine/Plugins/Experimental/LearningAgents/Source/LearningAgents/Public/LearningAgentsManagerComponent.h
--- a/Engine/Plugins/Experimental/LearningAgents/Source/LearningAgents/Public/LearningAgentsManagerComponent.h
+++ b/Engine/Plugins/Experimental/LearningAgents/Source/LearningAgents/Public/LearningAgentsManagerComponent.h
@@ -31,25 +31,40 @@ public:
 
 	virtual void PostInitProperties() override;
 
+	/**
+	 * During this event, any additional logic required for when agents are added can be executed.
+	 */
+	UFUNCTION(BlueprintNativeEvent, Category = "LearningAgents")
+	void AgentsAdded(const TArray<int32>& AgentIds);
+
+	/**
+	 * During this event, any additional logic required for when agents are removed can be executed.
+	 */
+	UFUNCTION(BlueprintNativeEvent, Category = "LearningAgents")
+	void AgentsRemoved(const TArray<int32>& AgentIds);
+
+	/**
+	 * During this event, any additional logic required for when agents are reset can be executed.
+	 */
+	UFUNCTION(BlueprintNativeEvent, Category = "LearningAgents")
+	void AgentsReset(const TArray<int32>& AgentIds);
+
 	/**
 	 * Called whenever agents are added to the parent ALearningAgentsManager object.
 	 * @param AgentIds Array of agent ids which have been added
 	 */
-	UFUNCTION(BlueprintCallable, Category = "LearningAgents")
 	virtual void OnAgentsAdded(const TArray<int32>& AgentIds);
 
 	/**
 	 * Called whenever agents are removed from the parent ALearningAgentsManager object.
 	 * @param AgentIds Array of agent ids which have been removed
 	 */
-	UFUNCTION(BlueprintCallable, Category = "LearningAgents")
 	virtual void OnAgentsRemoved(const TArray<int32>& AgentIds);
 
 	/**
 	 * Called whenever agents are reset on the parent ALearningAgentsManager object.
 	 * @param AgentIds Array of agent ids which have been reset
 	 */
-	UFUNCTION(BlueprintCallable, Category = "LearningAgents")
 	virtual void OnAgentsReset(const TArray<int32>& AgentIds);
 
 	/** Returns true if this component has been setup. Otherwise, false. */
@@ -81,6 +96,16 @@ protected:
 	UFUNCTION(BlueprintPure = false, Category = "LearningAgents", meta = (DeterminesOutputType = "AgentClass", DynamicOutputParam = "OutAgents"))
 	void GetAgents(const TArray<int32>& AgentIds, const TSubclassOf<UObject> AgentClass, TArray<UObject*>& OutAgents) const;
 
+	/**
+	 * Gets all added agents from the manager. Calling this from blueprint with the appropriate AgentClass will 
+	 * automatically cast the object to the given type.
+	 * @param AgentClass The class to cast the agent objects to (in blueprint).
+	 * @param OutAgents The output array of agent objects.
+	 * @param OutAgentIds The output array of agent ids.
+	 */
+	UFUNCTION(BlueprintPure = false, Category = "LearningAgents", meta = (DeterminesOutputType = "AgentClass", DynamicOutputParam = "OutAgents"))
+	void GetAllAgents(TArray<UObject*>& OutAgents, TArray<int32>& OutAgentIds, const TSubclassOf<UObject> AgentClass) const;
+
 	/**
 	 * Gets the agent manager associated with this component.
 	 * @param AgentClass The class to cast the agent manager to (in blueprint).
diff --git a/Engine/Plugins/Experimental/LearningAgents/Source/LearningAgents/Public/LearningAgentsObservations.h b/Engine/Plugins/Experimental/LearningAgents/Source/LearningAgents/Public/LearningAgentsObservations.h
--- a/Engine/Plugins/Experimental/LearningAgents/Source/LearningAgents/Public/LearningAgentsObservations.h
+++ b/Engine/Plugins/Experimental/LearningAgents/Source/LearningAgents/Public/LearningAgentsObservations.h
@@ -50,8 +50,29 @@ public:
 
 public:
 
-	/** Number of times this observation has been set for all agents */
-	TLearningArray<1, uint64, TInlineAllocator<32>> AgentIteration;
+	/** Initialize the internal state for a given maximum number of agents */
+	void Init(const int32 MaxAgentNum);
+
+	/**
+	 * Called whenever agents are added to the associated ULearningAgentsInteractor object.
+	 * @param AgentIds Array of agent ids which have been added
+	 */
+	virtual void OnAgentsAdded(const TArray<int32>& AgentIds);
+
+	/**
+	 * Called whenever agents are removed from the associated ULearningAgentsInteractor object.
+	 * @param AgentIds Array of agent ids which have been removed
+	 */
+	virtual void OnAgentsRemoved(const TArray<int32>& AgentIds);
+
+	/**
+	 * Called whenever agents are reset on the associated ULearningAgentsInteractor object.
+	 * @param AgentIds Array of agent ids which have been reset
+	 */
+	virtual void OnAgentsReset(const TArray<int32>& AgentIds);
+
+	/** Get the number of times an observation has been set for the given agent id. */
+	uint64 GetAgentIteration(const int32 AgentId) const;
 
 public:
 #if UE_LEARNING_AGENTS_ENABLE_VISUAL_LOG
@@ -62,6 +83,11 @@ public:
 	/** Describes this observation to the visual logger for debugging purposes. */
 	virtual void VisualLog(const UE::Learning::FIndexSet Instances) const {}
 #endif
+
+protected:
+
+	/** Number of times this observation has been set for all agents */
+	TLearningArray<1, uint64, TInlineAllocator<32>> AgentIteration;
 };
 
 //------------------------------------------------------------------
diff --git a/Engine/Plugins/Experimental/LearningAgents/Source/LearningAgentsTraining/Private/LearningAgentsCompletions.cpp b/Engine/Plugins/Experimental/LearningAgents/Source/LearningAgentsTraining/Private/LearningAgentsCompletions.cpp
--- a/Engine/Plugins/Experimental/LearningAgents/Source/LearningAgentsTraining/Private/LearningAgentsCompletions.cpp
+++ b/Engine/Plugins/Experimental/LearningAgents/Source/LearningAgentsTraining/Private/LearningAgentsCompletions.cpp
@@ -29,8 +29,10 @@ namespace UE::Learning::Agents::Completions::Private
 			return nullptr;
 		}
 
-		CompletionUObject* Completion = NewObject<CompletionUObject>(InAgentTrainer, Name);
+		const FName UniqueName = MakeUniqueObjectName(InAgentTrainer, CompletionUObject::StaticClass(), Name, EUniqueObjectNameOptions::GloballyUnique);
 
+		CompletionUObject* Completion = NewObject<CompletionUObject>(InAgentTrainer, UniqueName);
+		Completion->Init(InAgentTrainer->GetAgentManager()->GetMaxAgentNum());
 		Completion->AgentTrainer = InAgentTrainer;
 		Completion->CompletionObject = MakeShared<CompletionFObject>(
 			Completion->GetFName(),
@@ -38,9 +40,6 @@ namespace UE::Learning::Agents::Completions::Private
 			InAgentTrainer->GetAgentManager()->GetMaxAgentNum(),
 			Forward<InArgTypes>(Args)...);
 
-		Completion->AgentIteration.SetNumUninitialized({ InAgentTrainer->GetAgentManager()->GetMaxAgentNum() });
-		UE::Learning::Array::Set<1, uint64>(Completion->AgentIteration, INDEX_NONE);
-
 		InAgentTrainer->AddCompletion(Completion, Completion->CompletionObject.ToSharedRef());
 
 		return Completion;
@@ -49,6 +48,34 @@ namespace UE::Learning::Agents::Completions::Private
 
 //------------------------------------------------------------------
 
+void ULearningAgentsCompletion::Init(const int32 MaxAgentNum)
+{
+	AgentIteration.SetNumUninitialized({ MaxAgentNum });
+	UE::Learning::Array::Set<1, uint64>(AgentIteration, INDEX_NONE);
+}
+
+void ULearningAgentsCompletion::OnAgentsAdded(const TArray<int32>& AgentIds)
+{
+	UE::Learning::Array::Set<1, uint64>(AgentIteration, 0, AgentIds);
+}
+
+void ULearningAgentsCompletion::OnAgentsRemoved(const TArray<int32>& AgentIds)
+{
+	UE::Learning::Array::Set<1, uint64>(AgentIteration, INDEX_NONE, AgentIds);
+}
+
+void ULearningAgentsCompletion::OnAgentsReset(const TArray<int32>& AgentIds)
+{
+	UE::Learning::Array::Set<1, uint64>(AgentIteration, 0, AgentIds);
+}
+
+uint64 ULearningAgentsCompletion::GetAgentIteration(const int32 AgentId) const
+{
+	return AgentIteration[AgentId];
+}
+
+//------------------------------------------------------------------
+
 UConditionalCompletion* UConditionalCompletion::AddConditionalCompletion(ULearningAgentsTrainer* InAgentTrainer, const FName Name, const ELearningAgentsCompletion InCompletionMode)
 {
 	return UE::Learning::Agents::Completions::Private::AddCompletion<UConditionalCompletion, UE::Learning::FConditionalCompletion>(InAgentTrainer, Name, TEXT("AddConditionalCompletion"), UE::Learning::Agents::GetCompletionMode(InCompletionMode));
diff --git a/Engine/Plugins/Experimental/LearningAgents/Source/LearningAgentsTraining/Private/LearningAgentsRecorder.cpp b/Engine/Plugins/Experimental/LearningAgents/Source/LearningAgentsTraining/Private/LearningAgentsRecorder.cpp
--- a/Engine/Plugins/Experimental/LearningAgents/Source/LearningAgentsTraining/Private/LearningAgentsRecorder.cpp
+++ b/Engine/Plugins/Experimental/LearningAgents/Source/LearningAgentsTraining/Private/LearningAgentsRecorder.cpp
@@ -5,6 +5,7 @@
 #include "LearningAgentsManager.h"
 #include "LearningAgentsInteractor.h"
 #include "LearningAgentsRecording.h"
+#include "LearningAgentsHelpers.h"
 #include "LearningFeatureObject.h"
 #include "LearningLog.h"
 #include "LearningTrainer.h"
@@ -133,7 +134,9 @@ void ULearningAgentsRecorder::SetupRecorder(
 	}
 	else
 	{
-		Recording = NewObject<ULearningAgentsRecording>(this, TEXT("Recording"));
+		const FName UniqueName = MakeUniqueObjectName(this, ULearningAgentsRecording::StaticClass(), TEXT("Recording"), EUniqueObjectNameOptions::GloballyUnique);
+
+		Recording = NewObject<ULearningAgentsRecording>(this, UniqueName);
 	}
 
 	RecordingDirectory = RecorderPathSettings.IntermediateRelativePath.Path / TEXT("LearningAgents") / RecorderPathSettings.RecordingsSubdirectory;
@@ -142,6 +145,8 @@ void ULearningAgentsRecorder::SetupRecorder(
 	RecordBuffers.SetNum(Manager->GetMaxAgentNum());
 
 	bIsSetup = true;
+
+	OnAgentsAdded(Manager->GetAllAgentIds());
 }
 
 void ULearningAgentsRecorder::OnAgentsRemoved(const TArray<int32>& AgentIds)
@@ -156,6 +161,13 @@ void ULearningAgentsRecorder::OnAgentsRemoved(const TArray<int32>& AgentIds)
 				RecordBuffers[AgentId].Empty();
 			}
 		}
+
+		for (ULearningAgentsHelper* Helper : HelperObjects)
+		{
+			Helper->OnAgentsRemoved(AgentIds);
+		}
+
+		AgentsRemoved(AgentIds);
 	}
 }
 
@@ -171,6 +183,13 @@ void ULearningAgentsRecorder::OnAgentsReset(const TArray<int32>& AgentIds)
 				RecordBuffers[AgentId].Empty();
 			}
 		}
+
+		for (ULearningAgentsHelper* Helper : HelperObjects)
+		{
+			Helper->OnAgentsReset(AgentIds);
+		}
+
+		AgentsReset(AgentIds);
 	}
 }
 
diff --git a/Engine/Plugins/Experimental/LearningAgents/Source/LearningAgentsTraining/Private/LearningAgentsRecording.cpp b/Engine/Plugins/Experimental/LearningAgents/Source/LearningAgentsTraining/Private/LearningAgentsRecording.cpp
--- a/Engine/Plugins/Experimental/LearningAgents/Source/LearningAgentsTraining/Private/LearningAgentsRecording.cpp
+++ b/Engine/Plugins/Experimental/LearningAgents/Source/LearningAgentsTraining/Private/LearningAgentsRecording.cpp
@@ -128,7 +128,7 @@ void ULearningAgentsRecording::SaveRecordingToFile(const FFilePath& File) const
 
 void ULearningAgentsRecording::AppendRecordingFromFile(const FFilePath& File)
 {
-	ULearningAgentsRecording* TempRecording = NewObject<ULearningAgentsRecording>(this, TEXT("TempRecording"));
+	ULearningAgentsRecording* TempRecording = NewObject<ULearningAgentsRecording>(this);
 	TempRecording->LoadRecordingFromFile(File);
 
 	Records.Append(TempRecording->Records);
diff --git a/Engine/Plugins/Experimental/LearningAgents/Source/LearningAgentsTraining/Private/LearningAgentsRewards.cpp b/Engine/Plugins/Experimental/LearningAgents/Source/LearningAgentsTraining/Private/LearningAgentsRewards.cpp
--- a/Engine/Plugins/Experimental/LearningAgents/Source/LearningAgentsTraining/Private/LearningAgentsRewards.cpp
+++ b/Engine/Plugins/Experimental/LearningAgents/Source/LearningAgentsTraining/Private/LearningAgentsRewards.cpp
@@ -28,8 +28,10 @@ namespace UE::Learning::Agents::Rewards::Private
 			return nullptr;
 		}
 
-		RewardUObject* Reward = NewObject<RewardUObject>(InAgentTrainer, Name);
+		const FName UniqueName = MakeUniqueObjectName(InAgentTrainer, RewardUObject::StaticClass(), Name, EUniqueObjectNameOptions::GloballyUnique);
 
+		RewardUObject* Reward = NewObject<RewardUObject>(InAgentTrainer, UniqueName);
+		Reward->Init(InAgentTrainer->GetAgentManager()->GetMaxAgentNum());
 		Reward->AgentTrainer = InAgentTrainer;
 		Reward->RewardObject = MakeShared<RewardFObject>(
 			Reward->GetFName(),
@@ -37,9 +39,6 @@ namespace UE::Learning::Agents::Rewards::Private
 			InAgentTrainer->GetAgentManager()->GetMaxAgentNum(),
 			Forward<InArgTypes>(Args)...);
 
-		Reward->AgentIteration.SetNumUninitialized({ InAgentTrainer->GetAgentManager()->GetMaxAgentNum() });
-		UE::Learning::Array::Set<1, uint64>(Reward->AgentIteration, INDEX_NONE);
-
 		InAgentTrainer->AddReward(Reward, Reward->RewardObject.ToSharedRef());
 
 		return Reward;
@@ -48,6 +47,34 @@ namespace UE::Learning::Agents::Rewards::Private
 
 //------------------------------------------------------------------
 
+void ULearningAgentsReward::Init(const int32 MaxAgentNum)
+{
+	AgentIteration.SetNumUninitialized({ MaxAgentNum });
+	UE::Learning::Array::Set<1, uint64>(AgentIteration, INDEX_NONE);
+}
+
+void ULearningAgentsReward::OnAgentsAdded(const TArray<int32>& AgentIds)
+{
+	UE::Learning::Array::Set<1, uint64>(AgentIteration, 0, AgentIds);
+}
+
+void ULearningAgentsReward::OnAgentsRemoved(const TArray<int32>& AgentIds)
+{
+	UE::Learning::Array::Set<1, uint64>(AgentIteration, INDEX_NONE, AgentIds);
+}
+
+void ULearningAgentsReward::OnAgentsReset(const TArray<int32>& AgentIds)
+{
+	UE::Learning::Array::Set<1, uint64>(AgentIteration, 0, AgentIds);
+}
+
+uint64 ULearningAgentsReward::GetAgentIteration(const int32 AgentId) const
+{
+	return AgentIteration[AgentId];
+}
+
+//------------------------------------------------------------------
+
 UFloatReward* UFloatReward::AddFloatReward(ULearningAgentsTrainer* InAgentTrainer, const FName Name, const float Weight)
 {
 	return UE::Learning::Agents::Rewards::Private::AddReward<UFloatReward, UE::Learning::FFloatReward>(InAgentTrainer, Name, TEXT("AddFloatReward"), Weight);
@@ -93,6 +120,50 @@ void UFloatReward::VisualLog(const UE::Learning::FIndexSet Instances) const
 
 //------------------------------------------------------------------
 
+UConditionalReward* UConditionalReward::AddConditionalReward(ULearningAgentsTrainer* InAgentTrainer, const FName Name, const float Value)
+{
+	return UE::Learning::Agents::Rewards::Private::AddReward<UConditionalReward, UE::Learning::FConditionalConstantReward>(InAgentTrainer, Name, TEXT("AddConditionalReward"), Value);
+}
+
+void UConditionalReward::SetConditionalReward(const int32 AgentId, const bool bCondition)
+{
+	if (!AgentTrainer->HasAgent(AgentId))
+	{
+		UE_LOG(LogLearning, Error, TEXT("%s: AgentId %d not found in the agents set."), *GetName(), AgentId);
+		return;
+	}
+
+	RewardObject->InstanceData->View(RewardObject->ConditionHandle)[AgentId] = bCondition;
+	AgentIteration[AgentId]++;
+}
+
+#if UE_LEARNING_AGENTS_ENABLE_VISUAL_LOG
+void UConditionalReward::VisualLog(const UE::Learning::FIndexSet Instances) const
+{
+	UE_LEARNING_TRACE_CPUPROFILER_EVENT_SCOPE(UBoolReward::VisualLog);
+
+	const TLearningArrayView<1, const bool> ConditionView = RewardObject->InstanceData->ConstView(RewardObject->ConditionHandle);
+	const TLearningArrayView<1, const float> RewardView = RewardObject->InstanceData->ConstView(RewardObject->RewardHandle);
+
+	for (const int32 Instance : Instances)
+	{
+		if (const AActor* Actor = Cast<AActor>(AgentTrainer->GetAgent(Instance)))
+		{
+			UE_LEARNING_AGENTS_VLOG_STRING(this, LogLearning, Display,
+				Actor->GetActorLocation(),
+				VisualLogColor.ToFColor(true),
+				TEXT("Agent %i\nValue: [% 6.2f]\nCondition: %s\nReward: [% 6.3f]"),
+				Instance,
+				RewardObject->Value,
+				ConditionView[Instance] ? TEXT("true") : TEXT("false"),
+				RewardView[Instance]);
+		}
+	}
+}
+#endif
+
+//------------------------------------------------------------------
+
 UScalarVelocityReward* UScalarVelocityReward::AddScalarVelocityReward(ULearningAgentsTrainer* InAgentTrainer, const FName Name, const float Weight, const float Scale)
 {
 	return UE::Learning::Agents::Rewards::Private::AddReward<UScalarVelocityReward, UE::Learning::FScalarVelocityReward>(InAgentTrainer, Name, TEXT("AddScalarVelocityReward"), Weight, Scale);
diff --git a/Engine/Plugins/Experimental/LearningAgents/Source/LearningAgentsTraining/Private/LearningAgentsTrainer.cpp b/Engine/Plugins/Experimental/LearningAgents/Source/LearningAgentsTraining/Private/LearningAgentsTrainer.cpp
--- a/Engine/Plugins/Experimental/LearningAgents/Source/LearningAgentsTraining/Private/LearningAgentsTrainer.cpp
+++ b/Engine/Plugins/Experimental/LearningAgents/Source/LearningAgentsTraining/Private/LearningAgentsTrainer.cpp
@@ -9,6 +9,7 @@
 #include "LearningAgentsObservations.h"
 #include "LearningAgentsPolicy.h"
 #include "LearningAgentsCritic.h"
+#include "LearningAgentsHelpers.h"
 #include "LearningArray.h"
 #include "LearningArrayMap.h"
 #include "LearningExperience.h"
@@ -222,22 +223,38 @@ void ULearningAgentsTrainer::SetupTrainer(
 	CompletionEvaluatedAgentIteration.SetNumUninitialized({ Manager->GetMaxAgentNum() });
 	UE::Learning::Array::Set<1, uint64>(RewardEvaluatedAgentIteration, INDEX_NONE);
 	UE::Learning::Array::Set<1, uint64>(CompletionEvaluatedAgentIteration, INDEX_NONE);
-	SetAgentIteration(Manager->GetAllAgentSet(), 0);
 
 	bIsSetup = true;
+
+	OnAgentsAdded(Manager->GetAllAgentIds());
 }
 
 void ULearningAgentsTrainer::OnAgentsAdded(const TArray<int32>& AgentIds)
 {
 	if (IsSetup())
 	{
-		UE::Learning::FIndexSet AgentsSet = AgentIds;
-		AgentsSet.TryMakeSlice();
-
 		ResetEpisodes(AgentIds);
-		EpisodeBuffer->Reset(AgentsSet);
+		EpisodeBuffer->Reset(AgentIds);
+
+		UE::Learning::Array::Set<1, uint64>(RewardEvaluatedAgentIteration, 0, AgentIds);
+		UE::Learning::Array::Set<1, uint64>(CompletionEvaluatedAgentIteration, 0, AgentIds);
+
+		for (ULearningAgentsReward* RewardObject : RewardObjects)
+		{
+			RewardObject->OnAgentsAdded(AgentIds);
+		}
 
-		SetAgentIteration(AgentsSet, 0);
+		for (ULearningAgentsCompletion* CompletionObject : CompletionObjects)
+		{
+			CompletionObject->OnAgentsAdded(AgentIds);
+		}
+
+		for (ULearningAgentsHelper* Helper : HelperObjects)
+		{
+			Helper->OnAgentsAdded(AgentIds);
+		}
+
+		AgentsAdded(AgentIds);
 	}
 }
 
@@ -245,10 +262,25 @@ void ULearningAgentsTrainer::OnAgentsRemoved(const TArray<int32>& AgentIds)
 {
 	if (IsSetup())
 	{
-		UE::Learning::FIndexSet AgentsSet = AgentIds;
-		AgentsSet.TryMakeSlice();
+		UE::Learning::Array::Set<1, uint64>(RewardEvaluatedAgentIteration, INDEX_NONE, AgentIds);
+		UE::Learning::Array::Set<1, uint64>(CompletionEvaluatedAgentIteration, INDEX_NONE, AgentIds);
+
+		for (ULearningAgentsReward* RewardObject : RewardObjects)
+		{
+			RewardObject->OnAgentsRemoved(AgentIds);
+		}
 
-		SetAgentIteration(AgentsSet, INDEX_NONE);
+		for (ULearningAgentsCompletion* CompletionObject : CompletionObjects)
+		{
+			CompletionObject->OnAgentsRemoved(AgentIds);
+		}
+
+		for (ULearningAgentsHelper* Helper : HelperObjects)
+		{
+			Helper->OnAgentsRemoved(AgentIds);
+		}
+
+		AgentsRemoved(AgentIds);
 	}
 }
 
@@ -256,29 +288,28 @@ void ULearningAgentsTrainer::OnAgentsReset(const TArray<int32>& AgentIds)
 {
 	if (IsSetup())
 	{
-		UE::Learning::FIndexSet AgentsSet = AgentIds;
-		AgentsSet.TryMakeSlice();
-
 		ResetEpisodes(AgentIds);
-		EpisodeBuffer->Reset(AgentsSet);
+		EpisodeBuffer->Reset(AgentIds);
 
-		SetAgentIteration(AgentsSet, 0);
-	}
-}
+		UE::Learning::Array::Set<1, uint64>(RewardEvaluatedAgentIteration, 0, AgentIds);
+		UE::Learning::Array::Set<1, uint64>(CompletionEvaluatedAgentIteration, 0, AgentIds);
 
-void ULearningAgentsTrainer::SetAgentIteration(const UE::Learning::FIndexSet AgentIds, const uint64 Value)
-{
-	UE::Learning::Array::Set<1, uint64>(RewardEvaluatedAgentIteration, Value, AgentIds);
-	UE::Learning::Array::Set<1, uint64>(CompletionEvaluatedAgentIteration, Value, AgentIds);
+		for (ULearningAgentsReward* RewardObject : RewardObjects)
+		{
+			RewardObject->OnAgentsReset(AgentIds);
+		}
 
-	for (ULearningAgentsReward* RewardObject : RewardObjects)
-	{
-		UE::Learning::Array::Set<1, uint64>(RewardObject->AgentIteration, Value, AgentIds);
-	}
+		for (ULearningAgentsCompletion* CompletionObject : CompletionObjects)
+		{
+			CompletionObject->OnAgentsReset(AgentIds);
+		}
 
-	for (ULearningAgentsCompletion* CompletionObject : CompletionObjects)
-	{
-		UE::Learning::Array::Set<1, uint64>(CompletionObject->AgentIteration, Value, AgentIds);
+		for (ULearningAgentsHelper* Helper : HelperObjects)
+		{
+			Helper->OnAgentsReset(AgentIds);
+		}
+
+		AgentsReset(AgentIds);
 	}
 }
 
@@ -650,23 +681,23 @@ void ULearningAgentsTrainer::EvaluateRewards()
 	{
 		for (const int32 AgentId : ValidAgentSet)
 		{
-			if (RewardObject->AgentIteration[AgentId] == RewardEvaluatedAgentIteration[AgentId])
+			if (RewardObject->GetAgentIteration(AgentId) == RewardEvaluatedAgentIteration[AgentId])
 			{
-				UE_LOG(LogLearning, Error, TEXT("%s: Reward %s for agent with id %i has not been set (got iteration %i, expected iteration %i) and so agent will not have rewards evaluated."), *GetName(), *RewardObject->GetName(), AgentId, RewardObject->AgentIteration[AgentId], RewardEvaluatedAgentIteration[AgentId] + 1);
+				UE_LOG(LogLearning, Error, TEXT("%s: Reward %s for agent with id %i has not been set (got iteration %i, expected iteration %i) and so agent will not have rewards evaluated."), *GetName(), *RewardObject->GetName(), AgentId, RewardObject->GetAgentIteration(AgentId), RewardEvaluatedAgentIteration[AgentId] + 1);
 				ValidAgentStatus[AgentId] = false;
 				continue;
 			}
 
-			if (RewardObject->AgentIteration[AgentId] > RewardEvaluatedAgentIteration[AgentId] + 1)
+			if (RewardObject->GetAgentIteration(AgentId) > RewardEvaluatedAgentIteration[AgentId] + 1)
 			{
-				UE_LOG(LogLearning, Error, TEXT("%s: Reward %s for agent with id %i appears to have been set multiple times (got iteration %i, expected iteration %i) and so agent will not have rewards evaluated."), *GetName(), *RewardObject->GetName(), AgentId, RewardObject->AgentIteration[AgentId], RewardEvaluatedAgentIteration[AgentId] + 1);
+				UE_LOG(LogLearning, Error, TEXT("%s: Reward %s for agent with id %i appears to have been set multiple times (got iteration %i, expected iteration %i) and so agent will not have rewards evaluated."), *GetName(), *RewardObject->GetName(), AgentId, RewardObject->GetAgentIteration(AgentId), RewardEvaluatedAgentIteration[AgentId] + 1);
 				ValidAgentStatus[AgentId] = false;
 				continue;
 			}
 
-			if (RewardObject->AgentIteration[AgentId] != RewardEvaluatedAgentIteration[AgentId] + 1)
+			if (RewardObject->GetAgentIteration(AgentId) != RewardEvaluatedAgentIteration[AgentId] + 1)
 			{
-				UE_LOG(LogLearning, Error, TEXT("%s: Reward %s for agent with id %i does not have a matching iteration number (got iteration %i, expected iteration %i) and so agent will not have rewards evaluated."), *GetName(), *RewardObject->GetName(), AgentId, RewardObject->AgentIteration[AgentId], RewardEvaluatedAgentIteration[AgentId] + 1);
+				UE_LOG(LogLearning, Error, TEXT("%s: Reward %s for agent with id %i does not have a matching iteration number (got iteration %i, expected iteration %i) and so agent will not have rewards evaluated."), *GetName(), *RewardObject->GetName(), AgentId, RewardObject->GetAgentIteration(AgentId), RewardEvaluatedAgentIteration[AgentId] + 1);
 				ValidAgentStatus[AgentId] = false;
 				continue;
 			}
@@ -746,23 +777,23 @@ void ULearningAgentsTrainer::EvaluateCompletions()
 	{
 		for (const int32 AgentId : ValidAgentSet)
 		{
-			if (CompletionObject->AgentIteration[AgentId] == CompletionEvaluatedAgentIteration[AgentId])
+			if (CompletionObject->GetAgentIteration(AgentId) == CompletionEvaluatedAgentIteration[AgentId])
 			{
-				UE_LOG(LogLearning, Error, TEXT("%s: Completion %s for agent with id %i has not been set (got iteration %i, expected iteration %i) and so agent will not have completions evaluated."), *GetName(), *CompletionObject->GetName(), AgentId, CompletionObject->AgentIteration[AgentId], CompletionEvaluatedAgentIteration[AgentId] + 1);
+				UE_LOG(LogLearning, Error, TEXT("%s: Completion %s for agent with id %i has not been set (got iteration %i, expected iteration %i) and so agent will not have completions evaluated."), *GetName(), *CompletionObject->GetName(), AgentId, CompletionObject->GetAgentIteration(AgentId), CompletionEvaluatedAgentIteration[AgentId] + 1);
 				ValidAgentStatus[AgentId] = false;
 				continue;
 			}
 
-			if (CompletionObject->AgentIteration[AgentId] > CompletionEvaluatedAgentIteration[AgentId] + 1)
+			if (CompletionObject->GetAgentIteration(AgentId) > CompletionEvaluatedAgentIteration[AgentId] + 1)
 			{
-				UE_LOG(LogLearning, Error, TEXT("%s: Completion %s for agent with id %i appears to have been set multiple times (got iteration %i, expected iteration %i) and so agent will not have completions evaluated."), *GetName(), *CompletionObject->GetName(), AgentId, CompletionObject->AgentIteration[AgentId], CompletionEvaluatedAgentIteration[AgentId] + 1);
+				UE_LOG(LogLearning, Error, TEXT("%s: Completion %s for agent with id %i appears to have been set multiple times (got iteration %i, expected iteration %i) and so agent will not have completions evaluated."), *GetName(), *CompletionObject->GetName(), AgentId, CompletionObject->GetAgentIteration(AgentId), CompletionEvaluatedAgentIteration[AgentId] + 1);
 				ValidAgentStatus[AgentId] = false;
 				continue;
 			}
 
-			if (CompletionObject->AgentIteration[AgentId] != CompletionEvaluatedAgentIteration[AgentId] + 1)
+			if (CompletionObject->GetAgentIteration(AgentId) != CompletionEvaluatedAgentIteration[AgentId] + 1)
 			{
-				UE_LOG(LogLearning, Error, TEXT("%s: Completion %s for agent with id %i does not have a matching iteration number (got iteration %i, expected iteration %i) and so agent will not have completions evaluated."), *GetName(), *CompletionObject->GetName(), AgentId, CompletionObject->AgentIteration[AgentId], CompletionEvaluatedAgentIteration[AgentId] + 1);
+				UE_LOG(LogLearning, Error, TEXT("%s: Completion %s for agent with id %i does not have a matching iteration number (got iteration %i, expected iteration %i) and so agent will not have completions evaluated."), *GetName(), *CompletionObject->GetName(), AgentId, CompletionObject->GetAgentIteration(AgentId), CompletionEvaluatedAgentIteration[AgentId] + 1);
 				ValidAgentStatus[AgentId] = false;
 				continue;
 			}
diff --git a/Engine/Plugins/Experimental/LearningAgents/Source/LearningAgentsTraining/Public/LearningAgentsCompletions.h b/Engine/Plugins/Experimental/LearningAgents/Source/LearningAgentsTraining/Public/LearningAgentsCompletions.h
--- a/Engine/Plugins/Experimental/LearningAgents/Source/LearningAgentsTraining/Public/LearningAgentsCompletions.h
+++ b/Engine/Plugins/Experimental/LearningAgents/Source/LearningAgentsTraining/Public/LearningAgentsCompletions.h
@@ -49,8 +49,29 @@ public:
 
 public:
 
-	/** Number of times this completion has been set for all agents */
-	TLearningArray<1, uint64, TInlineAllocator<32>> AgentIteration;
+	/** Initialize the internal state for a given maximum number of agents */
+	void Init(const int32 MaxAgentNum);
+
+	/**
+	 * Called whenever agents are added to the associated ULearningAgentsTrainer object.
+	 * @param AgentIds Array of agent ids which have been added
+	 */
+	virtual void OnAgentsAdded(const TArray<int32>& AgentIds);
+
+	/**
+	 * Called whenever agents are removed from the associated ULearningAgentsTrainer object.
+	 * @param AgentIds Array of agent ids which have been removed
+	 */
+	virtual void OnAgentsRemoved(const TArray<int32>& AgentIds);
+
+	/**
+	 * Called whenever agents are reset on the associated ULearningAgentsTrainer object.
+	 * @param AgentIds Array of agent ids which have been reset
+	 */
+	virtual void OnAgentsReset(const TArray<int32>& AgentIds);
+
+	/** Get the number of times a completion has been set for the given agent id. */
+	uint64 GetAgentIteration(const int32 AgentId) const;
 
 public:
 #if UE_LEARNING_AGENTS_ENABLE_VISUAL_LOG
@@ -60,6 +81,11 @@ public:
 	/** Describes this completion to the visual logger for debugging purposes. */
 	virtual void VisualLog(const UE::Learning::FIndexSet Instances) const {}
 #endif
+
+protected:
+
+	/** Number of times this completion has been set for all agents */
+	TLearningArray<1, uint64, TInlineAllocator<32>> AgentIteration;
 };
 
 //------------------------------------------------------------------
diff --git a/Engine/Plugins/Experimental/LearningAgents/Source/LearningAgentsTraining/Public/LearningAgentsRewards.h b/Engine/Plugins/Experimental/LearningAgents/Source/LearningAgentsTraining/Public/LearningAgentsRewards.h
--- a/Engine/Plugins/Experimental/LearningAgents/Source/LearningAgentsTraining/Public/LearningAgentsRewards.h
+++ b/Engine/Plugins/Experimental/LearningAgents/Source/LearningAgentsTraining/Public/LearningAgentsRewards.h
@@ -14,6 +14,7 @@ namespace UE::Learning
 {
 	struct FRewardObject;
 	struct FFloatReward;
+	struct FConditionalConstantReward;
 	struct FScalarVelocityReward;
 	struct FLocalDirectionalVelocityReward;
 	struct FPlanarPositionDifferencePenalty;
@@ -45,8 +46,29 @@ public:
 
 public:
 
-	/** Number of times this reward has been set for all agents */
-	TLearningArray<1, uint64, TInlineAllocator<32>> AgentIteration;
+	/** Initialize the internal state for a given maximum number of agents */
+	void Init(const int32 MaxAgentNum);
+
+	/**
+	 * Called whenever agents are added to the associated ULearningAgentsTrainer object.
+	 * @param AgentIds Array of agent ids which have been added
+	 */
+	virtual void OnAgentsAdded(const TArray<int32>& AgentIds);
+
+	/**
+	 * Called whenever agents are removed from the associated ULearningAgentsTrainer object.
+	 * @param AgentIds Array of agent ids which have been removed
+	 */
+	virtual void OnAgentsRemoved(const TArray<int32>& AgentIds);
+
+	/**
+	 * Called whenever agents are reset on the associated ULearningAgentsTrainer object.
+	 * @param AgentIds Array of agent ids which have been reset
+	 */
+	virtual void OnAgentsReset(const TArray<int32>& AgentIds);
+
+	/** Get the number of times a reward has been set for the given agent id. */
+	uint64 GetAgentIteration(const int32 AgentId) const;
 
 public:
 #if UE_LEARNING_AGENTS_ENABLE_VISUAL_LOG
@@ -56,6 +78,11 @@ public:
 	/** Describes this reward to the visual logger for debugging purposes. */
 	virtual void VisualLog(const UE::Learning::FIndexSet Instances) const {}
 #endif
+
+protected:
+
+	/** Number of times this reward has been set for all agents */
+	TLearningArray<1, uint64, TInlineAllocator<32>> AgentIteration;
 };
 
 //------------------------------------------------------------------
@@ -99,6 +126,45 @@ public:
 
 //------------------------------------------------------------------
 
+/** A simple conditional reward that gives some constant reward value when a condition is true. */
+UCLASS()
+class LEARNINGAGENTSTRAINING_API UConditionalReward : public ULearningAgentsReward
+{
+	GENERATED_BODY()
+
+public:
+
+	/**
+	 * Adds a new conditional reward to the given trainer. Call during ULearningAgentsTrainer::SetupRewards event.
+	 * @param InAgentTrainer The trainer to add this reward to.
+	 * @param Name The name of this new reward. Used for debugging.
+	 * @param Value The amount of reward to give the agent when the provided condition is true.
+	 * @return The newly created reward.
+	 */
+	UFUNCTION(BlueprintCallable, Category = "LearningAgents", meta = (DefaultToSelf = "InAgentTrainer"))
+	static UConditionalReward* AddConditionalReward(
+		ULearningAgentsTrainer* InAgentTrainer,
+		const FName Name = NAME_None,
+		const float Value = 1.0f);
+
+	/**
+	 * Sets if the agent should receive a reward. Call during ULearningAgentsTrainer::SetRewards event.
+	 * @param AgentId The agent id this data corresponds to.
+	 * @param bCondition If the agent should receive a reward this iteration.
+	 */
+	UFUNCTION(BlueprintCallable, Category = "LearningAgents", meta = (AgentId = "-1"))
+	void SetConditionalReward(const int32 AgentId, const bool bCondition);
+
+#if UE_LEARNING_AGENTS_ENABLE_VISUAL_LOG
+	/** Describes this reward to the visual logger for debugging purposes. */
+	virtual void VisualLog(const UE::Learning::FIndexSet Instances) const override;
+#endif
+
+	TSharedPtr<UE::Learning::FConditionalConstantReward> RewardObject;
+};
+
+//------------------------------------------------------------------
+
 /** A reward for maximizing speed. */
 UCLASS()
 class LEARNINGAGENTSTRAINING_API UScalarVelocityReward : public ULearningAgentsReward
diff --git a/Engine/Plugins/Experimental/LearningAgents/Source/LearningAgentsTraining/Public/LearningAgentsTrainer.h b/Engine/Plugins/Experimental/LearningAgents/Source/LearningAgentsTraining/Public/LearningAgentsTrainer.h
--- a/Engine/Plugins/Experimental/LearningAgents/Source/LearningAgentsTraining/Public/LearningAgentsTrainer.h
+++ b/Engine/Plugins/Experimental/LearningAgents/Source/LearningAgentsTraining/Public/LearningAgentsTrainer.h
@@ -567,12 +567,6 @@ private:
 // ----- Private Iteration Checks ----- 
 private:
 
-	/**
-	 * Sets the iteration number for the given set of agents for all observations, actions, as well as for the encoding
-	 * of observations and actions
-	 */
-	void SetAgentIteration(const UE::Learning::FIndexSet AgentIds, const uint64 Value);
-
 	/** Number of times rewards have been evaluated for all agents */
 	TLearningArray<1, uint64, TInlineAllocator<32>> RewardEvaluatedAgentIteration;
 
