LearningAgents: Updated constructors

diff --git a/Engine/Plugins/Experimental/LearningAgents/Source/LearningAgents/Private/LearningAgentsController.cpp b/Engine/Plugins/Experimental/LearningAgents/Source/LearningAgents/Private/LearningAgentsController.cpp
--- a/Engine/Plugins/Experimental/LearningAgents/Source/LearningAgents/Private/LearningAgentsController.cpp
+++ b/Engine/Plugins/Experimental/LearningAgents/Source/LearningAgents/Private/LearningAgentsController.cpp
@@ -8,9 +8,9 @@
 #include "LearningFeatureObject.h"
 #include "LearningLog.h"
 
-ULearningAgentsController::ULearningAgentsController(const FObjectInitializer& ObjectInitializer) : ULearningAgentsManagerComponent(ObjectInitializer) {}
-ULearningAgentsController::ULearningAgentsController(FVTableHelper& Helper) : ULearningAgentsController() {}
-ULearningAgentsController::~ULearningAgentsController() {}
+ULearningAgentsController::ULearningAgentsController() : Super(FObjectInitializer::Get()) {}
+ULearningAgentsController::ULearningAgentsController(FVTableHelper& Helper) : Super(Helper) {}
+ULearningAgentsController::~ULearningAgentsController() = default;
 
 void ULearningAgentsController::SetActions_Implementation(const TArray<int32>& AgentIds)
 {
diff --git a/Engine/Plugins/Experimental/LearningAgents/Source/LearningAgents/Private/LearningAgentsCritic.cpp b/Engine/Plugins/Experimental/LearningAgents/Source/LearningAgents/Private/LearningAgentsCritic.cpp
--- a/Engine/Plugins/Experimental/LearningAgents/Source/LearningAgents/Private/LearningAgentsCritic.cpp
+++ b/Engine/Plugins/Experimental/LearningAgents/Source/LearningAgents/Private/LearningAgentsCritic.cpp
@@ -15,9 +15,9 @@
 #include "Misc/Paths.h"
 #include "GameFramework/Actor.h"
 
-ULearningAgentsCritic::ULearningAgentsCritic(const FObjectInitializer& ObjectInitializer) : ULearningAgentsManagerComponent(ObjectInitializer) {}
-ULearningAgentsCritic::ULearningAgentsCritic(FVTableHelper& Helper) : ULearningAgentsCritic() {}
-ULearningAgentsCritic::~ULearningAgentsCritic() {}
+ULearningAgentsCritic::ULearningAgentsCritic() : Super(FObjectInitializer::Get()) {}
+ULearningAgentsCritic::ULearningAgentsCritic(FVTableHelper& Helper) : Super(Helper) {}
+ULearningAgentsCritic::~ULearningAgentsCritic() = default;
 
 void ULearningAgentsCritic::SetupCritic(
 	ULearningAgentsInteractor* InInteractor, 
diff --git a/Engine/Plugins/Experimental/LearningAgents/Source/LearningAgents/Private/LearningAgentsInteractor.cpp b/Engine/Plugins/Experimental/LearningAgents/Source/LearningAgents/Private/LearningAgentsInteractor.cpp
--- a/Engine/Plugins/Experimental/LearningAgents/Source/LearningAgents/Private/LearningAgentsInteractor.cpp
+++ b/Engine/Plugins/Experimental/LearningAgents/Source/LearningAgents/Private/LearningAgentsInteractor.cpp
@@ -13,9 +13,9 @@
 #include "LearningLog.h"
 #include "EngineDefines.h"
 
-ULearningAgentsInteractor::ULearningAgentsInteractor(const FObjectInitializer& ObjectInitializer) : ULearningAgentsManagerComponent(ObjectInitializer) {}
-ULearningAgentsInteractor::ULearningAgentsInteractor(FVTableHelper& Helper) : ULearningAgentsInteractor() {}
-ULearningAgentsInteractor::~ULearningAgentsInteractor() {}
+ULearningAgentsInteractor::ULearningAgentsInteractor() : Super(FObjectInitializer::Get()) {}
+ULearningAgentsInteractor::ULearningAgentsInteractor(FVTableHelper& Helper) : Super(Helper) {}
+ULearningAgentsInteractor::~ULearningAgentsInteractor() = default;
 
 void ULearningAgentsInteractor::SetupInteractor()
 {
diff --git a/Engine/Plugins/Experimental/LearningAgents/Source/LearningAgents/Private/LearningAgentsManager.cpp b/Engine/Plugins/Experimental/LearningAgents/Source/LearningAgents/Private/LearningAgentsManager.cpp
--- a/Engine/Plugins/Experimental/LearningAgents/Source/LearningAgents/Private/LearningAgentsManager.cpp
+++ b/Engine/Plugins/Experimental/LearningAgents/Source/LearningAgents/Private/LearningAgentsManager.cpp
@@ -10,9 +10,9 @@
 #include "UObject/ScriptInterface.h"
 #include "EngineDefines.h"
 
-ALearningAgentsManager::ALearningAgentsManager(const FObjectInitializer& ObjectInitializer) : AActor(ObjectInitializer) {}
-ALearningAgentsManager::ALearningAgentsManager(FVTableHelper& Helper) : ALearningAgentsManager() {}
-ALearningAgentsManager::~ALearningAgentsManager() {}
+ALearningAgentsManager::ALearningAgentsManager() : Super(FObjectInitializer::Get()) {}
+ALearningAgentsManager::ALearningAgentsManager(FVTableHelper& Helper) : Super(Helper) {}
+ALearningAgentsManager::~ALearningAgentsManager() = default;
 
 void ALearningAgentsManager::PostInitProperties()
 {
diff --git a/Engine/Plugins/Experimental/LearningAgents/Source/LearningAgents/Private/LearningAgentsManagerComponent.cpp b/Engine/Plugins/Experimental/LearningAgents/Source/LearningAgents/Private/LearningAgentsManagerComponent.cpp
--- a/Engine/Plugins/Experimental/LearningAgents/Source/LearningAgents/Private/LearningAgentsManagerComponent.cpp
+++ b/Engine/Plugins/Experimental/LearningAgents/Source/LearningAgents/Private/LearningAgentsManagerComponent.cpp
@@ -5,9 +5,9 @@
 #include "LearningAgentsHelpers.h"
 #include "LearningLog.h"
 
-ULearningAgentsManagerComponent::ULearningAgentsManagerComponent(const FObjectInitializer& ObjectInitializer) : UActorComponent(ObjectInitializer) {}
-ULearningAgentsManagerComponent::ULearningAgentsManagerComponent(FVTableHelper& Helper) : ULearningAgentsManagerComponent() {}
-ULearningAgentsManagerComponent::~ULearningAgentsManagerComponent() {}
+ULearningAgentsManagerComponent::ULearningAgentsManagerComponent(const FObjectInitializer& ObjectInitializer) : Super(ObjectInitializer) {}
+ULearningAgentsManagerComponent::ULearningAgentsManagerComponent(FVTableHelper& Helper) : Super(Helper) {}
+ULearningAgentsManagerComponent::~ULearningAgentsManagerComponent() = default;
 
 void ULearningAgentsManagerComponent::PostInitProperties()
 {
diff --git a/Engine/Plugins/Experimental/LearningAgents/Source/LearningAgents/Private/LearningAgentsNeuralNetwork.cpp b/Engine/Plugins/Experimental/LearningAgents/Source/LearningAgents/Private/LearningAgentsNeuralNetwork.cpp
--- a/Engine/Plugins/Experimental/LearningAgents/Source/LearningAgents/Private/LearningAgentsNeuralNetwork.cpp
+++ b/Engine/Plugins/Experimental/LearningAgents/Source/LearningAgents/Private/LearningAgentsNeuralNetwork.cpp
@@ -8,9 +8,9 @@
 #include "UObject/Package.h"
 #include "Misc/FileHelper.h"
 
-ULearningAgentsNeuralNetwork::ULearningAgentsNeuralNetwork() {}
-ULearningAgentsNeuralNetwork::ULearningAgentsNeuralNetwork(FVTableHelper& Helper) : ULearningAgentsNeuralNetwork() {}
-ULearningAgentsNeuralNetwork::~ULearningAgentsNeuralNetwork() {}
+ULearningAgentsNeuralNetwork::ULearningAgentsNeuralNetwork() = default;
+ULearningAgentsNeuralNetwork::ULearningAgentsNeuralNetwork(FVTableHelper& Helper) : Super(Helper) {}
+ULearningAgentsNeuralNetwork::~ULearningAgentsNeuralNetwork() = default;
 
 namespace UE::Learning::Agents
 {
diff --git a/Engine/Plugins/Experimental/LearningAgents/Source/LearningAgents/Private/LearningAgentsPolicy.cpp b/Engine/Plugins/Experimental/LearningAgents/Source/LearningAgents/Private/LearningAgentsPolicy.cpp
--- a/Engine/Plugins/Experimental/LearningAgents/Source/LearningAgents/Private/LearningAgentsPolicy.cpp
+++ b/Engine/Plugins/Experimental/LearningAgents/Source/LearningAgents/Private/LearningAgentsPolicy.cpp
@@ -14,9 +14,9 @@
 #include "Misc/Paths.h"
 #include "GameFramework/Actor.h"
 
-ULearningAgentsPolicy::ULearningAgentsPolicy(const FObjectInitializer& ObjectInitializer) : ULearningAgentsManagerComponent(ObjectInitializer) {}
-ULearningAgentsPolicy::ULearningAgentsPolicy(FVTableHelper& Helper) : ULearningAgentsPolicy() {}
-ULearningAgentsPolicy::~ULearningAgentsPolicy() {}
+ULearningAgentsPolicy::ULearningAgentsPolicy() : Super(FObjectInitializer::Get()) {}
+ULearningAgentsPolicy::ULearningAgentsPolicy(FVTableHelper& Helper) : Super(Helper) {}
+ULearningAgentsPolicy::~ULearningAgentsPolicy() = default;
 
 void ULearningAgentsPolicy::SetupPolicy(
 	ULearningAgentsInteractor* InInteractor, 
diff --git a/Engine/Plugins/Experimental/LearningAgents/Source/LearningAgents/Public/LearningAgentsController.h b/Engine/Plugins/Experimental/LearningAgents/Source/LearningAgents/Public/LearningAgentsController.h
--- a/Engine/Plugins/Experimental/LearningAgents/Source/LearningAgents/Public/LearningAgentsController.h
+++ b/Engine/Plugins/Experimental/LearningAgents/Source/LearningAgents/Public/LearningAgentsController.h
@@ -19,7 +19,7 @@ class LEARNINGAGENTS_API ULearningAgentsController : public ULearningAgentsManag
 public:
 
 	// These constructors/destructors are needed to make forward declarations happy
-	ULearningAgentsController(const FObjectInitializer& ObjectInitializer = FObjectInitializer::Get());
+	ULearningAgentsController();
 	ULearningAgentsController(FVTableHelper& Helper);
 	virtual ~ULearningAgentsController();
 
diff --git a/Engine/Plugins/Experimental/LearningAgents/Source/LearningAgents/Public/LearningAgentsCritic.h b/Engine/Plugins/Experimental/LearningAgents/Source/LearningAgents/Public/LearningAgentsCritic.h
--- a/Engine/Plugins/Experimental/LearningAgents/Source/LearningAgents/Public/LearningAgentsCritic.h
+++ b/Engine/Plugins/Experimental/LearningAgents/Source/LearningAgents/Public/LearningAgentsCritic.h
@@ -50,7 +50,7 @@ class LEARNINGAGENTS_API ULearningAgentsCritic : public ULearningAgentsManagerCo
 public:
 
 	// These constructors/destructors are needed to make forward declarations happy
-	ULearningAgentsCritic(const FObjectInitializer& ObjectInitializer = FObjectInitializer::Get());
+	ULearningAgentsCritic();
 	ULearningAgentsCritic(FVTableHelper& Helper);
 	virtual ~ULearningAgentsCritic();
 
diff --git a/Engine/Plugins/Experimental/LearningAgents/Source/LearningAgents/Public/LearningAgentsInteractor.h b/Engine/Plugins/Experimental/LearningAgents/Source/LearningAgents/Public/LearningAgentsInteractor.h
--- a/Engine/Plugins/Experimental/LearningAgents/Source/LearningAgents/Public/LearningAgentsInteractor.h
+++ b/Engine/Plugins/Experimental/LearningAgents/Source/LearningAgents/Public/LearningAgentsInteractor.h
@@ -42,7 +42,7 @@ class LEARNINGAGENTS_API ULearningAgentsInteractor : public ULearningAgentsManag
 public:
 
 	// These constructors/destructors are needed to make forward declarations happy
-	ULearningAgentsInteractor(const FObjectInitializer& ObjectInitializer = FObjectInitializer::Get());
+	ULearningAgentsInteractor();
 	ULearningAgentsInteractor(FVTableHelper& Helper);
 	virtual ~ULearningAgentsInteractor();
 
diff --git a/Engine/Plugins/Experimental/LearningAgents/Source/LearningAgents/Public/LearningAgentsManager.h b/Engine/Plugins/Experimental/LearningAgents/Source/LearningAgents/Public/LearningAgentsManager.h
--- a/Engine/Plugins/Experimental/LearningAgents/Source/LearningAgents/Public/LearningAgentsManager.h
+++ b/Engine/Plugins/Experimental/LearningAgents/Source/LearningAgents/Public/LearningAgentsManager.h
@@ -28,7 +28,7 @@ class ULearningAgentsManagerComponent;
  * If you have multiple different types of objects you want controlled by Learning Agents, you should consider creating
  * one agent manager per object type, rather than trying to share an agent manager.
  */
-UCLASS(Abstract, BlueprintType, Blueprintable)
+UCLASS(BlueprintType, Blueprintable)
 class LEARNINGAGENTS_API ALearningAgentsManager : public AActor
 {
 	GENERATED_BODY()
@@ -36,7 +36,7 @@ class LEARNINGAGENTS_API ALearningAgentsManager : public AActor
 public:
 
 	// These constructors/destructors are needed to make forward declarations happy
-	ALearningAgentsManager(const FObjectInitializer& ObjectInitializer = FObjectInitializer::Get());
+	ALearningAgentsManager();
 	ALearningAgentsManager(FVTableHelper& Helper);
 	virtual ~ALearningAgentsManager();
 
diff --git a/Engine/Plugins/Experimental/LearningAgents/Source/LearningAgents/Public/LearningAgentsPolicy.h b/Engine/Plugins/Experimental/LearningAgents/Source/LearningAgents/Public/LearningAgentsPolicy.h
--- a/Engine/Plugins/Experimental/LearningAgents/Source/LearningAgents/Public/LearningAgentsPolicy.h
+++ b/Engine/Plugins/Experimental/LearningAgents/Source/LearningAgents/Public/LearningAgentsPolicy.h
@@ -66,7 +66,7 @@ class LEARNINGAGENTS_API ULearningAgentsPolicy : public ULearningAgentsManagerCo
 public:
 
 	// These constructors/destructors are needed to make forward declarations happy
-	ULearningAgentsPolicy(const FObjectInitializer& ObjectInitializer = FObjectInitializer::Get());
+	ULearningAgentsPolicy();
 	ULearningAgentsPolicy(FVTableHelper& Helper);
 	virtual ~ULearningAgentsPolicy();
 
diff --git a/Engine/Plugins/Experimental/LearningAgents/Source/LearningAgentsTraining/Private/LearningAgentsImitationTrainer.cpp b/Engine/Plugins/Experimental/LearningAgents/Source/LearningAgentsTraining/Private/LearningAgentsImitationTrainer.cpp
--- a/Engine/Plugins/Experimental/LearningAgents/Source/LearningAgentsTraining/Private/LearningAgentsImitationTrainer.cpp
+++ b/Engine/Plugins/Experimental/LearningAgents/Source/LearningAgentsTraining/Private/LearningAgentsImitationTrainer.cpp
@@ -16,9 +16,9 @@
 #include "LearningNeuralNetworkObject.h"
 #include "Misc/Paths.h"
 
-ULearningAgentsImitationTrainer::ULearningAgentsImitationTrainer(const FObjectInitializer& ObjectInitializer) : UActorComponent(ObjectInitializer) {}
-ULearningAgentsImitationTrainer::ULearningAgentsImitationTrainer(FVTableHelper& Helper) : ULearningAgentsImitationTrainer() {}
-ULearningAgentsImitationTrainer::~ULearningAgentsImitationTrainer() {}
+ULearningAgentsImitationTrainer::ULearningAgentsImitationTrainer() : Super(FObjectInitializer::Get()) {}
+ULearningAgentsImitationTrainer::ULearningAgentsImitationTrainer(FVTableHelper& Helper) : Super(Helper) {}
+ULearningAgentsImitationTrainer::~ULearningAgentsImitationTrainer() = default;
 
 void ULearningAgentsImitationTrainer::EndPlay(const EEndPlayReason::Type EndPlayReason)
 {
diff --git a/Engine/Plugins/Experimental/LearningAgents/Source/LearningAgentsTraining/Private/LearningAgentsRecorder.cpp b/Engine/Plugins/Experimental/LearningAgents/Source/LearningAgentsTraining/Private/LearningAgentsRecorder.cpp
--- a/Engine/Plugins/Experimental/LearningAgents/Source/LearningAgentsTraining/Private/LearningAgentsRecorder.cpp
+++ b/Engine/Plugins/Experimental/LearningAgents/Source/LearningAgentsTraining/Private/LearningAgentsRecorder.cpp
@@ -83,9 +83,9 @@ void ULearningAgentsRecorder::FAgentRecordBuffer::CopyToRecord(FLearningAgentsRe
 	}
 }
 
-ULearningAgentsRecorder::ULearningAgentsRecorder(const FObjectInitializer& ObjectInitializer) : ULearningAgentsManagerComponent(ObjectInitializer) {}
-ULearningAgentsRecorder::ULearningAgentsRecorder(FVTableHelper& Helper) : ULearningAgentsRecorder() {}
-ULearningAgentsRecorder::~ULearningAgentsRecorder() {}
+ULearningAgentsRecorder::ULearningAgentsRecorder() : Super(FObjectInitializer::Get()) {}
+ULearningAgentsRecorder::ULearningAgentsRecorder(FVTableHelper& Helper) : Super(Helper) {}
+ULearningAgentsRecorder::~ULearningAgentsRecorder() = default;
 
 void ULearningAgentsRecorder::EndPlay(const EEndPlayReason::Type EndPlayReason)
 {
diff --git a/Engine/Plugins/Experimental/LearningAgents/Source/LearningAgentsTraining/Private/LearningAgentsRecording.cpp b/Engine/Plugins/Experimental/LearningAgents/Source/LearningAgentsTraining/Private/LearningAgentsRecording.cpp
--- a/Engine/Plugins/Experimental/LearningAgents/Source/LearningAgentsTraining/Private/LearningAgentsRecording.cpp
+++ b/Engine/Plugins/Experimental/LearningAgents/Source/LearningAgentsTraining/Private/LearningAgentsRecording.cpp
@@ -17,9 +17,9 @@ bool FLearningAgentsRecord::Serialize(FArchive& Ar)
 	return true;
 }
 
-ULearningAgentsRecording::ULearningAgentsRecording() {}
-ULearningAgentsRecording::ULearningAgentsRecording(FVTableHelper& Helper) {}
-ULearningAgentsRecording::~ULearningAgentsRecording() {}
+ULearningAgentsRecording::ULearningAgentsRecording() = default;
+ULearningAgentsRecording::ULearningAgentsRecording(FVTableHelper& Helper) : Super(Helper) {}
+ULearningAgentsRecording::~ULearningAgentsRecording() = default;
 
 namespace UE::Learning::Agents::Recording::Private
 {
diff --git a/Engine/Plugins/Experimental/LearningAgents/Source/LearningAgentsTraining/Private/LearningAgentsTrainer.cpp b/Engine/Plugins/Experimental/LearningAgents/Source/LearningAgentsTraining/Private/LearningAgentsTrainer.cpp
--- a/Engine/Plugins/Experimental/LearningAgents/Source/LearningAgentsTraining/Private/LearningAgentsTrainer.cpp
+++ b/Engine/Plugins/Experimental/LearningAgents/Source/LearningAgentsTraining/Private/LearningAgentsTrainer.cpp
@@ -104,9 +104,9 @@ FString FLearningAgentsTrainerPathSettings::GetIntermediatePath() const
 	return IntermediateRelativePath.Path;
 }
 
-ULearningAgentsTrainer::ULearningAgentsTrainer(const FObjectInitializer& ObjectInitializer) : ULearningAgentsManagerComponent(ObjectInitializer) {}
-ULearningAgentsTrainer::ULearningAgentsTrainer(FVTableHelper& Helper) : ULearningAgentsTrainer() {}
-ULearningAgentsTrainer::~ULearningAgentsTrainer() {}
+ULearningAgentsTrainer::ULearningAgentsTrainer() : Super(FObjectInitializer::Get()) {}
+ULearningAgentsTrainer::ULearningAgentsTrainer(FVTableHelper& Helper) : Super(Helper) {}
+ULearningAgentsTrainer::~ULearningAgentsTrainer() = default;
 
 void ULearningAgentsTrainer::SetupTrainer(
 	ULearningAgentsInteractor* InInteractor,
diff --git a/Engine/Plugins/Experimental/LearningAgents/Source/LearningAgentsTraining/Public/LearningAgentsImitationTrainer.h b/Engine/Plugins/Experimental/LearningAgents/Source/LearningAgentsTraining/Public/LearningAgentsImitationTrainer.h
--- a/Engine/Plugins/Experimental/LearningAgents/Source/LearningAgentsTraining/Public/LearningAgentsImitationTrainer.h
+++ b/Engine/Plugins/Experimental/LearningAgents/Source/LearningAgentsTraining/Public/LearningAgentsImitationTrainer.h
@@ -96,7 +96,7 @@ class LEARNINGAGENTSTRAINING_API ULearningAgentsImitationTrainer : public UActor
 public:
 
 	// These constructors/destructors are needed to make forward declarations happy
-	ULearningAgentsImitationTrainer(const FObjectInitializer& ObjectInitializer = FObjectInitializer::Get());
+	ULearningAgentsImitationTrainer();
 	ULearningAgentsImitationTrainer(FVTableHelper& Helper);
 	virtual ~ULearningAgentsImitationTrainer();
 
diff --git a/Engine/Plugins/Experimental/LearningAgents/Source/LearningAgentsTraining/Public/LearningAgentsRecorder.h b/Engine/Plugins/Experimental/LearningAgents/Source/LearningAgentsTraining/Public/LearningAgentsRecorder.h
--- a/Engine/Plugins/Experimental/LearningAgents/Source/LearningAgentsTraining/Public/LearningAgentsRecorder.h
+++ b/Engine/Plugins/Experimental/LearningAgents/Source/LearningAgentsTraining/Public/LearningAgentsRecorder.h
@@ -42,7 +42,7 @@ class LEARNINGAGENTSTRAINING_API ULearningAgentsRecorder : public ULearningAgent
 public:
 
 	// These constructors/destructors are needed to make forward declarations happy
-	ULearningAgentsRecorder(const FObjectInitializer& ObjectInitializer = FObjectInitializer::Get());
+	ULearningAgentsRecorder();
 	ULearningAgentsRecorder(FVTableHelper& Helper);
 	virtual ~ULearningAgentsRecorder();
 
diff --git a/Engine/Plugins/Experimental/LearningAgents/Source/LearningAgentsTraining/Public/LearningAgentsTrainer.h b/Engine/Plugins/Experimental/LearningAgents/Source/LearningAgentsTraining/Public/LearningAgentsTrainer.h
--- a/Engine/Plugins/Experimental/LearningAgents/Source/LearningAgentsTraining/Public/LearningAgentsTrainer.h
+++ b/Engine/Plugins/Experimental/LearningAgents/Source/LearningAgentsTraining/Public/LearningAgentsTrainer.h
@@ -332,7 +332,7 @@ class LEARNINGAGENTSTRAINING_API ULearningAgentsTrainer : public ULearningAgents
 public:
 
 	// These constructors/destructors are needed to make forward declarations happy
-	ULearningAgentsTrainer(const FObjectInitializer& ObjectInitializer = FObjectInitializer::Get());
+	ULearningAgentsTrainer();
 	ULearningAgentsTrainer(FVTableHelper& Helper);
 	virtual ~ULearningAgentsTrainer();
 
