Learning Agents: Add ChildCanTick to manager

#rnx

diff --git a/Engine/Plugins/Experimental/LearningAgents/Source/LearningAgents/Public/LearningAgentsManager.h b/Engine/Plugins/Experimental/LearningAgents/Source/LearningAgents/Public/LearningAgentsManager.h
--- a/Engine/Plugins/Experimental/LearningAgents/Source/LearningAgents/Public/LearningAgentsManager.h
+++ b/Engine/Plugins/Experimental/LearningAgents/Source/LearningAgents/Public/LearningAgentsManager.h
@@ -28,7 +28,7 @@ class ULearningAgentsManagerComponent;
  * If you have multiple different types of objects you want controlled by Learning Agents, you should consider creating
  * one agent manager per object type, rather than trying to share an agent manager.
  */
-UCLASS(BlueprintType, Blueprintable)
+UCLASS(BlueprintType, Blueprintable, meta=(ChildCanTick))
 class LEARNINGAGENTS_API ALearningAgentsManager : public AActor
 {
 	GENERATED_BODY()
