LearningAgents: bug & QOL fixes - recorder, trainer issues

#rnx

diff --git a/Engine/Plugins/Experimental/LearningAgents/Source/LearningAgentsTraining/Private/LearningAgentsImitationTrainer.cpp b/Engine/Plugins/Experimental/LearningAgents/Source/LearningAgentsTraining/Private/LearningAgentsImitationTrainer.cpp
--- a/Engine/Plugins/Experimental/LearningAgents/Source/LearningAgentsTraining/Private/LearningAgentsImitationTrainer.cpp
+++ b/Engine/Plugins/Experimental/LearningAgents/Source/LearningAgentsTraining/Private/LearningAgentsImitationTrainer.cpp
@@ -198,6 +198,7 @@ void ULearningAgentsImitationTrainer::BeginTraining(
 	if (Response != UE::Learning::ETrainerResponse::Success)
 	{
 		UE_LOG(LogLearning, Error, TEXT("%s: Error sending or receiving policy from trainer: %s. Check log for errors."), *GetName(), UE::Learning::Trainer::GetResponseString(Response));
+		bHasTrainingFailed = true;
 		ImitationTrainer->Terminate();
 		return;
 	}
@@ -211,6 +212,7 @@ void ULearningAgentsImitationTrainer::BeginTraining(
 	if (Response != UE::Learning::ETrainerResponse::Success)
 	{
 		UE_LOG(LogLearning, Error, TEXT("%s: Error sending experience to trainer: %s. Check log for errors."), *GetName(), UE::Learning::Trainer::GetResponseString(Response));
+		bHasTrainingFailed = true;
 		ImitationTrainer->Terminate();
 		return;
 	}
@@ -266,6 +268,7 @@ void ULearningAgentsImitationTrainer::IterateTraining()
 		else if (Response != UE::Learning::ETrainerResponse::Success)
 		{
 			UE_LOG(LogLearning, Error, TEXT("%s: Error receiving policy from trainer. Check log for errors."), *GetName());
+			bHasTrainingFailed = true;
 			EndTraining();
 			return;
 		}
@@ -280,6 +283,12 @@ void ULearningAgentsImitationTrainer::RunTraining(
 	const FLearningAgentsTrainerPathSettings& ImitationTrainerPathSettings,
 	const bool bReinitializePolicyNetwork)
 {
+	if (bHasTrainingFailed)
+	{
+		UE_LOG(LogLearning, Error, TEXT("%s: Training has failed. Check log for errors."), *GetName());
+		return;
+	}
+
 	// If we aren't training yet, then start training and do the first inference step.
 	if (!IsTraining())
 	{
@@ -306,3 +315,8 @@ bool ULearningAgentsImitationTrainer::IsTraining() const
 {
 	return bIsTraining;
 }
+
+bool ULearningAgentsImitationTrainer::HasTrainingFailed() const
+{
+	return bHasTrainingFailed;
+}
diff --git a/Engine/Plugins/Experimental/LearningAgents/Source/LearningAgentsTraining/Private/LearningAgentsRecorder.cpp b/Engine/Plugins/Experimental/LearningAgents/Source/LearningAgentsTraining/Private/LearningAgentsRecorder.cpp
--- a/Engine/Plugins/Experimental/LearningAgents/Source/LearningAgentsTraining/Private/LearningAgentsRecorder.cpp
+++ b/Engine/Plugins/Experimental/LearningAgents/Source/LearningAgentsTraining/Private/LearningAgentsRecorder.cpp
@@ -259,12 +259,6 @@ void ULearningAgentsRecorder::EndRecording()
 		}
 	}
 
-	// Save Recording to Intermediate Directory
-
-	FFilePath RecordingFilePath;
-	RecordingFilePath.FilePath = RecordingDirectory / FString::Printf(TEXT("%s_%s.bin"), *GetName(), *FDateTime::Now().ToFormattedString(TEXT("%Y-%m-%d_%H-%M-%S")));
-	SaveRecordingToFile(RecordingFilePath);
-
 	bIsRecording = false;
 }
 
diff --git a/Engine/Plugins/Experimental/LearningAgents/Source/LearningAgentsTraining/Private/LearningAgentsTrainer.cpp b/Engine/Plugins/Experimental/LearningAgents/Source/LearningAgentsTraining/Private/LearningAgentsTrainer.cpp
--- a/Engine/Plugins/Experimental/LearningAgents/Source/LearningAgentsTraining/Private/LearningAgentsTrainer.cpp
+++ b/Engine/Plugins/Experimental/LearningAgents/Source/LearningAgentsTraining/Private/LearningAgentsTrainer.cpp
@@ -553,6 +553,7 @@ void ULearningAgentsTrainer::BeginTraining(
 	if (Response != UE::Learning::ETrainerResponse::Success)
 	{
 		UE_LOG(LogLearning, Error, TEXT("%s: Error sending or receiving policy from trainer: %s. Check log for errors."), *GetName(), UE::Learning::Trainer::GetResponseString(Response));
+		bHasTrainingFailed = true;
 		Trainer->Terminate();
 		return;
 	}
@@ -576,6 +577,7 @@ void ULearningAgentsTrainer::BeginTraining(
 		if (Response != UE::Learning::ETrainerResponse::Success)
 		{
 			UE_LOG(LogLearning, Error, TEXT("%s: Error sending or receiving critic from trainer: %s. Check log for errors."), *GetName(), UE::Learning::Trainer::GetResponseString(Response));
+			bHasTrainingFailed = true;
 			Trainer->Terminate();
 			return;
 		}
@@ -938,6 +940,7 @@ void ULearningAgentsTrainer::ProcessExperience()
 			if (Response != UE::Learning::ETrainerResponse::Success)
 			{
 				UE_LOG(LogLearning, Error, TEXT("%s: Error waiting to push experience to trainer. Check log for errors."), *GetName());
+				bHasTrainingFailed = true;
 				EndTraining();
 				return;
 			}
@@ -957,6 +960,7 @@ void ULearningAgentsTrainer::ProcessExperience()
 			else if (Response != UE::Learning::ETrainerResponse::Success)
 			{
 				UE_LOG(LogLearning, Error, TEXT("%s: Error waiting for policy from trainer. Check log for errors."), *GetName());
+				bHasTrainingFailed = true;
 				EndTraining();
 				return;
 			}
@@ -970,6 +974,7 @@ void ULearningAgentsTrainer::ProcessExperience()
 				if (Response != UE::Learning::ETrainerResponse::Success)
 				{
 					UE_LOG(LogLearning, Error, TEXT("%s: Error waiting for critic from trainer. Check log for errors."), *GetName());
+					bHasTrainingFailed = true;
 					EndTraining();
 					return;
 				}
@@ -998,6 +1003,12 @@ void ULearningAgentsTrainer::RunTraining(
 		return;
 	}
 
+	if (bHasTrainingFailed)
+	{
+		UE_LOG(LogLearning, Error, TEXT("%s: Training has failed. Check log for errors."), *GetName());
+		return;
+	}
+
 	// If we aren't training yet, then start training and do the first inference step.
 	if (!IsTraining())
 	{
@@ -1085,6 +1096,12 @@ bool ULearningAgentsTrainer::IsCompleted(const int32 AgentId, ELearningAgentsCom
 		return true;
 	}
 }
+
+bool ULearningAgentsTrainer::HasTrainingFailed() const
+{
+	return bHasTrainingFailed;
+}
+
 void ULearningAgentsTrainer::ResetEpisodes_Implementation(const TArray<int32>& AgentId)
 {
 	// Can be overridden to reset agent without blueprints
diff --git a/Engine/Plugins/Experimental/LearningAgents/Source/LearningAgentsTraining/Public/LearningAgentsImitationTrainer.h b/Engine/Plugins/Experimental/LearningAgents/Source/LearningAgentsTraining/Public/LearningAgentsImitationTrainer.h
--- a/Engine/Plugins/Experimental/LearningAgents/Source/LearningAgentsTraining/Public/LearningAgentsImitationTrainer.h
+++ b/Engine/Plugins/Experimental/LearningAgents/Source/LearningAgentsTraining/Public/LearningAgentsImitationTrainer.h
@@ -107,6 +107,15 @@ public:
 	UFUNCTION(BlueprintPure, Category = "LearningAgents")
 	bool IsTraining() const;
 
+	/**
+	 * Returns true if the trainer has failed to communicate with the external training process. This can be used in
+	 * combination with RunTraining to avoid filling the logs with errors.
+	 *
+	 * @returns				True if the training has failed. Otherwise, false.
+	 */
+	UFUNCTION(BlueprintPure, Category = "LearningAgents")
+	bool HasTrainingFailed() const;
+
 	/**
 	 * Begins the training process with the provided settings.
 	 * @param InPolicy The policy to train.
@@ -163,6 +172,13 @@ private:
 	UPROPERTY(VisibleAnywhere, Transient, Category = "LearningAgents")
 	bool bIsTraining = false;
 
+	/**
+	 * True if trainer encountered an unrecoverable error during training (e.g. the trainer process timed out). Otherwise, false.
+	 * This exists mainly to keep the editor from locking up if something goes wrong during training.
+	 */
+	UPROPERTY(VisibleAnywhere, Transient, Category = "LearningAgents")
+	bool bHasTrainingFailed = false;
+
 	float TrainerTimeout = 10.0f;
 
 	void DoneTraining();
diff --git a/Engine/Plugins/Experimental/LearningAgents/Source/LearningAgentsTraining/Public/LearningAgentsRecorder.h b/Engine/Plugins/Experimental/LearningAgents/Source/LearningAgentsTraining/Public/LearningAgentsRecorder.h
--- a/Engine/Plugins/Experimental/LearningAgents/Source/LearningAgentsTraining/Public/LearningAgentsRecorder.h
+++ b/Engine/Plugins/Experimental/LearningAgents/Source/LearningAgentsTraining/Public/LearningAgentsRecorder.h
@@ -109,7 +109,7 @@ public:
 	void LoadRecordingFromFile(const FFilePath& File);
 
 	/** Saves the current recording object to a file */
-	UFUNCTION(BlueprintCallable, Category = "LearningAgents", meta = (RelativePath))
+	UFUNCTION(BlueprintCallable, BlueprintPure = false, Category = "LearningAgents", meta = (RelativePath))
 	void SaveRecordingToFile(const FFilePath& File) const;
 
 	/** Append to the current recording object from a file. */
diff --git a/Engine/Plugins/Experimental/LearningAgents/Source/LearningAgentsTraining/Public/LearningAgentsRecording.h b/Engine/Plugins/Experimental/LearningAgents/Source/LearningAgentsTraining/Public/LearningAgentsRecording.h
--- a/Engine/Plugins/Experimental/LearningAgents/Source/LearningAgentsTraining/Public/LearningAgentsRecording.h
+++ b/Engine/Plugins/Experimental/LearningAgents/Source/LearningAgentsTraining/Public/LearningAgentsRecording.h
@@ -64,7 +64,7 @@ public:
 	void LoadRecordingFromFile(const FFilePath& File);
 
 	/** Save this recording to a file. */
-	UFUNCTION(BlueprintCallable, Category = "LearningAgents", meta = (RelativePath))
+	UFUNCTION(BlueprintCallable, BlueprintPure = false, Category = "LearningAgents", meta = (RelativePath))
 	void SaveRecordingToFile(const FFilePath& File) const;
 
 	/** Append to this recording from a file. */
diff --git a/Engine/Plugins/Experimental/LearningAgents/Source/LearningAgentsTraining/Public/LearningAgentsTrainer.h b/Engine/Plugins/Experimental/LearningAgents/Source/LearningAgentsTraining/Public/LearningAgentsTrainer.h
--- a/Engine/Plugins/Experimental/LearningAgents/Source/LearningAgentsTraining/Public/LearningAgentsTrainer.h
+++ b/Engine/Plugins/Experimental/LearningAgents/Source/LearningAgentsTraining/Public/LearningAgentsTrainer.h
@@ -519,6 +519,15 @@ public:
 	 */
 	UFUNCTION(BlueprintPure, Category = "LearningAgents", meta = (AgentId = "-1"))
 	bool IsCompleted(const int32 AgentId, ELearningAgentsCompletion& OutCompletion) const;
+
+	/**
+	 * Returns true if the trainer has failed to communicate with the external training process. This can be used in
+	 * combination with RunTraining to avoid filling the logs with errors.
+	 *
+	 * @returns				True if the training has failed. Otherwise, false.
+	 */
+	UFUNCTION(BlueprintPure, Category = "LearningAgents")
+	bool HasTrainingFailed() const;
 	
 // ----- Private Data ----- 
 private:
@@ -539,6 +548,13 @@ private:
 	UPROPERTY(VisibleAnywhere, Transient, Category = "LearningAgents")
 	bool bIsTraining = false;
 
+	/**
+	 * True if trainer encountered an unrecoverable error during training (e.g. the trainer process timed out). Otherwise, false.
+	 * This exists mainly to keep the editor from locking up if something goes wrong during training.
+	 */
+	UPROPERTY(VisibleAnywhere, Transient, Category = "LearningAgents")
+	bool bHasTrainingFailed = false;
+
 	/** The list of current reward objects. */
 	UPROPERTY(VisibleAnywhere, Transient, Category = "LearningAgents")
 	TArray<TObjectPtr<ULearningAgentsReward>> RewardObjects;
