LearningAgents: Make recorder paths BP read/write

#rnx

diff --git a/Engine/Plugins/Experimental/LearningAgents/Source/LearningAgentsTraining/Public/LearningAgentsRecorder.h b/Engine/Plugins/Experimental/LearningAgents/Source/LearningAgentsTraining/Public/LearningAgentsRecorder.h
--- a/Engine/Plugins/Experimental/LearningAgents/Source/LearningAgentsTraining/Public/LearningAgentsRecorder.h
+++ b/Engine/Plugins/Experimental/LearningAgents/Source/LearningAgentsTraining/Public/LearningAgentsRecorder.h
@@ -24,11 +24,11 @@ public:
 	FLearningAgentsRecorderPathSettings();
 
 	/** The relative path to the Intermediate directory. Defaults to FPaths::ProjectIntermediateDir. */
-	UPROPERTY(EditAnywhere, Category = "LearningAgents", meta = (RelativePath))
+	UPROPERTY(BlueprintReadWrite, EditAnywhere, Category = "LearningAgents", meta = (RelativePath))
 	FDirectoryPath IntermediateRelativePath;
 
 	/** The name of the sub-directory to use in the intermediate directory */
-	UPROPERTY(EditAnywhere, Category = "LearningAgents", meta = (RelativePath))
+	UPROPERTY(BlueprintReadWrite, EditAnywhere, Category = "LearningAgents", meta = (RelativePath))
 	FString RecordingsSubdirectory = TEXT("Recordings");
 };
 
@@ -86,7 +86,7 @@ public:
 	void EndRecording();
 
 	/** Returns true if recording is active; Otherwise, false. */
-	UFUNCTION(BlueprintCallable, Category = "LearningAgents")
+	UFUNCTION(BlueprintPure, Category = "LearningAgents")
 	bool IsRecording() const;
 
 	/**
