LearningAgents: Switched to NNERuntimeBasicCpu for Inference.

diff --git a/Engine/Plugins/Experimental/LearningAgents/LearningAgents.uplugin b/Engine/Plugins/Experimental/LearningAgents/LearningAgents.uplugin
--- a/Engine/Plugins/Experimental/LearningAgents/LearningAgents.uplugin
+++ b/Engine/Plugins/Experimental/LearningAgents/LearningAgents.uplugin
@@ -46,6 +46,10 @@
 		{
 			"Name": "PythonFoundationPackages",
 			"Enabled": true
+		},
+		{
+			"Name": "NNERuntimeBasicCpu",
+			"Enabled": true
 		}
 	]
 }
\ No newline at end of file
diff --git a/Engine/Plugins/Experimental/LearningAgents/Source/Learning/Private/Learning.ispc b/Engine/Plugins/Experimental/LearningAgents/Source/Learning/Private/Learning.ispc
--- a/Engine/Plugins/Experimental/LearningAgents/Source/Learning/Private/Learning.ispc
+++ b/Engine/Plugins/Experimental/LearningAgents/Source/Learning/Private/Learning.ispc
@@ -150,76 +150,58 @@ export void LearningLayerMatMulPlusBias(
 	const uniform float Weights[],
 	const uniform float Biases[],
 	const uniform int BatchNum,
+	const uniform int BufferSize,
 	const uniform int RowNum,
 	const uniform int ColNum)
 {
 	foreach (BatchIdx = 0 ... BatchNum, ColIdx = 0 ... ColNum)
 	{
-		Output[BatchIdx * ColNum + ColIdx] = Biases[ColIdx];
+		Output[BatchIdx * BufferSize + ColIdx] = Biases[ColIdx];
 	}
 	
 	for (uniform int BatchIdx = 0; BatchIdx < BatchNum; BatchIdx++)
 	{
 		for (uniform int RowIdx = 0; RowIdx < RowNum; RowIdx++)
 		{
-			uniform const float Value = Input[BatchIdx * RowNum + RowIdx];
+			uniform const float Value = Input[BatchIdx * BufferSize + RowIdx];
 
 			if (Value != 0.0)
 			{
 				foreach(ColIdx = 0 ... ColNum)
 				{
-					Output[BatchIdx * ColNum + ColIdx] += Value * Weights[RowIdx * ColNum + ColIdx];
+					Output[BatchIdx * BufferSize + ColIdx] += Value * Weights[RowIdx * ColNum + ColIdx];
 				}
 			}
 		}
 	}
 }
 
-export void LearningLayerMatMulVecPlusBias(
-	uniform float Output[],
-	const uniform float Input[],
-	const uniform float Weights[],
-	const uniform float Biases[],
-	const uniform int RowNum,
+export void LearningLayerReLU(
+	uniform float InputOutput[],
+	const uniform int BatchNum,
+	const uniform int BufferSize,
 	const uniform int ColNum)
 {
-	foreach(ColIdx = 0 ... ColNum)
-	{
-		Output[ColIdx] = Biases[ColIdx];
-	}
-
-	for (uniform int RowIdx = 0; RowIdx < RowNum; RowIdx++)
+	foreach(BatchIdx = 0 ... BatchNum, ColIdx = 0 ... ColNum)
 	{
-		uniform const float Value = Input[RowIdx];
-
-		if (Value != 0.0)
-		{
-			foreach(ColIdx = 0 ... ColNum)
-			{
-				Output[ColIdx] += Value * Weights[RowIdx * ColNum + ColIdx];
-			}
-		}
+		InputOutput[BatchIdx * BufferSize + ColIdx] = max(InputOutput[BatchIdx * BufferSize + ColIdx], 0.0f);
 	}
 }
 
-
-export void LearningLayerReLU(
-	uniform float InputOutput[],
-	const uniform int Num)
+inline float LearningELU(const float X)
 {
-	foreach(Idx = 0 ... Num)
-	{
-		InputOutput[Idx] = max(InputOutput[Idx], 0.0f);
-	}
+	return X > 0.0f ? X : exp(X) - 1.0f;
 }
 
 export void LearningLayerELU(
 	uniform float InputOutput[],
-	const uniform int Num)
+	const uniform int BatchNum,
+	const uniform int BufferSize,
+	const uniform int ColNum)
 {
-	foreach (Idx = 0 ... Num)
+	foreach (BatchIdx = 0 ... BatchNum, ColIdx = 0 ... ColNum)
 	{
-		InputOutput[Idx] = InputOutput[Idx] > 0.0f ? InputOutput[Idx] : LearningInvExpApprox(-InputOutput[Idx]) - 1.0f;
+		InputOutput[BatchIdx * BufferSize + ColIdx] = LearningELU(InputOutput[BatchIdx * BufferSize + ColIdx]);
 	}
 }
 
@@ -231,11 +213,13 @@ inline float LearningTanH(const float X, const float Min=-10.0f, const float Max
 
 export void LearningLayerTanH(
 	uniform float InputOutput[],
-	const uniform int Num)
+	const uniform int BatchNum,
+	const uniform int BufferSize,
+	const uniform int ColNum)
 {
-	foreach(Idx = 0 ... Num)
+	foreach(BatchIdx = 0 ... BatchNum, ColIdx = 0 ... ColNum)
 	{
-		InputOutput[Idx] = LearningTanH(InputOutput[Idx]);
+		InputOutput[BatchIdx * BufferSize + ColIdx] = LearningTanH(InputOutput[BatchIdx * BufferSize + ColIdx]);
 	}
 }
 
@@ -248,20 +232,19 @@ export void LearningLayerActionNoise(
 	uniform float Output[],
 	uniform float OutputMean[],
 	uniform float OutputStd[],
-	const uniform float Input[],
+	const uniform float OutputNetwork[],
 	const uniform uint32 Seed[],
 	const uniform float ActionNoiseScale[],
 	const uniform int BatchNum,
-	const uniform int InputNum,
 	const uniform int OutputNum,
 	const uniform float LogActionNoiseMin,
 	const uniform float LogActionNoiseMax)
 {
 	foreach (BatchIdx = 0 ... BatchNum, OutputIdx = 0 ... OutputNum)
 	{
-		OutputMean[BatchIdx * OutputNum + OutputIdx] = Input[BatchIdx * InputNum + OutputIdx];
+		OutputMean[BatchIdx * OutputNum + OutputIdx] = OutputNetwork[BatchIdx * 2 * OutputNum + OutputIdx];
 		OutputStd[BatchIdx * OutputNum + OutputIdx] = ActionNoiseScale[BatchIdx] *
-			exp(LearningSigmoid(Input[BatchIdx * InputNum + OutputNum + OutputIdx]) * (LogActionNoiseMax - LogActionNoiseMin) + LogActionNoiseMin);
+			exp(LearningSigmoid(OutputNetwork[BatchIdx * 2 * OutputNum + OutputNum + OutputIdx]) * (LogActionNoiseMax - LogActionNoiseMin) + LogActionNoiseMin);
 
 		Output[BatchIdx * OutputNum + OutputIdx] = LearningRandomGaussian(
 			Seed[BatchIdx] ^ 0xab744615 ^ LearningRandomInt(OutputIdx ^ 0xf8a88a27),
@@ -274,7 +257,7 @@ export void LearningLayerActionNoiseSingleBatch(
 	uniform float Output[],
 	uniform float OutputMean[],
 	uniform float OutputStd[],
-	const uniform float Input[],
+	const uniform float OutputNetwork[],
 	const uniform uint32 Seed,
 	const uniform float ActionNoiseScale,
 	const uniform int OutputNum,
@@ -283,9 +266,9 @@ export void LearningLayerActionNoiseSingleBatch(
 {
 	foreach(OutputIdx = 0 ... OutputNum)
 	{
-		OutputMean[OutputIdx] = Input[OutputIdx];
+		OutputMean[OutputIdx] = OutputNetwork[OutputIdx];
 		OutputStd[OutputIdx] = ActionNoiseScale *
-			exp(LearningSigmoid(Input[OutputNum + OutputIdx]) * (LogActionNoiseMax - LogActionNoiseMin) + LogActionNoiseMin);
+			exp(LearningSigmoid(OutputNetwork[OutputNum + OutputIdx]) * (LogActionNoiseMax - LogActionNoiseMin) + LogActionNoiseMin);
 
 		Output[OutputIdx] = LearningRandomGaussian(
 			Seed ^ 0xab744615 ^ LearningRandomInt(OutputIdx ^ 0xf8a88a27),
diff --git a/Engine/Plugins/Experimental/LearningAgents/Source/Learning/Private/LearningNeuralNetwork.cpp b/Engine/Plugins/Experimental/LearningAgents/Source/Learning/Private/LearningNeuralNetwork.cpp
--- a/Engine/Plugins/Experimental/LearningAgents/Source/Learning/Private/LearningNeuralNetwork.cpp
+++ b/Engine/Plugins/Experimental/LearningAgents/Source/Learning/Private/LearningNeuralNetwork.cpp
@@ -2,73 +2,167 @@
 
 #include "LearningNeuralNetwork.h"
 
+#if UE_LEARNING_ISPC
+#include "Learning.ispc.generated.h"
+#endif
+
 namespace UE::Learning
 {
-	const TCHAR* GetActivationFunctionString(const EActivationFunction ActivationFunction)
+	namespace NeuralNetwork
 	{
-		switch (ActivationFunction)
+		static inline void MatMulPlusBias(
+			TLearningArrayView<2, float> Output,
+			const TLearningArrayView<2, const float> Input,
+			const TLearningArrayView<2, const float> Weights,
+			const TLearningArrayView<1, const float> Biases,
+			const int32 BatchOffset,
+			const int32 BatchNum)
 		{
-		case EActivationFunction::ReLU: return TEXT("ReLU");
-		case EActivationFunction::ELU: return TEXT("ELU");
-		case EActivationFunction::TanH: return TEXT("TanH");
-		default: UE_LEARNING_NOT_IMPLEMENTED(); return TEXT("Unknown");
+			UE_LEARNING_TRACE_CPUPROFILER_EVENT_SCOPE(Learning::NeuralNetwork::MatMulPlusBias);
+
+			const int32 RowNum = Weights.Num<0>();
+			const int32 ColNum = Weights.Num<1>();
+			const int32 BufferSize = Input.Num<1>();
+
+#if UE_LEARNING_ISPC
+			ispc::LearningLayerMatMulPlusBias(
+				Output.Slice(BatchOffset, BatchNum).GetData(),
+				Input.Slice(BatchOffset, BatchNum).GetData(),
+				Weights.GetData(),
+				Biases.GetData(),
+				BatchNum,
+				BufferSize,
+				RowNum,
+				ColNum);
+#else
+			for (const int32 BatchIdx = BatchOffset; BatchIdx < BatchOffset + BatchNum; BatchIdx++)
+			{
+				for (int32 ColIdx = 0; ColIdx < ColNum; ColIdx++)
+				{
+					Output[BatchIdx][ColIdx] = Biases[ColIdx];
+				}
+
+				for (int32 RowIdx = 0; RowIdx < RowNum; RowIdx++)
+				{
+					const float Value = Input[BatchIdx][RowIdx];
+
+					if (Value != 0.0)
+					{
+						for (int32 ColIdx = 0; ColIdx < ColNum; ColIdx++)
+						{
+							Output[BatchIdx][ColIdx] += Value * Weights[RowIdx][ColIdx];
+						}
+					}
+				}
+			}
+#endif
 		}
-	}
 
-	void FNeuralNetwork::Resize(
-		const int32 InputNum,
-		const int32 OutputNum,
-		const int32 HiddenNum,
-		const int32 LayerNum)
-	{
-		UE_LEARNING_TRACE_CPUPROFILER_EVENT_SCOPE(Learning::FNeuralNetwork::Resize);
-		UE_LEARNING_CHECKF(LayerNum >= 2, TEXT("At least two layers required (input and output layers)"));
+		static inline float ReLU(const float X)
+		{
+			return FMath::Max(X, 0.0f);
+		}
 
-		Weights.SetNum(LayerNum);
-		Biases.SetNum(LayerNum);
+		static inline float ELU(const float X)
+		{
+			return X > 0.0f ? X : FMath::Exp(X) - 1.0f;
+		}
 
-		Weights[0].SetNumUninitialized({ InputNum, HiddenNum });
-		Biases[0].SetNumUninitialized({ HiddenNum });
+		static inline float TanH(const float X)
+		{
+			return FMath::Tanh(X);
+		}
 
-		for (int32 LayerIdx = 0; LayerIdx < LayerNum - 2; LayerIdx++)
+		static inline void ActivationReLU(TLearningArrayView<2, float> InputOutput, const int32 BatchOffset, const int32 BatchNum, const int32 InputOutputNum)
 		{
-			Weights[LayerIdx + 1].SetNumUninitialized({ HiddenNum, HiddenNum });
-			Biases[LayerIdx + 1].SetNumUninitialized({ HiddenNum });
+			UE_LEARNING_TRACE_CPUPROFILER_EVENT_SCOPE(Learning::NeuralNetwork::ActivationReLU);
+
+			const int32 BufferSize = InputOutput.Num<1>();
+
+#if UE_LEARNING_ISPC
+			ispc::LearningLayerReLU(
+				InputOutput.Slice(BatchOffset, BatchNum).GetData(),
+				BatchNum,
+				BufferSize, 
+				InputOutputNum);
+#else
+			for (int32 BatchIdx = BatchOffset; BatchIdx < BatchOffset + BatchNum; BatchIdx++)
+			{
+				for (int32 ColIdx = 0; ColIdx < InputOutputNum; ColIdx++)
+				{
+					InputOutput[BatchIdx][ColIdx] = ReLU(InputOutput[BatchIdx][ColIdx]);
+				}
+			}
+#endif
 		}
 
-		Weights[LayerNum - 1].SetNumUninitialized({ HiddenNum, OutputNum });
-		Biases[LayerNum - 1].SetNumUninitialized({ OutputNum });
+		static inline void ActivationELU(TLearningArrayView<2, float> InputOutput, const int32 BatchOffset, const int32 BatchNum, const int32 InputOutputNum)
+		{
+			UE_LEARNING_TRACE_CPUPROFILER_EVENT_SCOPE(Learning::NeuralNetwork::ActivationELU);
 
-		for (int32 LayerIdx = 0; LayerIdx < LayerNum; LayerIdx++)
+			const int32 BufferSize = InputOutput.Num<1>();
+
+#if UE_LEARNING_ISPC
+			ispc::LearningLayerELU(
+				InputOutput.Slice(BatchOffset, BatchNum).GetData(),
+				BatchNum,
+				BufferSize,
+				InputOutputNum);
+#else
+			for (int32 BatchIdx = BatchOffset; BatchIdx < BatchOffset + BatchNum; BatchIdx++)
+			{
+				for (int32 ColIdx = 0; ColIdx < InputOutputNum; ColIdx++)
+				{
+					InputOutput[BatchIdx][ColIdx] = ELU(InputOutput[BatchIdx][ColIdx]);
+				}
+			}
+#endif
+		}
+
+		static inline void ActivationTanH(TLearningArrayView<2, float> InputOutput, const int32 BatchOffset, const int32 BatchNum, const int32 InputOutputNum)
 		{
-			Array::Zero(Weights[LayerIdx]);
-			Array::Zero(Biases[LayerIdx]);
+			UE_LEARNING_TRACE_CPUPROFILER_EVENT_SCOPE(Learning::NeuralNetwork::ActivationTanH);
+
+			const int32 BufferSize = InputOutput.Num<1>();
+
+#if UE_LEARNING_ISPC
+			ispc::LearningLayerTanH(
+				InputOutput.Slice(BatchOffset, BatchNum).GetData(),
+				BatchNum,
+				BufferSize,
+				InputOutputNum);
+#else
+			for (int32 BatchIdx = BatchOffset; BatchIdx < BatchOffset + BatchNum; BatchIdx++)
+			{
+				for (int32 ValueIdx = 0; ValueIdx < InputOutputNum; ValueIdx++)
+				{
+					InputOutput[BatchIdx][ValueIdx] = TanH(InputOutput[BatchIdx][ValueIdx]);
+				}
+			}
+#endif
 		}
-	}
 
-	int32 FNeuralNetwork::GetInputNum() const
-	{
-		return Weights[0].Num<0>();
+		static constexpr int32 MagicNumber = 0x0cd353cf;
+		static constexpr int32 VersionNumber = 1;
 	}
 
-	int32 FNeuralNetwork::GetOutputNum() const
+	bool FNeuralNetworkMLP::DeserializeFromBytes(int32& InOutOffset, const TLearningArrayView<1, const uint8> RawBytes)
 	{
-		return Weights[Weights.Num() - 1].Num<1>();
-	}
+		UE_LEARNING_TRACE_CPUPROFILER_EVENT_SCOPE(Learning::FNeuralNetworkMLP::DeserializeFromBytes);
 
-	int32 FNeuralNetwork::GetLayerNum() const
-	{
-		return Weights.Num();
-	}
+		// At least Magic Number and Version Number Required
+		if (RawBytes.Num() < sizeof(int32) * 2)
+		{
+			return false;
+		}
 
-	int32 FNeuralNetwork::GetHiddenNum() const
-	{
-		return Weights[0].Num<1>();
-	}
+		int32 Magic;
+		UE::Learning::DeserializeFromBytes(InOutOffset, RawBytes, Magic);
+		if (Magic != NeuralNetwork::MagicNumber) { return false; }
 
-	void FNeuralNetwork::DeserializeFromBytes(int32& InOutOffset, const TLearningArrayView<1, const uint8> RawBytes)
-	{
-		UE_LEARNING_TRACE_CPUPROFILER_EVENT_SCOPE(Learning::FNeuralNetwork::DeserializeFromBytes);
+		int32 Version;
+		UE::Learning::DeserializeFromBytes(InOutOffset, RawBytes, Version);
+		if (Version != NeuralNetwork::VersionNumber) { return false; }
 
 		int32 ActivationFunctionInt;
 		UE::Learning::DeserializeFromBytes(InOutOffset, RawBytes, ActivationFunctionInt);
@@ -85,45 +179,191 @@ namespace UE::Learning
 			UE::Learning::Array::DeserializeFromBytes(InOutOffset, RawBytes, Weights[LayerIdx]);
 			UE::Learning::Array::DeserializeFromBytes(InOutOffset, RawBytes, Biases[LayerIdx]);
 		}
+
+		return true;
 	}
 
-	void FNeuralNetwork::SerializeToBytes(int32& InOutOffset, TLearningArrayView<1, uint8> OutRawBytes) const
+	void FNeuralNetworkMLP::SerializeToBytes(int32& InOutOffset, TLearningArrayView<1, uint8> OutRawBytes) const
 	{
-		UE_LEARNING_TRACE_CPUPROFILER_EVENT_SCOPE(Learning::FNeuralNetwork::SerializeToBytes);
+		UE_LEARNING_TRACE_CPUPROFILER_EVENT_SCOPE(Learning::FNeuralNetworkMLP::SerializeToBytes);
+
+		const int32 LayerNum = Weights.Num();
 
+		UE::Learning::SerializeToBytes(InOutOffset, OutRawBytes, NeuralNetwork::MagicNumber);
+		UE::Learning::SerializeToBytes(InOutOffset, OutRawBytes, NeuralNetwork::VersionNumber);
 		UE::Learning::SerializeToBytes(InOutOffset, OutRawBytes, (int32)ActivationFunction);
-		UE::Learning::SerializeToBytes(InOutOffset, OutRawBytes, GetLayerNum());
+		UE::Learning::SerializeToBytes(InOutOffset, OutRawBytes, LayerNum);
 
-		for (int32 LayerIdx = 0; LayerIdx < GetLayerNum(); LayerIdx++)
+		for (int32 LayerIdx = 0; LayerIdx < LayerNum; LayerIdx++)
 		{
 			UE::Learning::Array::SerializeToBytes(InOutOffset, OutRawBytes, Weights[LayerIdx]);
 			UE::Learning::Array::SerializeToBytes(InOutOffset, OutRawBytes, Biases[LayerIdx]);
 		}
 	}
 
-	int32 FNeuralNetwork::GetSerializationByteNum(
+	int32 FNeuralNetworkMLP::GetSerializationByteNum() const
+	{
+		const int32 LayerNum = GetLayerNum();
+
+		int32 Total = 0;
+
+		Total += sizeof(int32);  // Magic
+		Total += sizeof(int32);  // Version
+		Total += sizeof(int32);  // Activation
+		Total += sizeof(int32);  // Layer Num
+
+		for (int32 LayerIdx = 0; LayerIdx < LayerNum; LayerIdx++)
+		{
+			Total += UE::Learning::Array::SerializationByteNum<2, float>(Weights[LayerIdx].Shape());
+			Total += UE::Learning::Array::SerializationByteNum<1, float>(Biases[LayerIdx].Shape());
+		}
+
+		return Total;
+	}
+
+	TSharedRef<INeuralNetworkInference> FNeuralNetworkMLP::CreateInferenceObject(
+		const int32 MaxBatchSize,
+		const FNeuralNetworkInferenceSettings& Settings)
+	{
+		return MakeShared<FNeuralNetworkInferenceMLP>(*this, MaxBatchSize, Settings);
+	}
+
+	int32 FNeuralNetworkMLP::GetInputNum() const
+	{
+		UE_LEARNING_CHECK(Weights.Num() > 0);
+		return Weights[0].Num<0>();
+	}
+
+	int32 FNeuralNetworkMLP::GetOutputNum() const
+	{
+		UE_LEARNING_CHECK(Weights.Num() > 0);
+		return Weights.Last().Num<1>();
+	}
+
+	int32 FNeuralNetworkMLP::GetHiddenNum() const
+	{
+		UE_LEARNING_CHECK(Weights.Num() > 0);
+		return Weights[0].Num<1>();
+	}
+
+	int32 FNeuralNetworkMLP::GetLayerNum() const
+	{
+		return Weights.Num();
+	}
+
+	void FNeuralNetworkMLP::Resize(
 		const int32 InputNum,
 		const int32 OutputNum,
 		const int32 HiddenNum,
 		const int32 LayerNum)
 	{
-		int32 Total = 0;
-			
-		Total += sizeof(int32);  // Activation
-		Total += sizeof(int32);  // Layer Num
-		Total += UE::Learning::Array::SerializationByteNum<2, float>({ InputNum, HiddenNum });
-		Total += UE::Learning::Array::SerializationByteNum<1, float>({ HiddenNum });
+		UE_LEARNING_TRACE_CPUPROFILER_EVENT_SCOPE(Learning::FNeuralNetworkMLP::Resize);
+		UE_LEARNING_CHECKF(LayerNum >= 2, TEXT("At least two layers required (input and output layers)"));
+
+		Weights.SetNum(LayerNum);
+		Biases.SetNum(LayerNum);
+
+		Weights[0].SetNumUninitialized({ InputNum, HiddenNum });
+		Biases[0].SetNumUninitialized({ HiddenNum });
 
 		for (int32 LayerIdx = 0; LayerIdx < LayerNum - 2; LayerIdx++)
 		{
-			Total += UE::Learning::Array::SerializationByteNum<2, float>({ HiddenNum, HiddenNum });
-			Total += UE::Learning::Array::SerializationByteNum<1, float>({ HiddenNum });
+			Weights[LayerIdx + 1].SetNumUninitialized({ HiddenNum, HiddenNum });
+			Biases[LayerIdx + 1].SetNumUninitialized({ HiddenNum });
 		}
 
-		Total += UE::Learning::Array::SerializationByteNum<2, float>({ HiddenNum, OutputNum });
-		Total += UE::Learning::Array::SerializationByteNum<1, float>({ OutputNum });
+		Weights[LayerNum - 1].SetNumUninitialized({ HiddenNum, OutputNum });
+		Biases[LayerNum - 1].SetNumUninitialized({ OutputNum });
 
-		return Total;
+		for (int32 LayerIdx = 0; LayerIdx < LayerNum; LayerIdx++)
+		{
+			Array::Zero(Weights[LayerIdx]);
+			Array::Zero(Biases[LayerIdx]);
+		}
 	}
 
+	FNeuralNetworkInferenceMLP::FNeuralNetworkInferenceMLP(const FNeuralNetworkMLP& InNetwork, const int32 MaxBatchSize, const FNeuralNetworkInferenceSettings& InSettings)
+		: Network(InNetwork)
+		, Settings(InSettings)
+	{
+		const int32 BufferSize = FMath::Max(FMath::Max(Network.GetInputNum(), Network.GetOutputNum()), Network.GetHiddenNum());
+
+		FrontBuffer.SetNumUninitialized({ MaxBatchSize, BufferSize });
+		BackBuffer.SetNumUninitialized({ MaxBatchSize, BufferSize });
+	}
+
+	void FNeuralNetworkInferenceMLP::Evaluate(
+		TLearningArrayView<2, float> Output,
+		const TLearningArrayView<2, const float> Input,
+		const FIndexSet Instances) 
+	{
+		const int32 InputNum = Network.GetInputNum();
+		const int32 OutputNum = Network.GetOutputNum();
+		const int32 LayerNum = Network.GetLayerNum();
+
+		UE_LEARNING_CHECK(Input.Num<1>() == InputNum);
+		UE_LEARNING_CHECK(Output.Num<1>() == OutputNum);
+
+		// Gather Instance Inputs
+
+		Array::Check(Input, Instances);
+
+		const TLearningArrayView<2, float> InputBuffer = FrontBuffer;
+
+		for (int32 InstanceIdx = 0; InstanceIdx < Instances.Num(); InstanceIdx++)
+		{
+			Array::Copy(InputBuffer[InstanceIdx].Slice(0, InputNum), Input[Instances[InstanceIdx]]);
+		}
+
+		// Evaluate Network
+
+		auto EvaluationFunction = [this, LayerNum](const int32 BatchOffset, const int32 BatchSize)
+		{
+			for (int32 LayerIdx = 0; LayerIdx < LayerNum; LayerIdx++)
+			{
+				const TLearningArrayView<2, float> LayerInputBuffer = LayerIdx % 2 == 0 ? FrontBuffer : BackBuffer;
+				const TLearningArrayView<2, float> LayerOutputBuffer = LayerIdx % 2 == 0 ? BackBuffer : FrontBuffer;
+				const int32 LayerOutputNum = Network.Weights[LayerIdx].Num<1>();
+
+				NeuralNetwork::MatMulPlusBias(
+					LayerOutputBuffer,
+					LayerInputBuffer,
+					Network.Weights[LayerIdx],
+					Network.Biases[LayerIdx],
+					BatchOffset,
+					BatchSize);
+
+				if (LayerIdx != LayerNum - 1)
+				{
+					switch (Network.ActivationFunction)
+					{
+					case EActivationFunction::ReLU: NeuralNetwork::ActivationReLU(LayerOutputBuffer, BatchOffset, BatchSize, LayerOutputNum); break;
+					case EActivationFunction::ELU: NeuralNetwork::ActivationELU(LayerOutputBuffer, BatchOffset, BatchSize, LayerOutputNum); break;
+					case EActivationFunction::TanH: NeuralNetwork::ActivationTanH(LayerOutputBuffer, BatchOffset, BatchSize, LayerOutputNum); break;
+					default: UE_LEARNING_CHECKF(false, TEXT("Unimplemented"));
+					}
+				}
+			}
+		};
+
+		if (Settings.bParallelEvaluation && Instances.Num() > Settings.MinParallelBatchSize)
+		{
+			Learning::SlicedParallelFor(Instances.Num(), Settings.MinParallelBatchSize, EvaluationFunction);
+		}
+		else
+		{
+			EvaluationFunction(0, Instances.Num());
+		}
+
+		// Scatter Instance Outputs
+
+		const TLearningArrayView<2, float> OutputBuffer = LayerNum % 2 == 0 ? FrontBuffer : BackBuffer;
+
+		for (int32 InstanceIdx = 0; InstanceIdx < Instances.Num(); InstanceIdx++)
+		{
+			Array::Copy(Output[Instances[InstanceIdx]], OutputBuffer[InstanceIdx].Slice(0, OutputNum));
+		}
+
+		Array::Check(Output, Instances);
+	}
 }
diff --git a/Engine/Plugins/Experimental/LearningAgents/Source/Learning/Private/LearningNeuralNetworkObject.cpp b/Engine/Plugins/Experimental/LearningAgents/Source/Learning/Private/LearningNeuralNetworkObject.cpp
--- a/Engine/Plugins/Experimental/LearningAgents/Source/Learning/Private/LearningNeuralNetworkObject.cpp
+++ b/Engine/Plugins/Experimental/LearningAgents/Source/Learning/Private/LearningNeuralNetworkObject.cpp
@@ -2,7 +2,6 @@
 
 #include "LearningNeuralNetworkObject.h"
 
-#include "LearningNeuralNetwork.h"
 #include "LearningRandom.h"
 
 #if UE_LEARNING_ISPC
@@ -11,510 +10,177 @@
 
 namespace UE::Learning
 {
-	namespace NeuralNetworkObject
+	namespace NeuralNetworkPolicyFunction
 	{
-		static inline float ReLU(const float X)
-		{
-			return FMath::Max(X, 0.0f);
-		}
-
-		static inline float ELU(const float X)
-		{
-			return X > 0.0f ? X : FMath::InvExpApprox(-X) - 1.0f;
-		}
-
 		static inline float Sigmoid(const float X)
 		{
 			return 1.0f / (1.0f + FMath::InvExpApprox(X));
 		}
-
-		static inline float TanH(const float X)
-		{
-			return FMath::Tanh(X);
-		}
-
-		static inline void MatMulPlusBias(
-			TLearningArrayView<2, float> Output,
-			const TLearningArrayView<2, const float> Input,
-			const TLearningArrayView<2, const float> Weights,
-			const TLearningArrayView<1, const float> Biases,
-			const FIndexSet Instances)
-		{
-			UE_LEARNING_TRACE_CPUPROFILER_EVENT_SCOPE(Learning::NeuralNetworkObject::MatMulPlusBias);
-
-			const int32 RowNum = Weights.Num<0>();
-			const int32 ColNum = Weights.Num<1>();
-
-#if UE_LEARNING_ISPC
-			if (Instances.IsSlice())
-			{
-				ispc::LearningLayerMatMulPlusBias(
-					Output.Slice(Instances.GetSliceStart(), Instances.GetSliceNum()).GetData(),
-					Input.Slice(Instances.GetSliceStart(), Instances.GetSliceNum()).GetData(),
-					Weights.GetData(),
-					Biases.GetData(),
-					Instances.GetSliceNum(),
-					RowNum,
-					ColNum);
-			}
-			else
-			{
-				for (const int32 InstanceIdx : Instances)
-				{
-					ispc::LearningLayerMatMulVecPlusBias(
-						Output[InstanceIdx].GetData(),
-						Input[InstanceIdx].GetData(),
-						Weights.GetData(),
-						Biases.GetData(),
-						RowNum,
-						ColNum);
-				}
-			}
-#else
-			for (const int32 InstanceIdx : Instances)
-			{
-				Array::Copy(Output[InstanceIdx], Biases);
-
-				for (int32 RowIdx = 0; RowIdx < RowNum; RowIdx++)
-				{
-					const float Value = Input[InstanceIdx][RowIdx];
-
-					if (Value != 0.0)
-					{
-						for (int32 ColIdx = 0; ColIdx < ColNum; ColIdx++)
-						{
-							Output[InstanceIdx][ColIdx] += Value * Weights[RowIdx][ColIdx];
-						}
-					}
-				}
-			}
-#endif
-		}
-
-		static inline void ActivationReLU(TLearningArrayView<2, float> InputOutput, const FIndexSet Instances)
-		{
-			UE_LEARNING_TRACE_CPUPROFILER_EVENT_SCOPE(Learning::NeuralNetworkObject::ActivationReLU);
-
-			const int32 HiddenNum = InputOutput.Num<1>();
-
-#if UE_LEARNING_ISPC
-			if (Instances.IsSlice())
-			{
-				ispc::LearningLayerReLU(
-					InputOutput.Slice(Instances.GetSliceStart(), Instances.GetSliceNum()).GetData(),
-					Instances.GetSliceNum() * HiddenNum);
-			}
-			else
-			{
-				for (const int32 InstanceIdx : Instances)
-				{
-					ispc::LearningLayerReLU(
-						InputOutput[InstanceIdx].GetData(),
-						HiddenNum);
-				}
-			}
-#else
-			for (const int32 InstanceIdx : Instances)
-			{
-				for (int32 HiddenIdx = 0; HiddenIdx < HiddenNum; HiddenIdx++)
-				{
-					InputOutput[InstanceIdx][HiddenIdx] = ReLU(InputOutput[InstanceIdx][HiddenIdx]);
-				}
-			}
-#endif
-		}
-
-		static inline void ActivationELU(TLearningArrayView<2, float> InputOutput, const FIndexSet Instances)
-		{
-			UE_LEARNING_TRACE_CPUPROFILER_EVENT_SCOPE(Learning::NeuralNetworkObject::ActivationELU);
-
-			const int32 HiddenNum = InputOutput.Num<1>();
-
-#if UE_LEARNING_ISPC
-			if (Instances.IsSlice())
-			{
-				ispc::LearningLayerELU(
-					InputOutput.Slice(Instances.GetSliceStart(), Instances.GetSliceNum()).GetData(),
-					Instances.GetSliceNum() * HiddenNum);
-			}
-			else
-			{
-				for (const int32 InstanceIdx : Instances)
-				{
-					ispc::LearningLayerELU(
-						InputOutput[InstanceIdx].GetData(),
-						HiddenNum);
-				}
-			}
-#else
-			for (const int32 InstanceIdx : Instances)
-			{
-				for (int32 HiddenIdx = 0; HiddenIdx < HiddenNum; HiddenIdx++)
-				{
-					InputOutput[InstanceIdx][HiddenIdx] = ELU(InputOutput[InstanceIdx][HiddenIdx]);
-				}
-			}
-#endif
-		}
-
-		static inline void ActivationTanH(TLearningArrayView<2, float> InputOutput, const FIndexSet Instances)
-		{
-			UE_LEARNING_TRACE_CPUPROFILER_EVENT_SCOPE(Learning::NeuralNetworkObject::ActivationTanH);
-
-			const int32 HiddenNum = InputOutput.Num<1>();
-
-#if UE_LEARNING_ISPC
-			if (Instances.IsSlice())
-			{
-				ispc::LearningLayerTanH(
-					InputOutput.Slice(Instances.GetSliceStart(), Instances.GetSliceNum()).GetData(),
-					Instances.GetSliceNum() * HiddenNum);
-			}
-			else
-			{
-				for (const int32 InstanceIdx : Instances)
-				{
-					ispc::LearningLayerTanH(
-						InputOutput[InstanceIdx].GetData(),
-						HiddenNum);
-				}
-			}
-#else
-			for (const int32 InstanceIdx : Instances)
-			{
-				for (int32 HiddenIdx = 0; HiddenIdx < HiddenNum; HiddenIdx++)
-				{
-					InputOutput[InstanceIdx][HiddenIdx] = TanH(InputOutput[InstanceIdx][HiddenIdx]);
-				}
-			}
-#endif
-		}
-
-		static inline void ActionNoise(
-			TLearningArrayView<2, float> Output,
-			TLearningArrayView<2, float> OutputMean,
-			TLearningArrayView<2, float> OutputStd,
-			const TLearningArrayView<2, float> Input,
-			TLearningArrayView<1, uint32> Seed,
-			const TLearningArrayView<1, const float> ActionNoiseScale,
-			const float LogActionNoiseMin,
-			const float LogActionNoiseMax,
-			const FIndexSet Instances)
-		{
-			UE_LEARNING_TRACE_CPUPROFILER_EVENT_SCOPE(Learning::NeuralNetworkObject::ActionNoise);
-
-			const int32 InputNum = Input.Num<1>();
-			const int32 OutputNum = Output.Num<1>();
-
-#if UE_LEARNING_ISPC
-			if (Instances.IsSlice())
-			{
-				ispc::LearningLayerActionNoise(
-					Output.Slice(Instances.GetSliceStart(), Instances.GetSliceNum()).GetData(),
-					OutputMean.Slice(Instances.GetSliceStart(), Instances.GetSliceNum()).GetData(),
-					OutputStd.Slice(Instances.GetSliceStart(), Instances.GetSliceNum()).GetData(),
-					Input.Slice(Instances.GetSliceStart(), Instances.GetSliceNum()).GetData(),
-					Seed.Slice(Instances.GetSliceStart(), Instances.GetSliceNum()).GetData(),
-					ActionNoiseScale.Slice(Instances.GetSliceStart(), Instances.GetSliceNum()).GetData(),
-					Instances.GetSliceNum(),
-					InputNum,
-					OutputNum,
-					LogActionNoiseMin,
-					LogActionNoiseMax);
-			}
-			else
-			{
-				for (const int32 InstanceIdx : Instances)
-				{
-					ispc::LearningLayerActionNoiseSingleBatch(
-						Output[InstanceIdx].GetData(),
-						OutputMean[InstanceIdx].GetData(),
-						OutputStd[InstanceIdx].GetData(),
-						Input[InstanceIdx].GetData(),
-						Seed[InstanceIdx],
-						ActionNoiseScale[InstanceIdx],
-						OutputNum,
-						LogActionNoiseMin,
-						LogActionNoiseMax);
-				}
-			}
-#else
-			for (const int32 InstanceIdx : Instances)
-			{
-				for (int32 OutputIdx = 0; OutputIdx < OutputNum; OutputIdx++)
-				{
-					OutputMean[InstanceIdx][OutputIdx] = Input[InstanceIdx][OutputIdx];
-					OutputStd[InstanceIdx][OutputIdx] = ActionNoiseScale[InstanceIdx] *
-						FMath::Exp(Sigmoid(Input[InstanceIdx][OutputNum + OutputIdx]) * (LogActionNoiseMax - LogActionNoiseMin) + LogActionNoiseMin);
-
-					Output[InstanceIdx][OutputIdx] = Random::Gaussian(
-						Seed[InstanceIdx] ^ 0xab744615 ^ Random::Int(OutputIdx ^ 0xf8a88a27),
-						OutputMean[InstanceIdx][OutputIdx],
-						OutputStd[InstanceIdx][OutputIdx]);
-				}
-			}
-#endif
-
-			Random::ResampleStateArray(Seed, Instances);
-		}
-
-		static inline void EvaluatePolicyLayer(
-			TLearningArrayView<2, float> Outputs,
-			TLearningArrayView<2, float> OutputMeans,
-			TLearningArrayView<2, float> OutputStds,
-			TArrayView<TLearningArray<2, float, TInlineAllocator<128>>> Activations,
-			const TLearningArrayView<2, const float> Inputs,
-			TLearningArrayView<1, uint32> Seed,
-			const TLearningArrayView<1, const float> ActionNoiseScale,
-			const int32 LayerIdx,
-			const int32 LayerNum,
-			const float LogActionNoiseMin,
-			const float LogActionNoiseMax,
-			const FNeuralNetwork& NeuralNetwork,
-			const FIndexSet Instances)
-		{
-			UE_LEARNING_TRACE_CPUPROFILER_EVENT_SCOPE(Learning::NeuralNetworkObject::EvaluatePolicyLayer);
-
-			TLearningArrayView<2, float> LayerOutput = Activations[LayerIdx - 0];
-			const TLearningArrayView<2, const float> LayerInput = LayerIdx == 0 ? Inputs : Activations[LayerIdx - 1];
-
-			// Apply Linear Transformation
-
-			MatMulPlusBias(
-				LayerOutput,
-				LayerInput,
-				NeuralNetwork.Weights[LayerIdx],
-				NeuralNetwork.Biases[LayerIdx],
-				Instances);
-
-			// Apply Noise on final layer otherwise apply activation in-place
-
-			if (LayerIdx == LayerNum - 1)
-			{
-				ActionNoise(
-					Outputs,
-					OutputMeans,
-					OutputStds,
-					LayerOutput,
-					Seed,
-					ActionNoiseScale,
-					LogActionNoiseMin,
-					LogActionNoiseMax,
-					Instances);
-			}
-			else
-			{
-				switch (NeuralNetwork.ActivationFunction)
-				{
-				case EActivationFunction::ReLU: ActivationReLU(LayerOutput, Instances); break;
-				case EActivationFunction::ELU: ActivationELU(LayerOutput, Instances); break;
-				case EActivationFunction::TanH: ActivationTanH(LayerOutput, Instances); break;
-				default: UE_LEARNING_NOT_IMPLEMENTED();
-				}
-			}
-		}
-
-		static inline void EvaluateCriticLayer(
-			TLearningArrayView<1, float> Outputs,
-			TArrayView<TLearningArray<2, float, TInlineAllocator<128>>> Activations,
-			const TLearningArrayView<2, const float> Inputs,
-			const int32 LayerIdx,
-			const int32 LayerNum,
-			const FNeuralNetwork& NeuralNetwork,
-			const FIndexSet Instances)
-		{
-			UE_LEARNING_TRACE_CPUPROFILER_EVENT_SCOPE(Learning::NeuralNetworkObject::EvaluateCriticLayer);
-
-			TLearningArrayView<2, float> LayerOutput = LayerIdx == LayerNum - 1 ? TLearningArrayView<2, float>(Outputs.GetData(), { Outputs.Num<0>(), 1}) : Activations[LayerIdx - 0];
-			const TLearningArrayView<2, const float> LayerInput = LayerIdx == 0 ? Inputs : Activations[LayerIdx - 1];
-
-			// Apply Linear Transformation
-
-			MatMulPlusBias(
-				LayerOutput,
-				LayerInput,
-				NeuralNetwork.Weights[LayerIdx],
-				NeuralNetwork.Biases[LayerIdx],
-				Instances);
-
-			// Apply activation in-place
-
-			if (LayerIdx != LayerNum - 1)
-			{
-				switch (NeuralNetwork.ActivationFunction)
-				{
-				case EActivationFunction::ReLU: ActivationReLU(LayerOutput, Instances); break;
-				case EActivationFunction::ELU: ActivationELU(LayerOutput, Instances); break;
-				case EActivationFunction::TanH: ActivationTanH(LayerOutput, Instances); break;
-				default: UE_LEARNING_NOT_IMPLEMENTED();
-				}
-			}
-		}
 	}
 
 	FNeuralNetworkPolicyFunction::FNeuralNetworkPolicyFunction(
 		const FName& InIdentifier,
 		const TSharedRef<FArrayMap>& InInstanceData,
 		const int32 InMaxInstanceNum,
-		const TSharedRef<FNeuralNetwork>& InNeuralNetwork,
+		const TSharedPtr<INeuralNetwork>& InNeuralNetwork,
 		const uint32 InSeed,
+		const FNeuralNetworkInferenceSettings& InInferenceSettings,
 		const FNeuralNetworkPolicyFunctionSettings& InSettings)
 		: FFunctionObject(InInstanceData)
+		, MaxInstanceNum(InMaxInstanceNum)
 		, NeuralNetwork(InNeuralNetwork)
+		, InferenceSettings(InInferenceSettings)
 		, Settings(InSettings)
 	{
 		SeedHandle = InstanceData->Add<1, uint32>({ InIdentifier, TEXT("Seed") }, { InMaxInstanceNum });
 		InputHandle = InstanceData->Add<2, float>({ InIdentifier, TEXT("Input") }, { InMaxInstanceNum, NeuralNetwork->GetInputNum() }, 0.0f);
 		OutputHandle = InstanceData->Add<2, float>({ InIdentifier, TEXT("Output") }, { InMaxInstanceNum, NeuralNetwork->GetOutputNum() / 2 }, 0.0f);
+		OutputNetworkHandle = InstanceData->Add<2, float>({ InIdentifier, TEXT("OutputNetwork") }, { InMaxInstanceNum, NeuralNetwork->GetOutputNum() }, 0.0f);
 		OutputMeanHandle = InstanceData->Add<2, float>({ InIdentifier, TEXT("OutputMean") }, { InMaxInstanceNum, NeuralNetwork->GetOutputNum() / 2 }, 0.0f);
 		OutputStdHandle = InstanceData->Add<2, float>({ InIdentifier, TEXT("OutputStd") }, { InMaxInstanceNum, NeuralNetwork->GetOutputNum() / 2 }, 0.0f);
 		ActionNoiseScaleHandle = InstanceData->Add<1, float>({ InIdentifier, TEXT("ActionNoiseScale") }, { InMaxInstanceNum }, Settings.ActionNoiseScale);
 
+		NeuralNetworkInference = NeuralNetwork->CreateInferenceObject(InMaxInstanceNum, InInferenceSettings);
+		
 		Random::IntArray(InstanceData->View(SeedHandle), InSeed);
-
-		// Allocate temporary storage for activations
-
-		const int32 LayerNum = InNeuralNetwork->GetLayerNum();
-
-		Activations.SetNum(LayerNum);
-		for (int32 LayerIdx = 0; LayerIdx < LayerNum; LayerIdx++)
-		{
-			Activations[LayerIdx].SetNumUninitialized({ InMaxInstanceNum, NeuralNetwork->Weights[LayerIdx].Num<1>() });
-		}
 	}
 
 	void FNeuralNetworkPolicyFunction::Evaluate(const FIndexSet Instances)
 	{
 		UE_LEARNING_TRACE_CPUPROFILER_EVENT_SCOPE(Learning::FNeuralNetworkPolicyFunction::Evaluate);
 
-		const TLearningArrayView<2, const float> Inputs = InstanceData->ConstView(InputHandle);
+		const TLearningArrayView<2, const float> Input = InstanceData->ConstView(InputHandle);
 		const TLearningArrayView<1, const float> ActionNoiseScale = InstanceData->ConstView(ActionNoiseScaleHandle);
-		TLearningArrayView<2, float> Outputs = InstanceData->View(OutputHandle);
-		TLearningArrayView<2, float> OutputMeans = InstanceData->View(OutputMeanHandle);
-		TLearningArrayView<2, float> OutputStds = InstanceData->View(OutputStdHandle);
+		TLearningArrayView<2, float> Output = InstanceData->View(OutputHandle);
+		TLearningArrayView<2, float> OutputNetwork = InstanceData->View(OutputNetworkHandle);
+		TLearningArrayView<2, float> OutputMean = InstanceData->View(OutputMeanHandle);
+		TLearningArrayView<2, float> OutputStd = InstanceData->View(OutputStdHandle);
 		TLearningArrayView<1, uint32> Seed = InstanceData->View(SeedHandle);
 
-		UE_LEARNING_CHECK(NeuralNetwork->GetInputNum() == Inputs.Num<1>());
-		UE_LEARNING_CHECK(NeuralNetwork->GetOutputNum() == 2 * Outputs.Num<1>());
+		const int32 InstanceNum = Instances.Num();
+		const int32 InputNum = Input.Num<1>();
+		const int32 OutputNum = Output.Num<1>();
 
-		const int32 LayerNum = NeuralNetwork->GetLayerNum();
+		UE_LEARNING_CHECK(NeuralNetwork->GetInputNum() == InputNum);
+		UE_LEARNING_CHECK(NeuralNetwork->GetOutputNum() == 2 * OutputNum);
 
-		if (!ensureMsgf(LayerNum > 0, TEXT("Empty Neural Network used in Policy")))
-		{
-			Array::Zero(Outputs, Instances);
-			return;
-		}
+		Array::Check(Input, Instances);
+
+		NeuralNetworkInference->Evaluate(OutputNetwork, Input, Instances);
+
+		Array::Check(OutputNetwork, Instances);
+
+		// Apply Action Noise
 
 		UE_LEARNING_ARRAY_VALUE_CHECK(Settings.ActionNoiseMin >= 0.0f && Settings.ActionNoiseMax >= 0.0f);
 		const float LogActionNoiseMin = FMath::Loge(Settings.ActionNoiseMin + UE_KINDA_SMALL_NUMBER);
 		const float LogActionNoiseMax = FMath::Loge(Settings.ActionNoiseMax + UE_KINDA_SMALL_NUMBER);
 
-		// Compute Layers
-
-		auto LayerEvaluationFunction = [&](const int32 SliceStart, const int32 SliceNum)
+#if UE_LEARNING_ISPC
+		if (Instances.IsSlice())
+		{
+			ispc::LearningLayerActionNoise(
+				Output.Slice(Instances.GetSliceStart(), Instances.GetSliceNum()).GetData(),
+				OutputMean.Slice(Instances.GetSliceStart(), Instances.GetSliceNum()).GetData(),
+				OutputStd.Slice(Instances.GetSliceStart(), Instances.GetSliceNum()).GetData(),
+				OutputNetwork.Slice(Instances.GetSliceStart(), Instances.GetSliceNum()).GetData(),
+				Seed.Slice(Instances.GetSliceStart(), Instances.GetSliceNum()).GetData(),
+				ActionNoiseScale.Slice(Instances.GetSliceStart(), Instances.GetSliceNum()).GetData(),
+				Instances.GetSliceNum(),
+				OutputNum,
+				LogActionNoiseMin,
+				LogActionNoiseMax);
+		}
+		else
 		{
-			for (int32 LayerIdx = 0; LayerIdx < LayerNum; LayerIdx++)
+			for (const int32 InstanceIdx : Instances)
 			{
-				NeuralNetworkObject::EvaluatePolicyLayer(
-					Outputs,
-					OutputMeans,
-					OutputStds,
-					Activations,
-					Inputs,
-					Seed,
-					ActionNoiseScale,
-					LayerIdx,
-					LayerNum,
+				ispc::LearningLayerActionNoiseSingleBatch(
+					Output[InstanceIdx].GetData(),
+					OutputMean[InstanceIdx].GetData(),
+					OutputStd[InstanceIdx].GetData(),
+					OutputNetwork[InstanceIdx].GetData(),
+					Seed[InstanceIdx],
+					ActionNoiseScale[InstanceIdx],
+					OutputNum,
 					LogActionNoiseMin,
-					LogActionNoiseMax,
-					*NeuralNetwork,
-					Instances.Slice(SliceStart, SliceNum));
+					LogActionNoiseMax);
 			}
-		};
-
-		if (Settings.bParallelEvaluation && Instances.Num() > Settings.MinParallelBatchSize)
-		{
-			SlicedParallelFor(Instances.Num(), Settings.MinParallelBatchSize, LayerEvaluationFunction);
 		}
-		else
+#else
+		for (const int32 InstanceIdx : Instances)
 		{
-			LayerEvaluationFunction(0, Instances.Num());
+			for (int32 OutputIdx = 0; OutputIdx < OutputNum; OutputIdx++)
+			{
+				OutputMean[InstanceIdx][OutputIdx] = OutputNetwork[InstanceIdx][OutputIdx];
+				OutputStd[InstanceIdx][OutputIdx] = ActionNoiseScale[InstanceIdx] *
+					FMath::Exp(NeuralNetworkPolicyFunction::Private::Sigmoid(
+						OutputNetwork[InstanceIdx][OutputNum + OutputIdx]) * (LogActionNoiseMax - LogActionNoiseMin) + LogActionNoiseMin);
+
+				Output[InstanceIdx][OutputIdx] = Random::Gaussian(
+					Seed[InstanceIdx] ^ 0xab744615 ^ Random::Int(OutputIdx ^ 0xf8a88a27),
+					OutputMean[InstanceIdx][OutputIdx],
+					OutputStd[InstanceIdx][OutputIdx]);
+			}
 		}
+#endif
+
+		Random::ResampleStateArray(Seed, Instances);
+
+		Array::Check(OutputMean, Instances);
+		Array::Check(OutputStd, Instances);
+		Array::Check(Output, Instances);
+	}
 
-		Array::Check(Outputs, Instances);
+	void FNeuralNetworkPolicyFunction::UpdateNeuralNetwork(const TSharedPtr<INeuralNetwork>& NewNeuralNetwork)
+	{
+		if (NeuralNetwork != NewNeuralNetwork)
+		{
+			NeuralNetwork = NewNeuralNetwork;
+			NeuralNetworkInference = NeuralNetwork->CreateInferenceObject(MaxInstanceNum, InferenceSettings);
+		}
 	}
 
 	FNeuralNetworkCriticFunction::FNeuralNetworkCriticFunction(
 		const FName& InIdentifier,
 		const TSharedRef<FArrayMap>& InInstanceData,
 		const int32 InMaxInstanceNum,
-		const TSharedRef<FNeuralNetwork>& InNeuralNetwork,
-		const FNeuralNetworkCriticFunctionSettings& InSettings)
+		const TSharedPtr<INeuralNetwork>& InNeuralNetwork,
+		const FNeuralNetworkInferenceSettings& InInferenceSettings)
 		: FFunctionObject(InInstanceData)
+		, MaxInstanceNum(InMaxInstanceNum)
 		, NeuralNetwork(InNeuralNetwork)
-		, Settings(InSettings)
+		, InferenceSettings(InInferenceSettings)
 	{
 		InputHandle = InstanceData->Add<2, float>({ InIdentifier, TEXT("Input") }, { InMaxInstanceNum, NeuralNetwork->GetInputNum() }, 0.0f);
 		OutputHandle = InstanceData->Add<1, float>({ InIdentifier, TEXT("Output") }, { InMaxInstanceNum }, 0.0f);
 
-		// Allocate temporary storage for activations
-		// 
-		// We will use the output handle as the final layer directly
-		// so we don't need storage for this activation like we do with
-		// the policy network.
-		const int32 LayerNum = InNeuralNetwork->GetLayerNum() - 1;
-
-		Activations.SetNum(LayerNum);
-		for (int32 LayerIdx = 0; LayerIdx < LayerNum; LayerIdx++)
-		{
-			Activations[LayerIdx].SetNumUninitialized({ InMaxInstanceNum, NeuralNetwork->Weights[LayerIdx].Num<1>() });
-		}
+		NeuralNetworkInference = NeuralNetwork->CreateInferenceObject(InMaxInstanceNum, InInferenceSettings);
 	}
 
 	void FNeuralNetworkCriticFunction::Evaluate(const FIndexSet Instances)
 	{
 		UE_LEARNING_TRACE_CPUPROFILER_EVENT_SCOPE(Learning::FNeuralNetworkCriticFunction::Evaluate);
 
-		const TLearningArrayView<2, const float> Inputs = InstanceData->ConstView(InputHandle);
-		TLearningArrayView<1, float> Outputs = InstanceData->View(OutputHandle);
+		const TLearningArrayView<2, const float> Input = InstanceData->ConstView(InputHandle);
+		TLearningArrayView<1, float> Output = InstanceData->View(OutputHandle);
 
-		UE_LEARNING_CHECK(NeuralNetwork->GetInputNum() == Inputs.Num<1>());
+		UE_LEARNING_CHECK(NeuralNetwork->GetInputNum() == Input.Num<1>());
 		UE_LEARNING_CHECK(NeuralNetwork->GetOutputNum() == 1);
 
-		const int32 LayerNum = NeuralNetwork->GetLayerNum();
-
-		if (!ensureMsgf(LayerNum > 0, TEXT("Empty Neural Network used in Critic")))
-		{
-			Array::Zero(Outputs, Instances);
-			return;
-		}
-
-		// Compute Layers
+		NeuralNetworkInference->Evaluate(
+			TLearningArrayView<2, float>(Output.GetData(), { Output.Num(), 1 }),
+			Input, Instances);
 
-		auto LayerEvaluationFunction = [&](const int32 SliceStart, const int32 SliceNum)
-		{
-			for (int32 LayerIdx = 0; LayerIdx < LayerNum; LayerIdx++)
-			{
-				NeuralNetworkObject::EvaluateCriticLayer(
-					Outputs,
-					Activations,
-					Inputs,
-					LayerIdx,
-					LayerNum,
-					*NeuralNetwork,
-					Instances.Slice(SliceStart, SliceNum));
-			}
-		};
+		Array::Check(Output, Instances);
+	}
 
-		if (Settings.bParallelEvaluation && Instances.Num() > Settings.MinParallelBatchSize)
-		{
-			SlicedParallelFor(Instances.Num(), Settings.MinParallelBatchSize, LayerEvaluationFunction);
-		}
-		else
+	void FNeuralNetworkCriticFunction::UpdateNeuralNetwork(const TSharedPtr<INeuralNetwork>& NewNeuralNetwork)
+	{
+		if (NeuralNetwork != NewNeuralNetwork)
 		{
-			LayerEvaluationFunction(0, Instances.Num());
+			NeuralNetwork = NewNeuralNetwork;
+			NeuralNetworkInference = NeuralNetwork->CreateInferenceObject(MaxInstanceNum, InferenceSettings);
 		}
-
-		Array::Check(Outputs, Instances);
 	}
 }
diff --git a/Engine/Plugins/Experimental/LearningAgents/Source/Learning/Public/LearningNeuralNetwork.h b/Engine/Plugins/Experimental/LearningAgents/Source/Learning/Public/LearningNeuralNetwork.h
--- a/Engine/Plugins/Experimental/LearningAgents/Source/Learning/Public/LearningNeuralNetwork.h
+++ b/Engine/Plugins/Experimental/LearningAgents/Source/Learning/Public/LearningNeuralNetwork.h
@@ -4,10 +4,70 @@
 
 #include "LearningArray.h"
 
+#include "Templates/SharedPointer.h"
+
 namespace UE::Learning
 {
+	struct INeuralNetworkInference;
+	
+	/**
+	* Settings object for a neural network instance
+	*/
+	struct LEARNING_API FNeuralNetworkInferenceSettings
+	{
+		// If to allow for multi-threaded evaluation
+		bool bParallelEvaluation = true;
+
+		// Minimum batch size to use for multi-threaded evaluation
+		uint16 MinParallelBatchSize = 16;
+	};
+
+	/**
+	* Interface for a Neural Network
+	*/
+	struct INeuralNetwork
+	{
+		virtual ~INeuralNetwork() {}
+
+		/** Deserialize the network from the raw bytes starting at the given offset. Returns true if successful. */
+		virtual bool DeserializeFromBytes(int32& InOutOffset, const TLearningArrayView<1, const uint8> RawBytes) = 0;
+
+		/** Serialize the network to the raw bytes starting at the given offset. */
+		virtual void SerializeToBytes(int32& InOutOffset, TLearningArrayView<1, uint8> OutRawBytes) const = 0;
+
+		/** Get the number of bytes required to serialize this network. */
+		virtual int32 GetSerializationByteNum() const = 0;
+
+		/** Create a new inference object for this network with the given maximum batch size and inference settings. */
+		virtual TSharedRef<INeuralNetworkInference> CreateInferenceObject(
+			const int32 MaxBatchSize,
+			const FNeuralNetworkInferenceSettings& Settings = FNeuralNetworkInferenceSettings()) = 0;
+
+		/** Get the number of inputs expected by this network */
+		virtual int32 GetInputNum() const = 0;
+
+		/** Get the number of outputs expected by this network */
+		virtual int32 GetOutputNum() const = 0;
+
+		/** Get the name of the python class used to represent this type of network during training. */
+		virtual const TCHAR* GetPythonClassName() const = 0;
+	};
+
 	/**
-	* Activation Function for use in a Neural Network
+	* Interface for a Neural Network Inference
+	*/
+	struct INeuralNetworkInference
+	{
+		virtual ~INeuralNetworkInference() {}
+
+		virtual void Evaluate(
+			TLearningArrayView<2, float> Output,
+			const TLearningArrayView<2, const float> Input,
+			const FIndexSet Instances) = 0;
+	};
+
+	/**
+	* Activation Function for use in basic MLP Neural Network
 	*/
 	enum class EActivationFunction : uint8
 	{
@@ -21,21 +81,27 @@ namespace UE::Learning
 		TanH = 2,
 	};
 
-	LEARNING_API const TCHAR* GetActivationFunctionString(const EActivationFunction ActivationFunction);
-
 	/**
-	* Simple feed-forward network data.
+	* Basic Implementation for a MLP Neural Network
 	*/
-	struct LEARNING_API FNeuralNetwork
+	struct LEARNING_API FNeuralNetworkMLP : public INeuralNetwork
 	{
-		void DeserializeFromBytes(int32& InOutOffset, const TLearningArrayView<1, const uint8> RawBytes);
-		void SerializeToBytes(int32& InOutOffset, TLearningArrayView<1, uint8> OutRawBytes) const;
+		friend struct FNeuralNetworkInstanceMLP;
 
-		static int32 GetSerializationByteNum(
-			const int32 InputNum,
-			const int32 OutputNum,
-			const int32 HiddenNum,
-			const int32 LayerNum);
+		//~ Begin INeuralNetwork Interface
+		virtual bool DeserializeFromBytes(int32& InOutOffset, const TLearningArrayView<1, const uint8> RawBytes) override final;
+		virtual void SerializeToBytes(int32& InOutOffset, TLearningArrayView<1, uint8> OutRawBytes) const override final;
+		virtual int32 GetSerializationByteNum() const override final;
+		virtual TSharedRef<INeuralNetworkInference> CreateInferenceObject(
+			const int32 MaxBatchSize,
+			const FNeuralNetworkInferenceSettings& Settings = FNeuralNetworkInferenceSettings()) override final;
+		virtual int32 GetInputNum() const override final;
+		virtual int32 GetOutputNum() const override final;
+		virtual const TCHAR* GetPythonClassName() const override final { return TEXT("NeuralNetworkMLP"); };
+		//~ End INeuralNetwork Interface
+
+		int32 GetHiddenNum() const;
+		int32 GetLayerNum() const;
 
 		void Resize(
 			const int32 InputNum,
@@ -43,13 +109,31 @@ namespace UE::Learning
 			const int32 HiddenNum,
 			const int32 LayerNum);
 
-		int32 GetInputNum() const;
-		int32 GetOutputNum() const;
-		int32 GetLayerNum() const;
-		int32 GetHiddenNum() const;
-
 		EActivationFunction ActivationFunction = EActivationFunction::ELU;
 		TArray<TLearningArray<2, float>, TInlineAllocator<16>> Weights;
 		TArray<TLearningArray<1, float>, TInlineAllocator<16>> Biases;
 	};
+
+	/**
+	* Basic Implementation for a MLP Neural Network Inference
+	*/
+	struct LEARNING_API FNeuralNetworkInferenceMLP : public INeuralNetworkInference
+	{
+		FNeuralNetworkInferenceMLP(
+			const FNeuralNetworkMLP& InNetwork,
+			const int32 MaxBatchSize,
+			const FNeuralNetworkInferenceSettings& InSettings);
+
+		//~ Begin INeuralNetworkInference Interface
+		virtual void Evaluate(
+			TLearningArrayView<2, float> Output,
+			const TLearningArrayView<2, const float> Input,
+			const FIndexSet Instances) override final;
+		//~ End INeuralNetworkInference Interface
+
+		const FNeuralNetworkMLP& Network;
+		FNeuralNetworkInferenceSettings Settings;
+		TLearningArray<2, float> FrontBuffer;
+		TLearningArray<2, float> BackBuffer;
+	};
 }
diff --git a/Engine/Plugins/Experimental/LearningAgents/Source/Learning/Public/LearningNeuralNetworkObject.h b/Engine/Plugins/Experimental/LearningAgents/Source/Learning/Public/LearningNeuralNetworkObject.h
--- a/Engine/Plugins/Experimental/LearningAgents/Source/Learning/Public/LearningNeuralNetworkObject.h
+++ b/Engine/Plugins/Experimental/LearningAgents/Source/Learning/Public/LearningNeuralNetworkObject.h
@@ -4,13 +4,12 @@
 
 #include "LearningArrayMap.h"
 #include "LearningFunctionObject.h"
+#include "LearningNeuralNetwork.h" // Included for FNeuralNetworkInferenceSettings
 
 #include "Templates/SharedPointer.h"
 
 namespace UE::Learning
 {
-	struct FNeuralNetwork;
-
 	/**
 	* Settings object for a neural network based policy
 	*/
@@ -24,12 +23,6 @@ namespace UE::Learning
 
 		// Overall scale of the action noise
 		float ActionNoiseScale = 1.0f;
-
-		// If to allow for multi-threaded evaluation
-		bool bParallelEvaluation = true;
-
-		// Minimum batch size to use for multi-threaded evaluation
-		uint16 MinParallelBatchSize = 16;
 	};
 
 	/**
@@ -42,37 +35,31 @@ namespace UE::Learning
 			const FName& InIdentifier,
 			const TSharedRef<FArrayMap>& InInstanceData,
 			const int32 InMaxInstanceNum,
-			const TSharedRef<FNeuralNetwork>& InNeuralNetwork,
+			const TSharedPtr<INeuralNetwork>& InNeuralNetwork,
 			const uint32 InSeed,
+			const FNeuralNetworkInferenceSettings& InInferenceSettings = FNeuralNetworkInferenceSettings(),
 			const FNeuralNetworkPolicyFunctionSettings& InSettings = FNeuralNetworkPolicyFunctionSettings());
 
 		virtual void Evaluate(const FIndexSet Instances) override final;
 
-		TSharedRef<FNeuralNetwork> NeuralNetwork;
-		FNeuralNetworkPolicyFunctionSettings Settings;
+		/** Sets the NeuralNetwork and re-creates the NeuralNetworkInference object */
+		void UpdateNeuralNetwork(const TSharedPtr<INeuralNetwork>& NewNeuralNetwork);
 
-		TArray<TLearningArray<2, float, TInlineAllocator<128>>, TInlineAllocator<16>> Activations;
+		uint32 MaxInstanceNum = 0;
+		TSharedPtr<INeuralNetwork> NeuralNetwork;
+		TSharedPtr<INeuralNetworkInference> NeuralNetworkInference;
+		FNeuralNetworkInferenceSettings InferenceSettings;
+		FNeuralNetworkPolicyFunctionSettings Settings;
 
 		TArrayMapHandle<1, uint32> SeedHandle;
 		TArrayMapHandle<2, float> InputHandle;
 		TArrayMapHandle<2, float> OutputHandle;
+		TArrayMapHandle<2, float> OutputNetworkHandle;
 		TArrayMapHandle<2, float> OutputMeanHandle;
 		TArrayMapHandle<2, float> OutputStdHandle;
 		TArrayMapHandle<1, float> ActionNoiseScaleHandle;
 	};
 
-	/**
-	* Settings object for a neural network based critic
-	*/
-	struct LEARNING_API FNeuralNetworkCriticFunctionSettings
-	{
-		// If to allow for multi-threaded evaluation
-		bool bParallelEvaluation = true;
-
-		// Minimum batch size to use for multi-threaded evaluation
-		uint16 MinParallelBatchSize = 16;
-	};
-
 	/**
 	* Neural-network based critic object. Stores various settings and intermediate
 	* storage required to evaluate the given network for the provided number of instances.
@@ -83,15 +70,18 @@ namespace UE::Learning
 			const FName& InIdentifier,
 			const TSharedRef<FArrayMap>& InInstanceData,
 			const int32 InMaxInstanceNum,
-			const TSharedRef<FNeuralNetwork>& InNeuralNetwork,
-			const FNeuralNetworkCriticFunctionSettings& InSettings = FNeuralNetworkCriticFunctionSettings());
+			const TSharedPtr<INeuralNetwork>& InNeuralNetwork,
+			const FNeuralNetworkInferenceSettings& InInferenceSettings = FNeuralNetworkInferenceSettings());
 
 		virtual void Evaluate(const FIndexSet Instances) override final;
 
-		TSharedRef<FNeuralNetwork> NeuralNetwork;
-		FNeuralNetworkCriticFunctionSettings Settings;
+		/** Sets the NeuralNetwork and re-creates the NeuralNetworkInference object */
+		void UpdateNeuralNetwork(const TSharedPtr<INeuralNetwork>& NewNeuralNetwork);
 
-		TArray<TLearningArray<2, float, TInlineAllocator<128>>, TInlineAllocator<16>> Activations;
+		uint32 MaxInstanceNum = 0;
+		TSharedPtr<INeuralNetwork> NeuralNetwork;
+		TSharedPtr<INeuralNetworkInference> NeuralNetworkInference;
+		FNeuralNetworkInferenceSettings InferenceSettings;
 
 		TArrayMapHandle<2, float> InputHandle;
 		TArrayMapHandle<1, float> OutputHandle;
diff --git a/Engine/Plugins/Experimental/LearningAgents/Source/LearningAgents/LearningAgents.Build.cs b/Engine/Plugins/Experimental/LearningAgents/Source/LearningAgents/LearningAgents.Build.cs
--- a/Engine/Plugins/Experimental/LearningAgents/Source/LearningAgents/LearningAgents.Build.cs
+++ b/Engine/Plugins/Experimental/LearningAgents/Source/LearningAgents/LearningAgents.Build.cs
@@ -31,6 +31,8 @@ public class LearningAgents : ModuleRules
 			{
 				"CoreUObject",
 				"Engine",
+				"NNE",
+				"NNERuntimeBasicCpu",
 				// ... add private dependencies that you statically link with here ...	
 			});
 		
diff --git a/Engine/Plugins/Experimental/LearningAgents/Source/LearningAgents/Private/LearningAgentsCritic.cpp b/Engine/Plugins/Experimental/LearningAgents/Source/LearningAgents/Private/LearningAgentsCritic.cpp
--- a/Engine/Plugins/Experimental/LearningAgents/Source/LearningAgents/Private/LearningAgentsCritic.cpp
+++ b/Engine/Plugins/Experimental/LearningAgents/Source/LearningAgents/Private/LearningAgentsCritic.cpp
@@ -5,6 +5,7 @@
 #include "LearningAgentsManager.h"
 #include "LearningAgentsInteractor.h"
 #include "LearningAgentsHelpers.h"
+#include "LearningAgentsNeuralNetworkData.h"
 #include "LearningFeatureObject.h"
 #include "LearningNeuralNetwork.h"
 #include "LearningNeuralNetworkObject.h"
@@ -54,34 +55,27 @@ void ULearningAgentsCritic::SetupCritic(
 	{
 		// Use Existing Neural Network Asset
 
-		if (NeuralNetworkAsset->NeuralNetwork)
+		if (NeuralNetworkAsset->NeuralNetworkData)
 		{
-			if (NeuralNetworkAsset->NeuralNetwork->GetInputNum() != Interactor->GetObservationFeature().DimNum() ||
-				NeuralNetworkAsset->NeuralNetwork->GetOutputNum() != 2 * Interactor->GetActionFeature().DimNum())
+			if (NeuralNetworkAsset->NeuralNetworkData->GetNetworkInterface()->GetInputNum() != Interactor->GetObservationFeature().DimNum() ||
+				NeuralNetworkAsset->NeuralNetworkData->GetNetworkInterface()->GetOutputNum() != 2 * Interactor->GetActionFeature().DimNum())
 			{
 				UE_LOG(LogLearning, Error, TEXT("%s: Neural Network Asset provided during Setup is incorrect size: Inputs and outputs don't match."), *GetName());
 				return;
 			}
 
-			if (NeuralNetworkAsset->NeuralNetwork->GetHiddenNum() != CriticSettings.HiddenLayerSize ||
-				NeuralNetworkAsset->NeuralNetwork->GetLayerNum() != CriticSettings.LayerNum ||
-				NeuralNetworkAsset->NeuralNetwork->ActivationFunction != UE::Learning::Agents::GetActivationFunction(CriticSettings.ActivationFunction))
-			{
-				UE_LOG(LogLearning, Warning, TEXT("%s: Neural Network Asset settings don't match those given by CriticSettings"), *GetName());
-			}
-
 			Network = NeuralNetworkAsset;
 		}
 		else
 		{
 			Network = NeuralNetworkAsset;
-			Network->NeuralNetwork = MakeShared<UE::Learning::FNeuralNetwork>();
-			Network->NeuralNetwork->Resize(
+			Network->NeuralNetworkData = NewObject<ULearningAgentsDefaultNeuralNetworkData>(Network);
+			Network->NeuralNetworkData->CreateMLP(
 				Interactor->GetObservationFeature().DimNum(),
 				1,
 				CriticSettings.HiddenLayerSize,
-				CriticSettings.LayerNum);
-			Network->NeuralNetwork->ActivationFunction = UE::Learning::Agents::GetActivationFunction(CriticSettings.ActivationFunction);
+				CriticSettings.LayerNum,
+				CriticSettings.ActivationFunction);
 		}
 	}
 	else
@@ -91,13 +85,13 @@ void ULearningAgentsCritic::SetupCritic(
 		const FName UniqueName = MakeUniqueObjectName(this, ULearningAgentsNeuralNetwork::StaticClass(), TEXT("CriticNetwork"), EUniqueObjectNameOptions::GloballyUnique);
 
 		Network = NewObject<ULearningAgentsNeuralNetwork>(this, UniqueName);
-		Network->NeuralNetwork = MakeShared<UE::Learning::FNeuralNetwork>();
-		Network->NeuralNetwork->Resize(
+		Network->NeuralNetworkData = NewObject<ULearningAgentsDefaultNeuralNetworkData>(Network);
+		Network->NeuralNetworkData->CreateMLP(
 			Interactor->GetObservationFeature().DimNum(),
 			1,
 			CriticSettings.HiddenLayerSize,
-			CriticSettings.LayerNum);
-		Network->NeuralNetwork->ActivationFunction = UE::Learning::Agents::GetActivationFunction(CriticSettings.ActivationFunction);
+			CriticSettings.LayerNum,
+			CriticSettings.ActivationFunction);
 	}
 
 	// Create Critic Object
@@ -105,7 +99,7 @@ void ULearningAgentsCritic::SetupCritic(
 		TEXT("CriticObject"),
 		Manager->GetInstanceData().ToSharedRef(),
 		Manager->GetMaxAgentNum(),
-		Network->NeuralNetwork.ToSharedRef());
+		Network->NeuralNetworkData->GetNetworkInterface());
 
 	Manager->GetInstanceData()->Link(Interactor->GetObservationFeature().FeatureHandle, CriticObject->InputHandle);
 
@@ -167,9 +161,9 @@ ULearningAgentsNeuralNetwork* ULearningAgentsCritic::GetNetworkAsset()
 	return Network;
 }
 
-UE::Learning::FNeuralNetwork& ULearningAgentsCritic::GetCriticNetwork()
+UE::Learning::INeuralNetwork& ULearningAgentsCritic::GetCriticNetwork()
 {
-	return *Network->NeuralNetwork;
+	return *Network->NeuralNetworkData->GetNetworkInterface();
 }
 
 UE::Learning::FNeuralNetworkCriticFunction& ULearningAgentsCritic::GetCriticObject()
@@ -207,7 +201,7 @@ void ULearningAgentsCritic::UseCriticFromAsset(ULearningAgentsNeuralNetwork* Neu
 		return;
 	}
 
-	if (!NeuralNetworkAsset || !NeuralNetworkAsset->NeuralNetwork)
+	if (!NeuralNetworkAsset || !NeuralNetworkAsset->NeuralNetworkData)
 	{
 		UE_LOG(LogLearning, Error, TEXT("%s: Asset is invalid."), *GetName());
 		return;
@@ -219,17 +213,15 @@ void ULearningAgentsCritic::UseCriticFromAsset(ULearningAgentsNeuralNetwork* Neu
 		return;
 	}
 
-	if (NeuralNetworkAsset->NeuralNetwork->GetInputNum() != Network->NeuralNetwork->GetInputNum() ||
-		NeuralNetworkAsset->NeuralNetwork->GetOutputNum() != Network->NeuralNetwork->GetOutputNum() ||
-		NeuralNetworkAsset->NeuralNetwork->GetLayerNum() != Network->NeuralNetwork->GetLayerNum() ||
-		NeuralNetworkAsset->NeuralNetwork->ActivationFunction != Network->NeuralNetwork->ActivationFunction)
+	if (NeuralNetworkAsset->NeuralNetworkData->GetNetworkInterface()->GetInputNum() != Network->NeuralNetworkData->GetNetworkInterface()->GetInputNum() ||
+		NeuralNetworkAsset->NeuralNetworkData->GetNetworkInterface()->GetOutputNum() != Network->NeuralNetworkData->GetNetworkInterface()->GetOutputNum())
 	{
 		UE_LOG(LogLearning, Error, TEXT("%s: Failed to use asset as network settings don't match."), *GetName());
 		return;
 	}
 
 	Network = NeuralNetworkAsset;
-	CriticObject->NeuralNetwork = Network->NeuralNetwork.ToSharedRef();
+	CriticObject->UpdateNeuralNetwork(Network->NeuralNetworkData->GetNetworkInterface());
 }
 
 void ULearningAgentsCritic::LoadCriticFromAsset(ULearningAgentsNeuralNetwork* NeuralNetworkAsset)
diff --git a/Engine/Plugins/Experimental/LearningAgents/Source/LearningAgents/Private/LearningAgentsNeuralNetwork.cpp b/Engine/Plugins/Experimental/LearningAgents/Source/LearningAgents/Private/LearningAgentsNeuralNetwork.cpp
--- a/Engine/Plugins/Experimental/LearningAgents/Source/LearningAgents/Private/LearningAgentsNeuralNetwork.cpp
+++ b/Engine/Plugins/Experimental/LearningAgents/Source/LearningAgents/Private/LearningAgentsNeuralNetwork.cpp
@@ -2,8 +2,11 @@
 
 #include "LearningAgentsNeuralNetwork.h"
 
-#include "LearningNeuralNetwork.h"
+#include "LearningAgentsNeuralNetworkData.h"
+
 #include "LearningLog.h"
+#include "LearningArray.h"
+#include "LearningNeuralNetwork.h"
 
 #include "UObject/Package.h"
 #include "Misc/FileHelper.h"
@@ -12,158 +15,52 @@ ULearningAgentsNeuralNetwork::ULearningAgentsNeuralNetwork() = default;
 ULearningAgentsNeuralNetwork::ULearningAgentsNeuralNetwork(FVTableHelper& Helper) : Super(Helper) {}
 ULearningAgentsNeuralNetwork::~ULearningAgentsNeuralNetwork() = default;
 
-namespace UE::Learning::Agents
-{
-	ELearningAgentsActivationFunction GetLearningAgentsActivationFunction(const EActivationFunction ActivationFunction)
-	{
-		switch (ActivationFunction)
-		{
-		case EActivationFunction::ReLU: return ELearningAgentsActivationFunction::ReLU;
-		case EActivationFunction::TanH: return ELearningAgentsActivationFunction::TanH;
-		case EActivationFunction::ELU: return ELearningAgentsActivationFunction::ELU;
-		default:UE_LOG(LogLearning, Error, TEXT("Unknown Activation Function.")); return ELearningAgentsActivationFunction::ELU;
-		}
-	}
-
-	EActivationFunction GetActivationFunction(const ELearningAgentsActivationFunction ActivationFunction)
-	{
-		switch (ActivationFunction)
-		{
-		case ELearningAgentsActivationFunction::ReLU: return EActivationFunction::ReLU;
-		case ELearningAgentsActivationFunction::TanH: return EActivationFunction::TanH;
-		case ELearningAgentsActivationFunction::ELU: return EActivationFunction::ELU;
-		default:UE_LOG(LogLearning, Error, TEXT("Unknown Activation Function.")); return EActivationFunction::ELU;
-		}
-	}
-}
-
-namespace UE::Learning::Agents::NeuralNetwork::Private
-{
-	static constexpr int32 MagicNumber = 0x0cd353cf;
-	static constexpr int32 VersionNumber = 1;
-}
-
-void ULearningAgentsNeuralNetwork::Serialize(FArchive& Ar)
-{
-	Super::Serialize(Ar);
-
-	if (Ar.IsLoading())
-	{
-		bool bValid;
-		Ar << bValid;
-		if (bValid)
-		{
-			int32 VersionNumber;
-			Ar << VersionNumber;
-
-			if (VersionNumber != UE::Learning::Agents::NeuralNetwork::Private::VersionNumber)
-			{
-				UE_LOG(LogLearning, Warning, TEXT("%s: Unsupported Version Number %i."), *GetName(), VersionNumber);
-			}
-
-			NeuralNetwork = MakeShared<UE::Learning::FNeuralNetwork>();
-			TArray<uint8> Bytes;
-			Ar << Bytes;
-			int32 Offset = 0;
-			NeuralNetwork->DeserializeFromBytes(Offset, Bytes);
-			UE_LEARNING_CHECK(Offset == Bytes.Num());
-		}
-		else
-		{
-			NeuralNetwork.Reset();
-		}
-	}
-	else if (Ar.IsSaving())
-	{
-		bool bValid = NeuralNetwork.IsValid();
-		Ar << bValid;
-		if (bValid)
-		{
-			int32 VersionNumber = UE::Learning::Agents::NeuralNetwork::Private::VersionNumber;
-			Ar << VersionNumber;
-			TArray<uint8> Bytes;
-			Bytes.SetNumUninitialized(UE::Learning::FNeuralNetwork::GetSerializationByteNum(
-				NeuralNetwork->GetInputNum(),
-				NeuralNetwork->GetOutputNum(), 
-				NeuralNetwork->GetHiddenNum(),
-				NeuralNetwork->GetLayerNum()));
-			int32 Offset = 0;
-			NeuralNetwork->SerializeToBytes(Offset, Bytes);
-			UE_LEARNING_CHECK(Offset == Bytes.Num());
-			Ar << Bytes;
-		}
-	}
-}
-
 void ULearningAgentsNeuralNetwork::ResetNetwork()
 {
-	NeuralNetwork.Reset();
+	NeuralNetworkData->ConditionalBeginDestroy();
+	NeuralNetworkData = nullptr;
+	
 	ForceMarkDirty();
 }
 
 void ULearningAgentsNeuralNetwork::LoadNetworkFromSnapshot(const FFilePath& File)
 {
-	TArray<uint8> RecordingData;
+	TArray<uint8> NetworkData;
 
-	if (FFileHelper::LoadFileToArray(RecordingData, *File.FilePath))
+	if (FFileHelper::LoadFileToArray(NetworkData, *File.FilePath))
 	{
-		if (RecordingData.Num() < sizeof(int32) * 2)
-		{
-			UE_LOG(LogLearning, Error, TEXT("%s: Failed to load network. Incorrect Format."), *GetName());
-			return;
-		}
-
+		ULearningAgentsNeuralNetworkData* TempNeuralNetworkData = NewObject<ULearningAgentsDefaultNeuralNetworkData>(this);
 		int32 Offset = 0;
-
-		int32 MagicNumber;
-		UE::Learning::DeserializeFromBytes(Offset, RecordingData, MagicNumber);
-
-		if (MagicNumber != UE::Learning::Agents::NeuralNetwork::Private::MagicNumber)
+		bool bSuccess = TempNeuralNetworkData->GetNetworkInterface()->DeserializeFromBytes(Offset, NetworkData);
+		
+		if (!bSuccess)
 		{
-			UE_LOG(LogLearning, Error, TEXT("%s: Failed to load network. Incorrect Magic Number."), *GetName());
+			UE_LOG(LogLearning, Error, TEXT("%s: Failed to load network. Invalid Format: \"%s\""), *GetName(), *File.FilePath);
 			return;
 		}
+		
+		NetworkData.Empty();
 
-		int32 VersionNumber;
-		UE::Learning::DeserializeFromBytes(Offset, RecordingData, VersionNumber);
-
-		if (VersionNumber != UE::Learning::Agents::NeuralNetwork::Private::VersionNumber)
-		{
-			UE_LOG(LogLearning, Error, TEXT("%s: Failed to load network. Unsupported Version Number %i."), *GetName(), VersionNumber);
-			return;
-		}
-
-		TSharedPtr<UE::Learning::FNeuralNetwork> TempNeuralNetwork = MakeShared<UE::Learning::FNeuralNetwork>();
-		TempNeuralNetwork->DeserializeFromBytes(Offset, RecordingData);
-
-		if (Offset != RecordingData.Num())
-		{
-			UE_LOG(LogLearning, Error, TEXT("Failed to load network. Unexpected end of file."));
-			return;
-		}
-
-		if (NeuralNetwork)
+		if (NeuralNetworkData)
 		{
 			// If we already have a neural network check settings match
 
-			if (TempNeuralNetwork->GetInputNum() != NeuralNetwork->GetInputNum() ||
-				TempNeuralNetwork->GetOutputNum() != NeuralNetwork->GetOutputNum() ||
-				TempNeuralNetwork->GetLayerNum() != NeuralNetwork->GetLayerNum() ||
-				TempNeuralNetwork->ActivationFunction != NeuralNetwork->ActivationFunction)
+			if (TempNeuralNetworkData->GetNetworkInterface()->GetInputNum() != NeuralNetworkData->GetNetworkInterface()->GetInputNum() ||
+				TempNeuralNetworkData->GetNetworkInterface()->GetOutputNum() != NeuralNetworkData->GetNetworkInterface()->GetOutputNum())
 			{
 				UE_LOG(LogLearning, Error, TEXT("%s: Failed to load network from snapshot as settings don't match."), *GetName());
 				return;
 			}
 			else
 			{
-				*NeuralNetwork = *TempNeuralNetwork;
+				NeuralNetworkData->CopyFrom(TempNeuralNetworkData);
 			}
 		}
 		else
 		{
 			// Otherwise use loaded neural network as-is
 
-			NeuralNetwork = TempNeuralNetwork;
+			NeuralNetworkData = TempNeuralNetworkData;
 		}
 
 		ForceMarkDirty();
@@ -176,26 +73,20 @@ void ULearningAgentsNeuralNetwork::LoadNetworkFromSnapshot(const FFilePath& File
 
 void ULearningAgentsNeuralNetwork::SaveNetworkToSnapshot(const FFilePath& File)
 {
-	TArray<uint8> RecordingData;
-
-	int32 TotalByteNum =
-		sizeof(int32) + // Magic Num
-		sizeof(int32) + // Version Num
-		NeuralNetwork->GetSerializationByteNum(
-			NeuralNetwork->GetInputNum(), 
-			NeuralNetwork->GetOutputNum(), 
-			NeuralNetwork->GetHiddenNum(), 
-			NeuralNetwork->GetLayerNum());
+	if (!NeuralNetworkData)
+	{
+		UE_LOG(LogLearning, Error, TEXT("%s: No network data to save"), *GetName());
+		return;
+	}
 
-	RecordingData.SetNumUninitialized(TotalByteNum);
+	TArray<uint8> NetworkData;
+	NetworkData.SetNumUninitialized(NeuralNetworkData->GetNetworkInterface()->GetSerializationByteNum());
 
 	int32 Offset = 0;
-	UE::Learning::SerializeToBytes(Offset, RecordingData, UE::Learning::Agents::NeuralNetwork::Private::MagicNumber);
-	UE::Learning::SerializeToBytes(Offset, RecordingData, UE::Learning::Agents::NeuralNetwork::Private::VersionNumber);
-	NeuralNetwork->SerializeToBytes(Offset, RecordingData);
-	UE_LEARNING_CHECK(Offset == RecordingData.Num());
+	NeuralNetworkData->GetNetworkInterface()->SerializeToBytes(Offset, NetworkData);
+	UE_LEARNING_CHECK(Offset == NetworkData.Num());
 
-	if (!FFileHelper::SaveArrayToFile(RecordingData, *File.FilePath))
+	if (!FFileHelper::SaveArrayToFile(NetworkData, *File.FilePath))
 	{
 		UE_LOG(LogLearning, Error, TEXT("%s: Failed to save network to file: \"%s\""), *GetName(), *File.FilePath);
 	}
@@ -203,7 +94,7 @@ void ULearningAgentsNeuralNetwork::SaveNetworkToSnapshot(const FFilePath& File)
 
 void ULearningAgentsNeuralNetwork::LoadNetworkFromAsset(ULearningAgentsNeuralNetwork* NeuralNetworkAsset)
 {
-	if (!NeuralNetworkAsset || !NeuralNetworkAsset->NeuralNetwork)
+	if (!NeuralNetworkAsset || !NeuralNetworkAsset->NeuralNetworkData)
 	{
 		UE_LOG(LogLearning, Error, TEXT("%s: Asset is invalid."), *GetName());
 		return;
@@ -215,12 +106,10 @@ void ULearningAgentsNeuralNetwork::LoadNetworkFromAsset(ULearningAgentsNeuralNet
 		return;
 	}
 
-	if (NeuralNetwork)
+	if (NeuralNetworkData)
 	{
-		if (NeuralNetworkAsset->NeuralNetwork->GetInputNum() != NeuralNetwork->GetInputNum() ||
-			NeuralNetworkAsset->NeuralNetwork->GetOutputNum() != NeuralNetwork->GetOutputNum() ||
-			NeuralNetworkAsset->NeuralNetwork->GetLayerNum() != NeuralNetwork->GetLayerNum() ||
-			NeuralNetworkAsset->NeuralNetwork->ActivationFunction != NeuralNetwork->ActivationFunction)
+		if (NeuralNetworkAsset->NeuralNetworkData->GetNetworkInterface()->GetInputNum() != NeuralNetworkData->GetNetworkInterface()->GetInputNum() ||
+			NeuralNetworkAsset->NeuralNetworkData->GetNetworkInterface()->GetOutputNum() != NeuralNetworkData->GetNetworkInterface()->GetOutputNum())
 		{
 			UE_LOG(LogLearning, Error, TEXT("%s: Failed to load network from asset as settings don't match."), *GetName());
 			return;
@@ -228,10 +117,10 @@ void ULearningAgentsNeuralNetwork::LoadNetworkFromAsset(ULearningAgentsNeuralNet
 	}
 	else
 	{
-		NeuralNetwork = MakeShared<UE::Learning::FNeuralNetwork>();
+		NeuralNetworkData = NewObject<ULearningAgentsDefaultNeuralNetworkData>(this);
 	}
 
-	*NeuralNetwork = *NeuralNetworkAsset->NeuralNetwork;
+	NeuralNetworkData->CopyFrom(NeuralNetworkAsset->NeuralNetworkData);
 	ForceMarkDirty();
 }
 
@@ -249,12 +138,12 @@ void ULearningAgentsNeuralNetwork::SaveNetworkToAsset(ULearningAgentsNeuralNetwo
 		return;
 	}
 
-	if (!NeuralNetworkAsset->NeuralNetwork)
+	if (!NeuralNetworkAsset->NeuralNetworkData)
 	{
-		NeuralNetworkAsset->NeuralNetwork = MakeShared<UE::Learning::FNeuralNetwork>();
+		NeuralNetworkAsset->NeuralNetworkData = NewObject<ULearningAgentsDefaultNeuralNetworkData>(NeuralNetworkAsset);
 	}
 
-	*NeuralNetworkAsset->NeuralNetwork = *NeuralNetwork;
+	NeuralNetworkAsset->NeuralNetworkData->CopyFrom(NeuralNetworkData);
 	NeuralNetworkAsset->ForceMarkDirty();
 }
 
diff --git a/Engine/Plugins/Experimental/LearningAgents/Source/LearningAgents/Private/LearningAgentsNeuralNetworkData.cpp b/Engine/Plugins/Experimental/LearningAgents/Source/LearningAgents/Private/LearningAgentsNeuralNetworkData.cpp
new file mode 100644
--- /dev/null
+++ b/Engine/Plugins/Experimental/LearningAgents/Source/LearningAgents/Private/LearningAgentsNeuralNetworkData.cpp
@@ -0,0 +1,418 @@
+// Copyright Epic Games, Inc. All Rights Reserved.
+
+#include "LearningAgentsNeuralNetworkData.h"
+
+#include "LearningAgentsNeuralNetwork.h"
+#include "LearningNeuralNetwork.h"
+#include "Async/ParallelFor.h"
+
+#include "NNE.h"
+#include "NNEModelData.h"
+#include "NNERuntimeBasicCpu.h"
+
+void ULearningAgentsMLPNeuralNetworkData::Serialize(FArchive& Ar) 
+{
+	Super::Serialize(Ar);
+
+	if (Ar.IsLoading())
+	{
+		bool bValid;
+		Ar << bValid;
+
+		if (bValid)
+		{
+			if (!Network)
+			{
+				Network = MakeShared<UE::Learning::FNeuralNetworkMLP>();
+			}
+
+			TArray<uint8> Bytes;
+			Ar << Bytes;
+			int32 Offset = 0;
+			Network->DeserializeFromBytes(Offset, Bytes);
+			UE_LEARNING_CHECK(Offset == Bytes.Num());
+		}
+		else
+		{
+			Network.Reset();
+		}
+	}
+	else if (Ar.IsSaving())
+	{
+		bool bValid = Network != nullptr;
+		Ar << bValid;
+
+		if (bValid)
+		{
+			TArray<uint8> Bytes;
+			Bytes.SetNumUninitialized(Network->GetSerializationByteNum());
+			int32 Offset = 0;
+			Network->SerializeToBytes(Offset, Bytes);
+			UE_LEARNING_CHECK(Offset == Bytes.Num());
+			Ar << Bytes;
+		}
+	}
+}
+
+TSharedPtr<UE::Learning::INeuralNetwork> ULearningAgentsMLPNeuralNetworkData::GetNetworkInterface()
+{
+	return Network;
+}
+
+void ULearningAgentsMLPNeuralNetworkData::CreateMLP(
+	const uint32 InputSize,
+	const uint32 OutputSize,
+	const uint32 HiddenUnitNum,
+	const uint32 LayerNum,
+	const ELearningAgentsActivationFunction Activation)
+{
+	UE_LEARNING_TRACE_CPUPROFILER_EVENT_SCOPE(ULearningAgentsMLPNeuralNetworkData::CreateMLP);
+
+	if (!Network)
+	{
+		Network = MakeShared<UE::Learning::FNeuralNetworkMLP>();
+	}
+
+	Network->Resize(
+		InputSize,
+		OutputSize,
+		HiddenUnitNum,
+		LayerNum);
+
+	switch (Activation)
+	{
+	case ELearningAgentsActivationFunction::ELU: Network->ActivationFunction = UE::Learning::EActivationFunction::ELU; break;
+	case ELearningAgentsActivationFunction::ReLU: Network->ActivationFunction = UE::Learning::EActivationFunction::ReLU; break;
+	case ELearningAgentsActivationFunction::TanH: Network->ActivationFunction = UE::Learning::EActivationFunction::TanH; break;
+	default: UE_LEARNING_CHECKF(false, TEXT("Unimplemented"));
+	}
+}
+
+void ULearningAgentsMLPNeuralNetworkData::CopyFrom(const ULearningAgentsNeuralNetworkData* Other)
+{
+	UE_LEARNING_TRACE_CPUPROFILER_EVENT_SCOPE(ULearningAgentsMLPNeuralNetworkData::CopyFrom);
+
+	if (const ULearningAgentsMLPNeuralNetworkData* OtherData = Cast<ULearningAgentsMLPNeuralNetworkData>(Other))
+	{
+		if (OtherData->Network)
+		{
+			*Network = *OtherData->Network;
+		}
+		else
+		{
+			Network.Reset();
+		}
+	}
+	else
+	{
+		ensureMsgf(false, TEXT("Cannot copy from object of different derived type."));
+	}
+}
+
+namespace UE::Learning::Agents
+{
+	FNeuralNetworkNNEInference::FNeuralNetworkNNEInference(
+		UE::NNE::IModelCPU& InModel,
+		const FNeuralNetworkInferenceSettings& InSettings,
+		const int32 MaxBatchSize,
+		const int32 InputSize,
+		const int32 OutputSize)
+		: Settings(InSettings)
+	{
+		InputBuffer.SetNumUninitialized({ MaxBatchSize, InputSize });
+		OutputBuffer.SetNumUninitialized({ MaxBatchSize, OutputSize });
+
+		CreateInstances(InModel);
+	}
+
+	void FNeuralNetworkNNEInference::CreateInstances(NNE::IModelCPU& InModel)
+	{
+		UE_LEARNING_TRACE_CPUPROFILER_EVENT_SCOPE(Learning::Agents::FNeuralNetworkNNEInference::CreateInstances);
+
+		const int32 MaxBatchSize = InputBuffer.Num<0>();
+		const int32 InputSize = InputBuffer.Num<1>();
+		const int32 OutputSize = OutputBuffer.Num<1>();
+
+		if (Settings.bParallelEvaluation)
+		{
+			const int32 IdealInferenceInstanceNum = ParallelForImpl::GetNumberOfThreadTasks(MaxBatchSize, Settings.MinParallelBatchSize, EParallelForFlags::None);
+			const int32 InferenceInstanceSliceLength = FMath::DivideAndRoundUp((int32)MaxBatchSize, IdealInferenceInstanceNum);
+			const int32 InferenceInstanceNum = FMath::DivideAndRoundUp((int32)MaxBatchSize, InferenceInstanceSliceLength);
+
+			InferenceInstances.Empty(InferenceInstanceNum);
+
+			for (int32 InferenceInstanceIdx = 0; InferenceInstanceIdx < InferenceInstanceNum; InferenceInstanceIdx++)
+			{
+				InferenceInstances.Emplace(InModel.CreateModelInstanceCPU());
+				InferenceInstances.Last()->SetInputTensorShapes({ NNE::FTensorShape::Make({(uint32)InferenceInstanceSliceLength, (uint32)InputSize}) });
+			}
+		}
+		else
+		{
+			InferenceInstances.Empty(1);
+			InferenceInstances.Emplace(InModel.CreateModelInstanceCPU());
+			InferenceInstances.Last()->SetInputTensorShapes({ NNE::FTensorShape::Make({(uint32)MaxBatchSize, (uint32)InputSize}) });
+		}
+	}
+
+	void FNeuralNetworkNNEInference::Evaluate(
+		TLearningArrayView<2, float> Output,
+		const TLearningArrayView<2, const float> Input,
+		const FIndexSet Instances)
+	{
+		UE_LEARNING_TRACE_CPUPROFILER_EVENT_SCOPE(Learning::Agents::FNeuralNetworkNNEInference::Evaluate);
+
+		const int32 InstanceNum = Instances.Num();
+		const int32 InputNum = Input.Num<1>();
+		const int32 OutputNum = Output.Num<1>();
+
+		// Gather Inputs
+
+		Array::Check(Input, Instances);
+
+		for (int32 InstanceIdx = 0; InstanceIdx < InstanceNum; InstanceIdx++)
+		{
+			Array::Copy(InputBuffer[InstanceIdx], Input[Instances[InstanceIdx]]);
+		}
+
+		// Run Inference
+
+		Array::Check(InputBuffer.Slice(0, InstanceNum));
+
+		if (Settings.bParallelEvaluation && InstanceNum > Settings.MinParallelBatchSize)
+		{
+			const int32 IdealInferenceInstanceNum = ParallelForImpl::GetNumberOfThreadTasks(InstanceNum, Settings.MinParallelBatchSize, EParallelForFlags::None);
+			const int32 InferenceInstanceSliceLength = FMath::DivideAndRoundUp(InstanceNum, IdealInferenceInstanceNum);
+			const int32 InferenceInstanceNum = FMath::DivideAndRoundUp(InstanceNum, InferenceInstanceSliceLength);
+
+			check(InferenceInstanceNum <= InferenceInstances.Num());
+
+			ParallelFor(InferenceInstanceNum, [this, InstanceNum, InputNum, InferenceInstanceSliceLength](int32 InferenceInstanceIdx)
+				{
+					const int32 StartIndex = InferenceInstanceIdx * InferenceInstanceSliceLength;
+					const int32 StopIndex = FMath::Min((InferenceInstanceIdx + 1) * InferenceInstanceSliceLength, InstanceNum);
+
+					TLearningArrayView<2, float> InputBufferSlice = InputBuffer.Slice(StartIndex, StopIndex - StartIndex);
+					TLearningArrayView<2, float> OutputBufferSlice = OutputBuffer.Slice(StartIndex, StopIndex - StartIndex);
+
+					InferenceInstances[InferenceInstanceIdx]->SetInputTensorShapes({ NNE::FTensorShape::Make({(uint32)(StopIndex - StartIndex), (uint32)InputNum}) });
+					InferenceInstances[InferenceInstanceIdx]->RunSync(
+						{ { (void*)InputBufferSlice.GetData(), InputBufferSlice.Num() * sizeof(float) } },
+						{ { (void*)OutputBufferSlice.GetData(), OutputBufferSlice.Num() * sizeof(float) } });
+				});
+		}
+		else
+		{
+			InferenceInstances[0]->SetInputTensorShapes({ NNE::FTensorShape::Make({(uint32)InstanceNum, (uint32)InputNum}) });
+			InferenceInstances[0]->RunSync(
+				{ { (void*)InputBuffer.GetData(), InputBuffer.Num() * sizeof(float) } },
+				{ { (void*)OutputBuffer.GetData(), OutputBuffer.Num() * sizeof(float) } });
+		}
+
+		Array::Check(OutputBuffer.Slice(0, InstanceNum));
+
+		// Scatter Outputs
+
+		for (int32 InstanceIdx = 0; InstanceIdx < InstanceNum; InstanceIdx++)
+		{
+			Array::Copy(Output[Instances[InstanceIdx]], OutputBuffer[InstanceIdx]);
+		}
+
+		Array::Check(Output, Instances);
+	}
+
+	FNeuralNetworkNNE::FNeuralNetworkNNE(ULearningAgentsNNENeuralNetworkData& InParent) : Parent(InParent) {}
+
+	void FNeuralNetworkNNE::ReloadFromFileData()
+	{
+		UE_LEARNING_TRACE_CPUPROFILER_EVENT_SCOPE(Learning::Agents::FNeuralNetworkNNE::ReloadFromFileData);
+
+		if (!Parent.ModelData)
+		{
+			Parent.ModelData = NewObject<UNNEModelData>(&Parent);
+		}
+
+		Parent.ModelData->Init(TEXT("ubnne"), Parent.FileData);
+
+		TWeakInterfacePtr<INNERuntimeCPU> RuntimeCPU = NNE::GetRuntime<INNERuntimeCPU>(TEXT("NNERuntimeBasicCpu"));
+
+		if (ensureMsgf(RuntimeCPU.IsValid(), TEXT("Could not find requested NNE Runtime")))
+		{
+			Model = RuntimeCPU->CreateModelCPU(Parent.ModelData);
+		}
+
+		// If we are not in the editor we can now clear the FileData and FileType since these will be
+		// using additional memory and we are not going to save this asset and so don't require them.
+
+#if !WITH_EDITOR
+		Parent.ModelData->ClearFileDataAndFileType();
+#endif
+
+		// We need to tell all the created inference objects to update their instances since the Model has changed
+		// and so the instance objects themselves are either no longer valid or tied to the old model.
+
+		for (const TWeakPtr<FNeuralNetworkNNEInference>& InferenceObject : InferenceObjects)
+		{
+			if (TSharedPtr<FNeuralNetworkNNEInference> InferenceObjectPtr = InferenceObject.Pin())
+			{
+				InferenceObjectPtr->CreateInstances(*Model);
+			}
+		}
+	}
+
+	bool FNeuralNetworkNNE::DeserializeFromBytes(int32& InOutOffset, const TLearningArrayView<1, const uint8> RawBytes)
+	{
+		UE_LEARNING_TRACE_CPUPROFILER_EVENT_SCOPE(Learning::Agents::FNeuralNetworkNNE::DeserializeFromBytes);
+
+		if (RawBytes.Num() - InOutOffset < Parent.FileData.Num())
+		{
+			return false;
+		}
+
+		Array::Copy(
+			TLearningArrayView<1, uint8>(Parent.FileData),
+			RawBytes.Slice(InOutOffset, Parent.FileData.Num()));
+
+		InOutOffset += Parent.FileData.Num();
+
+		ReloadFromFileData();
+
+		return true;
+	}
+
+	void FNeuralNetworkNNE::SerializeToBytes(int32& InOutOffset, TLearningArrayView<1, uint8> OutRawBytes) const
+	{
+		UE_LEARNING_TRACE_CPUPROFILER_EVENT_SCOPE(Learning::Agents::FNeuralNetworkNNE::SerializeToBytes);
+
+		Array::Copy(
+			OutRawBytes.Slice(InOutOffset, Parent.FileData.Num()),
+			TLearningArrayView<1, const uint8>(Parent.FileData));
+
+		InOutOffset += Parent.FileData.Num();
+	}
+
+	int32 FNeuralNetworkNNE::GetSerializationByteNum() const
+	{
+		return Parent.FileData.Num();
+	}
+
+	TSharedRef<INeuralNetworkInference> FNeuralNetworkNNE::CreateInferenceObject(
+		const int32 MaxBatchSize,
+		const FNeuralNetworkInferenceSettings& InSettings)
+	{
+		TSharedRef<FNeuralNetworkNNEInference> InferenceObject = MakeShared<FNeuralNetworkNNEInference>(
+			*Model,
+			InSettings,
+			MaxBatchSize,
+			Parent.InputNum,
+			Parent.OutputNum);
+
+		InferenceObjects.Emplace(InferenceObject.ToWeakPtr());
+
+		return InferenceObject;
+	}
+
+	int32 FNeuralNetworkNNE::GetInputNum() const { return Parent.InputNum; }
+	int32 FNeuralNetworkNNE::GetOutputNum() const { return Parent.OutputNum; }
+}
+
+
+void ULearningAgentsNNENeuralNetworkData::PostLoad()
+{
+	if (FileData.Num() > 0)
+	{
+		Network = MakeShared<UE::Learning::Agents::FNeuralNetworkNNE>(*this);
+	}
+}
+
+TSharedPtr<UE::Learning::INeuralNetwork> ULearningAgentsNNENeuralNetworkData::GetNetworkInterface()
+{
+	return Network;
+}
+
+void ULearningAgentsNNENeuralNetworkData::CreateMLP(
+	const uint32 InputSize,
+	const uint32 OutputSize,
+	const uint32 HiddenUnitNum,
+	const uint32 LayerNum,
+	const ELearningAgentsActivationFunction Activation)
+{
+	UE_LEARNING_TRACE_CPUPROFILER_EVENT_SCOPE(ULearningAgentsNNENeuralNetworkData::CreateMLP);
+
+	UE_LEARNING_CHECK(InputSize > 0);
+	UE_LEARNING_CHECK(OutputSize > 0);
+	UE_LEARNING_CHECK(HiddenUnitNum > 0);
+	UE_LEARNING_CHECK(LayerNum > 0);
+
+	TArray<TArray<float>> LinearWeightData;
+	TArray<TArray<float>> LinearBiasData;
+
+	UE::NNE::RuntimeBasic::FSequentialModelBuilder Builder;
+
+	for (uint32 LayerIdx = 0; LayerIdx < LayerNum; LayerIdx++)
+	{
+		const int32 LayerInputNum = LayerIdx == 0 ? InputSize : HiddenUnitNum;
+		const int32 LayerOutputNum = LayerIdx == LayerNum - 1 ? OutputSize : HiddenUnitNum;
+
+		LinearWeightData.AddDefaulted();
+		LinearWeightData.Last().Init(0.0f, LayerInputNum * LayerOutputNum);
+
+		LinearBiasData.AddDefaulted();
+		LinearBiasData.Last().Init(0.0f, LayerOutputNum);
+
+		Builder.AddLinear(LayerInputNum, LayerOutputNum, LinearWeightData.Last(), LinearBiasData.Last());
+
+		if (LayerIdx != LayerNum - 1)
+		{
+			switch (Activation)
+			{
+			case ELearningAgentsActivationFunction::ELU: Builder.AddELU(); break;
+			case ELearningAgentsActivationFunction::ReLU: Builder.AddReLU(); break;
+			case ELearningAgentsActivationFunction::TanH: Builder.AddTanH(); break;
+			default: UE_LEARNING_CHECKF(false, TEXT("Unimplemented")); break;
+			}
+		}
+	}
+
+	FileData.SetNumUninitialized(Builder.GetWriteByteNum());
+	Builder.WriteAndReset(FileData);
+
+	LinearWeightData.Empty();
+	LinearBiasData.Empty();
+
+	InputNum = InputSize;
+	OutputNum = OutputSize;
+
+	if (!Network)
+	{
+		Network = MakeShared<UE::Learning::Agents::FNeuralNetworkNNE>(*this);
+	}
+
+	Network->ReloadFromFileData();
+}
+
+void ULearningAgentsNNENeuralNetworkData::CopyFrom(const ULearningAgentsNeuralNetworkData* Other)
+{
+	UE_LEARNING_TRACE_CPUPROFILER_EVENT_SCOPE(ULearningAgentsNNENeuralNetworkData::CopyFrom);
+
+	if (const ULearningAgentsNNENeuralNetworkData* OtherData = Cast<ULearningAgentsNNENeuralNetworkData>(Other))
+	{
+		FileData = OtherData->FileData;
+		InputNum = OtherData->InputNum;
+		OutputNum = OtherData->OutputNum;
+
+		if (!Network)
+		{
+			Network = MakeShared<UE::Learning::Agents::FNeuralNetworkNNE>(*this);
+		}
+
+		Network->ReloadFromFileData();
+	}
+	else
+	{
+		ensureMsgf(false, TEXT("Cannot copy from object of different derived type."));
+	}
+}
+
+
+
diff --git a/Engine/Plugins/Experimental/LearningAgents/Source/LearningAgents/Private/LearningAgentsNeuralNetworkData.h b/Engine/Plugins/Experimental/LearningAgents/Source/LearningAgents/Private/LearningAgentsNeuralNetworkData.h
new file mode 100644
--- /dev/null
+++ b/Engine/Plugins/Experimental/LearningAgents/Source/LearningAgents/Private/LearningAgentsNeuralNetworkData.h
@@ -0,0 +1,155 @@
+// Copyright Epic Games, Inc. All Rights Reserved.
+
+#pragma once
+
+#include "LearningNeuralNetwork.h"
+
+#include "LearningAgentsNeuralNetworkData.generated.h"
+
+enum class ELearningAgentsActivationFunction : uint8;
+
+class UNNEModelData;
+
+namespace UE::NNE
+{
+	class IModelCPU;
+	class IModelInstanceCPU;
+}
+
+UCLASS()
+class LEARNINGAGENTS_API ULearningAgentsNeuralNetworkData : public UObject 
+{
+	GENERATED_BODY()
+
+public:
+
+	virtual TSharedPtr<UE::Learning::INeuralNetwork> GetNetworkInterface() { return nullptr; }
+
+	virtual void CopyFrom(const ULearningAgentsNeuralNetworkData* Other) {}
+
+	virtual void CreateMLP(
+		const uint32 InputSize,
+		const uint32 OutputSize,
+		const uint32 HiddenUnitNum,
+		const uint32 LayerNum,
+		const ELearningAgentsActivationFunction Activation) {}
+};
+
+UCLASS()
+class LEARNINGAGENTS_API ULearningAgentsMLPNeuralNetworkData : public ULearningAgentsNeuralNetworkData 
+{
+	GENERATED_BODY()
+
+public:
+
+	virtual void Serialize(FArchive& Ar) override final;
+
+	virtual TSharedPtr<UE::Learning::INeuralNetwork> GetNetworkInterface() override final;
+
+	virtual void CopyFrom(const ULearningAgentsNeuralNetworkData* Other) override final;
+
+	virtual void CreateMLP(
+		const uint32 InputSize,
+		const uint32 OutputSize,
+		const uint32 HiddenUnitNum,
+		const uint32 LayerNum,
+		const ELearningAgentsActivationFunction Activation) override final;
+
+	TSharedPtr<UE::Learning::FNeuralNetworkMLP> Network;
+};
+
+class ULearningAgentsNNENeuralNetworkData;
+
+namespace UE::Learning::Agents
+{
+	struct FNeuralNetworkNNE;
+	struct FNeuralNetworkNNEInference;
+
+	struct FNeuralNetworkNNE : public INeuralNetwork
+	{
+		FNeuralNetworkNNE(ULearningAgentsNNENeuralNetworkData& InParent);
+
+		//~ Begin INeuralNetwork Interface
+		virtual bool DeserializeFromBytes(int32& InOutOffset, const TLearningArrayView<1, const uint8> RawBytes) override final;
+		virtual void SerializeToBytes(int32& InOutOffset, TLearningArrayView<1, uint8> OutRawBytes) const override final;
+		virtual int32 GetSerializationByteNum() const override final;
+		virtual TSharedRef<INeuralNetworkInference> CreateInferenceObject(
+			const int32 MaxBatchSize,
+			const FNeuralNetworkInferenceSettings& Settings = FNeuralNetworkInferenceSettings()) override final;
+		virtual int32 GetInputNum() const override final;
+		virtual int32 GetOutputNum() const override final;
+		virtual const TCHAR* GetPythonClassName() const override final { return TEXT("NeuralNetworkNNE"); };
+		//~ End INeuralNetwork Interface
+
+		void ReloadFromFileData();
+
+		ULearningAgentsNNENeuralNetworkData& Parent;
+		TSharedPtr<NNE::IModelCPU> Model;
+		TArray<TWeakPtr<FNeuralNetworkNNEInference>, TInlineAllocator<64>> InferenceObjects;
+	};
+
+	struct FNeuralNetworkNNEInference : public INeuralNetworkInference
+	{
+		FNeuralNetworkNNEInference(
+			NNE::IModelCPU& InModel,
+			const FNeuralNetworkInferenceSettings& InSettings,
+			const int32 MaxBatchSize,
+			const int32 InputSize,
+			const int32 OutputSize);
+
+		//~ Begin INeuralNetworkInference Interface
+		virtual void Evaluate(
+			TLearningArrayView<2, float> Output,
+			const TLearningArrayView<2, const float> Input,
+			const FIndexSet Instances) override final;
+		//~ End INeuralNetworkInference Interface
+
+		void CreateInstances(NNE::IModelCPU& InModel);
+
+		FNeuralNetworkInferenceSettings Settings;
+		TLearningArray<2, float> InputBuffer;
+		TLearningArray<2, float> OutputBuffer;
+		TArray<TSharedPtr<NNE::IModelInstanceCPU>, TInlineAllocator<64>> InferenceInstances;
+	};
+}
+
+UCLASS()
+class LEARNINGAGENTS_API ULearningAgentsNNENeuralNetworkData : public ULearningAgentsNeuralNetworkData
+{
+	GENERATED_BODY()
+
+	friend struct UE::Learning::Agents::FNeuralNetworkNNE;
+
+public:
+
+	virtual void PostLoad() override final;
+
+	virtual TSharedPtr<UE::Learning::INeuralNetwork> GetNetworkInterface() override final;
+
+	virtual void CopyFrom(const ULearningAgentsNeuralNetworkData* Other) override final;
+
+	virtual void CreateMLP(
+		const uint32 InputSize,
+		const uint32 OutputSize,
+		const uint32 HiddenUnitNum,
+		const uint32 LayerNum,
+		const ELearningAgentsActivationFunction Activation) override final;
+
+private:
+
+	UPROPERTY()
+	int32 InputNum = 0;
+
+	UPROPERTY()
+	int32 OutputNum = 0;
+
+	UPROPERTY()
+	TArray<uint8> FileData;
+
+	UPROPERTY()
+	TObjectPtr<UNNEModelData> ModelData;
+
+	TSharedPtr<UE::Learning::Agents::FNeuralNetworkNNE> Network;
+};
+
+using ULearningAgentsDefaultNeuralNetworkData = ULearningAgentsNNENeuralNetworkData;
\ No newline at end of file
diff --git a/Engine/Plugins/Experimental/LearningAgents/Source/LearningAgents/Private/LearningAgentsPolicy.cpp b/Engine/Plugins/Experimental/LearningAgents/Source/LearningAgents/Private/LearningAgentsPolicy.cpp
--- a/Engine/Plugins/Experimental/LearningAgents/Source/LearningAgents/Private/LearningAgentsPolicy.cpp
+++ b/Engine/Plugins/Experimental/LearningAgents/Source/LearningAgents/Private/LearningAgentsPolicy.cpp
@@ -4,6 +4,7 @@
 
 #include "LearningAgentsManager.h"
 #include "LearningAgentsInteractor.h"
+#include "LearningAgentsNeuralNetworkData.h"
 #include "LearningFeatureObject.h"
 #include "LearningNeuralNetwork.h"
 #include "LearningNeuralNetworkObject.h"
@@ -55,34 +56,27 @@ void ULearningAgentsPolicy::SetupPolicy(
 	{
 		// Use Existing Neural Network Asset
 
-		if (NeuralNetworkAsset->NeuralNetwork)
+		if (NeuralNetworkAsset->NeuralNetworkData)
 		{
-			if (NeuralNetworkAsset->NeuralNetwork->GetInputNum() != Interactor->GetObservationFeature().DimNum() ||
-				NeuralNetworkAsset->NeuralNetwork->GetOutputNum() != 2 * Interactor->GetActionFeature().DimNum())
+			if (NeuralNetworkAsset->NeuralNetworkData->GetNetworkInterface()->GetInputNum() != Interactor->GetObservationFeature().DimNum() ||
+				NeuralNetworkAsset->NeuralNetworkData->GetNetworkInterface()->GetOutputNum() != 2 * Interactor->GetActionFeature().DimNum())
 			{
 				UE_LOG(LogLearning, Error, TEXT("%s: Neural Network Asset provided during Setup is incorrect size: Inputs and outputs don't match what is required."), *GetName());
 				return;
 			}
 
-			if (NeuralNetworkAsset->NeuralNetwork->GetHiddenNum() != PolicySettings.HiddenLayerSize ||
-				NeuralNetworkAsset->NeuralNetwork->GetLayerNum() != PolicySettings.LayerNum ||
-				NeuralNetworkAsset->NeuralNetwork->ActivationFunction != UE::Learning::Agents::GetActivationFunction(PolicySettings.ActivationFunction))
-			{
-				UE_LOG(LogLearning, Warning, TEXT("%s: Neural Network Asset settings don't match those given by PolicySettings"), *GetName());
-			}
-
 			Network = NeuralNetworkAsset;
 		}
 		else
 		{
 			Network = NeuralNetworkAsset;
-			Network->NeuralNetwork = MakeShared<UE::Learning::FNeuralNetwork>();
-			Network->NeuralNetwork->Resize(
+			Network->NeuralNetworkData = NewObject<ULearningAgentsDefaultNeuralNetworkData>(Network);
+			Network->NeuralNetworkData->CreateMLP(
 				Interactor->GetObservationFeature().DimNum(),
 				2 * Interactor->GetActionFeature().DimNum(),
 				PolicySettings.HiddenLayerSize,
-				PolicySettings.LayerNum);
-			Network->NeuralNetwork->ActivationFunction = UE::Learning::Agents::GetActivationFunction(PolicySettings.ActivationFunction);
+				PolicySettings.LayerNum,
+				PolicySettings.ActivationFunction);
 		}
 	}
 	else
@@ -92,13 +86,13 @@ void ULearningAgentsPolicy::SetupPolicy(
 		const FName UniqueName = MakeUniqueObjectName(this, ULearningAgentsNeuralNetwork::StaticClass(), TEXT("PolicyNetwork"), EUniqueObjectNameOptions::GloballyUnique);
 
 		Network = NewObject<ULearningAgentsNeuralNetwork>(this, UniqueName);
-		Network->NeuralNetwork = MakeShared<UE::Learning::FNeuralNetwork>();
-		Network->NeuralNetwork->Resize(
+		Network->NeuralNetworkData = NewObject<ULearningAgentsDefaultNeuralNetworkData>(Network);
+		Network->NeuralNetworkData->CreateMLP(
 			Interactor->GetObservationFeature().DimNum(),
 			2 * Interactor->GetActionFeature().DimNum(),
 			PolicySettings.HiddenLayerSize,
-			PolicySettings.LayerNum);
-		Network->NeuralNetwork->ActivationFunction = UE::Learning::Agents::GetActivationFunction(PolicySettings.ActivationFunction);
+			PolicySettings.LayerNum,
+			PolicySettings.ActivationFunction);
 	}
 
 	// Create Policy Object
@@ -111,8 +105,9 @@ void ULearningAgentsPolicy::SetupPolicy(
 		TEXT("PolicyObject"),
 		Manager->GetInstanceData().ToSharedRef(),
 		Manager->GetMaxAgentNum(),
-		Network->NeuralNetwork.ToSharedRef(),
+		Network->NeuralNetworkData->GetNetworkInterface(),
 		PolicySettings.ActionNoiseSeed,
+		UE::Learning::FNeuralNetworkInferenceSettings(),
 		PolicyFunctionSettings);
 
 	Manager->GetInstanceData()->Link(Interactor->GetObservationFeature().FeatureHandle, PolicyObject->InputHandle);
@@ -128,9 +123,9 @@ ULearningAgentsNeuralNetwork* ULearningAgentsPolicy::GetNetworkAsset()
 	return Network;
 }
 
-UE::Learning::FNeuralNetwork& ULearningAgentsPolicy::GetPolicyNetwork()
+UE::Learning::INeuralNetwork& ULearningAgentsPolicy::GetPolicyNetwork()
 {
-	return *Network->NeuralNetwork;
+	return *Network->NeuralNetworkData->GetNetworkInterface();
 }
 
 UE::Learning::FNeuralNetworkPolicyFunction& ULearningAgentsPolicy::GetPolicyObject()
@@ -168,7 +163,7 @@ void ULearningAgentsPolicy::UsePolicyFromAsset(ULearningAgentsNeuralNetwork* Neu
 		return;
 	}
 
-	if (!NeuralNetworkAsset || !NeuralNetworkAsset->NeuralNetwork)
+	if (!NeuralNetworkAsset || !NeuralNetworkAsset->NeuralNetworkData->GetNetworkInterface())
 	{
 		UE_LOG(LogLearning, Error, TEXT("%s: Asset is invalid."), *GetName());
 		return;
@@ -180,17 +175,15 @@ void ULearningAgentsPolicy::UsePolicyFromAsset(ULearningAgentsNeuralNetwork* Neu
 		return;
 	}
 
-	if (NeuralNetworkAsset->NeuralNetwork->GetInputNum() != Network->NeuralNetwork->GetInputNum() ||
-		NeuralNetworkAsset->NeuralNetwork->GetOutputNum() != Network->NeuralNetwork->GetOutputNum() ||
-		NeuralNetworkAsset->NeuralNetwork->GetLayerNum() != Network->NeuralNetwork->GetLayerNum() ||
-		NeuralNetworkAsset->NeuralNetwork->ActivationFunction != Network->NeuralNetwork->ActivationFunction)
+	if (NeuralNetworkAsset->NeuralNetworkData->GetNetworkInterface()->GetInputNum() != Network->NeuralNetworkData->GetNetworkInterface()->GetInputNum() ||
+		NeuralNetworkAsset->NeuralNetworkData->GetNetworkInterface()->GetOutputNum() != Network->NeuralNetworkData->GetNetworkInterface()->GetOutputNum())
 	{
 		UE_LOG(LogLearning, Error, TEXT("%s: Failed to use asset as network settings don't match."), *GetName());
 		return;
 	}
 
 	Network = NeuralNetworkAsset;
-	PolicyObject->NeuralNetwork = Network->NeuralNetwork.ToSharedRef();
+	PolicyObject->UpdateNeuralNetwork(Network->NeuralNetworkData->GetNetworkInterface());
 }
 
 void ULearningAgentsPolicy::LoadPolicyFromAsset(ULearningAgentsNeuralNetwork* NeuralNetworkAsset)
diff --git a/Engine/Plugins/Experimental/LearningAgents/Source/LearningAgents/Public/LearningAgentsCritic.h b/Engine/Plugins/Experimental/LearningAgents/Source/LearningAgents/Public/LearningAgentsCritic.h
--- a/Engine/Plugins/Experimental/LearningAgents/Source/LearningAgents/Public/LearningAgentsCritic.h
+++ b/Engine/Plugins/Experimental/LearningAgents/Source/LearningAgents/Public/LearningAgentsCritic.h
@@ -13,7 +13,7 @@
 
 namespace UE::Learning
 {
-	struct FNeuralNetwork;
+	struct INeuralNetwork;
 	struct FNeuralNetworkCriticFunction;
 }
 
@@ -144,7 +144,7 @@ public:
 	ULearningAgentsNeuralNetwork* GetNetworkAsset();
 
 	/** Get a reference to this critic's neural network. */
-	UE::Learning::FNeuralNetwork& GetCriticNetwork();
+	UE::Learning::INeuralNetwork& GetCriticNetwork();
 	
 	/** Get a reference to this critic's critic function object. */
 	UE::Learning::FNeuralNetworkCriticFunction& GetCriticObject();
diff --git a/Engine/Plugins/Experimental/LearningAgents/Source/LearningAgents/Public/LearningAgentsNeuralNetwork.h b/Engine/Plugins/Experimental/LearningAgents/Source/LearningAgents/Public/LearningAgentsNeuralNetwork.h
--- a/Engine/Plugins/Experimental/LearningAgents/Source/LearningAgents/Public/LearningAgentsNeuralNetwork.h
+++ b/Engine/Plugins/Experimental/LearningAgents/Source/LearningAgents/Public/LearningAgentsNeuralNetwork.h
@@ -6,12 +6,6 @@
 
 #include "LearningAgentsNeuralNetwork.generated.h"
 
-namespace UE::Learning
-{
-	struct FNeuralNetwork;
-	enum class EActivationFunction : uint8;
-}
-
 /** Activation functions for neural networks. */
 UENUM(BlueprintType, Category = "LearningAgents")
 enum class ELearningAgentsActivationFunction : uint8
@@ -26,14 +20,7 @@ enum class ELearningAgentsActivationFunction : uint8
 	TanH	UMETA(DisplayName = "TanH"),
 };
 
-namespace UE::Learning::Agents
-{
-	/** Get the learning agents activation function from the UE::Learning activation function. */
-	LEARNINGAGENTS_API ELearningAgentsActivationFunction GetLearningAgentsActivationFunction(const EActivationFunction ActivationFunction);
-
-	/** Get the UE::Learning activation function from the learning agents activation function. */
-	LEARNINGAGENTS_API EActivationFunction GetActivationFunction(const ELearningAgentsActivationFunction ActivationFunction);
-}
+class ULearningAgentsNeuralNetworkData;
 
 /** A neural network data asset. */
 UCLASS(BlueprintType)
@@ -47,9 +34,6 @@ public:
 	ULearningAgentsNeuralNetwork(FVTableHelper& Helper);
 	virtual ~ULearningAgentsNeuralNetwork();
 
-	/** Serialize this neural network to/from the given archive. */
-	virtual void Serialize(FArchive& Ar) override;
-
 	/**
 	 * Resets this network asset to be empty.
 	 */
@@ -91,6 +75,7 @@ public:
 
 public:
 
-	/** Pointer to the internal Neural Network Data */
-	TSharedPtr<UE::Learning::FNeuralNetwork> NeuralNetwork;
+	/** The internal Neural Network Data */
+	UPROPERTY();
+	TObjectPtr<ULearningAgentsNeuralNetworkData> NeuralNetworkData;
 };
diff --git a/Engine/Plugins/Experimental/LearningAgents/Source/LearningAgents/Public/LearningAgentsPolicy.h b/Engine/Plugins/Experimental/LearningAgents/Source/LearningAgents/Public/LearningAgentsPolicy.h
--- a/Engine/Plugins/Experimental/LearningAgents/Source/LearningAgents/Public/LearningAgentsPolicy.h
+++ b/Engine/Plugins/Experimental/LearningAgents/Source/LearningAgents/Public/LearningAgentsPolicy.h
@@ -13,7 +13,7 @@
 
 namespace UE::Learning
 {
-	struct FNeuralNetwork;
+	struct INeuralNetwork;
 	struct FNeuralNetworkPolicyFunction;
 }
 
@@ -154,7 +154,7 @@ public:
 	ULearningAgentsNeuralNetwork* GetNetworkAsset();
 
 	/** Get a reference to this policy's neural network. */
-	UE::Learning::FNeuralNetwork& GetPolicyNetwork();
+	UE::Learning::INeuralNetwork& GetPolicyNetwork();
 
 	/** Get a reference to this policy's policy function object. */
 	UE::Learning::FNeuralNetworkPolicyFunction& GetPolicyObject();
diff --git a/Engine/Plugins/Experimental/LearningAgents/Source/LearningAgentsTraining/Private/LearningAgentsImitationTrainer.cpp b/Engine/Plugins/Experimental/LearningAgents/Source/LearningAgentsTraining/Private/LearningAgentsImitationTrainer.cpp
--- a/Engine/Plugins/Experimental/LearningAgents/Source/LearningAgentsTraining/Private/LearningAgentsImitationTrainer.cpp
+++ b/Engine/Plugins/Experimental/LearningAgents/Source/LearningAgentsTraining/Private/LearningAgentsImitationTrainer.cpp
@@ -2,18 +2,18 @@
 
 #include "LearningAgentsImitationTrainer.h"
 
-#include "GenericPlatform/GenericPlatformMisc.h"
-#include "HAL/FileManager.h"
 #include "LearningAgentsInteractor.h"
 #include "LearningArrayMap.h"
 #include "LearningExperience.h"
 #include "LearningFeatureObject.h"
 #include "LearningLog.h"
-#include "LearningNeuralNetwork.h"
 #include "LearningImitationTrainer.h"
 #include "LearningAgentsRecording.h"
 #include "LearningAgentsPolicy.h"
 #include "LearningNeuralNetworkObject.h"
+
+#include "GenericPlatform/GenericPlatformMisc.h"
+#include "HAL/FileManager.h"
 #include "Misc/Paths.h"
 
 ULearningAgentsImitationTrainer::ULearningAgentsImitationTrainer() : Super(FObjectInitializer::Get()) {}
@@ -179,30 +179,38 @@ void ULearningAgentsImitationTrainer::BeginTraining(
 		RecordedObservations.Num<0>(),
 		RecordedObservations.Num<1>(),
 		RecordedActions.Num<1>(),
+		Policy->GetPolicyNetwork(),
 		ImitationTrainingSettings);
 
 	UE_LOG(LogLearning, Display, TEXT("%s: Sending / Receiving initial policy..."), *GetName());
 
 	UE::Learning::ETrainerResponse Response = UE::Learning::ETrainerResponse::Success;
 
-	if (bReinitializePolicyNetwork)
-	{
-		Response = ImitationTrainer->RecvPolicy(Policy->GetPolicyNetwork(), TrainerTimeout);
-		Policy->GetNetworkAsset()->ForceMarkDirty();
-	}
-	else
-	{
-		Response = ImitationTrainer->SendPolicy(Policy->GetPolicyNetwork(), TrainerTimeout);
-	}
+	Response = ImitationTrainer->SendPolicy(Policy->GetPolicyNetwork(), TrainerTimeout);
 
 	if (Response != UE::Learning::ETrainerResponse::Success)
 	{
-		UE_LOG(LogLearning, Error, TEXT("%s: Error sending or receiving policy from trainer: %s. Check log for errors."), *GetName(), UE::Learning::Trainer::GetResponseString(Response));
-		bHasTrainingFailed = true;
+		UE_LOG(LogLearning, Error, TEXT("%s: Error sending policy to trainer: %s. Check log for errors."), *GetName(), UE::Learning::Trainer::GetResponseString(Response));
 		ImitationTrainer->Terminate();
+		bHasTrainingFailed = true;
 		return;
 	}
 
+	if (!bReinitializePolicyNetwork)
+	{
+		Response = ImitationTrainer->RecvPolicy(Policy->GetPolicyNetwork(), TrainerTimeout);
+
+		if (Response != UE::Learning::ETrainerResponse::Success)
+		{
+			UE_LOG(LogLearning, Error, TEXT("%s: Error receiving policy from trainer: %s. Check log for errors."), *GetName(), UE::Learning::Trainer::GetResponseString(Response));
+			bHasTrainingFailed = true;
+			ImitationTrainer->Terminate();
+			return;
+		}
+
+		Policy->GetNetworkAsset()->ForceMarkDirty();
+	}
+
 	UE_LOG(LogLearning, Display, TEXT("%s: Sending Experience..."), *GetName());
 
 	// Send Experience
diff --git a/Engine/Plugins/Experimental/LearningAgents/Source/LearningAgentsTraining/Private/LearningAgentsTrainer.cpp b/Engine/Plugins/Experimental/LearningAgents/Source/LearningAgentsTraining/Private/LearningAgentsTrainer.cpp
--- a/Engine/Plugins/Experimental/LearningAgents/Source/LearningAgentsTraining/Private/LearningAgentsTrainer.cpp
+++ b/Engine/Plugins/Experimental/LearningAgents/Source/LearningAgentsTraining/Private/LearningAgentsTrainer.cpp
@@ -15,7 +15,6 @@
 #include "LearningExperience.h"
 #include "LearningFeatureObject.h"
 #include "LearningLog.h"
-#include "LearningNeuralNetwork.h"
 #include "LearningNeuralNetworkObject.h"
 #include "LearningPPOTrainer.h"
 #include "LearningRewardObject.h"
@@ -144,14 +143,19 @@ void ULearningAgentsTrainer::SetupTrainer(
 		return;
 	}
 
+	if (!InCritic)
+	{
+		UE_LOG(LogLearning, Error, TEXT("%s: InCritic is nullptr."), *GetName());
+		return;
+	}
+
 	if (!InPolicy->IsSetup())
 	{
 		UE_LOG(LogLearning, Error, TEXT("%s: %s's Setup must be run before it can be used."), *GetName(), *InPolicy->GetName());
 		return;
 	}
 
-	// The critic is optional unlike the other components
-	if (InCritic && !InCritic->IsSetup())
+	if (!InCritic->IsSetup())
 	{
 		UE_LOG(LogLearning, Error, TEXT("%s: %s's Setup must be run before it can be used."), *GetName(), *InCritic->GetName());
 		return;
@@ -368,7 +372,6 @@ void ULearningAgentsTrainer::BeginTraining(
 	const FLearningAgentsTrainerTrainingSettings& TrainerTrainingSettings,
 	const FLearningAgentsTrainerGameSettings& TrainerGameSettings,
 	const FLearningAgentsTrainerPathSettings& TrainerPathSettings,
-	const FLearningAgentsCriticSettings& CriticSettings,
 	const bool bReinitializePolicyNetwork,
 	const bool bReinitializeCriticNetwork,
 	const bool bResetAgentsOnBegin)
@@ -490,37 +493,8 @@ void ULearningAgentsTrainer::BeginTraining(
 	UE::Learning::FPPOTrainerNetworkSettings PPONetworkSettings;
 	PPONetworkSettings.PolicyActionNoiseMin = Policy->GetPolicyObject().Settings.ActionNoiseMin;
 	PPONetworkSettings.PolicyActionNoiseMax = Policy->GetPolicyObject().Settings.ActionNoiseMax;
-	PPONetworkSettings.PolicyActivationFunction = Policy->GetPolicyNetwork().ActivationFunction;
-	PPONetworkSettings.PolicyHiddenLayerSize = Policy->GetPolicyNetwork().GetHiddenNum();
-	PPONetworkSettings.PolicyLayerNum = Policy->GetPolicyNetwork().GetLayerNum();
-
-	if (Critic)
-	{
-		if (CriticSettings.HiddenLayerSize != Critic->GetCriticNetwork().GetHiddenNum() ||
-			CriticSettings.LayerNum != Critic->GetCriticNetwork().GetLayerNum() ||
-			UE::Learning::Agents::GetActivationFunction(CriticSettings.ActivationFunction) != Critic->GetCriticNetwork().ActivationFunction)
-		{
-			UE_LOG(LogLearning, Warning, TEXT("%s: BeginTraining got different Critic Network Settings to those provided to SetupCritic."), *GetName());
-		}
-
-		PPONetworkSettings.CriticHiddenLayerSize = Critic->GetCriticNetwork().GetHiddenNum();
-		PPONetworkSettings.CriticLayerNum = Critic->GetCriticNetwork().GetLayerNum();
-		PPONetworkSettings.CriticActivationFunction = Critic->GetCriticNetwork().ActivationFunction;
-	}
-	else
-	{
-		PPONetworkSettings.CriticHiddenLayerSize = CriticSettings.HiddenLayerSize;
-		PPONetworkSettings.CriticLayerNum = CriticSettings.LayerNum;
-		PPONetworkSettings.CriticActivationFunction = UE::Learning::Agents::GetActivationFunction(CriticSettings.ActivationFunction);
-	}
-
-	// We assume that if the critic has been setup on the agent interactor, then
-	// the user wants the critic network to be synced during training.
-	UE::Learning::EPPOTrainerFlags TrainerFlags = 
-		Critic ?
-		UE::Learning::EPPOTrainerFlags::SynchronizeCriticNetwork :
-		UE::Learning::EPPOTrainerFlags::None;
 
+	UE::Learning::EPPOTrainerFlags TrainerFlags = UE::Learning::EPPOTrainerFlags::None;
 	if (!bReinitializePolicyNetwork) { TrainerFlags |= UE::Learning::EPPOTrainerFlags::UseInitialPolicyNetwork; }
 	if (!bReinitializeCriticNetwork && Critic) { TrainerFlags |= UE::Learning::EPPOTrainerFlags::UseInitialCriticNetwork; }
 
@@ -532,6 +506,8 @@ void ULearningAgentsTrainer::BeginTraining(
 		PythonContentPath,
 		IntermediatePath,
 		*ReplayBuffer,
+		Policy->GetPolicyNetwork(),
+		Critic->GetCriticNetwork(),
 		PPOTrainingSettings,
 		PPONetworkSettings,
 		TrainerFlags);
@@ -540,47 +516,54 @@ void ULearningAgentsTrainer::BeginTraining(
 
 	UE::Learning::ETrainerResponse Response = UE::Learning::ETrainerResponse::Success;
 
-	if ((bool)(TrainerFlags & UE::Learning::EPPOTrainerFlags::UseInitialPolicyNetwork))
-	{
-		Response = Trainer->SendPolicy(Policy->GetPolicyNetwork(), TrainerTimeout);
-	}
-	else
+	Response = Trainer->SendPolicy(Policy->GetPolicyNetwork(), TrainerTimeout);
+
+	if (Response != UE::Learning::ETrainerResponse::Success)
 	{
-		Response = Trainer->RecvPolicy(Policy->GetPolicyNetwork(), TrainerTimeout);
-		Policy->GetNetworkAsset()->ForceMarkDirty();
+		UE_LOG(LogLearning, Error, TEXT("%s: Error sending policy to trainer: %s. Check log for errors."), *GetName(), UE::Learning::Trainer::GetResponseString(Response));
+		bHasTrainingFailed = true;
+		Trainer->Terminate();
+		return;
 	}
 
+	Response = Trainer->SendCritic(Critic->GetCriticNetwork(), TrainerTimeout);
+
 	if (Response != UE::Learning::ETrainerResponse::Success)
 	{
-		UE_LOG(LogLearning, Error, TEXT("%s: Error sending or receiving policy from trainer: %s. Check log for errors."), *GetName(), UE::Learning::Trainer::GetResponseString(Response));
+		UE_LOG(LogLearning, Error, TEXT("%s: Error sending critic to trainer: %s. Check log for errors."), *GetName(), UE::Learning::Trainer::GetResponseString(Response));
 		bHasTrainingFailed = true;
 		Trainer->Terminate();
 		return;
 	}
 
-	if (Critic)
+	if (!(bool)(TrainerFlags & UE::Learning::EPPOTrainerFlags::UseInitialPolicyNetwork))
 	{
-		if ((bool)(TrainerFlags & UE::Learning::EPPOTrainerFlags::UseInitialCriticNetwork))
-		{
-			Response = Trainer->SendCritic(Critic->GetCriticNetwork(), TrainerTimeout);
-		}
-		else if ((bool)(TrainerFlags & UE::Learning::EPPOTrainerFlags::SynchronizeCriticNetwork))
-		{
-			Response = Trainer->RecvCritic(Critic->GetCriticNetwork(), TrainerTimeout);
-			Critic->GetNetworkAsset()->ForceMarkDirty();
-		}
-		else
+		Response = Trainer->RecvPolicy(Policy->GetPolicyNetwork(), TrainerTimeout);
+
+		if (Response != UE::Learning::ETrainerResponse::Success)
 		{
-			Response = UE::Learning::ETrainerResponse::Success;
+			UE_LOG(LogLearning, Error, TEXT("%s: Error receiving policy from trainer: %s. Check log for errors."), *GetName(), UE::Learning::Trainer::GetResponseString(Response));
+			bHasTrainingFailed = true;
+			Trainer->Terminate();
+			return;
 		}
 
+		Policy->GetNetworkAsset()->ForceMarkDirty();
+	}
+
+	if (!(bool)(TrainerFlags & UE::Learning::EPPOTrainerFlags::UseInitialCriticNetwork))
+	{
+		Response = Trainer->RecvCritic(Critic->GetCriticNetwork(), TrainerTimeout);
+
 		if (Response != UE::Learning::ETrainerResponse::Success)
 		{
-			UE_LOG(LogLearning, Error, TEXT("%s: Error sending or receiving critic from trainer: %s. Check log for errors."), *GetName(), UE::Learning::Trainer::GetResponseString(Response));
+			UE_LOG(LogLearning, Error, TEXT("%s: Error receiving critic from trainer: %s. Check log for errors."), *GetName(), UE::Learning::Trainer::GetResponseString(Response));
 			bHasTrainingFailed = true;
 			Trainer->Terminate();
 			return;
 		}
+
+		Critic->GetNetworkAsset()->ForceMarkDirty();
 	}
 
 	// Reset Agents, Replay Buffer
@@ -940,7 +923,6 @@ void ULearningAgentsTrainer::ProcessExperience()
 			if (Response != UE::Learning::ETrainerResponse::Success)
 			{
 				UE_LOG(LogLearning, Error, TEXT("%s: Error waiting to push experience to trainer. Check log for errors."), *GetName());
-				bHasTrainingFailed = true;
 				EndTraining();
 				return;
 			}
@@ -966,18 +948,15 @@ void ULearningAgentsTrainer::ProcessExperience()
 			}
 
 			// Get Updated Critic
-			if (Critic)
+			Response = Trainer->RecvCritic(Critic->GetCriticNetwork(), TrainerTimeout);
+			Critic->GetNetworkAsset()->ForceMarkDirty();
+
+			if (Response != UE::Learning::ETrainerResponse::Success)
 			{
-				Response = Trainer->RecvCritic(Critic->GetCriticNetwork(), TrainerTimeout);
-				Critic->GetNetworkAsset()->ForceMarkDirty();
-
-				if (Response != UE::Learning::ETrainerResponse::Success)
-				{
-					UE_LOG(LogLearning, Error, TEXT("%s: Error waiting for critic from trainer. Check log for errors."), *GetName());
-					bHasTrainingFailed = true;
-					EndTraining();
-					return;
-				}
+				UE_LOG(LogLearning, Error, TEXT("%s: Error waiting for critic from trainer. Check log for errors."), *GetName());
+				bHasTrainingFailed = true;
+				EndTraining();
+				return;
 			}
 
 			// Mark all agents for reset since we have a new policy
@@ -992,7 +971,6 @@ void ULearningAgentsTrainer::RunTraining(
 	const FLearningAgentsTrainerTrainingSettings& TrainerTrainingSettings,
 	const FLearningAgentsTrainerGameSettings& TrainerGameSettings,
 	const FLearningAgentsTrainerPathSettings& TrainerPathSettings,
-	const FLearningAgentsCriticSettings& CriticSettings,
 	const bool bReinitializePolicyNetwork,
 	const bool bReinitializeCriticNetwork,
 	const bool bResetAgentsOnBegin)
@@ -1016,7 +994,6 @@ void ULearningAgentsTrainer::RunTraining(
 			TrainerTrainingSettings,
 			TrainerGameSettings, 
 			TrainerPathSettings, 
-			CriticSettings, 
 			bReinitializePolicyNetwork, 
 			bReinitializeCriticNetwork,
 			bResetAgentsOnBegin);
diff --git a/Engine/Plugins/Experimental/LearningAgents/Source/LearningAgentsTraining/Public/LearningAgentsTrainer.h b/Engine/Plugins/Experimental/LearningAgents/Source/LearningAgentsTraining/Public/LearningAgentsTrainer.h
--- a/Engine/Plugins/Experimental/LearningAgents/Source/LearningAgentsTraining/Public/LearningAgentsTrainer.h
+++ b/Engine/Plugins/Experimental/LearningAgents/Source/LearningAgentsTraining/Public/LearningAgentsTrainer.h
@@ -343,14 +343,14 @@ public:
 	 * Initializes this object and runs the setup functions for rewards and completions.
 	 * @param InInteractor The agent interactor we are training with.
 	 * @param InPolicy The policy to be trained.
-	 * @param InCritic Optional - only needs to be provided if we want the critic to be accessible at runtime.
+	 * @param InCritic The critic to be trained.
 	 * @param TrainerSettings The trainer settings to use.
 	 */
 	UFUNCTION(BlueprintCallable, Category = "LearningAgents")
 	void SetupTrainer(
 		ULearningAgentsInteractor* InInteractor,
 		ULearningAgentsPolicy* InPolicy,
-		ULearningAgentsCritic* InCritic = nullptr,
+		ULearningAgentsCritic* InCritic,
 		const FLearningAgentsTrainerSettings& TrainerSettings = FLearningAgentsTrainerSettings());
 
 public: 
@@ -431,10 +431,10 @@ public:
 
 	/**
 	 * Begins the training process with the provided settings.
+	 * 
 	 * @param TrainerTrainingSettings The settings for this training run.
 	 * @param TrainerGameSettings The settings that will affect the game's simulation.
 	 * @param TrainerPathSettings The path settings used by the trainer.
-	 * @param CriticSettings The settings for the critic (if we are using one).
 	 * @param bReinitializePolicyNetwork If true, reinitialize the policy. Set this to false if your policy is pre-trained, e.g. with imitation learning.
 	 * @param bReinitializeCriticNetwork If true, reinitialize the critic. Set this to false if your critic is pre-trained.
 	 * @param bResetAgentsOnBegin If true, reset all agents at the beginning of training.
@@ -444,7 +444,6 @@ public:
 		const FLearningAgentsTrainerTrainingSettings& TrainerTrainingSettings = FLearningAgentsTrainerTrainingSettings(),
 		const FLearningAgentsTrainerGameSettings& TrainerGameSettings = FLearningAgentsTrainerGameSettings(),
 		const FLearningAgentsTrainerPathSettings& TrainerPathSettings = FLearningAgentsTrainerPathSettings(),
-		const FLearningAgentsCriticSettings& CriticSettings = FLearningAgentsCriticSettings(),
 		const bool bReinitializePolicyNetwork = true,
 		const bool bReinitializeCriticNetwork = true,
 		const bool bResetAgentsOnBegin = true);
@@ -482,10 +481,10 @@ public:
 	 * Convenience function that runs a basic training loop. If training has not been started, it will start it, and 
 	 * then call RunInference. On each following call to this function, it will call EvaluateRewards, 
 	 * EvaluateCompletions, and ProcessExperience, followed by RunInference.
+	 * 
 	 * @param TrainerTrainingSettings The settings for this training run.
 	 * @param TrainerGameSettings The settings that will affect the game's simulation.
 	 * @param TrainerPathSettings The path settings used by the trainer.
-	 * @param CriticSettings The settings for the critic (if we are using one).
 	 * @param bReinitializePolicyNetwork If true, reinitialize the policy. Set this to false if your policy is pre-trained, e.g. with imitation learning.
 	 * @param bReinitializeCriticNetwork If true, reinitialize the critic. Set this to false if your critic is pre-trained.
 	 * @param bResetAgentsOnBegin If true, reset all agents at the beginning of training.
@@ -495,7 +494,6 @@ public:
 		const FLearningAgentsTrainerTrainingSettings& TrainerTrainingSettings = FLearningAgentsTrainerTrainingSettings(),
 		const FLearningAgentsTrainerGameSettings& TrainerGameSettings = FLearningAgentsTrainerGameSettings(),
 		const FLearningAgentsTrainerPathSettings& TrainerPathSettings = FLearningAgentsTrainerPathSettings(),
-		const FLearningAgentsCriticSettings& CriticSettings = FLearningAgentsCriticSettings(),
 		const bool bReinitializePolicyNetwork = true,
 		const bool bReinitializeCriticNetwork = true,
 		const bool bResetAgentsOnBegin = true);
diff --git a/Engine/Plugins/Experimental/LearningAgents/Source/LearningTraining/Private/LearningImitationTrainer.cpp b/Engine/Plugins/Experimental/LearningAgents/Source/LearningTraining/Private/LearningImitationTrainer.cpp
--- a/Engine/Plugins/Experimental/LearningAgents/Source/LearningTraining/Private/LearningImitationTrainer.cpp
+++ b/Engine/Plugins/Experimental/LearningAgents/Source/LearningTraining/Private/LearningImitationTrainer.cpp
@@ -117,6 +117,7 @@ namespace UE::Learning
 		const int32 MaxSampleNum,
 		const int32 ObservationDimNum,
 		const int32 ActionDimNum,
+		const INeuralNetwork& PolicyNetwork,
 		const FImitationTrainerTrainingSettings& TrainingSettings,
 		const FImitationTrainerNetworkSettings& NetworkSettings,
 		const EImitationTrainerFlags TrainerFlags,
@@ -127,15 +128,9 @@ namespace UE::Learning
 		UE_LEARNING_CHECK(FPaths::DirectoryExists(PythonContentPath));
 		UE_LEARNING_CHECK(FPaths::DirectoryExists(SitePackagesPath));
 
-		const int32 TotalPolicyByteNum = FNeuralNetwork::GetSerializationByteNum(
-			ObservationDimNum,
-			2 * ActionDimNum,
-			NetworkSettings.PolicyHiddenLayerSize,
-			NetworkSettings.PolicyLayerNum);
-
 		// Allocate Shared Memory
 
-		Policy = SharedMemory::Allocate<1, uint8>({ TotalPolicyByteNum });
+		Policy = SharedMemory::Allocate<1, uint8>({ PolicyNetwork.GetSerializationByteNum() });
 		Controls = SharedMemory::Allocate<1, volatile int32>({ SharedMemoryTraining::GetControlNum() });
 		Observations = SharedMemory::Allocate<2, float>({ MaxSampleNum, ObservationDimNum });
 		Actions = SharedMemory::Allocate<2, float>({ MaxSampleNum, ActionDimNum });
@@ -163,6 +158,8 @@ namespace UE::Learning
 		ConfigObject->SetStringField(TEXT("SitePackagesPath"), *FileManager.ConvertToAbsolutePathForExternalAppForRead(*SitePackagesPath));
 		ConfigObject->SetStringField(TEXT("IntermediatePath"), *FileManager.ConvertToAbsolutePathForExternalAppForRead(*IntermediatePath));
 
+		ConfigObject->SetStringField(TEXT("PolicyNetworkClass"), PolicyNetwork.GetPythonClassName());
+
 		ConfigObject->SetStringField(TEXT("PolicyGuid"), *Policy.Guid.ToString(EGuidFormats::DigitsWithHyphensInBraces));
 		ConfigObject->SetStringField(TEXT("ControlsGuid"), *Controls.Guid.ToString(EGuidFormats::DigitsWithHyphensInBraces));
 		ConfigObject->SetStringField(TEXT("ObservationsGuid"), *Observations.Guid.ToString(EGuidFormats::DigitsWithHyphensInBraces));
@@ -172,10 +169,7 @@ namespace UE::Learning
 		ConfigObject->SetNumberField(TEXT("ActionVectorDimensionNum"), ActionDimNum);
 		ConfigObject->SetNumberField(TEXT("MaxSampleNum"), MaxSampleNum);
 
-		ConfigObject->SetNumberField(TEXT("PolicyNetworkByteNum"), TotalPolicyByteNum);
-		ConfigObject->SetNumberField(TEXT("PolicyHiddenUnitNum"), NetworkSettings.PolicyHiddenLayerSize);
-		ConfigObject->SetNumberField(TEXT("PolicyLayerNum"), NetworkSettings.PolicyLayerNum);
-		ConfigObject->SetStringField(TEXT("PolicyActivationFunction"), GetActivationFunctionString(NetworkSettings.PolicyActivationFunction));
+		ConfigObject->SetNumberField(TEXT("PolicyNetworkByteNum"), PolicyNetwork.GetSerializationByteNum());
 		ConfigObject->SetNumberField(TEXT("PolicyActionNoiseMin"), NetworkSettings.PolicyActionNoiseMin);
 		ConfigObject->SetNumberField(TEXT("PolicyActionNoiseMax"), NetworkSettings.PolicyActionNoiseMax);
 
@@ -236,7 +230,7 @@ namespace UE::Learning
 	}
 
 	ETrainerResponse FSharedMemoryImitationTrainer::RecvPolicy(
-		FNeuralNetwork& OutNetwork,
+		INeuralNetwork& OutNetwork,
 		const float Timeout,
 		FRWLock* NetworkLock,
 		const ELogSetting LogSettings)
@@ -251,7 +245,7 @@ namespace UE::Learning
 	}
 
 	ETrainerResponse FSharedMemoryImitationTrainer::SendPolicy(
-		const FNeuralNetwork& Network,
+		const INeuralNetwork& Network,
 		const float Timeout,
 		FRWLock* NetworkLock,
 		const ELogSetting LogSettings)
@@ -458,6 +452,7 @@ namespace UE::Learning
 		const int32 MaxSampleNum,
 		const int32 ObservationDimNum,
 		const int32 ActionDimNum,
+		const INeuralNetwork& PolicyNetwork,
 		const TCHAR* IpAddress,
 		const uint32 Port,
 		const float Timeout,
@@ -465,12 +460,6 @@ namespace UE::Learning
 		const FImitationTrainerNetworkSettings& NetworkSettings,
 		const EImitationTrainerFlags TrainerFlags)
 	{
-		const int32 TotalPolicyByteNum = FNeuralNetwork::GetSerializationByteNum(
-			ObservationDimNum,
-			2 * ActionDimNum,
-			NetworkSettings.PolicyHiddenLayerSize,
-			NetworkSettings.PolicyLayerNum);
-
 		// Write Config
 
 		const FString TimeStamp = FDateTime::Now().ToFormattedString(TEXT("%Y-%m-%d_%H-%M-%S"));
@@ -483,14 +472,13 @@ namespace UE::Learning
 		ConfigObject->SetStringField(TEXT("TrainerType"), TrainerType);
 		ConfigObject->SetStringField(TEXT("TimeStamp"), *TimeStamp);
 
+		ConfigObject->SetStringField(TEXT("PolicyNetworkClass"), PolicyNetwork.GetPythonClassName());
+
 		ConfigObject->SetNumberField(TEXT("ObservationVectorDimensionNum"), ObservationDimNum);
 		ConfigObject->SetNumberField(TEXT("ActionVectorDimensionNum"), ActionDimNum);
 		ConfigObject->SetNumberField(TEXT("MaxSampleNum"), MaxSampleNum);
 
-		ConfigObject->SetNumberField(TEXT("PolicyNetworkByteNum"), TotalPolicyByteNum);
-		ConfigObject->SetNumberField(TEXT("PolicyHiddenUnitNum"), NetworkSettings.PolicyHiddenLayerSize);
-		ConfigObject->SetNumberField(TEXT("PolicyLayerNum"), NetworkSettings.PolicyLayerNum);
-		ConfigObject->SetStringField(TEXT("PolicyActivationFunction"), GetActivationFunctionString(NetworkSettings.PolicyActivationFunction));
+		ConfigObject->SetNumberField(TEXT("PolicyNetworkByteNum"), PolicyNetwork.GetSerializationByteNum());
 		ConfigObject->SetNumberField(TEXT("PolicyActionNoiseMin"), NetworkSettings.PolicyActionNoiseMin);
 		ConfigObject->SetNumberField(TEXT("PolicyActionNoiseMax"), NetworkSettings.PolicyActionNoiseMax);
 
@@ -511,7 +499,7 @@ namespace UE::Learning
 
 		// Allocate buffer to receive network data in
 
-		NetworkBuffer.SetNumUninitialized({ TotalPolicyByteNum });
+		NetworkBuffer.SetNumUninitialized({ PolicyNetwork.GetSerializationByteNum() });
 
 		// Create Socket
 
@@ -580,7 +568,7 @@ namespace UE::Learning
 	}
 
 	ETrainerResponse FSocketImitationTrainer::RecvPolicy(
-		FNeuralNetwork& OutNetwork,
+		INeuralNetwork& OutNetwork,
 		const float Timeout,
 		FRWLock* NetworkLock,
 		const ELogSetting LogSettings)
@@ -589,7 +577,7 @@ namespace UE::Learning
 	}
 
 	ETrainerResponse FSocketImitationTrainer::SendPolicy(
-		const FNeuralNetwork& Network,
+		const INeuralNetwork& Network,
 		const float Timeout,
 		FRWLock* NetworkLock,
 		const ELogSetting LogSettings)
@@ -610,7 +598,7 @@ namespace UE::Learning
 	{
 		ETrainerResponse Train(
 			IImitationTrainer& Trainer,
-			FNeuralNetwork& Network,
+			INeuralNetwork& Network,
 			const TLearningArrayView<2, const float> ObservationVectors,
 			const TLearningArrayView<2, const float> ActionVectors,
 			const EImitationTrainerFlags TrainerFlags, 
@@ -621,29 +609,27 @@ namespace UE::Learning
 		{
 			ETrainerResponse Response = ETrainerResponse::Success;
 
-			if ((bool)(TrainerFlags & EImitationTrainerFlags::UseInitialPolicyNetwork))
+			// Send initial Policy
+
+			if (LogSettings != ELogSetting::Silent)
 			{
-				// Send initial Policy
+				UE_LOG(LogLearning, Display, TEXT("Sending initial Policy..."));
+			}
+
+			Response = Trainer.SendPolicy(Network, 20.0f, NetworkLock);
 
+			if (Response != ETrainerResponse::Success)
+			{
 				if (LogSettings != ELogSetting::Silent)
 				{
-					UE_LOG(LogLearning, Display, TEXT("Sending initial Policy..."));
+					UE_LOG(LogLearning, Error, TEXT("Error sending initial policy from trainer: %s. Check log for errors."), Trainer::GetResponseString(Response));
 				}
 
-				Response = Trainer.SendPolicy(Network, 20.0f, NetworkLock);
-
-				if (Response != ETrainerResponse::Success)
-				{
-					if (LogSettings != ELogSetting::Silent)
-					{
-						UE_LOG(LogLearning, Error, TEXT("Error sending initial policy from trainer: %s. Check log for errors."), Trainer::GetResponseString(Response));
-					}
-
-					Trainer.Terminate();
-					return Response;
-				}
+				Trainer.Terminate();
+				return Response;
 			}
-			else
+
+			if (!(bool)(TrainerFlags & EImitationTrainerFlags::UseInitialPolicyNetwork))
 			{
 				// Receive initial Policy
 
diff --git a/Engine/Plugins/Experimental/LearningAgents/Source/LearningTraining/Private/LearningPPOTrainer.cpp b/Engine/Plugins/Experimental/LearningAgents/Source/LearningTraining/Private/LearningPPOTrainer.cpp
--- a/Engine/Plugins/Experimental/LearningAgents/Source/LearningTraining/Private/LearningPPOTrainer.cpp
+++ b/Engine/Plugins/Experimental/LearningAgents/Source/LearningTraining/Private/LearningPPOTrainer.cpp
@@ -118,6 +118,8 @@ namespace UE::Learning
 		const FString& PythonContentPath,
 		const FString& IntermediatePath,
 		const FReplayBuffer& ReplayBuffer,
+		const INeuralNetwork& PolicyNetwork,
+		const INeuralNetwork& CriticNetwork,
 		const FPPOTrainerTrainingSettings& TrainingSettings,
 		const FPPOTrainerNetworkSettings& NetworkSettings,
 		const EPPOTrainerFlags TrainerFlags,
@@ -133,20 +135,6 @@ namespace UE::Learning
 		const int32 ObservationVectorDimensionNum = ReplayBuffer.GetObservations().Num<1>();
 		const int32 ActionVectorDimensionNum = ReplayBuffer.GetActions().Num<1>();
 
-		const int32 TotalPolicyByteNum = FNeuralNetwork::GetSerializationByteNum(
-			ObservationVectorDimensionNum,
-			2 * ActionVectorDimensionNum,
-			NetworkSettings.PolicyHiddenLayerSize,
-			NetworkSettings.PolicyLayerNum);
-
-		const int32 TotalCriticByteNum = FNeuralNetwork::GetSerializationByteNum(
-				ObservationVectorDimensionNum,
-				1,
-				NetworkSettings.CriticHiddenLayerSize,
-				NetworkSettings.CriticLayerNum);
-
-		const bool bRequiresCritic = (bool)(TrainerFlags & EPPOTrainerFlags::SynchronizeCriticNetwork) || (bool)(TrainerFlags & EPPOTrainerFlags::UseInitialCriticNetwork);
-
 		if (!ensure(Policy.Region == nullptr))
 		{
 			UE_LOG(LogLearning, Warning, TEXT("Training already started!"));
@@ -162,8 +150,8 @@ namespace UE::Learning
 		{
 			// Allocate Shared Memory
 
-			Policy = SharedMemory::Allocate<1, uint8>({ TotalPolicyByteNum });
-			if (bRequiresCritic) { Critic = SharedMemory::Allocate<1, uint8>({ TotalCriticByteNum }); }
+			Policy = SharedMemory::Allocate<1, uint8>({ PolicyNetwork.GetSerializationByteNum() });
+			Critic = SharedMemory::Allocate<1, uint8>({ CriticNetwork.GetSerializationByteNum() });
 			Controls = SharedMemory::Allocate<2, volatile int32>({ ProcessNum, SharedMemoryTraining::GetControlNum() });
 			EpisodeStarts = SharedMemory::Allocate<2, int32>({ ProcessNum, ReplayBuffer.GetMaxEpisodeNum() });
 			EpisodeLengths = SharedMemory::Allocate<2, int32>({ ProcessNum, ReplayBuffer.GetMaxEpisodeNum() });
@@ -189,7 +177,7 @@ namespace UE::Learning
 
 				SubprocessCommandLine += FString::Printf(TEXT(" -LearningProcessIdx %i"), SubprocessIdx);
 				SubprocessCommandLine += FString(TEXT(" -LearningPolicyGuid ")) + Policy.Guid.ToString();
-				if (bRequiresCritic) { SubprocessCommandLine += FString(TEXT(" -LearningCriticGuid ")) + Critic.Guid.ToString(); }
+				SubprocessCommandLine += FString(TEXT(" -LearningCriticGuid ")) + Critic.Guid.ToString();
 				SubprocessCommandLine += FString(TEXT(" -LearningControlsGuid ")) + Controls.Guid.ToString();
 				SubprocessCommandLine += FString(TEXT(" -LearningEpisodeStartsGuid ")) + EpisodeStarts.Guid.ToString();
 				SubprocessCommandLine += FString(TEXT(" -LearningEpisodeLengthsGuid ")) + EpisodeLengths.Guid.ToString();
@@ -236,8 +224,11 @@ namespace UE::Learning
 			ConfigObject->SetStringField(TEXT("SitePackagesPath"), *FileManager.ConvertToAbsolutePathForExternalAppForRead(*SitePackagesPath));
 			ConfigObject->SetStringField(TEXT("IntermediatePath"), *FileManager.ConvertToAbsolutePathForExternalAppForRead(*IntermediatePath));
 
+			ConfigObject->SetStringField(TEXT("PolicyNetworkClass"), PolicyNetwork.GetPythonClassName());
+			ConfigObject->SetStringField(TEXT("CriticNetworkClass"), CriticNetwork.GetPythonClassName());
+
 			ConfigObject->SetStringField(TEXT("PolicyGuid"), *Policy.Guid.ToString(EGuidFormats::DigitsWithHyphensInBraces));
-			if (bRequiresCritic) { ConfigObject->SetStringField(TEXT("CriticGuid"), *Critic.Guid.ToString(EGuidFormats::DigitsWithHyphensInBraces)); }
+			ConfigObject->SetStringField(TEXT("CriticGuid"), *Critic.Guid.ToString(EGuidFormats::DigitsWithHyphensInBraces));
 			ConfigObject->SetStringField(TEXT("ControlsGuid"), *Controls.Guid.ToString(EGuidFormats::DigitsWithHyphensInBraces));
 			ConfigObject->SetStringField(TEXT("EpisodeStartsGuid"), *EpisodeStarts.Guid.ToString(EGuidFormats::DigitsWithHyphensInBraces));
 			ConfigObject->SetStringField(TEXT("EpisodeLengthsGuid"), *EpisodeLengths.Guid.ToString(EGuidFormats::DigitsWithHyphensInBraces));
@@ -252,17 +243,11 @@ namespace UE::Learning
 			ConfigObject->SetNumberField(TEXT("MaxEpisodeNum"), ReplayBuffer.GetMaxEpisodeNum());
 			ConfigObject->SetNumberField(TEXT("MaxStepNum"), ReplayBuffer.GetMaxStepNum());
 
-			ConfigObject->SetNumberField(TEXT("PolicyNetworkByteNum"), TotalPolicyByteNum);
-			ConfigObject->SetNumberField(TEXT("PolicyHiddenUnitNum"), NetworkSettings.PolicyHiddenLayerSize);
-			ConfigObject->SetNumberField(TEXT("PolicyLayerNum"), NetworkSettings.PolicyLayerNum);
-			ConfigObject->SetStringField(TEXT("PolicyActivationFunction"), GetActivationFunctionString(NetworkSettings.PolicyActivationFunction));
+			ConfigObject->SetNumberField(TEXT("PolicyNetworkByteNum"), PolicyNetwork.GetSerializationByteNum());
 			ConfigObject->SetNumberField(TEXT("PolicyActionNoiseMin"), NetworkSettings.PolicyActionNoiseMin);
 			ConfigObject->SetNumberField(TEXT("PolicyActionNoiseMax"), NetworkSettings.PolicyActionNoiseMax);
 
-			ConfigObject->SetNumberField(TEXT("CriticNetworkByteNum"), TotalCriticByteNum);
-			ConfigObject->SetNumberField(TEXT("CriticHiddenUnitNum"), NetworkSettings.CriticHiddenLayerSize);
-			ConfigObject->SetNumberField(TEXT("CriticLayerNum"), NetworkSettings.CriticLayerNum);
-			ConfigObject->SetStringField(TEXT("CriticActivationFunction"), GetActivationFunctionString(NetworkSettings.CriticActivationFunction));
+			ConfigObject->SetNumberField(TEXT("CriticNetworkByteNum"), CriticNetwork.GetSerializationByteNum());
 
 			ConfigObject->SetNumberField(TEXT("ProcessNum"), ProcessNum);
 
@@ -288,7 +273,6 @@ namespace UE::Learning
 
 			ConfigObject->SetBoolField(TEXT("UseInitialPolicyNetwork"), (bool)(TrainerFlags& EPPOTrainerFlags::UseInitialPolicyNetwork));
 			ConfigObject->SetBoolField(TEXT("UseInitialCriticNetwork"), (bool)(TrainerFlags& EPPOTrainerFlags::UseInitialCriticNetwork));
-			ConfigObject->SetBoolField(TEXT("SynchronizeCriticNetwork"), (bool)(TrainerFlags& EPPOTrainerFlags::SynchronizeCriticNetwork));
 
 			ConfigObject->SetBoolField(TEXT("LoggingEnabled"), LogSettings == ELogSetting::Silent ? false : true);
 
@@ -324,7 +308,7 @@ namespace UE::Learning
 			// Parse Guids from command line args
 
 			FGuid PolicyGuid; ensure(FParse::Value(FCommandLine::Get(), TEXT("LearningPolicyGuid"), PolicyGuid));
-			FGuid CriticGuid; if (bRequiresCritic) { ensure(FParse::Value(FCommandLine::Get(), TEXT("LearningCriticGuid"), CriticGuid)); }
+			FGuid CriticGuid; ensure(FParse::Value(FCommandLine::Get(), TEXT("LearningCriticGuid"), CriticGuid));
 			FGuid ControlsGuid; ensure(FParse::Value(FCommandLine::Get(), TEXT("LearningControlsGuid"), ControlsGuid));
 			FGuid EpisodeStartsGuid; ensure(FParse::Value(FCommandLine::Get(), TEXT("LearningEpisodeStartsGuid"), EpisodeStartsGuid));
 			FGuid EpisodeLengthsGuid; ensure(FParse::Value(FCommandLine::Get(), TEXT("LearningEpisodeLengthsGuid"), EpisodeLengthsGuid));
@@ -336,8 +320,8 @@ namespace UE::Learning
 
 			// Map shared memory
 
-			Policy = SharedMemory::Map<1, uint8>(PolicyGuid, { TotalPolicyByteNum });
-			if (bRequiresCritic) { Critic = SharedMemory::Map<1, uint8>(CriticGuid, { TotalCriticByteNum }); }
+			Policy = SharedMemory::Map<1, uint8>(PolicyGuid, { PolicyNetwork.GetSerializationByteNum() });
+			Critic = SharedMemory::Map<1, uint8>(CriticGuid, { CriticNetwork.GetSerializationByteNum() });
 			Controls = SharedMemory::Map<2, volatile int32>(ControlsGuid, { ProcessNum, SharedMemoryTraining::GetControlNum() });
 			EpisodeStarts = SharedMemory::Map<2, int32>(EpisodeStartsGuid, { ProcessNum, ReplayBuffer.GetMaxEpisodeNum() });
 			EpisodeLengths = SharedMemory::Map<2, int32>(EpisodeLengthsGuid, { ProcessNum, ReplayBuffer.GetMaxEpisodeNum() });
@@ -394,7 +378,7 @@ namespace UE::Learning
 	}
 
 	ETrainerResponse FSharedMemoryPPOTrainer::RecvPolicy(
-		FNeuralNetwork& OutNetwork,
+		INeuralNetwork& OutNetwork,
 		const float Timeout,
 		FRWLock* NetworkLock,
 		const ELogSetting LogSettings)
@@ -409,7 +393,7 @@ namespace UE::Learning
 	}
 
 	ETrainerResponse FSharedMemoryPPOTrainer::RecvCritic(
-		FNeuralNetwork& OutNetwork,
+		INeuralNetwork& OutNetwork,
 		const float Timeout,
 		FRWLock* NetworkLock,
 		const ELogSetting LogSettings)
@@ -427,7 +411,7 @@ namespace UE::Learning
 	}
 
 	ETrainerResponse FSharedMemoryPPOTrainer::SendPolicy(
-		const FNeuralNetwork& Network,
+		const INeuralNetwork& Network,
 		const float Timeout,
 		FRWLock* NetworkLock,
 		const ELogSetting LogSettings)
@@ -442,7 +426,7 @@ namespace UE::Learning
 	}
 
 	ETrainerResponse FSharedMemoryPPOTrainer::SendCritic(
-		const FNeuralNetwork& Network,
+		const INeuralNetwork& Network,
 		const float Timeout,
 		FRWLock* NetworkLock,
 		const ELogSetting LogSettings)
@@ -648,6 +632,8 @@ namespace UE::Learning
 		ETrainerResponse& OutResponse,
 		const FString& TaskName,
 		const FReplayBuffer& ReplayBuffer,
+		const INeuralNetwork& PolicyNetwork,
+		const INeuralNetwork& CriticNetwork,
 		const TCHAR* IpAddress,
 		const uint32 Port,
 		const float Timeout,
@@ -658,18 +644,6 @@ namespace UE::Learning
 		const int32 ObservationVectorDimensionNum = ReplayBuffer.GetObservations().Num<1>();
 		const int32 ActionVectorDimensionNum = ReplayBuffer.GetActions().Num<1>();
 
-		const int32 TotalPolicyByteNum = FNeuralNetwork::GetSerializationByteNum(
-			ObservationVectorDimensionNum,
-			2 * ActionVectorDimensionNum,
-			NetworkSettings.PolicyHiddenLayerSize,
-			NetworkSettings.PolicyLayerNum);
-
-		const int32 TotalCriticByteNum = FNeuralNetwork::GetSerializationByteNum(
-			ObservationVectorDimensionNum,
-			1,
-			NetworkSettings.CriticHiddenLayerSize,
-			NetworkSettings.CriticLayerNum);
-
 		// Write Config
 
 		TSharedRef<FJsonObject> ConfigObject = MakeShared<FJsonObject>();
@@ -678,22 +652,19 @@ namespace UE::Learning
 		ConfigObject->SetStringField(TEXT("TrainerType"), TEXT("Network"));
 		ConfigObject->SetStringField(TEXT("TimeStamp"), *FDateTime::Now().ToFormattedString(TEXT("%Y-%m-%d_%H-%M-%S")));
 
+		ConfigObject->SetStringField(TEXT("PolicyNetworkClass"), PolicyNetwork.GetPythonClassName());
+		ConfigObject->SetStringField(TEXT("CriticNetworkClass"), CriticNetwork.GetPythonClassName());
+
 		ConfigObject->SetNumberField(TEXT("ObservationVectorDimensionNum"), ObservationVectorDimensionNum);
 		ConfigObject->SetNumberField(TEXT("ActionVectorDimensionNum"), ActionVectorDimensionNum);
 		ConfigObject->SetNumberField(TEXT("MaxEpisodeNum"), ReplayBuffer.GetMaxEpisodeNum());
 		ConfigObject->SetNumberField(TEXT("MaxStepNum"), ReplayBuffer.GetMaxStepNum());
 
-		ConfigObject->SetNumberField(TEXT("PolicyNetworkByteNum"), TotalPolicyByteNum);
-		ConfigObject->SetNumberField(TEXT("PolicyHiddenUnitNum"), NetworkSettings.PolicyHiddenLayerSize);
-		ConfigObject->SetNumberField(TEXT("PolicyLayerNum"), NetworkSettings.PolicyLayerNum);
-		ConfigObject->SetStringField(TEXT("PolicyActivationFunction"), GetActivationFunctionString(NetworkSettings.PolicyActivationFunction));
+		ConfigObject->SetNumberField(TEXT("PolicyNetworkByteNum"), PolicyNetwork.GetSerializationByteNum());
 		ConfigObject->SetNumberField(TEXT("PolicyActionNoiseMin"), NetworkSettings.PolicyActionNoiseMin);
 		ConfigObject->SetNumberField(TEXT("PolicyActionNoiseMax"), NetworkSettings.PolicyActionNoiseMax);
 
-		ConfigObject->SetNumberField(TEXT("CriticNetworkByteNum"), TotalCriticByteNum);
-		ConfigObject->SetNumberField(TEXT("CriticHiddenUnitNum"), NetworkSettings.CriticHiddenLayerSize);
-		ConfigObject->SetNumberField(TEXT("CriticLayerNum"), NetworkSettings.CriticLayerNum);
-		ConfigObject->SetStringField(TEXT("CriticActivationFunction"), GetActivationFunctionString(NetworkSettings.CriticActivationFunction));
+		ConfigObject->SetNumberField(TEXT("CriticNetworkByteNum"), CriticNetwork.GetSerializationByteNum());
 
 		ConfigObject->SetNumberField(TEXT("IterationNum"), TrainingSettings.IterationNum);
 		ConfigObject->SetNumberField(TEXT("LearningRatePolicy"), TrainingSettings.LearningRatePolicy);
@@ -717,7 +688,6 @@ namespace UE::Learning
 
 		ConfigObject->SetBoolField(TEXT("UseInitialPolicyNetwork"), (bool)(TrainerFlags & EPPOTrainerFlags::UseInitialPolicyNetwork));
 		ConfigObject->SetBoolField(TEXT("UseInitialCriticNetwork"), (bool)(TrainerFlags & EPPOTrainerFlags::UseInitialCriticNetwork));
-		ConfigObject->SetBoolField(TEXT("SynchronizeCriticNetwork"), (bool)(TrainerFlags & EPPOTrainerFlags::SynchronizeCriticNetwork));
 
 		FString JsonString;
 		TSharedRef<TJsonWriter<>> JsonWriter = TJsonWriterFactory<>::Create(&JsonString, 0);
@@ -725,12 +695,8 @@ namespace UE::Learning
 
 		// Allocate buffer to receive network data in
 
-		PolicyNetworkBuffer.SetNumUninitialized({ TotalPolicyByteNum });
-
-		if ((bool)(TrainerFlags & EPPOTrainerFlags::SynchronizeCriticNetwork) || (bool)(TrainerFlags & EPPOTrainerFlags::UseInitialCriticNetwork))
-		{
-			CriticNetworkBuffer.SetNumUninitialized({ TotalCriticByteNum });
-		}
+		PolicyNetworkBuffer.SetNumUninitialized({ PolicyNetwork.GetSerializationByteNum() });
+		CriticNetworkBuffer.SetNumUninitialized({ CriticNetwork.GetSerializationByteNum() });
 
 		// Create Socket
 
@@ -793,7 +759,7 @@ namespace UE::Learning
 	}
 
 	ETrainerResponse FSocketPPOTrainer::RecvPolicy(
-		FNeuralNetwork& OutNetwork,
+		INeuralNetwork& OutNetwork,
 		const float Timeout,
 		FRWLock* NetworkLock,
 		const ELogSetting LogSettings)
@@ -802,7 +768,7 @@ namespace UE::Learning
 	}
 
 	ETrainerResponse FSocketPPOTrainer::RecvCritic(
-		FNeuralNetwork& OutNetwork,
+		INeuralNetwork& OutNetwork,
 		const float Timeout,
 		FRWLock* NetworkLock,
 		const ELogSetting LogSettings)
@@ -814,7 +780,7 @@ namespace UE::Learning
 	}
 
 	ETrainerResponse FSocketPPOTrainer::SendPolicy(
-		const FNeuralNetwork& Network,
+		const INeuralNetwork& Network,
 		const float Timeout,
 		FRWLock* NetworkLock,
 		const ELogSetting LogSettings)
@@ -823,7 +789,7 @@ namespace UE::Learning
 	}
 
 	ETrainerResponse FSocketPPOTrainer::SendCritic(
-		const FNeuralNetwork& Network,
+		const INeuralNetwork& Network,
 		const float Timeout,
 		FRWLock* NetworkLock,
 		const ELogSetting LogSettings)
@@ -849,8 +815,8 @@ namespace UE::Learning
 			FReplayBuffer& ReplayBuffer,
 			FEpisodeBuffer& EpisodeBuffer,
 			FResetInstanceBuffer& ResetBuffer,
-			FNeuralNetwork& PolicyNetwork,
-			FNeuralNetwork* CriticNetwork,
+			INeuralNetwork& PolicyNetwork,
+			INeuralNetwork& CriticNetwork,
 			TLearningArrayView<2, float> ObservationVectorBuffer,
 			TLearningArrayView<2, float> ActionVectorBuffer,
 			TLearningArrayView<1, float> RewardBuffer,
@@ -876,29 +842,47 @@ namespace UE::Learning
 
 			ETrainerResponse Response = ETrainerResponse::Success;
 
-			if ((bool)(TrainerFlags & EPPOTrainerFlags::UseInitialPolicyNetwork))
+			// Send initial Policy
+
+			if (LogSettings != ELogSetting::Silent)
 			{
-				// Send initial Policy
+				UE_LOG(LogLearning, Display, TEXT("Sending initial Policy..."));
+			}
+
+			Response = Trainer.SendPolicy(PolicyNetwork, 20.0f, PolicyNetworkLock);
 
+			if (Response != ETrainerResponse::Success)
+			{
 				if (LogSettings != ELogSetting::Silent)
 				{
-					UE_LOG(LogLearning, Display, TEXT("Sending initial Policy..."));
+					UE_LOG(LogLearning, Error, TEXT("Error sending initial policy to trainer: %s. Check log for errors."), Trainer::GetResponseString(Response));
 				}
 
-				Response = Trainer.SendPolicy(PolicyNetwork, 20.0f, PolicyNetworkLock);
+				Trainer.Terminate();
+				return Response;
+			}
 
-				if (Response != ETrainerResponse::Success)
-				{
-					if (LogSettings != ELogSetting::Silent)
-					{
-						UE_LOG(LogLearning, Error, TEXT("Error sending initial policy to trainer: %s. Check log for errors."), Trainer::GetResponseString(Response));
-					}
+			// Send initial Critic
 
-					Trainer.Terminate();
-					return Response;
+			if (LogSettings != ELogSetting::Silent)
+			{
+				UE_LOG(LogLearning, Display, TEXT("Sending initial Critic..."));
+			}
+
+			Response = Trainer.SendCritic(CriticNetwork, 20.0f, CriticNetworkLock);
+
+			if (Response != ETrainerResponse::Success)
+			{
+				if (LogSettings != ELogSetting::Silent)
+				{
+					UE_LOG(LogLearning, Error, TEXT("Error sending initial critic to trainer: %s. Check log for errors."), Trainer::GetResponseString(Response));
 				}
+
+				Trainer.Terminate();
+				return Response;
 			}
-			else
+
+			if (!(bool)(TrainerFlags & EPPOTrainerFlags::UseInitialPolicyNetwork))
 			{
 				// Receive initial Policy
 
@@ -926,34 +910,8 @@ namespace UE::Learning
 				}
 			}
 
-			if ((bool)(TrainerFlags & EPPOTrainerFlags::UseInitialCriticNetwork))
+			if (!(bool)(TrainerFlags & EPPOTrainerFlags::UseInitialCriticNetwork))
 			{
-				UE_LEARNING_CHECKF(CriticNetwork, TEXT("Valid Critic Network must be provided if you want to send it"));
-
-				// Send initial Critic
-
-				if (LogSettings != ELogSetting::Silent)
-				{
-					UE_LOG(LogLearning, Display, TEXT("Sending initial Critic..."));
-				}
-
-				Response = Trainer.SendCritic(*CriticNetwork, 20.0f, CriticNetworkLock);
-
-				if (Response != ETrainerResponse::Success)
-				{
-					if (LogSettings != ELogSetting::Silent)
-					{
-						UE_LOG(LogLearning, Error, TEXT("Error sending initial critic to trainer: %s. Check log for errors."), Trainer::GetResponseString(Response));
-					}
-
-					Trainer.Terminate();
-					return Response;
-				}
-			}
-			else if ((bool)(TrainerFlags & EPPOTrainerFlags::SynchronizeCriticNetwork))
-			{
-				UE_LEARNING_CHECKF(CriticNetwork, TEXT("Valid Critic Network must be provided if you want to sync it"));
-
 				// Receive initial Critic
 
 				if (LogSettings != ELogSetting::Silent)
@@ -961,7 +919,7 @@ namespace UE::Learning
 					UE_LOG(LogLearning, Display, TEXT("Receiving initial Critic..."));
 				}
 
-				Response = Trainer.RecvCritic(*CriticNetwork, 20.0f, CriticNetworkLock);
+				Response = Trainer.RecvCritic(CriticNetwork, 20.0f, CriticNetworkLock);
 
 				if (Response != ETrainerResponse::Success)
 				{
@@ -1066,25 +1024,20 @@ namespace UE::Learning
 					*bPolicyNetworkUpdatedSignal = true;
 				}
 
-				if ((bool)(TrainerFlags & EPPOTrainerFlags::SynchronizeCriticNetwork))
-				{
-					UE_LEARNING_CHECKF(CriticNetwork, TEXT("Valid Critic Network must be provided if you want to sync it"));
-
-					Response = Trainer.RecvCritic(*CriticNetwork, 10.0f, CriticNetworkLock);
+				Response = Trainer.RecvCritic(CriticNetwork, 10.0f, CriticNetworkLock);
 
-					if (Response != ETrainerResponse::Success)
+				if (Response != ETrainerResponse::Success)
+				{
+					if (LogSettings != ELogSetting::Silent)
 					{
-						if (LogSettings != ELogSetting::Silent)
-						{
-							UE_LOG(LogLearning, Error, TEXT("Error receiving critic from trainer: %s. Check log for errors."), Trainer::GetResponseString(Response));
-						}
-						break;
+						UE_LOG(LogLearning, Error, TEXT("Error receiving critic from trainer: %s. Check log for errors."), Trainer::GetResponseString(Response));
 					}
+					break;
+				}
 
-					if (bCriticNetworkUpdatedSignal)
-					{
-						*bCriticNetworkUpdatedSignal = true;
-					}
+				if (bCriticNetworkUpdatedSignal)
+				{
+					*bCriticNetworkUpdatedSignal = true;
 				}
 			}
 
diff --git a/Engine/Plugins/Experimental/LearningAgents/Source/LearningTraining/Private/LearningSharedMemoryTraining.cpp b/Engine/Plugins/Experimental/LearningAgents/Source/LearningTraining/Private/LearningSharedMemoryTraining.cpp
--- a/Engine/Plugins/Experimental/LearningAgents/Source/LearningTraining/Private/LearningSharedMemoryTraining.cpp
+++ b/Engine/Plugins/Experimental/LearningAgents/Source/LearningTraining/Private/LearningSharedMemoryTraining.cpp
@@ -42,7 +42,7 @@ namespace UE::Learning::SharedMemoryTraining
 
 	ETrainerResponse RecvPolicy(
 		TLearningArrayView<1, volatile int32> Controls,
-		FNeuralNetwork& OutNetwork,
+		INeuralNetwork& OutNetwork,
 		const TLearningArrayView<1, const uint8> Policy,
 		const float Timeout,
 		FRWLock* NetworkLock,
@@ -77,21 +77,27 @@ namespace UE::Learning::SharedMemoryTraining
 		}
 
 		// Read the policy
+		bool bSuccess = false;
 		{
 			FScopeNullableWriteLock ScopeLock(NetworkLock);
 			int32 Offset = 0;
-			OutNetwork.DeserializeFromBytes(Offset, Policy);
+			bSuccess = OutNetwork.DeserializeFromBytes(Offset, Policy);
+
+			if (!bSuccess)
+			{
+				UE_LOG(LogLearning, Error, TEXT("Error receiving Policy network. Format invalid."));
+			}
 		}
 
 		// Confirm we have read the policy
 		Controls[(uint8)EControls::PolicySignal] = false;
 
-		return ETrainerResponse::Success;
+		return bSuccess ? ETrainerResponse::Success : ETrainerResponse::Unexpected;
 	}
 
 	ETrainerResponse RecvCritic(
 		TLearningArrayView<1, volatile int32> Controls,
-		FNeuralNetwork& OutNetwork,
+		INeuralNetwork& OutNetwork,
 		const TLearningArrayView<1, const uint8> Critic,
 		const float Timeout,
 		FRWLock* NetworkLock,
@@ -118,22 +124,28 @@ namespace UE::Learning::SharedMemoryTraining
 		}
 
 		// Read the critic
+		bool bSuccess = false;
 		{
 			FScopeNullableWriteLock ScopeLock(NetworkLock);
 			int32 Offset = 0;
-			OutNetwork.DeserializeFromBytes(Offset, Critic);
+			bSuccess = OutNetwork.DeserializeFromBytes(Offset, Critic);
+
+			if (!bSuccess)
+			{
+				UE_LOG(LogLearning, Error, TEXT("Error receiving Critic network. Format invalid."));
+			}
 		}
 
 		// Confirm we have read the critic
 		Controls[(uint8)EControls::CriticSignal] = false;
 
-		return ETrainerResponse::Success;
+		return bSuccess ? ETrainerResponse::Success : ETrainerResponse::Unexpected;
 	}
 
 	ETrainerResponse SendPolicy(
 		TLearningArrayView<1, volatile int32> Controls,
 		TLearningArrayView<1, uint8> Policy,
-		const FNeuralNetwork& Network,
+		const INeuralNetwork& Network,
 		const float Timeout,
 		FRWLock* NetworkLock,
 		const ELogSetting LogSettings)
@@ -174,7 +186,7 @@ namespace UE::Learning::SharedMemoryTraining
 	ETrainerResponse SendCritic(
 		TLearningArrayView<1, volatile int32> Controls,
 		TLearningArrayView<1, uint8> Critic,
-		const FNeuralNetwork& Network,
+		const INeuralNetwork& Network,
 		const float Timeout,
 		FRWLock* NetworkLock,
 		const ELogSetting LogSettings)
diff --git a/Engine/Plugins/Experimental/LearningAgents/Source/LearningTraining/Private/LearningSharedMemoryTraining.h b/Engine/Plugins/Experimental/LearningAgents/Source/LearningTraining/Private/LearningSharedMemoryTraining.h
--- a/Engine/Plugins/Experimental/LearningAgents/Source/LearningTraining/Private/LearningSharedMemoryTraining.h
+++ b/Engine/Plugins/Experimental/LearningAgents/Source/LearningTraining/Private/LearningSharedMemoryTraining.h
@@ -7,7 +7,7 @@
 namespace UE::Learning
 {
 	enum class ECompletionMode : uint8;
-	struct FNeuralNetwork;
+	struct INeuralNetwork;
 	struct FReplayBuffer;
 
 	namespace SharedMemoryTraining
@@ -16,7 +16,7 @@ namespace UE::Learning
 
 		LEARNINGTRAINING_API ETrainerResponse RecvPolicy(
 			TLearningArrayView<1, volatile int32> Controls,
-			FNeuralNetwork& OutNetwork,
+			INeuralNetwork& OutNetwork,
 			const TLearningArrayView<1, const uint8> Policy,
 			const float Timeout = Trainer::DefaultTimeout,
 			FRWLock* NetworkLock = nullptr,
@@ -24,7 +24,7 @@ namespace UE::Learning
 
 		LEARNINGTRAINING_API ETrainerResponse RecvCritic(
 			TLearningArrayView<1, volatile int32> Controls,
-			FNeuralNetwork& OutNetwork,
+			INeuralNetwork& OutNetwork,
 			const TLearningArrayView<1, const uint8> Critic,
 			const float Timeout = Trainer::DefaultTimeout,
 			FRWLock* NetworkLock = nullptr,
@@ -38,7 +38,7 @@ namespace UE::Learning
 		LEARNINGTRAINING_API ETrainerResponse SendPolicy(
 			TLearningArrayView<1, volatile int32> Controls,
 			TLearningArrayView<1, uint8> Policy,
-			const FNeuralNetwork& Network,
+			const INeuralNetwork& Network,
 			const float Timeout = Trainer::DefaultTimeout,
 			FRWLock* NetworkLock = nullptr,
 			const ELogSetting LogSettings = Trainer::DefaultLogSettings);
@@ -46,7 +46,7 @@ namespace UE::Learning
 		LEARNINGTRAINING_API ETrainerResponse SendCritic(
 			TLearningArrayView<1, volatile int32> Controls,
 			TLearningArrayView<1, uint8> Critic,
-			const FNeuralNetwork& Network,
+			const INeuralNetwork& Network,
 			const float Timeout = Trainer::DefaultTimeout,
 			FRWLock* NetworkLock = nullptr,
 			const ELogSetting LogSettings = Trainer::DefaultLogSettings);
diff --git a/Engine/Plugins/Experimental/LearningAgents/Source/LearningTraining/Private/LearningSocketTraining.cpp b/Engine/Plugins/Experimental/LearningAgents/Source/LearningTraining/Private/LearningSocketTraining.cpp
--- a/Engine/Plugins/Experimental/LearningAgents/Source/LearningTraining/Private/LearningSocketTraining.cpp
+++ b/Engine/Plugins/Experimental/LearningAgents/Source/LearningTraining/Private/LearningSocketTraining.cpp
@@ -88,7 +88,7 @@ namespace UE::Learning::SocketTraining
 
 	ETrainerResponse RecvPolicy(
 		FSocket& Socket,
-		FNeuralNetwork& OutNetwork,
+		INeuralNetwork& OutNetwork,
 		TLearningArrayView<1, uint8> OutNetworkBuffer,
 		const float Timeout,
 		FRWLock* NetworkLock,
@@ -116,18 +116,24 @@ namespace UE::Learning::SocketTraining
 		Response = RecvWithTimeout(Socket, OutNetworkBuffer.GetData(), OutNetworkBuffer.Num(), Timeout);
 		if (Response != ETrainerResponse::Success) { return Response; }
 
+		bool bSuccess = false;
 		{
 			FScopeNullableWriteLock ScopeLock(NetworkLock);
 			int32 Offset = 0;
-			OutNetwork.DeserializeFromBytes(Offset, OutNetworkBuffer);
+			bSuccess = OutNetwork.DeserializeFromBytes(Offset, OutNetworkBuffer);
+
+			if (!bSuccess)
+			{
+				UE_LOG(LogLearning, Error, TEXT("Error receiving Policy network. Format invalid."));
+			}
 		}
 
-		return ETrainerResponse::Success;
+		return bSuccess ? ETrainerResponse::Success : ETrainerResponse::Unexpected;
 	}
 
 	ETrainerResponse RecvCritic(
 		FSocket& Socket,
-		FNeuralNetwork& OutNetwork,
+		INeuralNetwork& OutNetwork,
 		TLearningArrayView<1, uint8> OutNetworkBuffer,
 		const float Timeout,
 		FRWLock* NetworkLock,
@@ -150,13 +156,19 @@ namespace UE::Learning::SocketTraining
 		Response = RecvWithTimeout(Socket, OutNetworkBuffer.GetData(), OutNetworkBuffer.Num(), Timeout);
 		if (Response != ETrainerResponse::Success) { return Response; }
 
+		bool bSuccess = false;
 		{
 			FScopeNullableWriteLock ScopeLock(NetworkLock);
 			int32 Offset = 0;
-			OutNetwork.DeserializeFromBytes(Offset, OutNetworkBuffer);
+			bSuccess = OutNetwork.DeserializeFromBytes(Offset, OutNetworkBuffer);
+
+			if (!bSuccess)
+			{
+				UE_LOG(LogLearning, Error, TEXT("Error receiving Critic network. Format invalid."));
+			}
 		}
 
-		return ETrainerResponse::Success;
+		return bSuccess ? ETrainerResponse::Success : ETrainerResponse::Unexpected;
 	}
 
 	ETrainerResponse SendWithTimeout(FSocket& Socket, const uint8* Bytes, const int32 ByteNum, const float Timeout)
@@ -222,7 +234,7 @@ namespace UE::Learning::SocketTraining
 	ETrainerResponse SendPolicy(
 		FSocket& Socket,
 		TLearningArrayView<1, uint8> NetworkBuffer,
-		const FNeuralNetwork& Network,
+		const INeuralNetwork& Network,
 		const float Timeout,
 		FRWLock* NetworkLock,
 		const ELogSetting LogSettings)
@@ -251,7 +263,7 @@ namespace UE::Learning::SocketTraining
 	ETrainerResponse SendCritic(
 		FSocket& Socket,
 		TLearningArrayView<1, uint8> NetworkBuffer,
-		const FNeuralNetwork& Network,
+		const INeuralNetwork& Network,
 		const float Timeout,
 		FRWLock* NetworkLock,
 		const ELogSetting LogSettings)
diff --git a/Engine/Plugins/Experimental/LearningAgents/Source/LearningTraining/Private/LearningSocketTraining.h b/Engine/Plugins/Experimental/LearningAgents/Source/LearningTraining/Private/LearningSocketTraining.h
--- a/Engine/Plugins/Experimental/LearningAgents/Source/LearningTraining/Private/LearningSocketTraining.h
+++ b/Engine/Plugins/Experimental/LearningAgents/Source/LearningTraining/Private/LearningSocketTraining.h
@@ -8,7 +8,7 @@ class FSocket;
 
 namespace UE::Learning
 {
-	struct FNeuralNetwork;
+	struct INeuralNetwork;
 	struct FReplayBuffer;
 
 	namespace SocketTraining
@@ -25,7 +25,7 @@ namespace UE::Learning
 
 		LEARNINGTRAINING_API ETrainerResponse RecvPolicy(
 			FSocket& Socket,
-			FNeuralNetwork& OutNetwork,
+			INeuralNetwork& OutNetwork,
 			TLearningArrayView<1, uint8> OutNetworkBuffer,
 			const float Timeout = Trainer::DefaultTimeout,
 			FRWLock* NetworkLock = nullptr,
@@ -33,7 +33,7 @@ namespace UE::Learning
 
 		LEARNINGTRAINING_API ETrainerResponse RecvCritic(
 			FSocket& Socket,
-			FNeuralNetwork& OutNetwork,
+			INeuralNetwork& OutNetwork,
 			TLearningArrayView<1, uint8> OutNetworkBuffer,
 			const float Timeout = Trainer::DefaultTimeout,
 			FRWLock* NetworkLock = nullptr,
@@ -59,7 +59,7 @@ namespace UE::Learning
 		LEARNINGTRAINING_API ETrainerResponse SendPolicy(
 			FSocket& Socket,
 			TLearningArrayView<1, uint8> NetworkBuffer,
-			const FNeuralNetwork& Network,
+			const INeuralNetwork& Network,
 			const float Timeout = Trainer::DefaultTimeout,
 			FRWLock* NetworkLock = nullptr,
 			const ELogSetting LogSettings = Trainer::DefaultLogSettings);
@@ -67,7 +67,7 @@ namespace UE::Learning
 		LEARNINGTRAINING_API ETrainerResponse SendCritic(
 			FSocket& Socket,
 			TLearningArrayView<1, uint8> NetworkBuffer,
-			const FNeuralNetwork& Network,
+			const INeuralNetwork& Network,
 			const float Timeout = Trainer::DefaultTimeout,
 			FRWLock* NetworkLock = nullptr,
 			const ELogSetting LogSettings = Trainer::DefaultLogSettings);
diff --git a/Engine/Plugins/Experimental/LearningAgents/Source/LearningTraining/Public/LearningImitationTrainer.h b/Engine/Plugins/Experimental/LearningAgents/Source/LearningTraining/Public/LearningImitationTrainer.h
--- a/Engine/Plugins/Experimental/LearningAgents/Source/LearningTraining/Public/LearningImitationTrainer.h
+++ b/Engine/Plugins/Experimental/LearningAgents/Source/LearningTraining/Public/LearningImitationTrainer.h
@@ -6,7 +6,6 @@
 #include "LearningLog.h"
 #include "LearningTrainer.h"
 #include "LearningSharedMemory.h"
-#include "LearningNeuralNetwork.h" // Included for EActivationFunction::ELU
 
 #include "Commandlets/Commandlet.h"
 #include "Templates/SharedPointer.h"
@@ -29,6 +28,8 @@ class LEARNINGTRAINING_API ULearningSocketImitationTrainerServerCommandlet : pub
 
 namespace UE::Learning
 {
+	struct INeuralNetwork;
+
 	/**
 	* Settings for the network used for training. These settings must match the Neural Network
 	* objects passed to ImitationTrainer::Train.
@@ -40,15 +41,6 @@ namespace UE::Learning
 
 		/** Maximum action noise used by the policy */
 		float PolicyActionNoiseMax = 0.25f;
-
-		/** Total layers for policy network including input, hidden, and output layers */
-		int32 PolicyLayerNum = 3;
-
-		/** Number of neurons in each hidden layer of the policy network */
-		int32 PolicyHiddenLayerSize = 128;
-
-		/** Activation function to use on hidden layers of the policy network */
-		EActivationFunction PolicyActivationFunction = EActivationFunction::ELU;
 	};
 
 	struct FImitationTrainerTrainingSettings
@@ -145,7 +137,7 @@ namespace UE::Learning
 		* @returns				Trainer response
 		*/
 		virtual ETrainerResponse RecvPolicy(
-			FNeuralNetwork& OutNetwork,
+			INeuralNetwork& OutNetwork,
 			const float Timeout = Trainer::DefaultTimeout,
 			FRWLock* NetworkLock = nullptr,
 			const ELogSetting LogSettings = Trainer::DefaultLogSettings) = 0;
@@ -160,7 +152,7 @@ namespace UE::Learning
 		* @returns				Trainer response
 		*/
 		virtual ETrainerResponse SendPolicy(
-			const FNeuralNetwork& Network,
+			const INeuralNetwork& Network,
 			const float Timeout = Trainer::DefaultTimeout,
 			FRWLock* NetworkLock = nullptr,
 			const ELogSetting LogSettings = Trainer::DefaultLogSettings) = 0;
@@ -197,6 +189,7 @@ namespace UE::Learning
 		* @param MaxSampleNum			Maximum number of samples in the training data
 		* @param ObservationDimNum		Number of dimensions in the observation vector
 		* @param ActionDimNum			Number of dimensions in the action vector
+		* @param PolicyNetwork			Policy Network to use
 		* @param TrainingSettings		Trainer Training settings
 		* @param NetworkSettings		Trainer Network settings
 		* @param TrainingProcessFlags	Training subprocess flags
@@ -211,6 +204,7 @@ namespace UE::Learning
 			const int32 MaxSampleNum,
 			const int32 ObservationDimNum,
 			const int32 ActionDimNum,
+			const INeuralNetwork& PolicyNetwork,
 			const FImitationTrainerTrainingSettings& TrainingSettings = FImitationTrainerTrainingSettings(),
 			const FImitationTrainerNetworkSettings& NetworkSettings = FImitationTrainerNetworkSettings(),
 			const EImitationTrainerFlags TrainerFlags = EImitationTrainerFlags::None,
@@ -228,13 +222,13 @@ namespace UE::Learning
 		virtual bool HasPolicyOrCompleted() override final;
 
 		virtual ETrainerResponse RecvPolicy(
-			FNeuralNetwork& OutNetwork,
+			INeuralNetwork& OutNetwork,
 			const float Timeout = Trainer::DefaultTimeout,
 			FRWLock* NetworkLock = nullptr,
 			const ELogSetting LogSettings = Trainer::DefaultLogSettings) override final;
 
 		virtual ETrainerResponse SendPolicy(
-			const FNeuralNetwork& Network,
+			const INeuralNetwork& Network,
 			const float Timeout = Trainer::DefaultTimeout,
 			FRWLock* NetworkLock = nullptr,
 			const ELogSetting LogSettings = Trainer::DefaultLogSettings) override final;
@@ -332,9 +326,10 @@ namespace UE::Learning
 		*
 		* @param OutResponse				Response to the initial connection
 		* @param TaskName					Name of the training task - used to help identify the logs, snapshots, and other files generated by training
-		* @param MaxSampleNum			Maximum number of samples in the training data
-		* @param ObservationDimNum		Number of dimensions in the observation vector
-		* @param ActionDimNum			Number of dimensions in the action vector
+		* @param MaxSampleNum				Maximum number of samples in the training data
+		* @param ObservationDimNum			Number of dimensions in the observation vector
+		* @param ActionDimNum				Number of dimensions in the action vector
+		* @param PolicyNetwork				Policy Network to use
 		* @param IpAddress					Server Ip address
 		* @param Port						Server Port
 		* @param Timeout					Timeout to wait in seconds for connection and initial data transfer
@@ -348,6 +343,7 @@ namespace UE::Learning
 			const int32 MaxSampleNum,
 			const int32 ObservationDimNum,
 			const int32 ActionDimNum,
+			const INeuralNetwork& PolicyNetwork,
 			const TCHAR* IpAddress = Trainer::DefaultIp,
 			const uint32 Port = Trainer::DefaultPort,
 			const float Timeout = Trainer::DefaultTimeout,
@@ -366,13 +362,13 @@ namespace UE::Learning
 		virtual bool HasPolicyOrCompleted() override final;
 
 		virtual ETrainerResponse RecvPolicy(
-			FNeuralNetwork& OutNetwork,
+			INeuralNetwork& OutNetwork,
 			const float Timeout = Trainer::DefaultTimeout,
 			FRWLock* NetworkLock = nullptr,
 			const ELogSetting LogSettings = Trainer::DefaultLogSettings) override final;
 
 		virtual ETrainerResponse SendPolicy(
-			const FNeuralNetwork& Network,
+			const INeuralNetwork& Network,
 			const float Timeout = Trainer::DefaultTimeout,
 			FRWLock* NetworkLock = nullptr,
 			const ELogSetting LogSettings = Trainer::DefaultLogSettings) override final;
@@ -407,7 +403,7 @@ namespace UE::Learning
 		*/
 		LEARNINGTRAINING_API ETrainerResponse Train(
 			IImitationTrainer& Trainer,
-			FNeuralNetwork& Network,
+			INeuralNetwork& Network,
 			const TLearningArrayView<2, const float> ObservationVectors,
 			const TLearningArrayView<2, const float> ActionVectors,
 			const EImitationTrainerFlags TrainerFlags = EImitationTrainerFlags::None,
diff --git a/Engine/Plugins/Experimental/LearningAgents/Source/LearningTraining/Public/LearningPPOTrainer.h b/Engine/Plugins/Experimental/LearningAgents/Source/LearningTraining/Public/LearningPPOTrainer.h
--- a/Engine/Plugins/Experimental/LearningAgents/Source/LearningTraining/Public/LearningPPOTrainer.h
+++ b/Engine/Plugins/Experimental/LearningAgents/Source/LearningTraining/Public/LearningPPOTrainer.h
@@ -6,7 +6,6 @@
 #include "LearningLog.h"
 #include "LearningTrainer.h"
 #include "LearningSharedMemory.h"
-#include "LearningNeuralNetwork.h" // Included for EActivationFunction::ELU
 
 #include "Commandlets/Commandlet.h"
 #include "Templates/SharedPointer.h"
@@ -29,6 +28,7 @@ class LEARNINGTRAINING_API ULearningSocketPPOTrainerServerCommandlet : public UC
 
 namespace UE::Learning
 {
+	struct INeuralNetwork;
 	struct FReplayBuffer;
 	struct FResetInstanceBuffer;
 	struct FEpisodeBuffer;
@@ -46,24 +46,6 @@ namespace UE::Learning
 
 		/** Maximum action noise used by the policy */
 		float PolicyActionNoiseMax = 0.25f;
-
-		/** Total layers for policy network including input, hidden, and output layers */
-		int32 PolicyLayerNum = 3;
-
-		/** Number of neurons in each hidden layer of the policy network */
-		int32 PolicyHiddenLayerSize = 128;
-
-		/** Activation function to use on hidden layers of the policy network */
-		EActivationFunction PolicyActivationFunction = EActivationFunction::ELU;
-
-		/** Total layers for critic network including input, hidden, and output layers */
-		int32 CriticLayerNum = 4;
-
-		/** Number of neurons in each hidden layer of the critic network */
-		int32 CriticHiddenLayerSize = 256;
-
-		/** Activation function to use on hidden layers of the critic network */
-		EActivationFunction CriticActivationFunction = EActivationFunction::ELU;
 	};
 
 	/**
@@ -171,10 +153,6 @@ namespace UE::Learning
 		// the start of training. Use this if you want to start from a network which has already been trained 
 		// such as via Imitation Training.
 		UseInitialCriticNetwork = 1 << 1,
-
-		// If the critic network should be synchronized between the training process 
-		// and experience gathering process during training
-		SynchronizeCriticNetwork = 1 << 2,
 	};
 	ENUM_CLASS_FLAGS(EPPOTrainerFlags)
 
@@ -208,7 +186,7 @@ namespace UE::Learning
 		* @returns				Trainer response
 		*/
 		virtual ETrainerResponse RecvPolicy(
-			FNeuralNetwork& OutNetwork,
+			INeuralNetwork& OutNetwork,
 			const float Timeout = Trainer::DefaultTimeout,
 			FRWLock* NetworkLock = nullptr,
 			const ELogSetting LogSettings = Trainer::DefaultLogSettings) = 0;
@@ -223,7 +201,7 @@ namespace UE::Learning
 		* @returns				Trainer response
 		*/
 		virtual ETrainerResponse RecvCritic(
-			FNeuralNetwork& OutNetwork,
+			INeuralNetwork& OutNetwork,
 			const float Timeout = Trainer::DefaultTimeout,
 			FRWLock* NetworkLock = nullptr,
 			const ELogSetting LogSettings = Trainer::DefaultLogSettings) = 0;
@@ -243,7 +221,7 @@ namespace UE::Learning
 		* @returns				Trainer response
 		*/
 		virtual ETrainerResponse SendPolicy(
-			const FNeuralNetwork& Network,
+			const INeuralNetwork& Network,
 			const float Timeout = Trainer::DefaultTimeout,
 			FRWLock* NetworkLock = nullptr,
 			const ELogSetting LogSettings = Trainer::DefaultLogSettings) = 0;
@@ -258,7 +236,7 @@ namespace UE::Learning
 		* @returns				Trainer response
 		*/
 		virtual ETrainerResponse SendCritic(
-			const FNeuralNetwork& Network,
+			const INeuralNetwork& Network,
 			const float Timeout = Trainer::DefaultTimeout,
 			FRWLock* NetworkLock = nullptr,
 			const ELogSetting LogSettings = Trainer::DefaultLogSettings) = 0;
@@ -294,6 +272,8 @@ namespace UE::Learning
 		* @param PythonContentPath			Path to the Python Content folder provided by the Learning plugin
 		* @param IntermediatePath			Path to the intermediate folder to write temporary files, logs, and snapshots to
 		* @param ReplayBuffer				Replay buffer used to collect experience
+		* @param PolicyNetwork				Policy Network to use
+		* @param CriticNetwork				Critic Network to use
 		* @param TrainingSettings			Trainer Training settings
 		* @param NetworkSettings			Trainer Network settings
 		* @param TrainerFlags				Flags for the trainer
@@ -328,6 +308,8 @@ namespace UE::Learning
 			const FString& PythonContentPath,
 			const FString& IntermediatePath,
 			const FReplayBuffer& ReplayBuffer,
+			const INeuralNetwork& PolicyNetwork,
+			const INeuralNetwork& CriticNetwork,
 			const FPPOTrainerTrainingSettings& TrainingSettings = FPPOTrainerTrainingSettings(),
 			const FPPOTrainerNetworkSettings& NetworkSettings = FPPOTrainerNetworkSettings(),
 			const EPPOTrainerFlags TrainerFlags = EPPOTrainerFlags::None,
@@ -343,13 +325,13 @@ namespace UE::Learning
 		virtual ETrainerResponse Wait(const float Timeout = Trainer::DefaultTimeout) override final;
 
 		virtual ETrainerResponse RecvPolicy(
-			FNeuralNetwork& OutNetwork,
+			INeuralNetwork& OutNetwork,
 			const float Timeout = Trainer::DefaultTimeout,
 			FRWLock* NetworkLock = nullptr,
 			const ELogSetting LogSettings = Trainer::DefaultLogSettings) override final;
 
 		virtual ETrainerResponse RecvCritic(
-			FNeuralNetwork& OutNetwork,
+			INeuralNetwork& OutNetwork,
 			const float Timeout = Trainer::DefaultTimeout,
 			FRWLock* NetworkLock = nullptr,
 			const ELogSetting LogSettings = Trainer::DefaultLogSettings) override final;
@@ -357,13 +339,13 @@ namespace UE::Learning
 		virtual ETrainerResponse SendStop(const float Timeout = Trainer::DefaultTimeout) override final;
 
 		virtual ETrainerResponse SendPolicy(
-			const FNeuralNetwork& Network,
+			const INeuralNetwork& Network,
 			const float Timeout = Trainer::DefaultTimeout,
 			FRWLock* NetworkLock = nullptr,
 			const ELogSetting LogSettings = Trainer::DefaultLogSettings) override final;
 
 		virtual ETrainerResponse SendCritic(
-			const FNeuralNetwork& Network,
+			const INeuralNetwork& Network,
 			const float Timeout = Trainer::DefaultTimeout,
 			FRWLock* NetworkLock = nullptr,
 			const ELogSetting LogSettings = Trainer::DefaultLogSettings) override final;
@@ -479,6 +461,8 @@ namespace UE::Learning
 		* @param OutResponse				Response to the initial connection
 		* @param TaskName					Name of the training task - used to help identify the logs, snapshots, and other files generated by training
 		* @param ReplayBuffer				Replay buffer used to collect experience
+		* @param PolicyNetwork				Policy Network to use
+		* @param CriticNetwork				Critic Network to use
 		* @param IpAddress					Server Ip address
 		* @param Port						Server Port
 		* @param Timeout					Timeout to wait in seconds for connection and initial data transfer
@@ -490,6 +474,8 @@ namespace UE::Learning
 			ETrainerResponse& OutResponse,
 			const FString& TaskName,
 			const FReplayBuffer& ReplayBuffer,
+			const INeuralNetwork& PolicyNetwork,
+			const INeuralNetwork& CriticNetwork,
 			const TCHAR* IpAddress = Trainer::DefaultIp,
 			const uint32 Port = Trainer::DefaultPort,
 			const float Timeout = Trainer::DefaultTimeout,
@@ -504,13 +490,13 @@ namespace UE::Learning
 		virtual ETrainerResponse Wait(const float Timeout = Trainer::DefaultTimeout) override final;
 
 		virtual ETrainerResponse RecvPolicy(
-			FNeuralNetwork& OutNetwork,
+			INeuralNetwork& OutNetwork,
 			const float Timeout = Trainer::DefaultTimeout,
 			FRWLock* NetworkLock = nullptr,
 			const ELogSetting LogSettings = Trainer::DefaultLogSettings) override final;
 
 		virtual ETrainerResponse RecvCritic(
-			FNeuralNetwork& OutNetwork,
+			INeuralNetwork& OutNetwork,
 			const float Timeout = Trainer::DefaultTimeout,
 			FRWLock* NetworkLock = nullptr,
 			const ELogSetting LogSettings = Trainer::DefaultLogSettings) override final;
@@ -518,13 +504,13 @@ namespace UE::Learning
 		virtual ETrainerResponse SendStop(const float Timeout = Trainer::DefaultTimeout) override final;
 
 		virtual ETrainerResponse SendPolicy(
-			const FNeuralNetwork& Network,
+			const INeuralNetwork& Network,
 			const float Timeout = Trainer::DefaultTimeout,
 			FRWLock* NetworkLock = nullptr,
 			const ELogSetting LogSettings = Trainer::DefaultLogSettings) override final;
 
 		virtual ETrainerResponse SendCritic(
-			const FNeuralNetwork& Network,
+			const INeuralNetwork& Network,
 			const float Timeout = Trainer::DefaultTimeout,
 			FRWLock* NetworkLock = nullptr,
 			const ELogSetting LogSettings = Trainer::DefaultLogSettings) override final;
@@ -580,8 +566,8 @@ namespace UE::Learning
 			FReplayBuffer& ReplayBuffer,
 			FEpisodeBuffer& EpisodeBuffer,
 			FResetInstanceBuffer& ResetBuffer,
-			FNeuralNetwork& PolicyNetwork,
-			FNeuralNetwork* CriticNetwork,
+			INeuralNetwork& PolicyNetwork,
+			INeuralNetwork& CriticNetwork,
 			TLearningArrayView<2, float> ObservationVectorBuffer,
 			TLearningArrayView<2, float> ActionVectorBuffer,
 			TLearningArrayView<1, float> RewardBuffer,
