LearningAgents: Fixed errors with non-ispc build.

diff --git a/Engine/Plugins/Experimental/LearningAgents/Source/Learning/Private/LearningNeuralNetwork.cpp b/Engine/Plugins/Experimental/LearningAgents/Source/Learning/Private/LearningNeuralNetwork.cpp
--- a/Engine/Plugins/Experimental/LearningAgents/Source/Learning/Private/LearningNeuralNetwork.cpp
+++ b/Engine/Plugins/Experimental/LearningAgents/Source/Learning/Private/LearningNeuralNetwork.cpp
@@ -35,7 +35,7 @@ namespace UE::Learning
 				RowNum,
 				ColNum);
 #else
-			for (const int32 BatchIdx = BatchOffset; BatchIdx < BatchOffset + BatchNum; BatchIdx++)
+			for (int32 BatchIdx = BatchOffset; BatchIdx < BatchOffset + BatchNum; BatchIdx++)
 			{
 				for (int32 ColIdx = 0; ColIdx < ColNum; ColIdx++)
 				{
diff --git a/Engine/Plugins/Experimental/LearningAgents/Source/Learning/Private/LearningNeuralNetworkObject.cpp b/Engine/Plugins/Experimental/LearningAgents/Source/Learning/Private/LearningNeuralNetworkObject.cpp
--- a/Engine/Plugins/Experimental/LearningAgents/Source/Learning/Private/LearningNeuralNetworkObject.cpp
+++ b/Engine/Plugins/Experimental/LearningAgents/Source/Learning/Private/LearningNeuralNetworkObject.cpp
@@ -10,7 +10,7 @@
 
 namespace UE::Learning
 {
-	namespace NeuralNetworkPolicyFunction
+	namespace NeuralNetworkPolicyFunction::Private
 	{
 		static inline float Sigmoid(const float X)
 		{
