LearningAgents: Missing Super::PostLoad() and ReloadFromFileData on ULearningAgentsNeuralNetworkData

diff --git a/Engine/Plugins/Experimental/LearningAgents/Source/LearningAgents/Private/LearningAgentsNeuralNetworkData.cpp b/Engine/Plugins/Experimental/LearningAgents/Source/LearningAgents/Private/LearningAgentsNeuralNetworkData.cpp
--- a/Engine/Plugins/Experimental/LearningAgents/Source/LearningAgents/Private/LearningAgentsNeuralNetworkData.cpp
+++ b/Engine/Plugins/Experimental/LearningAgents/Source/LearningAgents/Private/LearningAgentsNeuralNetworkData.cpp
@@ -220,9 +220,12 @@ namespace UE::Learning::Agents
 
 void ULearningAgentsNeuralNetworkData::PostLoad()
 {
+	Super::PostLoad();
+
 	if (FileData.Num() > 0)
 	{
 		Network = MakeShared<UE::Learning::Agents::FNeuralNetwork>(*this);
+		Network->ReloadFromFileData();
 	}
 }
 
