Learning Agents - Fix Enum Observations having an extra value in the one-hot encoding

#rnx

diff --git a/Engine/Plugins/Experimental/LearningAgents/Source/LearningAgents/Private/LearningAgentsObservations.cpp b/Engine/Plugins/Experimental/LearningAgents/Source/LearningAgents/Private/LearningAgentsObservations.cpp
--- a/Engine/Plugins/Experimental/LearningAgents/Source/LearningAgents/Private/LearningAgentsObservations.cpp
+++ b/Engine/Plugins/Experimental/LearningAgents/Source/LearningAgents/Private/LearningAgentsObservations.cpp
@@ -332,13 +332,14 @@ UEnumObservation* UEnumObservation::AddEnumObservation(ULearningAgentsInteractor
 		return nullptr;
 	}
 
-	if (EnumType->NumEnums() < 1)
+	int32 NumEnums = EnumType->NumEnums() - 1;
+	if (NumEnums < 1)
 	{
 		UE_LOG(LogLearning, Error, TEXT("AddEnumObservation: Enum requires at least one entry to be used as an observation."));
 		return nullptr;
 	}
 
-	UEnumObservation* Observation = UE::Learning::Agents::Observations::Private::AddObservation<UEnumObservation, UE::Learning::FFloatFeature>(InInteractor, Name, TEXT("AddEnumObservation"), EnumType->NumEnums());
+	UEnumObservation* Observation = UE::Learning::Agents::Observations::Private::AddObservation<UEnumObservation, UE::Learning::FFloatFeature>(InInteractor, Name, TEXT("AddEnumObservation"), NumEnums);
 	if (Observation) { Observation->Enum = EnumType; }
 	
 	return Observation;
@@ -421,7 +422,8 @@ UEnumArrayObservation* UEnumArrayObservation::AddEnumArrayObservation(ULearningA
 		return nullptr;
 	}
 
-	if (EnumType->NumEnums() < 1)
+	int32 NumEnums = EnumType->NumEnums() - 1;
+	if (NumEnums < 1)
 	{
 		UE_LOG(LogLearning, Error, TEXT("AddEnumArrayObservation: Enum requires at least one entry to be used as an observation."));
 		return nullptr;
@@ -433,7 +435,7 @@ UEnumArrayObservation* UEnumArrayObservation::AddEnumArrayObservation(ULearningA
 		return nullptr;
 	}
 
-	UEnumArrayObservation* Observation = UE::Learning::Agents::Observations::Private::AddObservation<UEnumArrayObservation, UE::Learning::FFloatFeature>(InInteractor, Name, TEXT("AddEnumArrayObservation"), EnumNum * EnumType->NumEnums());
+	UEnumArrayObservation* Observation = UE::Learning::Agents::Observations::Private::AddObservation<UEnumArrayObservation, UE::Learning::FFloatFeature>(InInteractor, Name, TEXT("AddEnumArrayObservation"), EnumNum * NumEnums);
 	if (Observation) { Observation->Enum = EnumType; }
 
 	return Observation;
