LearningAgents: make policy and critic BP spawnable

diff --git a/Engine/Plugins/Experimental/LearningAgents/Source/LearningAgents/Public/LearningAgentsCritic.h b/Engine/Plugins/Experimental/LearningAgents/Source/LearningAgents/Public/LearningAgentsCritic.h
--- a/Engine/Plugins/Experimental/LearningAgents/Source/LearningAgents/Public/LearningAgentsCritic.h
+++ b/Engine/Plugins/Experimental/LearningAgents/Source/LearningAgents/Public/LearningAgentsCritic.h
@@ -44,7 +44,7 @@ public:
 };
 
 /** A critic used by some algorithms for training the managed agents. */
-UCLASS(BlueprintType, Blueprintable)
+UCLASS(BlueprintType, Blueprintable, meta = (BlueprintSpawnableComponent))
 class LEARNINGAGENTS_API ULearningAgentsCritic : public ULearningAgentsManagerComponent
 {
 	GENERATED_BODY()
diff --git a/Engine/Plugins/Experimental/LearningAgents/Source/LearningAgents/Public/LearningAgentsPolicy.h b/Engine/Plugins/Experimental/LearningAgents/Source/LearningAgents/Public/LearningAgentsPolicy.h
--- a/Engine/Plugins/Experimental/LearningAgents/Source/LearningAgents/Public/LearningAgentsPolicy.h
+++ b/Engine/Plugins/Experimental/LearningAgents/Source/LearningAgents/Public/LearningAgentsPolicy.h
@@ -63,7 +63,7 @@ public:
 };
 
 /** A policy that maps from observations to actions for the managed agents. */
-UCLASS(BlueprintType, Blueprintable)
+UCLASS(BlueprintType, Blueprintable, meta = (BlueprintSpawnableComponent))
 class LEARNINGAGENTS_API ULearningAgentsPolicy : public ULearningAgentsManagerComponent
 {
 	GENERATED_BODY()
