Fix C4702 warnings
#rnx

diff --git a/Engine/Plugins/Experimental/CineCameraRigs/Source/CineCameraRigs/Private/CineCameraRigRail.cpp b/Engine/Plugins/Experimental/CineCameraRigs/Source/CineCameraRigs/Private/CineCameraRigRail.cpp
--- a/Engine/Plugins/Experimental/CineCameraRigs/Source/CineCameraRigs/Private/CineCameraRigRail.cpp
+++ b/Engine/Plugins/Experimental/CineCameraRigs/Source/CineCameraRigs/Private/CineCameraRigRail.cpp
@@ -327,7 +327,6 @@ FVector ACineCameraRigRail::GetVelocityAtPosition(const float InPosition, const
 		FVector const P1 = CineSplineComponent->GetLocationAtDistanceAlongSpline(t1 * SplineLen, ESplineCoordinateSpace::World);
 		return (P1 - P0) / (t1 - t0) * TimeMultiplier;
 	}
-	return FVector::ZeroVector;
 }
 
 void ACineCameraRigRail::OnSplineEdited()
diff --git a/Engine/Plugins/Experimental/LearningAgents/Source/LearningTraining/Private/LearningImitationTrainer.cpp b/Engine/Plugins/Experimental/LearningAgents/Source/LearningTraining/Private/LearningImitationTrainer.cpp
--- a/Engine/Plugins/Experimental/LearningAgents/Source/LearningTraining/Private/LearningImitationTrainer.cpp
+++ b/Engine/Plugins/Experimental/LearningAgents/Source/LearningTraining/Private/LearningImitationTrainer.cpp
@@ -31,6 +31,7 @@ int32 ULearningSocketImitationTrainerServerCommandlet::Main(const FString& Comma
 {
 	UE_LOG(LogLearning, Display, TEXT("Running Imitation Training Server Commandlet..."));
 
+#if WITH_EDITOR
 	TArray<FString> Tokens;
 	TArray<FString> Switches;
 	TMap<FString, FString> Params;
@@ -45,19 +46,11 @@ int32 ULearningSocketImitationTrainerServerCommandlet::Main(const FString& Comma
 	const FString* PortParam = Params.Find(TEXT("Port"));
 	const FString* LogSettingsParam = Params.Find(TEXT("LogSettings"));
 
-#if WITH_EDITOR
 	const FString PythonExecutablePath = PythonExecutablePathParam ? *PythonExecutablePathParam : UE::Learning::Trainer::GetPythonExecutablePath(FPaths::EngineDir());
 	const FString SitePackagesPath = SitePackagesPathParam ? *SitePackagesPathParam : UE::Learning::Trainer::GetSitePackagesPath(FPaths::EngineDir());
 	const FString PythonContentPath = PythonContentPathParam ? *PythonContentPathParam : UE::Learning::Trainer::GetPythonContentPath(FPaths::EngineDir());
 	const FString IntermediatePath = IntermediatePathParam ? *IntermediatePathParam : UE::Learning::Trainer::GetIntermediatePath(FPaths::EngineDir());
-#else
-	UE_LEARNING_NOT_IMPLEMENTED();
-	const FString PythonExecutablePath = TEXT("");
-	const FString SitePackagesPath = TEXT("");
-	const FString PythonContentPath = TEXT("");
-	const FString IntermediatePath = TEXT("");
-	return 0;
-#endif
+
 	const TCHAR* IpAddress = IpAddressParam ? *(*IpAddressParam) : UE::Learning::Trainer::DefaultIp;
 	const uint32 Port = PortParam ? FCString::Atoi(*(*PortParam)) : UE::Learning::Trainer::DefaultPort;
 
@@ -102,6 +95,9 @@ int32 ULearningSocketImitationTrainerServerCommandlet::Main(const FString& Comma
 	{
 		FPlatformProcess::Sleep(0.01f);
 	}
+#else
+	UE_LEARNING_NOT_IMPLEMENTED();
+#endif
 
 	return 0;
 }
diff --git a/Engine/Plugins/Experimental/LearningAgents/Source/LearningTraining/Private/LearningPPOTrainer.cpp b/Engine/Plugins/Experimental/LearningAgents/Source/LearningTraining/Private/LearningPPOTrainer.cpp
--- a/Engine/Plugins/Experimental/LearningAgents/Source/LearningTraining/Private/LearningPPOTrainer.cpp
+++ b/Engine/Plugins/Experimental/LearningAgents/Source/LearningTraining/Private/LearningPPOTrainer.cpp
@@ -34,6 +34,7 @@ int32 ULearningSocketPPOTrainerServerCommandlet::Main(const FString& Commandline
 {
 	UE_LOG(LogLearning, Display, TEXT("Running PPO Training Server Commandlet..."));
 
+#if WITH_EDITOR
 	TArray<FString> Tokens;
 	TArray<FString> Switches;
 	TMap<FString, FString> Params;
@@ -48,19 +49,11 @@ int32 ULearningSocketPPOTrainerServerCommandlet::Main(const FString& Commandline
 	const FString* PortParam = Params.Find(TEXT("Port"));
 	const FString* LogSettingsParam = Params.Find(TEXT("LogSettings"));
 
-#if WITH_EDITOR
 	const FString PythonExecutiblePath = PythonExecutiblePathParam ? *PythonExecutiblePathParam : UE::Learning::Trainer::GetPythonExecutablePath(FPaths::EngineDir());
 	const FString SitePackagesPath = SitePackagesPathParam ? *SitePackagesPathParam : UE::Learning::Trainer::GetSitePackagesPath(FPaths::EngineDir());
 	const FString PythonContentPath = PythonContentPathParam ? *PythonContentPathParam : UE::Learning::Trainer::GetPythonContentPath(FPaths::EngineDir());
 	const FString IntermediatePath = IntermediatePathParam ? *IntermediatePathParam : UE::Learning::Trainer::GetIntermediatePath(FPaths::EngineDir());
-#else
-	UE_LEARNING_NOT_IMPLEMENTED();
-	const FString PythonExecutiblePath = TEXT("");
-	const FString SitePackagesPath = TEXT("");
-	const FString PythonContentPath = TEXT("");
-	const FString IntermediatePath = TEXT("");
-	return 0;
-#endif
+
 	const TCHAR* IpAddress = IpAddressParam ? *(*IpAddressParam) : UE::Learning::Trainer::DefaultIp;
 	const uint32 Port = PortParam ? FCString::Atoi(*(*PortParam)) : UE::Learning::Trainer::DefaultPort;
 	
@@ -106,6 +99,10 @@ int32 ULearningSocketPPOTrainerServerCommandlet::Main(const FString& Commandline
 		FPlatformProcess::Sleep(0.01f);
 	}
 
+#else
+	UE_LEARNING_NOT_IMPLEMENTED();
+#endif
+
 	return 0;
 }
 
diff --git a/Engine/Plugins/Media/TimecodeSynchronizer/Source/TimecodeSynchronizer/Private/TimecodeSynchronizer.cpp b/Engine/Plugins/Media/TimecodeSynchronizer/Source/TimecodeSynchronizer/Private/TimecodeSynchronizer.cpp
--- a/Engine/Plugins/Media/TimecodeSynchronizer/Source/TimecodeSynchronizer/Private/TimecodeSynchronizer.cpp
+++ b/Engine/Plugins/Media/TimecodeSynchronizer/Source/TimecodeSynchronizer/Private/TimecodeSynchronizer.cpp
@@ -762,8 +762,6 @@ bool UTimecodeSynchronizer::Tick_TestTimecode()
 
 		return SynchronizationState == ETimecodeProviderSynchronizationState::Synchronized;
 	}
-
-	return true;
 }
 
 void UTimecodeSynchronizer::TickPreRolling_WaitGenlockTimecodeProvider()
diff --git a/Engine/Plugins/Runtime/MeshModelingToolset/Source/MeshModelingTools/Private/ToolActivities/PolyEditExtrudeActivity.cpp b/Engine/Plugins/Runtime/MeshModelingToolset/Source/MeshModelingTools/Private/ToolActivities/PolyEditExtrudeActivity.cpp
--- a/Engine/Plugins/Runtime/MeshModelingToolset/Source/MeshModelingTools/Private/ToolActivities/PolyEditExtrudeActivity.cpp
+++ b/Engine/Plugins/Runtime/MeshModelingToolset/Source/MeshModelingTools/Private/ToolActivities/PolyEditExtrudeActivity.cpp
@@ -592,7 +592,6 @@ FVector3d UPolyEditExtrudeActivity::GetExtrudeDirection() const
 	case EPolyEditExtrudeDirection::LocalZ:
 		return FTransformSRT3d(ActivityContext->Preview->PreviewMesh->GetTransform()).GetRotation().AxisZ();
 	}
-	return ActiveSelectionFrameWorld.Z();
 }
 
 void UPolyEditExtrudeActivity::Render(IToolsContextRenderAPI* RenderAPI)
