Fixed implicit cast error on Linux

diff --git a/Engine/Plugins/Experimental/LearningAgents/Source/LearningAgents/Private/LearningAgentsHelpers.cpp b/Engine/Plugins/Experimental/LearningAgents/Source/LearningAgents/Private/LearningAgentsHelpers.cpp
--- a/Engine/Plugins/Experimental/LearningAgents/Source/LearningAgents/Private/LearningAgentsHelpers.cpp
+++ b/Engine/Plugins/Experimental/LearningAgents/Source/LearningAgents/Private/LearningAgentsHelpers.cpp
@@ -54,7 +54,7 @@ FVector USplineComponentHelper::GetNearestPositionOnSpline(const int32 AgentId,
 	{
 		UE_LEARNING_AGENTS_VLOG_LOCATION(this, LogLearning, Display,
 			SplinePosition,
-			10.0f,
+			10,
 			VisualLogColor.ToFColor(true),
 			TEXT(""));
 
@@ -69,7 +69,7 @@ FVector USplineComponentHelper::GetNearestPositionOnSpline(const int32 AgentId,
 
 		UE_LEARNING_AGENTS_VLOG_LOCATION(this, LogLearning, Display,
 			Position,
-			10.0f,
+			10,
 			VisualLogColor.ToFColor(true),
 			TEXT(""));
 	}
@@ -100,7 +100,7 @@ float USplineComponentHelper::GetDistanceAlongSplineAtPosition(const int32 Agent
 
 		UE_LEARNING_AGENTS_VLOG_LOCATION(this, LogLearning, Display,
 			SplinePosition,
-			10.0f,
+			10,
 			VisualLogColor.ToFColor(true),
 			TEXT(""));
 
@@ -115,7 +115,7 @@ float USplineComponentHelper::GetDistanceAlongSplineAtPosition(const int32 Agent
 
 		UE_LEARNING_AGENTS_VLOG_LOCATION(this, LogLearning, Display,
 			Position,
-			10.0f,
+			10,
 			VisualLogColor.ToFColor(true),
 			TEXT(""));
 	}
@@ -144,7 +144,7 @@ FVector USplineComponentHelper::GetPositionAtDistanceAlongSpline(const int32 Age
 	{
 		UE_LEARNING_AGENTS_VLOG_LOCATION(this, LogLearning, Display,
 			Position,
-			10.0f,
+			10,
 			VisualLogColor.ToFColor(true),
 			TEXT("GetPositionAtDistanceAlongSpline\nAgent %i\nDistance Along Spline: [% 6.1f]\nPosition: [% 6.1f % 6.1f % 6.1f]"),
 			AgentId,
@@ -212,7 +212,7 @@ float USplineComponentHelper::GetProportionAlongSpline(const int32 AgentId, cons
 
 		UE_LEARNING_AGENTS_VLOG_LOCATION(this, LogLearning, Display,
 			Position,
-			10.0f,
+			10,
 			VisualLogColor.ToFColor(true),
 			TEXT("GetProportionAlongSpline\nAgent %i\nDistance Along Spline: [% 6.1f]\nProportion: [% 6.3f]"),
 			AgentId,
@@ -316,7 +316,7 @@ void USplineComponentHelper::GetPositionsAlongSpline(TArray<FVector>& OutPositio
 		{
 			UE_LEARNING_AGENTS_VLOG_LOCATION(this, LogLearning, Display,
 				OutPositions[PositionIdx],
-				2.5f,
+				2,
 				VisualLogColor.ToFColor(true),
 				TEXT("GetPositionsAlongSpline\nAgent: %i\nPosition: [% 6.1f % 6.1f % 6.1f]"),
 				AgentId,
@@ -599,7 +599,7 @@ void UMeshComponentHelper::GetMeshBonePositions(TArray<FVector>& OutBonePosition
 		{
 			UE_LEARNING_AGENTS_VLOG_LOCATION(this, LogLearning, Display,
 				OutBonePositions[BoneIdx],
-				2.5f,
+				2,
 				VisualLogColor.ToFColor(true),
 				TEXT("Agent %i\nBone: \"%s\"\nPosition: [% 6.1f % 6.1f % 6.1f]"),
 				AgentId,
@@ -726,7 +726,7 @@ void URayCastHelper::RayCastGridHeights(
 
 				UE_LEARNING_AGENTS_VLOG_LOCATION(this, LogLearning, Display,
 					ImpactPoint,
-					2.5f,
+					2,
 					VisualLogColor.ToFColor(true),
 					TEXT("Agent %i\nRow: %i\nCol: %i\nImpact Point: [% 6.1f % 6.1f % 6.1f]\nHeight: [% 6.1f]"),
 					AgentId,
@@ -822,7 +822,7 @@ void URayCastHelper::RayCastRadial(
 
 			UE_LEARNING_AGENTS_VLOG_LOCATION(this, LogLearning, Display,
 				ImpactPoint,
-				2.5f,
+				2,
 				VisualLogColor.ToFColor(true),
 				TEXT("Agent %i\nRay: %i\nImpact Point: [% 6.1f % 6.1f % 6.1f]\nDistance: [% 6.1f]"),
 				AgentId,
diff --git a/Engine/Plugins/Experimental/LearningAgents/Source/LearningAgents/Private/LearningAgentsObservations.cpp b/Engine/Plugins/Experimental/LearningAgents/Source/LearningAgents/Private/LearningAgentsObservations.cpp
--- a/Engine/Plugins/Experimental/LearningAgents/Source/LearningAgents/Private/LearningAgentsObservations.cpp
+++ b/Engine/Plugins/Experimental/LearningAgents/Source/LearningAgents/Private/LearningAgentsObservations.cpp
@@ -1283,7 +1283,7 @@ void UPositionObservation::VisualLog(const UE::Learning::FIndexSet Instances) co
 
 		UE_LEARNING_AGENTS_VLOG_LOCATION(this, LogLearning, Display,
 			Position,
-			10.0f,
+			10,
 			VisualLogColor.ToFColor(true),
 			TEXT("Position: [% 6.1f % 6.1f % 6.1f]\nLocal Position: [% 6.1f % 6.1f % 6.1f]"),
 			Position.X, Position.Y, Position.Z,
@@ -1365,7 +1365,7 @@ void UPositionArrayObservation::VisualLog(const UE::Learning::FIndexSet Instance
 
 			UE_LEARNING_AGENTS_VLOG_LOCATION(this, LogLearning, Display,
 				PositionView[Instance][PositionIdx],
-				10.0f,
+				10,
 				VisualLogColor.ToFColor(true),
 				TEXT("Position: [% 6.1f % 6.1f % 6.1f]\nLocal Position: [% 6.1f % 6.1f % 6.1f]"),
 				PositionView[Instance][PositionIdx].X,
@@ -1576,7 +1576,7 @@ void UPlanarPositionObservation::VisualLog(const UE::Learning::FIndexSet Instanc
 
 		UE_LEARNING_AGENTS_VLOG_LOCATION(this, LogLearning, Display,
 			Position,
-			10.0f,
+			10,
 			VisualLogColor.ToFColor(true),
 			TEXT("Position: [% 6.1f % 6.1f % 6.1f]\nLocal Position: [% 6.1f % 6.1f % 6.1f]"),
 			Position.X, Position.Y, Position.Z,
@@ -1673,7 +1673,7 @@ void UPlanarPositionArrayObservation::VisualLog(const UE::Learning::FIndexSet In
 
 			UE_LEARNING_AGENTS_VLOG_LOCATION(this, LogLearning, Display,
 				PositionView[Instance][PositionIdx],
-				10.0f,
+				10,
 				VisualLogColor.ToFColor(true),
 				TEXT("Position: [% 6.1f % 6.1f % 6.1f]\nLocal Position: [% 6.1f % 6.1f % 6.1f]"),
 				PositionView[Instance][PositionIdx].X, 
diff --git a/Engine/Plugins/Experimental/LearningAgents/Source/LearningAgents/Public/LearningAgentsDebug.h b/Engine/Plugins/Experimental/LearningAgents/Source/LearningAgents/Public/LearningAgentsDebug.h
--- a/Engine/Plugins/Experimental/LearningAgents/Source/LearningAgents/Public/LearningAgentsDebug.h
+++ b/Engine/Plugins/Experimental/LearningAgents/Source/LearningAgents/Public/LearningAgentsDebug.h
@@ -55,7 +55,7 @@
 		UE_LEARNING_AGENTS_VLOG_CIRCLE(Owner, Category, Verbosity, Location, FVector::UpVector, Radius, Color, TEXT("")); \
 		UE_LEARNING_AGENTS_VLOG_SEGMENT(Owner, Category, Verbosity, Location, Location + Radius * FVector(FMath::Sin(RelativeAngle), FMath::Cos(RelativeAngle), 0.0f), Color, TEXT("")); \
 		UE_LEARNING_AGENTS_VLOG_SEGMENT(Owner, Category, Verbosity, Location, Location + Radius * FVector(FMath::Sin(Angle), FMath::Cos(Angle), 0.0f), Color, TEXT("")); \
-		UE_LEARNING_AGENTS_VLOG_LOCATION(Owner, Category, Verbosity, Location + Radius * FVector(FMath::Sin(Angle), FMath::Cos(Angle), 0.0f), Radius / 20.0f, Color, Format, ##__VA_ARGS__); \
+		UE_LEARNING_AGENTS_VLOG_LOCATION(Owner, Category, Verbosity, Location + Radius * FVector(FMath::Sin(Angle), FMath::Cos(Angle), 0.0f), static_cast<uint16>(Radius / 20.0f), Color, Format, ##__VA_ARGS__); \
 	}
 
 #endif
diff --git a/Engine/Plugins/Experimental/LearningAgents/Source/LearningAgentsTraining/Private/LearningAgentsCompletions.cpp b/Engine/Plugins/Experimental/LearningAgents/Source/LearningAgentsTraining/Private/LearningAgentsCompletions.cpp
--- a/Engine/Plugins/Experimental/LearningAgents/Source/LearningAgentsTraining/Private/LearningAgentsCompletions.cpp
+++ b/Engine/Plugins/Experimental/LearningAgents/Source/LearningAgentsTraining/Private/LearningAgentsCompletions.cpp
@@ -218,7 +218,7 @@ void UPlanarPositionDifferenceCompletion::VisualLog(const UE::Learning::FIndexSe
 
 			UE_LEARNING_AGENTS_VLOG_LOCATION(this, LogLearning, Display,
 				Position0,
-				10.0f,
+				10,
 				VisualLogColor.ToFColor(true),
 				TEXT("Position0: [% 6.1f % 6.1f % 6.1f]\nPlanar Position0: [% 6.1f % 6.1f]"),
 				Position0.X, Position0.Y, Position0.Z,
@@ -234,7 +234,7 @@ void UPlanarPositionDifferenceCompletion::VisualLog(const UE::Learning::FIndexSe
 
 			UE_LEARNING_AGENTS_VLOG_LOCATION(this, LogLearning, Display,
 				Position1,
-				10.0f,
+				10,
 				VisualLogColor.ToFColor(true),
 				TEXT("Position1: [% 6.1f % 6.1f % 6.1f]\nPlanar Position1: [% 6.1f % 6.1f]"),
 				Position1.X, Position1.Y, Position1.Z,
@@ -318,7 +318,7 @@ void UPlanarPositionSimilarityCompletion::VisualLog(const UE::Learning::FIndexSe
 
 			UE_LEARNING_AGENTS_VLOG_LOCATION(this, LogLearning, Display,
 				Position0,
-				10.0f,
+				10,
 				VisualLogColor.ToFColor(true),
 				TEXT("Position0: [% 6.1f % 6.1f % 6.1f]\nPlanar Position0: [% 6.1f % 6.1f]"),
 				Position0.X, Position0.Y, Position0.Z,
@@ -334,7 +334,7 @@ void UPlanarPositionSimilarityCompletion::VisualLog(const UE::Learning::FIndexSe
 
 			UE_LEARNING_AGENTS_VLOG_LOCATION(this, LogLearning, Display,
 				Position1,
-				10.0f,
+				10,
 				VisualLogColor.ToFColor(true),
 				TEXT("Position1: [% 6.1f % 6.1f % 6.1f]\nPlanar Position1: [% 6.1f % 6.1f]"),
 				Position1.X, Position1.Y, Position1.Z,
diff --git a/Engine/Plugins/Experimental/LearningAgents/Source/LearningAgentsTraining/Private/LearningAgentsRewards.cpp b/Engine/Plugins/Experimental/LearningAgents/Source/LearningAgentsTraining/Private/LearningAgentsRewards.cpp
--- a/Engine/Plugins/Experimental/LearningAgents/Source/LearningAgentsTraining/Private/LearningAgentsRewards.cpp
+++ b/Engine/Plugins/Experimental/LearningAgents/Source/LearningAgentsTraining/Private/LearningAgentsRewards.cpp
@@ -347,7 +347,7 @@ void UPlanarPositionDifferencePenalty::VisualLog(const UE::Learning::FIndexSet I
 
 			UE_LEARNING_AGENTS_VLOG_LOCATION(this, LogLearning, Display,
 				Position0,
-				10.0f,
+				10,
 				VisualLogColor.ToFColor(true),
 				TEXT("Position0: [% 6.1f % 6.1f % 6.1f]\nPlanar Position0: [% 6.1f % 6.1f]"),
 				Position0.X, Position0.Y, Position0.Z,
@@ -363,7 +363,7 @@ void UPlanarPositionDifferencePenalty::VisualLog(const UE::Learning::FIndexSet I
 
 			UE_LEARNING_AGENTS_VLOG_LOCATION(this, LogLearning, Display,
 				Position1,
-				10.0f,
+				10,
 				VisualLogColor.ToFColor(true),
 				TEXT("Position1: [% 6.1f % 6.1f % 6.1f]\nPlanar Position1: [% 6.1f % 6.1f]"),
 				Position1.X, Position1.Y, Position1.Z,
@@ -486,7 +486,7 @@ void UPositionArraySimilarityReward::VisualLog(const UE::Learning::FIndexSet Ins
 
 				UE_LEARNING_AGENTS_VLOG_LOCATION(this, LogLearning, Display,
 					Position0,
-					10.0f,
+					10,
 					VisualLogColor.ToFColor(true),
 					TEXT("Position0: [% 6.1f % 6.1f % 6.1f]\nLocal Position0: [% 6.1f % 6.1f % 6.1f]"),
 					Position0.X, Position0.Y, Position0.Z,
@@ -494,7 +494,7 @@ void UPositionArraySimilarityReward::VisualLog(const UE::Learning::FIndexSet Ins
 
 				UE_LEARNING_AGENTS_VLOG_LOCATION(this, LogLearning, Display,
 					Position1,
-					10.0f,
+					10,
 					VisualLogColor.ToFColor(true),
 					TEXT("Position1: [% 6.1f % 6.1f % 6.1f]\nLocal Position1: [% 6.1f % 6.1f % 6.1f]"),
 					Position1.X, Position1.Y, Position1.Z,
