LearningAgents: Large restructuring to allow for structured actions and observations.

#virtualized

diff --git a/Engine/Plugins/Animation/MLDeformer/NearestNeighborModel/Source/NearestNeighborModel/Private/NearestNeighborOptimizedNetwork.cpp b/Engine/Plugins/Animation/MLDeformer/NearestNeighborModel/Source/NearestNeighborModel/Private/NearestNeighborOptimizedNetwork.cpp
--- a/Engine/Plugins/Animation/MLDeformer/NearestNeighborModel/Source/NearestNeighborModel/Private/NearestNeighborOptimizedNetwork.cpp
+++ b/Engine/Plugins/Animation/MLDeformer/NearestNeighborModel/Source/NearestNeighborModel/Private/NearestNeighborOptimizedNetwork.cpp
@@ -9,7 +9,7 @@
 #include "NNERuntime.h"
 #include "NNERuntimeCPU.h"
 #include "NNEModelData.h"
-#include "NNERuntimeBasicCpu.h"
+#include "NNERuntimeBasicCpuBuilder.h"
 
 //--------------------------------------------------------------------------
 // UNearestNeighborOptimizedNetwork
@@ -22,33 +22,30 @@ namespace UE::NearestNeighborModel::Private
 	/** Creates the FileData from the legacy network format and clears it. */
 	static inline void CreateFileDataAndClearLayers(TArray<uint8>& OutFileData, TArray<TObjectPtr<UNearestNeighborNetworkLayer>>& Layers)
 	{
-		UE::NNE::RuntimeBasic::FSequentialModelBuilder Builder;
+		UE::NNE::RuntimeBasic::FModelBuilder Builder;
 
-		// Allocate data for the PReLuAlpha values as these need to be kept around until we run "Builder.WriteAndReset"
-		TArray<TArray<float>> PreluAlphaData;
+		TArray<UE::NNE::RuntimeBasic::FModelBuilderElement, TInlineAllocator<32>> LayerElements;
+		LayerElements.Reserve(2 * Layers.Num());
 
 		for (UNearestNeighborNetworkLayer* Layer : Layers)
 		{
 			if (UNearestNeighborNetworkLayer_Gemm_Prelu* GemmPreluLayer = Cast<UNearestNeighborNetworkLayer_Gemm_Prelu>(Layer))
 			{
-				Builder.AddLinear(
+				LayerElements.Add(Builder.MakeLinear(
 					GemmPreluLayer->NumInputs,
 					GemmPreluLayer->NumOutputs,
 					GemmPreluLayer->Parameters[0].Values,
-					GemmPreluLayer->Parameters[1].Values);
+					GemmPreluLayer->Parameters[1].Values));
 
-				PreluAlphaData.AddDefaulted();
-				PreluAlphaData.Last().Init(GemmPreluLayer->Parameters[2].Values[0], GemmPreluLayer->NumOutputs);
-
-				Builder.AddPReLU(PreluAlphaData.Last());
+				LayerElements.Add(Builder.MakePReLU(GemmPreluLayer->NumOutputs, Builder.MakeWeightsConstant(GemmPreluLayer->NumOutputs, GemmPreluLayer->Parameters[2].Values[0])));
 			}
 			else if (UNearestNeighborNetworkLayer_Gemm* GemmLayer = Cast<UNearestNeighborNetworkLayer_Gemm>(Layer))
 			{
-				Builder.AddLinear(
+				LayerElements.Add(Builder.MakeLinear(
 					GemmLayer->NumInputs,
 					GemmLayer->NumOutputs,
 					GemmLayer->Parameters[0].Values,
-					GemmLayer->Parameters[1].Values);
+					GemmLayer->Parameters[1].Values));
 			}
 			else
 			{
@@ -56,10 +53,10 @@ namespace UE::NearestNeighborModel::Private
 			}
 		}
 
-		Layers.Empty();
+		uint32 InputSize, OutputSize;
+		Builder.WriteFileDataAndReset(OutFileData, InputSize, OutputSize, Builder.MakeSequence(LayerElements));
 
-		OutFileData.SetNumUninitialized(Builder.GetWriteByteNum());
-		Builder.WriteAndReset(OutFileData);
+		Layers.Empty();
 	}
 
 	PRAGMA_ENABLE_DEPRECATION_WARNINGS
diff --git a/Engine/Plugins/Animation/MLDeformer/NeuralMorphModel/Source/NeuralMorphModel/Private/NeuralMorphNetwork.cpp b/Engine/Plugins/Animation/MLDeformer/NeuralMorphModel/Source/NeuralMorphModel/Private/NeuralMorphNetwork.cpp
--- a/Engine/Plugins/Animation/MLDeformer/NeuralMorphModel/Source/NeuralMorphModel/Private/NeuralMorphNetwork.cpp
+++ b/Engine/Plugins/Animation/MLDeformer/NeuralMorphModel/Source/NeuralMorphModel/Private/NeuralMorphNetwork.cpp
@@ -11,7 +11,7 @@
 #include "NNERuntime.h"
 #include "NNERuntimeCPU.h"
 #include "NNEModelData.h"
-#include "NNERuntimeBasicCpu.h"
+#include "NNERuntimeBasicCpuBuilder.h"
 
 //--------------------------------------------------------------------------
 // UNeuralMorphMLP
@@ -275,10 +275,13 @@ namespace UE::NeuralMorphModel::Private
 
 	static inline void ConvertMLPToFileData(TArray<uint8>& OutBytes, TObjectPtr<UNeuralMorphMLP>& MLP)
 	{
-		UE::NNE::RuntimeBasic::FSequentialModelBuilder Builder;
+		UE::NNE::RuntimeBasic::FModelBuilder Builder;
 
 		const uint32 LayerNum = MLP->GetNumLayers();
 
+		TArray<UE::NNE::RuntimeBasic::FModelBuilderElement, TInlineAllocator<32>> LayerElements;
+		LayerElements.Reserve(2 * LayerNum);
+
 		for (uint32 LayerIdx = 0; LayerIdx < LayerNum; LayerIdx++)
 		{
 			const uint32 InputSize = MLP->GetLayer(LayerIdx).NumInputs;
@@ -289,20 +292,20 @@ namespace UE::NeuralMorphModel::Private
 
 			if (BlockNum == 1)
 			{
-				Builder.AddLinear(InputSize, OutputSize, Weights, Biases);
+				LayerElements.Add(Builder.MakeLinear(InputSize, OutputSize, Weights, Biases));
 			}
 			else
 			{
-				Builder.AddMultiLinear(InputSize, OutputSize, BlockNum, Weights, Biases);
+				LayerElements.Add(Builder.MakeMultiLinear(InputSize, OutputSize, BlockNum, Weights, Biases));
 			}
 
-			Builder.AddELU();
+			LayerElements.Add(Builder.MakeELU(OutputSize * BlockNum));
 		}
 
-		MLP = nullptr;
+		uint32 InputSize, OutputSize;
+		Builder.WriteFileDataAndReset(OutBytes, InputSize, OutputSize, Builder.MakeSequence(LayerElements));
 
-		OutBytes.SetNumUninitialized(Builder.GetWriteByteNum());
-		Builder.WriteAndReset(OutBytes);
+		MLP = nullptr;
 	}
 
 	PRAGMA_ENABLE_DEPRECATION_WARNINGS
diff --git a/Engine/Plugins/Experimental/LearningAgents/LearningAgents.uplugin b/Engine/Plugins/Experimental/LearningAgents/LearningAgents.uplugin
--- a/Engine/Plugins/Experimental/LearningAgents/LearningAgents.uplugin
+++ b/Engine/Plugins/Experimental/LearningAgents/LearningAgents.uplugin
@@ -31,11 +31,6 @@
 			"Type": "Runtime",
 			"LoadingPhase": "Default"
 		},
-		{
-			"Name": "LearningAgentsEditor",
-			"Type": "Editor",
-			"LoadingPhase": "PostEngineInit"
-		},
 		{
 			"Name": "LearningAgentsTraining",
 			"Type": "Runtime",
diff --git a/Engine/Plugins/Experimental/LearningAgents/Source/Learning/Learning.Build.cs b/Engine/Plugins/Experimental/LearningAgents/Source/Learning/Learning.Build.cs
--- a/Engine/Plugins/Experimental/LearningAgents/Source/Learning/Learning.Build.cs
+++ b/Engine/Plugins/Experimental/LearningAgents/Source/Learning/Learning.Build.cs
@@ -38,7 +38,10 @@ public class Learning : ModuleRules
 		PrivateDependencyModuleNames.AddRange(
 			new string[]
 			{
+				"CoreUObject",
 				"Engine",
+				"NNE",
+				"NNERuntimeBasicCpu",
 				// ... add private dependencies that you statically link with here ...	
 			}
 			);
diff --git a/Engine/Plugins/Experimental/LearningAgents/Source/Learning/Private/Learning.ispc b/Engine/Plugins/Experimental/LearningAgents/Source/Learning/Private/Learning.ispc
--- a/Engine/Plugins/Experimental/LearningAgents/Source/Learning/Private/Learning.ispc
+++ b/Engine/Plugins/Experimental/LearningAgents/Source/Learning/Private/Learning.ispc
@@ -5,9 +5,18 @@
 #endif
 
 #ifndef INDEX_NONE
-#define INDEX_NONE		-1
+#define INDEX_NONE		(-1)
 #endif
 
+#ifndef UE_SMALL_NUMBER
+#define UE_SMALL_NUMBER (1.e-8f)
+#endif
+
+inline float LearningSigmoid(const float X)
+{
+	return 1.0f / (1.0f + exp(-X));
+}
+
 /////////////////////////////////////////////////////////
 // 
 // Functions matching those in LearningRandom.cpp
@@ -59,16 +68,16 @@ inline uniform int32 LearningRandomIntInRange(uniform const uint32 State, unifor
 
 inline float LearningRandomUniform(
 	const uint32 State,
-	const float Min,
-	const float Max)
+	const float Min = 0.0f,
+	const float Max = 1.0f)
 {
 	return (Max - Min) * LearningRandomFloat(State ^ 0x72404541) + Min;
 }
 
 inline float LearningRandomGaussian(
 	const uint32 State,
-	const float Mean,
-	const float Std)
+	const float Mean = 0.0f,
+	const float Std = 1.0f)
 {
 	return Std * LearningRandomUniformToGaussian(
 		LearningRandomFloat(State ^ 0x4855e88f),
@@ -123,154 +132,50 @@ export void LearningRandomGaussianArray(
 	}
 }
 
-export void LearningResampleStateArray(
-	uniform uint32 States[],
-	uniform const int DimNum)
-{
-	foreach (DimIdx = 0 ... DimNum)
-	{
-		States[DimIdx] = LearningRandomInt(States[DimIdx] ^ 0x3616dcc8 ^ LearningRandomInt(DimIdx ^ 0x6a837ffd));
-	}
-}
-
-/////////////////////////////////////////////////////////
-// 
-// Some functions related to evaluation of Neural Networks
-// 
-/////////////////////////////////////////////////////////
-
-export void LearningLayerMatMulPlusBias(
+export void LearningRandomDistributionIndependantNormal(
 	uniform float Output[],
-	const uniform float Input[],
-	const uniform float Weights[],
-	const uniform float Biases[],
-	const uniform int BatchNum,
-	const uniform int BufferSize,
-	const uniform int RowNum,
-	const uniform int ColNum)
+	uniform const float Mean[],
+	uniform const float LogStd[],
+	uniform const int DimNum,
+	uniform const uint32 State,
+	uniform const float Scale)
 {
-	foreach (BatchIdx = 0 ... BatchNum, ColIdx = 0 ... ColNum)
+	if (Scale < UE_SMALL_NUMBER)
 	{
-		Output[BatchIdx * BufferSize + ColIdx] = Biases[ColIdx];
-	}
-	
-	for (uniform int BatchIdx = 0; BatchIdx < BatchNum; BatchIdx++)
-	{
-		for (uniform int RowIdx = 0; RowIdx < RowNum; RowIdx++)
+		foreach(DimIdx = 0 ... DimNum)
 		{
-			uniform const float Value = Input[BatchIdx * BufferSize + RowIdx];
-
-			if (Value != 0.0)
-			{
-				foreach(ColIdx = 0 ... ColNum)
-				{
-					Output[BatchIdx * BufferSize + ColIdx] += Value * Weights[RowIdx * ColNum + ColIdx];
-				}
-			}
+			Output[DimIdx] = Mean[DimIdx];
 		}
 	}
-}
-
-export void LearningLayerReLU(
-	uniform float InputOutput[],
-	const uniform int BatchNum,
-	const uniform int BufferSize,
-	const uniform int ColNum)
-{
-	foreach(BatchIdx = 0 ... BatchNum, ColIdx = 0 ... ColNum)
+	else
 	{
-		InputOutput[BatchIdx * BufferSize + ColIdx] = max(InputOutput[BatchIdx * BufferSize + ColIdx], 0.0f);
-	}
-}
-
-inline float LearningELU(const float X)
-{
-	return X > 0.0f ? X : exp(X) - 1.0f;
-}
-
-export void LearningLayerELU(
-	uniform float InputOutput[],
-	const uniform int BatchNum,
-	const uniform int BufferSize,
-	const uniform int ColNum)
-{
-	foreach (BatchIdx = 0 ... BatchNum, ColIdx = 0 ... ColNum)
-	{
-		InputOutput[BatchIdx * BufferSize + ColIdx] = LearningELU(InputOutput[BatchIdx * BufferSize + ColIdx]);
-	}
-}
-
-inline float LearningTanH(const float X, const float Min=-10.0f, const float Max=10.0f)
-{
-	const float Exp2X = exp(2.0f * clamp(X, Min, Max));
-	return (Exp2X - 1) / (Exp2X + 1);
-}
-
-export void LearningLayerTanH(
-	uniform float InputOutput[],
-	const uniform int BatchNum,
-	const uniform int BufferSize,
-	const uniform int ColNum)
-{
-	foreach(BatchIdx = 0 ... BatchNum, ColIdx = 0 ... ColNum)
-	{
-		InputOutput[BatchIdx * BufferSize + ColIdx] = LearningTanH(InputOutput[BatchIdx * BufferSize + ColIdx]);
+		foreach(DimIdx = 0 ... DimNum)
+		{
+			Output[DimIdx] = LearningRandomGaussian(State ^ 0x7db0e4e9 ^ LearningRandomInt(DimIdx ^ 0xf4976a00), Mean[DimIdx], Scale * exp(min(LogStd[DimIdx], 10.0f)));
+		}
 	}
 }
 
-inline float LearningSigmoid(const float X)
-{
-	return 1.0f / (1.0f + exp(-X));
-}
-
-export void LearningLayerActionNoise(
+export void LearningRandomDistributionBernoulli(
 	uniform float Output[],
-	uniform float OutputMean[],
-	uniform float OutputStd[],
-	const uniform float OutputNetwork[],
-	const uniform uint32 Seed[],
-	const uniform float ActionNoiseScale[],
-	const uniform int BatchNum,
-	const uniform int ActionNum,
-	const uniform int MemoryStateNum,
-	const uniform float LogActionNoiseMin,
-	const uniform float LogActionNoiseMax)
+	uniform const float Logits[],
+	uniform const int DimNum,
+	uniform const uint32 State,
+	uniform const float Scale)
 {
-	foreach (BatchIdx = 0 ... BatchNum, ActionIdx = 0 ... ActionNum)
+	if (Scale < UE_SMALL_NUMBER)
 	{
-		OutputMean[BatchIdx * ActionNum + ActionIdx] = OutputNetwork[BatchIdx * (2 * ActionNum + MemoryStateNum) + ActionIdx];
-		OutputStd[BatchIdx * ActionNum + ActionIdx] = ActionNoiseScale[BatchIdx] *
-			exp(LearningSigmoid(OutputNetwork[BatchIdx * (2 * ActionNum + MemoryStateNum) + ActionNum + ActionIdx]) * (LogActionNoiseMax - LogActionNoiseMin) + LogActionNoiseMin);
-
-		Output[BatchIdx * ActionNum + ActionIdx] = LearningRandomGaussian(
-			Seed[BatchIdx] ^ 0xab744615 ^ LearningRandomInt(ActionIdx ^ 0xf8a88a27),
-			OutputMean[BatchIdx * ActionNum + ActionIdx],
-			OutputStd[BatchIdx * ActionNum + ActionIdx]);
+		foreach(DimIdx = 0 ... DimNum)
+		{
+			Output[DimIdx] = select(Logits[DimIdx] > 0.0f, 1.0f, 0.0f);
+		}
 	}
-}
-
-export void LearningLayerActionNoiseSingleBatch(
-	uniform float Output[],
-	uniform float OutputMean[],
-	uniform float OutputStd[],
-	const uniform float OutputNetwork[],
-	const uniform uint32 Seed,
-	const uniform float ActionNoiseScale,
-	const uniform int ActionNum,
-	const uniform int MemoryStateNum,
-	const uniform float LogActionNoiseMin,
-	const uniform float LogActionNoiseMax)
-{
-	foreach(ActionIdx = 0 ... ActionNum)
+	else
 	{
-		OutputMean[ActionIdx] = OutputNetwork[ActionIdx];
-		OutputStd[ActionIdx] = ActionNoiseScale *
-			exp(LearningSigmoid(OutputNetwork[ActionNum + ActionIdx]) * (LogActionNoiseMax - LogActionNoiseMin) + LogActionNoiseMin);
-
-		Output[ActionIdx] = LearningRandomGaussian(
-			Seed ^ 0xab744615 ^ LearningRandomInt(ActionIdx ^ 0xf8a88a27),
-			OutputMean[ActionIdx],
-			OutputStd[ActionIdx]);
+		foreach(DimIdx = 0 ... DimNum)
+		{
+			Output[DimIdx] = select(LearningSigmoid(Logits[DimIdx] / Scale) > LearningRandomUniform(State ^ 0xf4021a46 ^ LearningRandomInt(DimIdx ^ 0x7a8cc64e)), 1.0f, 0.0f);
+		}
 	}
 }
 
@@ -515,70 +420,6 @@ export void LearningUpdatePSOBest(
 	}
 }
 
-/////////////////////////////////////////////////////////
-// 
-// Functions related to constructing feature vectors
-// 
-/////////////////////////////////////////////////////////
-
-export void LearningCombineFeature(
-	uniform float Output[],
-	uniform const float Buffer[],
-	uniform const int DimOffset,
-	uniform const int DimNum,
-	uniform const int OutputDimNum,
-	uniform const int InstanceNum,
-	uniform const float Scale)
-{
-	foreach (InstanceIdx = 0 ... InstanceNum, DimIdx = 0 ... DimNum)
-	{
-		Output[InstanceIdx * OutputDimNum + DimOffset + DimIdx] = Scale * Buffer[InstanceIdx * DimNum + DimIdx];
-	}
-}
-
-export void LearningCombineFeatureSingleInstance(
-	uniform float Output[],
-	uniform const float Buffer[],
-	uniform const int DimOffset,
-	uniform const int DimNum,
-	uniform const float Scale)
-{
-	foreach(DimIdx = 0 ... DimNum)
-	{
-		Output[DimOffset + DimIdx] = Scale * Buffer[DimIdx];
-	}
-}
-
-export void LearningSeparateFeature(
-	uniform float Buffer[],
-	uniform const float Input[],
-	uniform const int DimOffset,
-	uniform const int DimNum,
-	uniform const int OutputDimNum,
-	uniform const int InstanceNum,
-	uniform const float Scale,
-	uniform const float Epsilon)
-{
-	foreach (InstanceIdx = 0 ... InstanceNum, DimIdx = 0 ... DimNum)
-	{
-		Buffer[InstanceIdx * DimNum + DimIdx] = Input[InstanceIdx * OutputDimNum + DimOffset + DimIdx] / max(Scale, Epsilon);
-	}
-}
-
-export void LearningSeparateFeatureSingleInstance(
-	uniform float Buffer[],
-	uniform const float Input[],
-	uniform const int DimOffset,
-	uniform const int DimNum,
-	uniform const float Scale,
-	uniform const float Epsilon)
-{
-	foreach(DimIdx = 0 ... DimNum)
-	{
-		Buffer[DimIdx] = Input[DimOffset + DimIdx] / max(Scale, Epsilon);
-	}
-}
-
 /////////////////////////////////////////////////////////
 // 
 // Functions related to K-Means Clustering
diff --git a/Engine/Plugins/Experimental/LearningAgents/Source/Learning/Private/LearningCompletion.cpp b/Engine/Plugins/Experimental/LearningAgents/Source/Learning/Private/LearningCompletion.cpp
--- a/Engine/Plugins/Experimental/LearningAgents/Source/Learning/Private/LearningCompletion.cpp
+++ b/Engine/Plugins/Experimental/LearningAgents/Source/Learning/Private/LearningCompletion.cpp
@@ -4,45 +4,41 @@
 
 namespace UE::Learning
 {
-	void FResetInstanceBuffer::Resize(const int32 InMaxInstanceNum)
+	void FResetInstanceBuffer::Reserve(const int32 InMaxInstanceNum)
 	{
-		MaxInstanceNum = InMaxInstanceNum;
-		ResetInstanceNum = 0;
-		ResetInstances.SetNumUninitialized({ MaxInstanceNum });
+		ResetInstances.Reserve(InMaxInstanceNum);
 	}
 
 	void FResetInstanceBuffer::SetResetInstances(const FIndexSet Instances)
 	{
-		ResetInstanceNum = 0;
-
+		ResetInstances.Reset();
 		for (const int32 InstanceIdx : Instances)
 		{
-			ResetInstances[ResetInstanceNum] = InstanceIdx;
-			ResetInstanceNum++;
+			ResetInstances.Add(InstanceIdx);
 		}
 
-		ResetInstancesSet = ResetInstances.Slice(0, ResetInstanceNum);
+		ResetInstancesSet = ResetInstances;
 		ResetInstancesSet.TryMakeSlice();
 	}
 
 	void FResetInstanceBuffer::SetResetInstancesFromCompletions(const TLearningArrayView<1, const ECompletionMode> Completions, const FIndexSet Instances)
 	{
-		ResetInstanceNum = 0;
-
+		ResetInstances.Reset();
 		for (const int32 InstanceIdx : Instances)
 		{
 			if (Completions[InstanceIdx] != ECompletionMode::Running)
 			{
-				ResetInstances[ResetInstanceNum] = InstanceIdx;
-				ResetInstanceNum++;
+				ResetInstances.Add(InstanceIdx);
 			}
 		}
 
-		ResetInstancesSet = ResetInstances.Slice(0, ResetInstanceNum);
+		ResetInstancesSet = ResetInstances;
 		ResetInstancesSet.TryMakeSlice();
 	}
 
-	const int32 FResetInstanceBuffer::GetResetInstanceNum() const { return ResetInstanceNum; }
+	const int32 FResetInstanceBuffer::GetResetInstanceNum() const { return ResetInstances.Num(); }
+
+	const TArray<int32>& FResetInstanceBuffer::GetResetInstancesArray() const { return ResetInstances; }
 
 	const FIndexSet FResetInstanceBuffer::GetResetInstances() const { return ResetInstancesSet; }
 
@@ -76,17 +72,11 @@ namespace UE::Learning
 			TLearningArrayView<1, ECompletionMode> InOutCompletions,
 			const TLearningArrayView<1, const int32> EpisodeStepNums,
 			const int32 EpisodeMaxStepNum,
-			const ECompletionMode EpisodeEndCompletionMode,
 			const FIndexSet Instances)
 		{
-			UE_LEARNING_TRACE_CPUPROFILER_EVENT_SCOPE(Learning::Completion::EvaluateEndOfEpisodeCompletions);
-
 			for (const int32 InstanceIdx : Instances)
 			{
-				if (InOutCompletions[InstanceIdx] == ECompletionMode::Running && EpisodeStepNums[InstanceIdx] == EpisodeMaxStepNum)
-				{
-					InOutCompletions[InstanceIdx] = EpisodeEndCompletionMode;
-				}
+				InOutCompletions[InstanceIdx] = EpisodeStepNums[InstanceIdx] == EpisodeMaxStepNum ? ECompletionMode::Truncated : ECompletionMode::Running;
 			}
 		}
 	}
diff --git a/Engine/Plugins/Experimental/LearningAgents/Source/Learning/Private/LearningCritic.cpp b/Engine/Plugins/Experimental/LearningAgents/Source/Learning/Private/LearningCritic.cpp
new file mode 100644
--- /dev/null
+++ b/Engine/Plugins/Experimental/LearningAgents/Source/Learning/Private/LearningCritic.cpp
@@ -0,0 +1,67 @@
+// Copyright Epic Games, Inc. All Rights Reserved.
+
+#include "LearningCritic.h"
+
+namespace UE::Learning
+{
+	FNeuralNetworkCritic::FNeuralNetworkCritic(
+		const int32 InMaxInstanceNum,
+		const int32 InObservationEncodedNum,
+		const int32 InMemoryStateNum,
+		const TSharedPtr<FNeuralNetwork>& InNeuralNetwork,
+		const FNeuralNetworkInferenceSettings& InInferenceSettings)
+		: MaxInstanceNum(InMaxInstanceNum)
+		, ObservationEncodedNum(InObservationEncodedNum)
+		, MemoryStateNum(InMemoryStateNum)
+		, NeuralNetwork(InNeuralNetwork)
+		, InferenceSettings(InInferenceSettings)
+	{
+		UE_LEARNING_CHECK(NeuralNetwork->GetInputSize() == ObservationEncodedNum + MemoryStateNum);
+		UE_LEARNING_CHECK(NeuralNetwork->GetOutputSize() == 1);
+
+		Input.SetNumUninitialized({ MaxInstanceNum, ObservationEncodedNum + MemoryStateNum });
+		Array::Zero(Input);
+
+		NeuralNetworkInference = NeuralNetwork->CreateInferenceObject(MaxInstanceNum, InferenceSettings);
+	}
+
+	void FNeuralNetworkCritic::Evaluate(
+		TLearningArrayView<1, float> OutputReturns,
+		const TLearningArrayView<2, const float> InputObservationVectorsEncoded,
+		const TLearningArrayView<2, const float> InputMemoryState,
+		const FIndexSet Instances)
+	{
+		UE_LEARNING_TRACE_CPUPROFILER_EVENT_SCOPE(Learning::FNeuralNetworkCritic::Evaluate);
+
+		Array::Check(InputObservationVectorsEncoded, Instances);
+		Array::Check(InputMemoryState, Instances);
+
+		// Copy in Observation and Memory State into network input
+
+		for (const int32 InstanceIdx : Instances)
+		{
+			Array::Copy(Input[InstanceIdx].Slice(0, ObservationEncodedNum), InputObservationVectorsEncoded[InstanceIdx]);
+			Array::Copy(Input[InstanceIdx].Slice(ObservationEncodedNum, MemoryStateNum), InputMemoryState[InstanceIdx]);
+		}
+
+		// Evaluate Network
+
+		NeuralNetworkInference->Evaluate(
+			TLearningArrayView<2, float>(OutputReturns.GetData(), { OutputReturns.Num(), 1 }),
+			Input, 
+			Instances);
+
+		Array::Check(OutputReturns, Instances);
+	}
+
+	void FNeuralNetworkCritic::UpdateNeuralNetwork(const TSharedPtr<FNeuralNetwork>& NewNeuralNetwork)
+	{
+		if (NeuralNetwork != NewNeuralNetwork)
+		{
+			UE_LEARNING_CHECK(NewNeuralNetwork->GetInputSize() == ObservationEncodedNum + MemoryStateNum);
+			UE_LEARNING_CHECK(NewNeuralNetwork->GetOutputSize() == 1);
+			NeuralNetwork = NewNeuralNetwork;
+			NeuralNetworkInference = NeuralNetwork->CreateInferenceObject(MaxInstanceNum, InferenceSettings);
+		}
+	}
+}
diff --git a/Engine/Plugins/Experimental/LearningAgents/Source/Learning/Private/LearningNeuralNetwork.cpp b/Engine/Plugins/Experimental/LearningAgents/Source/Learning/Private/LearningNeuralNetwork.cpp
--- a/Engine/Plugins/Experimental/LearningAgents/Source/Learning/Private/LearningNeuralNetwork.cpp
+++ b/Engine/Plugins/Experimental/LearningAgents/Source/Learning/Private/LearningNeuralNetwork.cpp
@@ -2,368 +2,377 @@
 
 #include "LearningNeuralNetwork.h"
 
-#if UE_LEARNING_ISPC
-#include "Learning.ispc.generated.h"
-#endif
+#include "LearningProgress.h"
 
-namespace UE::Learning
-{
-	namespace NeuralNetwork
-	{
-		static inline void MatMulPlusBias(
-			TLearningArrayView<2, float> Output,
-			const TLearningArrayView<2, const float> Input,
-			const TLearningArrayView<2, const float> Weights,
-			const TLearningArrayView<1, const float> Biases,
-			const int32 BatchOffset,
-			const int32 BatchNum)
-		{
-			UE_LEARNING_TRACE_CPUPROFILER_EVENT_SCOPE(Learning::NeuralNetwork::MatMulPlusBias);
-
-			const int32 RowNum = Weights.Num<0>();
-			const int32 ColNum = Weights.Num<1>();
-			const int32 BufferSize = Input.Num<1>();
-
-#if UE_LEARNING_ISPC
-			ispc::LearningLayerMatMulPlusBias(
-				Output.Slice(BatchOffset, BatchNum).GetData(),
-				Input.Slice(BatchOffset, BatchNum).GetData(),
-				Weights.GetData(),
-				Biases.GetData(),
-				BatchNum,
-				BufferSize,
-				RowNum,
-				ColNum);
-#else
-			for (int32 BatchIdx = BatchOffset; BatchIdx < BatchOffset + BatchNum; BatchIdx++)
-			{
-				for (int32 ColIdx = 0; ColIdx < ColNum; ColIdx++)
-				{
-					Output[BatchIdx][ColIdx] = Biases[ColIdx];
-				}
+#include "Async/ParallelFor.h"
+#include "Modules/ModuleManager.h"
 
-				for (int32 RowIdx = 0; RowIdx < RowNum; RowIdx++)
-				{
-					const float Value = Input[BatchIdx][RowIdx];
-
-					if (Value != 0.0)
-					{
-						for (int32 ColIdx = 0; ColIdx < ColNum; ColIdx++)
-						{
-							Output[BatchIdx][ColIdx] += Value * Weights[RowIdx][ColIdx];
-						}
-					}
-				}
-			}
-#endif
-		}
+#include "NNE.h"
+#include "NNEModelData.h"
+#include "NNERuntimeCPU.h"
 
-		static inline float ReLU(const float X)
-		{
-			return FMath::Max(X, 0.0f);
-		}
+void ULearningNeuralNetworkData::PostLoad()
+{
+	Super::PostLoad();
 
-		static inline float ELU(const float X)
-		{
-			return X > 0.0f ? X : FMath::Exp(X) - 1.0f;
-		}
+	UpdateNetwork();
+}
 
-		static inline float TanH(const float X)
-		{
-			return FMath::Tanh(X);
-		}
+void ULearningNeuralNetworkData::Init(const int32 InInputSize, const int32 InOutputSize, const int32 InCompatibilityHash, const TArrayView<const uint8> InFileData)
+{
+	InputSize = InInputSize;
+	OutputSize = InOutputSize;
+	CompatibilityHash = InCompatibilityHash;
+	FileData = InFileData;
+	UpdateNetwork();
+}
 
-		static inline void ActivationReLU(TLearningArrayView<2, float> InputOutput, const int32 BatchOffset, const int32 BatchNum, const int32 InputOutputNum)
-		{
-			UE_LEARNING_TRACE_CPUPROFILER_EVENT_SCOPE(Learning::NeuralNetwork::ActivationReLU);
-
-			const int32 BufferSize = InputOutput.Num<1>();
-
-#if UE_LEARNING_ISPC
-			ispc::LearningLayerReLU(
-				InputOutput.Slice(BatchOffset, BatchNum).GetData(),
-				BatchNum,
-				BufferSize, 
-				InputOutputNum);
-#else
-			for (int32 BatchIdx = BatchOffset; BatchIdx < BatchOffset + BatchNum; BatchIdx++)
-			{
-				for (int32 ColIdx = 0; ColIdx < InputOutputNum; ColIdx++)
-				{
-					InputOutput[BatchIdx][ColIdx] = ReLU(InputOutput[BatchIdx][ColIdx]);
-				}
-			}
-#endif
-		}
+void ULearningNeuralNetworkData::InitFrom(const ULearningNeuralNetworkData* OtherNetworkData)
+{
+	InputSize = OtherNetworkData->InputSize;
+	OutputSize = OtherNetworkData->OutputSize;
+	CompatibilityHash = OtherNetworkData->CompatibilityHash;
+	FileData = OtherNetworkData->FileData;
+	UpdateNetwork();
+}
 
-		static inline void ActivationELU(TLearningArrayView<2, float> InputOutput, const int32 BatchOffset, const int32 BatchNum, const int32 InputOutputNum)
-		{
-			UE_LEARNING_TRACE_CPUPROFILER_EVENT_SCOPE(Learning::NeuralNetwork::ActivationELU);
-
-			const int32 BufferSize = InputOutput.Num<1>();
-
-#if UE_LEARNING_ISPC
-			ispc::LearningLayerELU(
-				InputOutput.Slice(BatchOffset, BatchNum).GetData(),
-				BatchNum,
-				BufferSize,
-				InputOutputNum);
-#else
-			for (int32 BatchIdx = BatchOffset; BatchIdx < BatchOffset + BatchNum; BatchIdx++)
-			{
-				for (int32 ColIdx = 0; ColIdx < InputOutputNum; ColIdx++)
-				{
-					InputOutput[BatchIdx][ColIdx] = ELU(InputOutput[BatchIdx][ColIdx]);
-				}
-			}
-#endif
-		}
+namespace UE::Learning::NeuralNetwork::Private
+{
+	static constexpr int32 MagicNumber = 0x1e9b0c80;
+	static constexpr int32 VersionNumber = 1;
+}
 
-		static inline void ActivationTanH(TLearningArrayView<2, float> InputOutput, const int32 BatchOffset, const int32 BatchNum, const int32 InputOutputNum)
-		{
-			UE_LEARNING_TRACE_CPUPROFILER_EVENT_SCOPE(Learning::NeuralNetwork::ActivationTanH);
-
-			const int32 BufferSize = InputOutput.Num<1>();
-
-#if UE_LEARNING_ISPC
-			ispc::LearningLayerTanH(
-				InputOutput.Slice(BatchOffset, BatchNum).GetData(),
-				BatchNum,
-				BufferSize,
-				InputOutputNum);
-#else
-			for (int32 BatchIdx = BatchOffset; BatchIdx < BatchOffset + BatchNum; BatchIdx++)
-			{
-				for (int32 ValueIdx = 0; ValueIdx < InputOutputNum; ValueIdx++)
-				{
-					InputOutput[BatchIdx][ValueIdx] = TanH(InputOutput[BatchIdx][ValueIdx]);
-				}
-			}
-#endif
-		}
+bool ULearningNeuralNetworkData::LoadFromSnapshot(const TArrayView<const uint8> InBytes)
+{
+	if (!ensureMsgf(InBytes.Num() > sizeof(int32) * 6, TEXT("Snapshot format invalid. File not large enough.")))
+	{
+		return false;
+	}
 
-		static constexpr int32 MagicNumber = 0x0cd353cf;
-		static constexpr int32 VersionNumber = 1;
+	int32 Offset = 0;
+	
+	int32 MagicNumber = 0;
+	UE::Learning::DeserializeFromBytes(Offset, InBytes, MagicNumber);
+	if (!ensureMsgf(MagicNumber == UE::Learning::NeuralNetwork::Private::MagicNumber, TEXT("Snapshot format invalid. Incorrect Magic Number. Got %i, expected %i."), MagicNumber, UE::Learning::NeuralNetwork::Private::MagicNumber))
+	{
+		return false;
 	}
 
-	bool FNeuralNetworkMLP::DeserializeFromBytes(int32& InOutOffset, const TLearningArrayView<1, const uint8> RawBytes)
+	int32 VersionNumber = 0;
+	UE::Learning::DeserializeFromBytes(Offset, InBytes, VersionNumber);
+	if (!ensureMsgf(VersionNumber == UE::Learning::NeuralNetwork::Private::VersionNumber, TEXT("Snapshot format invalid. Unsupported version number. Got %i, expected %i."), VersionNumber, UE::Learning::NeuralNetwork::Private::VersionNumber))
 	{
-		UE_LEARNING_TRACE_CPUPROFILER_EVENT_SCOPE(Learning::FNeuralNetworkMLP::DeserializeFromBytes);
+		return false;
+	}
 
-		// At least Magic Number and Version Number Required
-		if (RawBytes.Num() < sizeof(int32) * 2)
-		{
-			return false;
-		}
+	UE::Learning::DeserializeFromBytes(Offset, InBytes, InputSize);
+	UE::Learning::DeserializeFromBytes(Offset, InBytes, OutputSize);
+	UE::Learning::DeserializeFromBytes(Offset, InBytes, CompatibilityHash);
 
-		int32 Magic;
-		UE::Learning::DeserializeFromBytes(InOutOffset, RawBytes, Magic);
-		if (Magic != NeuralNetwork::MagicNumber) { return false; }
+	int32 FileDataSize = 0;
+	UE::Learning::DeserializeFromBytes(Offset, InBytes, FileDataSize);
+	FileData = InBytes.Slice(Offset, FileDataSize);
 
-		int32 Version;
-		UE::Learning::DeserializeFromBytes(InOutOffset, RawBytes, Version);
-		if (Version != NeuralNetwork::VersionNumber) { return false; }
+	UpdateNetwork();
+	return true;
+}
 
-		int32 ActivationFunctionInt;
-		UE::Learning::DeserializeFromBytes(InOutOffset, RawBytes, ActivationFunctionInt);
-		ActivationFunction = (EActivationFunction)ActivationFunctionInt;
+void ULearningNeuralNetworkData::SaveToSnapshot(TArrayView<uint8> OutBytes) const
+{
+	int32 Offset = 0;
+	UE::Learning::SerializeToBytes(Offset, OutBytes, UE::Learning::NeuralNetwork::Private::MagicNumber);
+	UE::Learning::SerializeToBytes(Offset, OutBytes, UE::Learning::NeuralNetwork::Private::VersionNumber);
+	UE::Learning::SerializeToBytes(Offset, OutBytes, InputSize);
+	UE::Learning::SerializeToBytes(Offset, OutBytes, OutputSize);
+	UE::Learning::SerializeToBytes(Offset, OutBytes, CompatibilityHash);
+	UE::Learning::SerializeToBytes(Offset, OutBytes, FileData.Num());
+	FMemory::Memcpy(OutBytes.Slice(Offset, FileData.Num()).GetData(), FileData.GetData(), FileData.Num());
+}
 
-		int32 LayerNum;
-		UE::Learning::DeserializeFromBytes(InOutOffset, RawBytes, LayerNum);
+int32 ULearningNeuralNetworkData::GetSnapshotByteNum() const
+{
+	return
+		sizeof(int32) + // Magic Number
+		sizeof(int32) + // Version Number
+		sizeof(int32) + // Input Size
+		sizeof(int32) + // Output Size
+		sizeof(int32) + // Compatibility Hash
+		sizeof(int32) + // FileData size
+		FileData.Num(); // FileData
+}
 
-		Weights.SetNum(LayerNum);
-		Biases.SetNum(LayerNum);
+TSharedPtr<UE::Learning::FNeuralNetwork>& ULearningNeuralNetworkData::GetNetwork()
+{
+	return Network;
+}
 
-		for (int32 LayerIdx = 0; LayerIdx < LayerNum; LayerIdx++)
-		{
-			UE::Learning::Array::DeserializeFromBytes(InOutOffset, RawBytes, Weights[LayerIdx]);
-			UE::Learning::Array::DeserializeFromBytes(InOutOffset, RawBytes, Biases[LayerIdx]);
-		}
+bool ULearningNeuralNetworkData::IsEmpty() const
+{
+	return Network->IsEmpty();
+}
 
-		return true;
-	}
+int32 ULearningNeuralNetworkData::GetInputSize() const
+{
+	return InputSize;
+}
 
-	void FNeuralNetworkMLP::SerializeToBytes(int32& InOutOffset, TLearningArrayView<1, uint8> OutRawBytes) const
-	{
-		UE_LEARNING_TRACE_CPUPROFILER_EVENT_SCOPE(Learning::FNeuralNetworkMLP::SerializeToBytes);
+int32 ULearningNeuralNetworkData::GetOutputSize() const
+{
+	return OutputSize;
+}
 
-		const int32 LayerNum = Weights.Num();
+int32 ULearningNeuralNetworkData::GetCompatibilityHash() const
+{
+	return CompatibilityHash;
+}
 
-		UE::Learning::SerializeToBytes(InOutOffset, OutRawBytes, NeuralNetwork::MagicNumber);
-		UE::Learning::SerializeToBytes(InOutOffset, OutRawBytes, NeuralNetwork::VersionNumber);
-		UE::Learning::SerializeToBytes(InOutOffset, OutRawBytes, (int32)ActivationFunction);
-		UE::Learning::SerializeToBytes(InOutOffset, OutRawBytes, LayerNum);
+void ULearningNeuralNetworkData::UpdateNetwork()
+{
+	if (FileData.Num() > 0)
+	{
+		if (!ModelData)
+		{
+			ModelData = NewObject<UNNEModelData>(this);
+		}
 
-		for (int32 LayerIdx = 0; LayerIdx < LayerNum; LayerIdx++)
+		ModelData->Init(TEXT("ubnne"), FileData);
+
+		if (!Network)
 		{
-			UE::Learning::Array::SerializeToBytes(InOutOffset, OutRawBytes, Weights[LayerIdx]);
-			UE::Learning::Array::SerializeToBytes(InOutOffset, OutRawBytes, Biases[LayerIdx]);
+			Network = MakeShared<UE::Learning::FNeuralNetwork>();
 		}
-	}
 
-	int32 FNeuralNetworkMLP::GetSerializationByteNum() const
-	{
-		const int32 LayerNum = GetLayerNum();
+		ensureMsgf(FModuleManager::Get().LoadModule(TEXT("NNERuntimeBasicCpu")), TEXT("Unable to load module for NNE runtime NNERuntimeBasicCpu."));
 
-		int32 Total = 0;
+		TWeakInterfacePtr<INNERuntimeCPU> RuntimeCPU = UE::NNE::GetRuntime<INNERuntimeCPU>(TEXT("NNERuntimeBasicCpu"));
 
-		Total += sizeof(int32);  // Magic
-		Total += sizeof(int32);  // Version
-		Total += sizeof(int32);  // Activation
-		Total += sizeof(int32);  // Layer Num
+		TSharedPtr<UE::NNE::IModelCPU> UpdatedModel = nullptr;
 
-		for (int32 LayerIdx = 0; LayerIdx < LayerNum; LayerIdx++)
+		if (ensureMsgf(RuntimeCPU.IsValid(), TEXT("Could not find requested NNE Runtime")))
 		{
-			Total += UE::Learning::Array::SerializationByteNum<2, float>(Weights[LayerIdx].Shape());
-			Total += UE::Learning::Array::SerializationByteNum<1, float>(Biases[LayerIdx].Shape());
+			UpdatedModel = RuntimeCPU->CreateModelCPU(ModelData);
 		}
 
-		return Total;
+		// If we are not in the editor we can now clear the FileData and FileType since these will be
+		// using additional memory and we are not going to save this asset and so don't require them.
+
+#if !WITH_EDITOR
+		ModelData->ClearFileDataAndFileType();
+#endif
+
+		Network->UpdateModel(UpdatedModel, InputSize, OutputSize);
 	}
+}
 
-	TSharedRef<INeuralNetworkInference> FNeuralNetworkMLP::CreateInferenceObject(
+namespace UE::Learning
+{
+	TSharedRef<FNeuralNetworkInference> FNeuralNetwork::CreateInferenceObject(
 		const int32 MaxBatchSize,
 		const FNeuralNetworkInferenceSettings& Settings)
 	{
-		return MakeShared<FNeuralNetworkInferenceMLP>(*this, MaxBatchSize, Settings);
-	}
+		TSharedRef<FNeuralNetworkInference> InferenceObject = MakeShared<FNeuralNetworkInference>(
+			*Model,
+			MaxBatchSize,
+			InputSize,
+			OutputSize,
+			Settings);
 
-	int32 FNeuralNetworkMLP::GetInputNum() const
-	{
-		UE_LEARNING_CHECK(Weights.Num() > 0);
-		return Weights[0].Num<0>();
+		InferenceObjects.Emplace(InferenceObject.ToWeakPtr());
+
+		return InferenceObject;
 	}
 
-	int32 FNeuralNetworkMLP::GetOutputNum() const
+	bool FNeuralNetwork::IsEmpty() const
 	{
-		UE_LEARNING_CHECK(Weights.Num() > 0);
-		return Weights.Last().Num<1>();
+		return Model != nullptr;
 	}
-
-	int32 FNeuralNetworkMLP::GetHiddenNum() const
+	
+	int32 FNeuralNetwork::GetInputSize() const
 	{
-		UE_LEARNING_CHECK(Weights.Num() > 0);
-		return Weights[0].Num<1>();
+		return InputSize;
 	}
 
-	int32 FNeuralNetworkMLP::GetLayerNum() const
+	int32 FNeuralNetwork::GetOutputSize() const
 	{
-		return Weights.Num();
+		return OutputSize;
 	}
 
-	void FNeuralNetworkMLP::Resize(
-		const int32 InputNum,
-		const int32 OutputNum,
-		const int32 HiddenNum,
-		const int32 LayerNum)
+	void FNeuralNetwork::UpdateModel(const TSharedPtr<NNE::IModelCPU>& InModel, const int32 InInputSize, const int32 InOutputSize)
 	{
-		UE_LEARNING_TRACE_CPUPROFILER_EVENT_SCOPE(Learning::FNeuralNetworkMLP::Resize);
-		UE_LEARNING_CHECKF(LayerNum >= 2, TEXT("At least two layers required (input and output layers)"));
-
-		Weights.SetNum(LayerNum);
-		Biases.SetNum(LayerNum);
+		// Update Model
 
-		Weights[0].SetNumUninitialized({ InputNum, HiddenNum });
-		Biases[0].SetNumUninitialized({ HiddenNum });
+		InputSize = InInputSize;
+		OutputSize = InOutputSize;
+		Model = InModel;
 
-		for (int32 LayerIdx = 0; LayerIdx < LayerNum - 2; LayerIdx++)
+		if (Model)
 		{
-			Weights[LayerIdx + 1].SetNumUninitialized({ HiddenNum, HiddenNum });
-			Biases[LayerIdx + 1].SetNumUninitialized({ HiddenNum });
-		}
+			// We need to tell all the created inference objects to update their instances since the Model has changed
+			// and so the instance objects themselves are either no longer valid or tied to the old model.
+
+			for (const TWeakPtr<UE::Learning::FNeuralNetworkInference>& InferenceObject : InferenceObjects)
+			{
+				if (TSharedPtr<UE::Learning::FNeuralNetworkInference> InferenceObjectPtr = InferenceObject.Pin())
+				{
+					InferenceObjectPtr->ReloadModelInstances(*Model);
+				}
+			}
 
-		Weights[LayerNum - 1].SetNumUninitialized({ HiddenNum, OutputNum });
-		Biases[LayerNum - 1].SetNumUninitialized({ OutputNum });
+			// Remove all invalid weak pointers
 
-		for (int32 LayerIdx = 0; LayerIdx < LayerNum; LayerIdx++)
+			InferenceObjects.RemoveAll(
+				[](const TWeakPtr<UE::Learning::FNeuralNetworkInference>& InferenceObject) { return !InferenceObject.IsValid(); });
+		}
+		else
 		{
-			Array::Zero(Weights[LayerIdx]);
-			Array::Zero(Biases[LayerIdx]);
+			InferenceObjects.Empty();
 		}
 	}
 
-	FNeuralNetworkInferenceMLP::FNeuralNetworkInferenceMLP(const FNeuralNetworkMLP& InNetwork, const int32 MaxBatchSize, const FNeuralNetworkInferenceSettings& InSettings)
-		: Network(InNetwork)
-		, Settings(InSettings)
+	FNeuralNetworkInference::FNeuralNetworkInference(
+		UE::NNE::IModelCPU& InModel,
+		const int32 MaxBatchSize,
+		const int32 InInputSize,
+		const int32 InOutputSize,
+		const FNeuralNetworkInferenceSettings& InSettings)
+		: Settings(InSettings)
+	{
+		InputBuffer.SetNumUninitialized({ MaxBatchSize, InInputSize });
+		OutputBuffer.SetNumUninitialized({ MaxBatchSize, InOutputSize });
+
+		ReloadModelInstances(InModel);
+	}
+
+	void FNeuralNetworkInference::ReloadModelInstances(NNE::IModelCPU& InModel)
 	{
-		const int32 BufferSize = FMath::Max(FMath::Max(Network.GetInputNum(), Network.GetOutputNum()), Network.GetHiddenNum());
+		UE_LEARNING_TRACE_CPUPROFILER_EVENT_SCOPE(Learning::FNeuralNetworkInference::ReloadModelInstances);
+
+		const int32 MaxBatchSize = InputBuffer.Num<0>();
+		const int32 InputSize = InputBuffer.Num<1>();
+		const int32 OutputSize = OutputBuffer.Num<1>();
 
-		FrontBuffer.SetNumUninitialized({ MaxBatchSize, BufferSize });
-		BackBuffer.SetNumUninitialized({ MaxBatchSize, BufferSize });
+		if (Settings.bParallelEvaluation)
+		{
+			const int32 IdealInferenceInstanceNum = ParallelForImpl::GetNumberOfThreadTasks(MaxBatchSize, Settings.MinParallelBatchSize, EParallelForFlags::None);
+			const int32 InferenceInstanceSliceLength = FMath::DivideAndRoundUp((int32)MaxBatchSize, IdealInferenceInstanceNum);
+			const int32 InferenceInstanceNum = FMath::DivideAndRoundUp((int32)MaxBatchSize, InferenceInstanceSliceLength);
+
+			ModelInstances.Empty(InferenceInstanceNum);
+
+			for (int32 InferenceInstanceIdx = 0; InferenceInstanceIdx < InferenceInstanceNum; InferenceInstanceIdx++)
+			{
+				ModelInstances.Emplace(InModel.CreateModelInstanceCPU());
+				ModelInstances.Last()->SetInputTensorShapes({ NNE::FTensorShape::Make({(uint32)InferenceInstanceSliceLength, (uint32)InputSize}) });
+			}
+		}
+		else
+		{
+			ModelInstances.Empty(1);
+			ModelInstances.Emplace(InModel.CreateModelInstanceCPU());
+			ModelInstances.Last()->SetInputTensorShapes({ NNE::FTensorShape::Make({(uint32)MaxBatchSize, (uint32)InputSize}) });
+		}
 	}
 
-	void FNeuralNetworkInferenceMLP::Evaluate(
+	void FNeuralNetworkInference::Evaluate(
 		TLearningArrayView<2, float> Output,
 		const TLearningArrayView<2, const float> Input,
-		const FIndexSet Instances) 
+		const FIndexSet Instances)
 	{
-		const int32 InputNum = Network.GetInputNum();
-		const int32 OutputNum = Network.GetOutputNum();
-		const int32 LayerNum = Network.GetLayerNum();
+		UE_LEARNING_TRACE_CPUPROFILER_EVENT_SCOPE(Learning::FNeuralNetworkInference::Evaluate);
 
-		UE_LEARNING_CHECK(Input.Num<1>() == InputNum);
-		UE_LEARNING_CHECK(Output.Num<1>() == OutputNum);
+		const int32 InstanceNum = Instances.Num();
+		const int32 InputNum = Input.Num<1>();
+		const int32 OutputNum = Output.Num<1>();
 
-		// Gather Instance Inputs
+		if (InstanceNum == 0) { return; }
 
-		Array::Check(Input, Instances);
+		// Evaluation is not thread-safe due to the scatter and gather we do into the InputBuffer and OutputBuffer
+		// so we take a lock here since the interface for this function makes it appear like this might be thread safe as 
+		// long as the Instances provided are not an overlapping set.
+		FScopeNullableWriteLock ScopeLock(&EvaluationLock);
+
+		// Gather Inputs
 
-		const TLearningArrayView<2, float> InputBuffer = FrontBuffer;
+		Array::Check(Input, Instances);
 
-		for (int32 InstanceIdx = 0; InstanceIdx < Instances.Num(); InstanceIdx++)
+		for (int32 InstanceIdx = 0; InstanceIdx < InstanceNum; InstanceIdx++)
 		{
-			Array::Copy(InputBuffer[InstanceIdx].Slice(0, InputNum), Input[Instances[InstanceIdx]]);
+			Array::Copy(InputBuffer[InstanceIdx], Input[Instances[InstanceIdx]]);
 		}
 
-		// Evaluate Network
+		// Run Inference
 
-		auto EvaluationFunction = [this, LayerNum](const int32 BatchOffset, const int32 BatchSize)
+		Array::Check(InputBuffer.Slice(0, InstanceNum));
+
+		if (Settings.bParallelEvaluation && InstanceNum > Settings.MinParallelBatchSize)
 		{
-			for (int32 LayerIdx = 0; LayerIdx < LayerNum; LayerIdx++)
+			const int32 IdealInferenceInstanceNum = ParallelForImpl::GetNumberOfThreadTasks(InstanceNum, Settings.MinParallelBatchSize, EParallelForFlags::None);
+			const int32 InferenceInstanceSliceLength = FMath::DivideAndRoundUp(InstanceNum, IdealInferenceInstanceNum);
+			const int32 InferenceInstanceNum = FMath::DivideAndRoundUp(InstanceNum, InferenceInstanceSliceLength);
+
+			UE_LEARNING_CHECK(InferenceInstanceNum <= ModelInstances.Num());
+
+			ParallelFor(InferenceInstanceNum, [this, InstanceNum, InputNum, InferenceInstanceSliceLength](int32 InferenceInstanceIdx)
 			{
-				const TLearningArrayView<2, float> LayerInputBuffer = LayerIdx % 2 == 0 ? FrontBuffer : BackBuffer;
-				const TLearningArrayView<2, float> LayerOutputBuffer = LayerIdx % 2 == 0 ? BackBuffer : FrontBuffer;
-				const int32 LayerOutputNum = Network.Weights[LayerIdx].Num<1>();
-
-				NeuralNetwork::MatMulPlusBias(
-					LayerOutputBuffer,
-					LayerInputBuffer,
-					Network.Weights[LayerIdx],
-					Network.Biases[LayerIdx],
-					BatchOffset,
-					BatchSize);
-
-				if (LayerIdx != LayerNum - 1)
-				{
-					switch (Network.ActivationFunction)
-					{
-					case EActivationFunction::ReLU: NeuralNetwork::ActivationReLU(LayerOutputBuffer, BatchOffset, BatchSize, LayerOutputNum); break;
-					case EActivationFunction::ELU: NeuralNetwork::ActivationELU(LayerOutputBuffer, BatchOffset, BatchSize, LayerOutputNum); break;
-					case EActivationFunction::TanH: NeuralNetwork::ActivationTanH(LayerOutputBuffer, BatchOffset, BatchSize, LayerOutputNum); break;
-					default: UE_LEARNING_CHECKF(false, TEXT("Unimplemented"));
-					}
-				}
-			}
-		};
+				const int32 StartIndex = InferenceInstanceIdx * InferenceInstanceSliceLength;
+				const int32 StopIndex = FMath::Min((InferenceInstanceIdx + 1) * InferenceInstanceSliceLength, InstanceNum);
 
-		if (Settings.bParallelEvaluation && Instances.Num() > Settings.MinParallelBatchSize)
-		{
-			Learning::SlicedParallelFor(Instances.Num(), Settings.MinParallelBatchSize, EvaluationFunction);
+				TLearningArrayView<2, float> InputBufferSlice = InputBuffer.Slice(StartIndex, StopIndex - StartIndex);
+				TLearningArrayView<2, float> OutputBufferSlice = OutputBuffer.Slice(StartIndex, StopIndex - StartIndex);
+
+				ModelInstances[InferenceInstanceIdx]->SetInputTensorShapes({ NNE::FTensorShape::Make({(uint32)(StopIndex - StartIndex), (uint32)InputNum}) });
+				ModelInstances[InferenceInstanceIdx]->RunSync(
+					{ { (void*)InputBufferSlice.GetData(), InputBufferSlice.Num() * sizeof(float) } },
+					{ { (void*)OutputBufferSlice.GetData(), OutputBufferSlice.Num() * sizeof(float) } });
+			});
 		}
 		else
 		{
-			EvaluationFunction(0, Instances.Num());
+			ModelInstances[0]->SetInputTensorShapes({ NNE::FTensorShape::Make({(uint32)InstanceNum, (uint32)InputNum}) });
+			ModelInstances[0]->RunSync(
+				{ { (void*)InputBuffer.GetData(), InputBuffer.Slice(0, InstanceNum).Num() * sizeof(float) } },
+				{ { (void*)OutputBuffer.GetData(), OutputBuffer.Slice(0, InstanceNum).Num() * sizeof(float) } });
 		}
 
-		// Scatter Instance Outputs
+		Array::Check(OutputBuffer.Slice(0, InstanceNum));
 
-		const TLearningArrayView<2, float> OutputBuffer = LayerNum % 2 == 0 ? FrontBuffer : BackBuffer;
+		// Scatter Outputs
 
-		for (int32 InstanceIdx = 0; InstanceIdx < Instances.Num(); InstanceIdx++)
+		for (int32 InstanceIdx = 0; InstanceIdx < InstanceNum; InstanceIdx++)
 		{
-			Array::Copy(Output[Instances[InstanceIdx]], OutputBuffer[InstanceIdx].Slice(0, OutputNum));
+			Array::Copy(Output[Instances[InstanceIdx]], OutputBuffer[InstanceIdx]);
 		}
 
 		Array::Check(Output, Instances);
 	}
+
+
+	FNeuralNetworkFunction::FNeuralNetworkFunction(
+		const int32 InMaxInstanceNum,
+		const TSharedPtr<FNeuralNetwork>& InNeuralNetwork,
+		const FNeuralNetworkInferenceSettings& InInferenceSettings) 
+		: MaxInstanceNum(InMaxInstanceNum)
+		, NeuralNetwork(InNeuralNetwork)
+		, InferenceSettings(InInferenceSettings)
+	{
+		NeuralNetworkInference = NeuralNetwork->CreateInferenceObject(InMaxInstanceNum, InInferenceSettings);
+	}
+
+	void FNeuralNetworkFunction::Evaluate(
+		TLearningArrayView<2, float> Output,
+		const TLearningArrayView<2, const float> Input,
+		const FIndexSet Instances)
+	{
+		NeuralNetworkInference->Evaluate(Output, Input, Instances);
+	}
+
+	void FNeuralNetworkFunction::UpdateNeuralNetwork(const TSharedPtr<FNeuralNetwork>& NewNeuralNetwork)
+	{
+		if (NeuralNetwork != NewNeuralNetwork)
+		{
+			NeuralNetwork = NewNeuralNetwork;
+			NeuralNetworkInference = NeuralNetwork->CreateInferenceObject(MaxInstanceNum, InferenceSettings);
+		}
+	}
+
 }
diff --git a/Engine/Plugins/Experimental/LearningAgents/Source/Learning/Private/LearningPolicy.cpp b/Engine/Plugins/Experimental/LearningAgents/Source/Learning/Private/LearningPolicy.cpp
new file mode 100644
--- /dev/null
+++ b/Engine/Plugins/Experimental/LearningAgents/Source/Learning/Private/LearningPolicy.cpp
@@ -0,0 +1,79 @@
+// Copyright Epic Games, Inc. All Rights Reserved.
+
+#include "LearningPolicy.h"
+
+namespace UE::Learning
+{
+	FNeuralNetworkPolicy::FNeuralNetworkPolicy(
+		const int32 InMaxInstanceNum,
+		const int32 InObservationEncodedNum,
+		const int32 InActionEncodedNum,
+		const int32 InMemoryStateNum,
+		const TSharedPtr<FNeuralNetwork>& InNeuralNetwork,
+		const FNeuralNetworkInferenceSettings& InInferenceSettings)
+		: MaxInstanceNum(InMaxInstanceNum)
+		, ObservationEncodedNum(InObservationEncodedNum)
+		, ActionEncodedNum(InActionEncodedNum)
+		, MemoryStateNum(InMemoryStateNum)
+		, NeuralNetwork(InNeuralNetwork)
+		, InferenceSettings(InInferenceSettings)
+	{
+		UE_LEARNING_CHECK(NeuralNetwork->GetInputSize() == ObservationEncodedNum + MemoryStateNum);
+		UE_LEARNING_CHECK(NeuralNetwork->GetOutputSize() == ActionEncodedNum + MemoryStateNum);
+
+		Input.SetNumUninitialized({ MaxInstanceNum, ObservationEncodedNum + MemoryStateNum });
+		Output.SetNumUninitialized({ MaxInstanceNum, ActionEncodedNum + MemoryStateNum });
+
+		Array::Zero(Input);
+		Array::Zero(Output);
+
+		NeuralNetworkInference = NeuralNetwork->CreateInferenceObject(MaxInstanceNum, InferenceSettings);
+	}
+
+	void FNeuralNetworkPolicy::Evaluate(
+		TLearningArrayView<2, float> OutputActionVectorsEncoded,
+		TLearningArrayView<2, float> OutputMemoryState,
+		const TLearningArrayView<2, const float> InputObservationVectorsEncoded,
+		const TLearningArrayView<2, const float> InputMemoryState, 
+		const FIndexSet Instances)
+	{
+		UE_LEARNING_TRACE_CPUPROFILER_EVENT_SCOPE(Learning::FNeuralNetworkPolicy::Evaluate);
+
+		const int32 InstanceNum = Instances.Num();
+
+		Array::Check(InputObservationVectorsEncoded, Instances);
+		Array::Check(InputMemoryState, Instances);
+
+		// Copy in Observation and Memory State into network input
+
+		for (const int32 InstanceIdx : Instances)
+		{
+			Array::Copy(Input[InstanceIdx].Slice(0, ObservationEncodedNum), InputObservationVectorsEncoded[InstanceIdx]);
+			Array::Copy(Input[InstanceIdx].Slice(ObservationEncodedNum, MemoryStateNum), InputMemoryState[InstanceIdx]);
+		}
+
+		NeuralNetworkInference->Evaluate(Output, Input, Instances);
+
+		// Copy Out Memory State
+
+		for (const int32 InstanceIdx : Instances)
+		{
+			Array::Copy(OutputActionVectorsEncoded[InstanceIdx], Output[InstanceIdx].Slice(0, ActionEncodedNum));
+			Array::Copy(OutputMemoryState[InstanceIdx], Output[InstanceIdx].Slice(ActionEncodedNum, MemoryStateNum));
+		}
+
+		Array::Check(OutputActionVectorsEncoded, Instances);
+		Array::Check(OutputMemoryState, Instances);
+	}
+
+	void FNeuralNetworkPolicy::UpdateNeuralNetwork(const TSharedPtr<FNeuralNetwork>& NewNeuralNetwork)
+	{
+		if (NeuralNetwork != NewNeuralNetwork)
+		{
+			UE_LEARNING_CHECK(NewNeuralNetwork->GetInputSize() == ObservationEncodedNum + MemoryStateNum);
+			UE_LEARNING_CHECK(NewNeuralNetwork->GetOutputSize() == ActionEncodedNum + MemoryStateNum);
+			NeuralNetwork = NewNeuralNetwork;
+			NeuralNetworkInference = NeuralNetwork->CreateInferenceObject(MaxInstanceNum, InferenceSettings);
+		}
+	}
+}
diff --git a/Engine/Plugins/Experimental/LearningAgents/Source/Learning/Private/LearningRandom.cpp b/Engine/Plugins/Experimental/LearningAgents/Source/Learning/Private/LearningRandom.cpp
--- a/Engine/Plugins/Experimental/LearningAgents/Source/Learning/Private/LearningRandom.cpp
+++ b/Engine/Plugins/Experimental/LearningAgents/Source/Learning/Private/LearningRandom.cpp
@@ -13,6 +13,11 @@ namespace UE::Learning::Random
 		return FMath::Sqrt(-2.0f * FMath::Loge(FMath::Max(R1, UE_SMALL_NUMBER))) * FMath::Cos(R2 * UE_TWO_PI);
 	}
 
+	static inline float Sigmoid(const float X)
+	{
+		return 1.0f / (1.0f + FMath::Exp(-X));
+	}
+
 	uint32 Int(const uint32 State)
 	{
 		// Here we use a simple xor and shift based 
@@ -143,6 +148,28 @@ namespace UE::Learning::Random
 #endif
 	}
 
+	LEARNING_API void IntArray(
+		TLearningArrayView<1, uint32> Output,
+		const uint32 State,
+		const FIndexSet Indices)
+	{
+		UE_LEARNING_TRACE_CPUPROFILER_EVENT_SCOPE(Learning::Random::IntArray);
+
+		if (Indices.IsSlice() && Indices.Num() > 4)
+		{
+			IntArray(
+				Output.Slice(Indices.GetSliceStart(), Indices.GetSliceNum()),
+				State);
+		}
+		else
+		{
+			for (const int32 ElementIdx : Indices)
+			{
+				Output[ElementIdx] = Int(State ^ 0xbed25b30 ^ Int(ElementIdx ^ 0xb521a8d9));
+			}
+		}
+	}
+
 	void FloatArray(
 		TLearningArrayView<1, float> Output,
 		const uint32 State)
@@ -249,6 +276,126 @@ namespace UE::Learning::Random
 		}
 	}
 
+	void DistributionIndependantNormal(
+		TLearningArrayView<1, float> Output,
+		const uint32 State,
+		const TLearningArrayView<1, const float> Mean,
+		const TLearningArrayView<1, const float> LogStd,
+		const float Scale)
+	{
+		UE_LEARNING_TRACE_CPUPROFILER_EVENT_SCOPE(Learning::Random::DistributionIndependantNormal);
+
+		const int32 ElementNum = Output.Num();
+
+#if UE_LEARNING_ISPC
+		ispc::LearningRandomDistributionIndependantNormal(
+			Output.GetData(),
+			Mean.GetData(),
+			LogStd.GetData(),
+			ElementNum,
+			State,
+			Scale);
+#else
+		if (Scale < UE_SMALL_NUMBER)
+		{
+			for (int32 ElementIdx = 0; ElementIdx < ElementNum; ElementIdx++)
+			{
+				Output[ElementIdx] = Mean[ElementIdx];
+			}
+		}
+		else
+		{
+			for (int32 ElementIdx = 0; ElementIdx < ElementNum; ElementIdx++)
+			{
+				Output[ElementIdx] = Gaussian(State ^ 0x7db0e4e9 ^ Int(ElementIdx ^ 0xf4976a00), Mean[ElementIdx], Scale * FMath::Exp(FMath::Min(LogStd[ElementIdx], 10.0f)));
+			}
+		}
+#endif
+	}
+
+	void DistributionMultinoulli(
+		TLearningArrayView<1, float> Output,
+		const uint32 State,
+		const TLearningArrayView<1, const float> Logits,
+		const float Scale)
+	{
+		UE_LEARNING_TRACE_CPUPROFILER_EVENT_SCOPE(Learning::Random::DistributionMultinoulli);
+
+		const int32 ElementNum = Output.Num();
+		if (ElementNum == 0)
+		{
+			return;
+		}
+
+		float MaxValue = -FLT_MAX;
+		uint32 MaxIdx = INDEX_NONE;
+
+		if (Scale < UE_SMALL_NUMBER)
+		{
+			for (int32 ElementIdx = 0; ElementIdx < ElementNum; ElementIdx++)
+			{
+				if (Logits[ElementIdx] > MaxValue)
+				{
+					MaxValue = Logits[ElementIdx];
+					MaxIdx = ElementIdx;
+				}
+			}
+		}
+		else
+		{
+			for (int32 ElementIdx = 0; ElementIdx < ElementNum; ElementIdx++)
+			{
+				const float ElementValue = Logits[ElementIdx] / Scale - FMath::Loge(-FMath::Loge(Uniform(State ^ 0x7a156b37 ^ Int(ElementIdx ^ 0x0c71bebb))));
+
+				if (ElementValue > MaxValue)
+				{
+					MaxValue = ElementValue;
+					MaxIdx = ElementIdx;
+				}
+			}
+		}
+
+		UE_LEARNING_CHECK(MaxIdx != INDEX_NONE);
+
+		Array::Zero(Output);
+		Output[MaxIdx] = 1.0f;
+	}
+
+	void DistributionBernoulli(
+		TLearningArrayView<1, float> Output,
+		const uint32 State,
+		const TLearningArrayView<1, const float> Logits,
+		const float Scale)
+	{
+		UE_LEARNING_TRACE_CPUPROFILER_EVENT_SCOPE(Learning::Random::DistributionBernoulli);
+
+		const int32 ElementNum = Output.Num();
+
+#if UE_LEARNING_ISPC
+		ispc::LearningRandomDistributionBernoulli(
+			Output.GetData(),
+			Logits.GetData(),
+			ElementNum,
+			State,
+			Scale);
+#else
+		if (Scale < UE_SMALL_NUMBER)
+		{
+			for (int32 ElementIdx = 0; ElementIdx < ElementNum; ElementIdx++)
+			{
+				Output[ElementIdx] = Logits[ElementIdx] > 0.0f ? 1.0f : 0.0f;
+			}
+		}
+		else
+		{
+			for (int32 ElementIdx = 0; ElementIdx < ElementNum; ElementIdx++)
+			{
+				Output[ElementIdx] = Sigmoid(Logits[ElementIdx] / Scale) > Uniform(State ^ 0xf4021a46 ^ Int(ElementIdx ^ 0x7a8cc64e)) ? 1.0f : 0.0f;
+			}
+		}
+#endif
+	}
+
 	////////////
 
 	uint32 SampleInt(uint32& State)
@@ -354,6 +501,15 @@ namespace UE::Learning::Random
 		IntArray(Output, State ^ 0x6adcdb41);
 	}
 
+	void SampleIntArray(
+		TLearningArrayView<1, uint32> Output,
+		uint32& State,
+		const FIndexSet Indices)
+	{
+		State = Int(State ^ 0xab9c9ee3);
+		IntArray(Output, State ^ 0x6adcdb41, Indices);
+	}
+
 	void SampleFloatArray(
 		TLearningArrayView<1, float> Output,
 		uint32& State)
@@ -382,29 +538,6 @@ namespace UE::Learning::Random
 		GaussianArray(Output, State ^ 0xf2381309, Mean, Std);
 	}
 
-	void ResampleStateArray(TLearningArrayView<1, uint32> InOutStates, const FIndexSet Indices)
-	{
-		UE_LEARNING_TRACE_CPUPROFILER_EVENT_SCOPE(Learning::Random::ResampleStateArray);
-
-		const int32 ElementNum = Indices.Num();
-
-		if (UE_LEARNING_ISPC && Indices.IsSlice() && Indices.Num() > 4)
-		{
-#if UE_LEARNING_ISPC
-			ispc::LearningResampleStateArray(
-				InOutStates.Slice(Indices.GetSliceStart(), Indices.GetSliceNum()).GetData(),
-				Indices.GetSliceNum());
-#endif
-		}
-		else
-		{
-			for (const int32 ElementIdx : Indices)
-			{
-				InOutStates[ElementIdx] = Int(InOutStates[ElementIdx] ^ 0x3616dcc8 ^ Int(ElementIdx ^ 0x6a837ffd));
-			}
-		}
-	}
-
 	void SamplePlanarClippedGaussianArray(
 		TLearningArrayView<1, FVector> Output,
 		uint32& State,
@@ -428,4 +561,48 @@ namespace UE::Learning::Random
 		PlanarDirectionArray(Output, State ^ 0x6bfd3e6a, Axis0, Axis1);
 	}
 
+	void SampleDistributionIndependantNormal(
+		TLearningArrayView<1, float> Output,
+		uint32& State,
+		const TLearningArrayView<1, const float> Mean,
+		const TLearningArrayView<1, const float> LogStd,
+		const float Scale)
+	{
+		State = Int(State ^ 0x984c4d5f);
+		DistributionIndependantNormal(
+			Output,
+			State,
+			Mean,
+			LogStd,
+			Scale);
+	}
+
+	void SampleDistributionMultinoulli(
+		TLearningArrayView<1, float> Output,
+		uint32& State,
+		const TLearningArrayView<1, const float> Logits,
+		const float Scale)
+	{
+		State = Int(State ^ 0xc9261d09);
+		DistributionMultinoulli(
+			Output,
+			State,
+			Logits,
+			Scale);
+	}
+
+	void SampleDistributionBernoulli(
+		TLearningArrayView<1, float> Output,
+		uint32& State,
+		const TLearningArrayView<1, const float> Logits,
+		const float Scale)
+	{
+		State = Int(State ^ 0xfeba3d63);
+		DistributionBernoulli(
+			Output,
+			State,
+			Logits,
+			Scale);
+	}
+
 }
\ No newline at end of file
diff --git a/Engine/Plugins/Experimental/LearningAgents/Source/Learning/Public/LearningAction.h b/Engine/Plugins/Experimental/LearningAgents/Source/Learning/Public/LearningAction.h
new file mode 100644
--- /dev/null
+++ b/Engine/Plugins/Experimental/LearningAgents/Source/Learning/Public/LearningAction.h
@@ -0,0 +1,490 @@
+// Copyright Epic Games, Inc. All Rights Reserved.
+
+#pragma once
+
+#include "LearningArray.h"
+
+#include "UObject/NameTypes.h"
+
+namespace UE::Learning::Action
+{
+	/**
+	 * Action Type
+	 * 
+	 * The core type of an action object for which different loss functions and behaviors are defined.
+	 */
+	enum class EType : uint8
+	{
+		// Empty Action
+		Null				= 0,	
+
+		// Vector of continuous float actions
+		Continuous			= 1,
+
+		// Set of exclusive discrete actions
+		DiscreteExclusive	= 2,
+
+		// Set of inclusive discrete actions
+		DiscreteInclusive	= 3,
+
+		// Combination of multiple actions
+		And					= 4,
+
+		// Exclusive choice from a set of actions
+		OrExclusive			= 5,
+
+		// Inclusive choice from a set of actions
+		OrInclusive			= 6,
+
+		// Fixed sized array of actions
+		Array				= 7,
+
+		// Encoding of another action
+		Encoding			= 8,
+	};
+
+	/**
+	 * Action Schema Element
+	 * 
+	 * A single element in the action schema representing a part of an action. Internally this consists of a index used by the schema to look up 
+	 * the associated action data and a generation id which can be used to check when this index is no longer valid.
+	 */
+	struct FSchemaElement { int32 Index = INDEX_NONE; uint32 Generation = INDEX_NONE; };
+
+	struct FSchemaContinuousParameters
+	{
+		// Number of values in the continuous action
+		int32 Num = 0;
+	};
+
+	struct FSchemaDiscreteExclusiveParameters
+	{
+		// Number of choices in the discrete action
+		int32 Num = 0;
+
+		// Prior Probabilities of the discrete actions.
+		TArrayView<const float> PriorProbabilities;
+	};
+
+	struct FSchemaDiscreteInclusiveParameters
+	{
+		// Number of choices in the discrete action
+		int32 Num = 0;
+
+		// Prior Probabilities of the discrete actions.
+		TArrayView<const float> PriorProbabilities;
+	};
+
+	struct FSchemaAndParameters
+	{
+		// Names of the sub-actions.
+		TArrayView<const FName> ElementNames;
+
+		// The associated sub-elements.
+		TArrayView<const FSchemaElement> Elements;
+	};
+
+	struct FSchemaOrExclusiveParameters
+	{
+		// Names of the sub-actions.
+		TArrayView<const FName> ElementNames;
+
+		// The associated sub-elements.
+		TArrayView<const FSchemaElement> Elements;
+
+		// Prior Probabilities of the sub-actions.
+		TArrayView<const float> PriorProbabilities;
+	};
+
+	struct FSchemaOrInclusiveParameters
+	{
+		// Names of the sub-actions.
+		TArrayView<const FName> ElementNames;
+
+		// The associated sub-elements.
+		TArrayView<const FSchemaElement> Elements;
+
+		// Prior Probabilities of the sub-actions.
+		TArrayView<const float> PriorProbabilities;
+	};
+
+	struct FSchemaArrayParameters
+	{
+		// The array sub-element.
+		FSchemaElement Element;
+
+		// The number of elements in the array.
+		int32 Num = 0;
+	};
+
+	struct FSchemaEncodingParameters
+	{
+		// The sub-element.
+		FSchemaElement Element;
+
+		// The size at which the sub-element should be encoded.
+		int32 EncodingSize = 32;
+	};
+
+	/**
+	 * Action Schema
+	 *
+	 * This object allows you to construct a description of the kind of actions you might want to be generated by a policy. Internally this object
+	 * contains a pool of individual elements. This allows them to be constructed performantly and in a cache efficient way. This object is therefore
+	 * required to access any data about the individual action elements that are created.
+	 */
+	struct LEARNING_API FSchema
+	{
+		FSchemaElement CreateNull(const FName Name = NAME_None);
+		FSchemaElement CreateContinuous(const FSchemaContinuousParameters Parameters, const FName Name = NAME_None);
+		FSchemaElement CreateDiscreteExclusive(const FSchemaDiscreteExclusiveParameters Parameters, const FName Name = NAME_None);
+		FSchemaElement CreateDiscreteInclusive(const FSchemaDiscreteInclusiveParameters Parameters, const FName Name = NAME_None);
+		FSchemaElement CreateAnd(const FSchemaAndParameters Parameters, const FName Name = NAME_None);
+		FSchemaElement CreateOrExclusive(const FSchemaOrExclusiveParameters Parameters, const FName Name = NAME_None);
+		FSchemaElement CreateOrInclusive(const FSchemaOrInclusiveParameters Parameters, const FName Name = NAME_None);
+		FSchemaElement CreateArray(const FSchemaArrayParameters Parameters, const FName Name = NAME_None);
+		FSchemaElement CreateEncoding(const FSchemaEncodingParameters Parameters, const FName Name = NAME_None);
+
+		FSchemaContinuousParameters GetContinuous(const FSchemaElement Element) const;
+		FSchemaDiscreteExclusiveParameters GetDiscreteExclusive(const FSchemaElement Element) const;
+		FSchemaDiscreteInclusiveParameters GetDiscreteInclusive(const FSchemaElement Element) const;
+		FSchemaAndParameters GetAnd(const FSchemaElement Element) const;
+		FSchemaOrExclusiveParameters GetOrExclusive(const FSchemaElement Element) const;
+		FSchemaOrInclusiveParameters GetOrInclusive(const FSchemaElement Element) const;
+		FSchemaArrayParameters GetArray(const FSchemaElement Element) const;
+		FSchemaEncodingParameters GetEncoding(const FSchemaElement Element) const;
+
+		// Checks if the given element is valid
+		bool IsValid(const FSchemaElement Element) const;
+
+		// Gets the type of the given element
+		EType GetType(const FSchemaElement Element) const;
+
+		// Gets the name of the given element
+		FName GetName(const FSchemaElement Element) const;
+
+		// Get the encoded vector size of the given element
+		int32 GetEncodedVectorSize(const FSchemaElement Element) const;
+
+		// Get the action vector size of the given element
+		int32 GetActionVectorSize(const FSchemaElement Element) const;
+
+		// Get the action distribution vector size of the given element
+		int32 GetActionDistributionVectorSize(const FSchemaElement Element) const;
+
+		// Get the current generation
+		uint32 GetGeneration() const;
+		
+		// Empty all internal buffers of elements. This invalidates all existing elements.
+		void Empty();
+
+		// Reset all internal buffers (without freeing memory). This invalidates all existing elements.
+		void Reset();
+
+	private:
+
+		struct FContinuousData
+		{
+			int32 Num = INDEX_NONE;
+		};
+
+		struct FDiscreteExclusiveData
+		{
+			int32 Num = INDEX_NONE;
+			int32 PriorProbabilitiesOffset = INDEX_NONE;
+		};
+
+		struct FDiscreteInclusiveData
+		{
+			int32 Num = INDEX_NONE;
+			int32 PriorProbabilitiesOffset = INDEX_NONE;
+		};
+
+		struct FAndData
+		{
+			int32 Num = INDEX_NONE;
+			int32 ElementsOffset = INDEX_NONE;
+		};
+
+		struct FOrExclusiveData
+		{
+			int32 Num = INDEX_NONE;
+			int32 ElementsOffset = INDEX_NONE;
+			int32 PriorProbabilitiesOffset = INDEX_NONE;
+		};
+
+		struct FOrInclusiveData
+		{
+			int32 Num = INDEX_NONE;
+			int32 ElementsOffset = INDEX_NONE;
+			int32 PriorProbabilitiesOffset = INDEX_NONE;
+		};
+
+		struct FArrayData
+		{
+			int32 Num = INDEX_NONE;
+			int32 ElementIndex = INDEX_NONE;
+		};
+
+		struct FEncodingData
+		{
+			int32 EncodingSize = INDEX_NONE;
+			int32 ElementIndex = INDEX_NONE;
+		};
+
+		uint32 Generation = 0;
+
+		/** These have entries for each Schema Element */
+		TArray<EType> Types;
+		TArray<FName> Names;
+		TArray<int32> EncodedVectorSizes;
+		TArray<int32> ActionVectorSizes;
+		TArray<int32> ActionDistributionVectorSizes;
+		TArray<int32> TypeDataIndices;
+
+		/** These are indexed based on the type and the index found in TypeDataIndices */
+		TArray<FContinuousData> ContinuousData;
+		TArray<FDiscreteExclusiveData> DiscreteExclusiveData;
+		TArray<FDiscreteInclusiveData> DiscreteInclusiveData;
+		TArray<FAndData> AndData;
+		TArray<FOrExclusiveData> OrExclusiveData;
+		TArray<FOrInclusiveData> OrInclusiveData;
+		TArray<FArrayData> ArrayData;
+		TArray<FEncodingData> EncodingData;
+
+		/** This is an array of all the SubElements and their names, referenced by other elements */
+		TArray<FName> SubElementNames;
+		TArray<FSchemaElement> SubElementObjects;
+
+		/** This is an array of all the prior probabilities, referenced by other elements. */
+		TArray<float> PriorProbabilities;
+	};
+
+	/**
+	 * Action Object Element
+	 *
+	 * A single element in the action object representing part of an action. Internally this consists of a index used by the object to look up
+	 * the associated action data and a generation id which can be used to check when this index is no longer valid.
+	 */
+	struct FObjectElement { int32 Index = INDEX_NONE; uint32 Generation = INDEX_NONE; };
+
+	struct FObjectContinuousParameters
+	{
+		// Continuous action values
+		TArrayView<const float> Values;
+	};
+
+	struct FObjectDiscreteExclusiveParameters
+	{
+		// Exclusive discrete action index
+		int32 DiscreteIndex = INDEX_NONE;
+	};
+
+	struct FObjectDiscreteInclusiveParameters
+	{
+		// Inclusive discrete action indices
+		TArrayView<const int32> DiscreteIndices;
+	};
+
+	struct FObjectAndParameters
+	{
+		// Names of the sub-actions.
+		TArrayView<const FName> ElementNames;
+
+		// The associated sub-elements.
+		TArrayView<const FObjectElement> Elements;
+	};
+
+	struct FObjectOrExclusiveParameters
+	{
+		// Name of the chosen sub-action.
+		FName ElementName;
+
+		// The associated chosen sub-element.
+		FObjectElement Element;
+	};
+
+	struct FObjectOrInclusiveParameters
+	{
+		// Names of the chosen sub-actions.
+		TArrayView<const FName> ElementNames;
+
+		// The associated chosen sub-elements.
+		TArrayView<const FObjectElement> Elements;
+	};
+
+	struct FObjectArrayParameters
+	{
+		// Array of sub-elements.
+		TArrayView<const FObjectElement> Elements;
+	};
+
+	struct FObjectEncodingParameters
+	{
+		// Encoded sub-element.
+		FObjectElement Element;
+	};
+
+	/**
+	 * Action Object
+	 *
+	 * This object allows you to construct or get data from an instance of an action you might have generated from a policy. Internally this object
+	 * contains a pool of individual elements. This allows them to be constructed performantly and in a cache efficient way. This object is therefore
+	 * required to access any data about the individual action elements that are created.
+	 */
+	struct LEARNING_API FObject
+	{
+		FObjectElement CreateNull(const FName Name = NAME_None);
+		FObjectElement CreateContinuous(const FObjectContinuousParameters Parameters, const FName Name = NAME_None);
+		FObjectElement CreateDiscreteExclusive(const FObjectDiscreteExclusiveParameters Parameters, const FName Name = NAME_None);
+		FObjectElement CreateDiscreteInclusive(const FObjectDiscreteInclusiveParameters Parameters, const FName Name = NAME_None);
+		FObjectElement CreateAnd(const FObjectAndParameters Parameters, const FName Name = NAME_None);
+		FObjectElement CreateOrExclusive(const FObjectOrExclusiveParameters Parameters, const FName Name = NAME_None);
+		FObjectElement CreateOrInclusive(const FObjectOrInclusiveParameters Parameters, const FName Name = NAME_None);
+		FObjectElement CreateArray(const FObjectArrayParameters Parameters, const FName Name = NAME_None);
+		FObjectElement CreateEncoding(const FObjectEncodingParameters Parameters, const FName Name = NAME_None);
+
+		FObjectContinuousParameters GetContinuous(const FObjectElement Element) const;
+		FObjectDiscreteExclusiveParameters GetDiscreteExclusive(const FObjectElement Element) const;
+		FObjectDiscreteInclusiveParameters GetDiscreteInclusive(const FObjectElement Element) const;
+		FObjectAndParameters GetAnd(const FObjectElement Element) const;
+		FObjectOrExclusiveParameters GetOrExclusive(const FObjectElement Element) const;
+		FObjectOrInclusiveParameters GetOrInclusive(const FObjectElement Element) const;
+		FObjectArrayParameters GetArray(const FObjectElement Element) const;
+		FObjectEncodingParameters GetEncoding(const FObjectElement Element) const;
+
+		// Checks if the given element is valid
+		bool IsValid(const FObjectElement Element) const;
+
+		// Gets the type of the given element
+		EType GetType(const FObjectElement Element) const;
+
+		// Gets the name of the given element
+		FName GetName(const FObjectElement Element) const;
+
+		// Get the current generation
+		uint32 GetGeneration() const;
+
+		// Empty all internal buffers of elements. This invalidates all existing elements.
+		void Empty();
+
+		// Reset all internal buffers (without freeing memory). This invalidates all existing elements.
+		void Reset();
+
+	private:
+
+		uint32 Generation = 0;
+
+		TArray<EType> Types;
+		TArray<FName> Names;
+		TArray<int32> ContinuousDataOffsets;
+		TArray<int32> ContinuousDataNums;
+		TArray<int32> DiscreteDataOffsets;
+		TArray<int32> DiscreteDataNums;
+		TArray<int32> ElementDataOffsets;
+		TArray<int32> ElementDataNums;
+
+		TArray<float> ContinuousValues;
+		TArray<int32> DiscreteValues;
+		TArray<FName> SubElementNames;
+		TArray<FObjectElement> SubElementObjects;
+	};
+
+	/**
+	 * Gets a hash value representing object compatibility between schemas i.e. if objects from one schema can be used by objects expecting another 
+	 * schema. This is not a crytographic hash, and so `AreSchemaObjectsCompatible` should still be used as the ultimate source of truth.
+	 * This function returns an int32 so that it can be used in blueprints.
+	 * 
+	 * @param Schema					Action Schema
+	 * @param SchemaElement				Action Schema Element
+	 * @param Salt						Hash salt
+	 */
+	LEARNING_API int32 GetSchemaObjectsCompatibilityHash(
+		const FSchema& Schema,
+		const FSchemaElement SchemaElement,
+		const int32 Salt = 0x49710e77);
+
+	/**
+	 * Test if two schemas are compatible i.e. if objects from one schema can be used by objects expecting another schema. 
+	 *
+	 * @param SchemaA					First Schema
+	 * @param SchemaElementA			First Schema Element
+	 * @param SchemaB					Second Schema
+	 * @param SchemaElementB			Second Schema Element
+	 * @returns							true when the schemas are compatible
+	 */
+	LEARNING_API bool AreSchemaObjectsCompatible(
+		const FSchema& SchemaA,
+		const FSchemaElement SchemaElementA,
+		const FSchema& SchemaB,
+		const FSchemaElement SchemaElementB);
+
+	/**
+	 * Generate FileData for a Neural Network that can act as a Decoder for the given schema. This network will take as input a vector of size
+	 * EncodedActionSize and produce a vector of size ActionDistributionVectorSize.
+	 *
+	 * @param OutFileData				Output File Data to write to
+	 * @param OutInputSize				Size of the vector this network takes as input
+	 * @param OutOutputSize				Size of the vector this network takes as output
+	 * @param Schema					Action Schema
+	 * @param SchemaElement				Action Schema Element
+	 * @param Seed						The random seed used in initializing the network weights
+	 */
+	LEARNING_API void GenerateDecoderNetworkFileDataFromSchema(
+		TArray<uint8>& OutFileData,
+		uint32& OutInputSize,
+		uint32& OutOutputSize,
+		const FSchema& Schema,
+		const FSchemaElement SchemaElement,
+		const uint32 Seed = 0x05dbd1f5);
+
+	/**
+	 * Sample an action vector from a vector representing the action distribution.
+	 *
+	 * @param InOutSeed					Seed to use when sampling the action vector
+	 * @param OutActionVector			Output action vector
+	 * @param ActionDistributionVector	Input Action Distribution vector
+	 * @param Schema					Action Schema
+	 * @param SchemaElement				Action Schema Element
+	 * @param ActionNoiseScale			Scale factor for the noise. Set this to zero to always choose the mean (expected) action.
+	 */
+	LEARNING_API void SampleVectorFromDistributionVector(
+		uint32& InOutSeed,
+		TLearningArrayView<1, float> OutActionVector,
+		const TLearningArrayView<1, const float> ActionDistributionVector,
+		const FSchema& Schema,
+		const FSchemaElement SchemaElement,
+		const float ActionNoiseScale = 1.0f);
+
+	/**
+	 * Convert an action object into an action vector.
+	 *
+	 * @param OutActionVector			Output action vector
+	 * @param Schema					Action Schema
+	 * @param SchemaElement				Action Schema Element
+	 * @param Object					Action Object
+	 * @param ObjectElement				Action Object Element
+	 */
+	LEARNING_API void SetVectorFromObject(
+		TLearningArrayView<1, float> OutActionVector,
+		const FSchema& Schema,
+		const FSchemaElement SchemaElement,
+		const FObject& Object,
+		const FObjectElement ObjectElement);
+
+	/**
+	 * Convert an action vector into an action object
+	 *
+	 * @param OutObject					Output Action Object
+	 * @param OutObjectElement			Output Action Object Element
+	 * @param Schema					Action Schema
+	 * @param SchemaElement				Action Schema Element
+	 * @param ActionVector				Input Action vector
+	 */
+	LEARNING_API void GetObjectFromVector(
+		FObject& OutObject,
+		FObjectElement& OutObjectElement,
+		const FSchema& Schema,
+		const FSchemaElement SchemaElement,
+		const TLearningArrayView<1, const float> ActionVector);
+}
\ No newline at end of file
diff --git a/Engine/Plugins/Experimental/LearningAgents/Source/Learning/Public/LearningArray.h b/Engine/Plugins/Experimental/LearningAgents/Source/Learning/Public/LearningArray.h
--- a/Engine/Plugins/Experimental/LearningAgents/Source/Learning/Public/LearningArray.h
+++ b/Engine/Plugins/Experimental/LearningAgents/Source/Learning/Public/LearningArray.h
@@ -310,6 +310,16 @@ namespace UE::Learning
 		InOutOffset += sizeof(int32);
 	}
 
+	/**
+	* Serialize an integer from bytes
+	*/
+	static inline void DeserializeFromBytes(int32& InOutOffset, TLearningArrayView<1, const uint8> Bytes, uint64& OutValue)
+	{
+		UE_LEARNING_CHECK(InOutOffset + sizeof(uint64) <= Bytes.Num());
+		FMemory::Memcpy((uint8*)&OutValue, &Bytes[InOutOffset], sizeof(uint64));
+		InOutOffset += sizeof(uint64);
+	}
+
 	/**
 	* Some additional functions that act on arrays such as copy, set, zero, etc.
 	* 
@@ -434,6 +444,32 @@ namespace UE::Learning
 #endif
 		}
 
+		/**
+		* Check that an array view does not contain any NaN, Inf, -Inf, MAX_flt, or -MAX_flt
+		*/
+		template<uint8 InDimNum, typename InElementType>
+		inline void Check(const TLearningArray<InDimNum, InElementType>& View, const FIndexSet Indices)
+		{
+#if UE_LEARNING_ARRAY_CHECK_VALUES == 1
+			if constexpr (InDimNum > 1)
+			{
+				for (const int32 Idx : Indices)
+				{
+					Check(View[Idx]);
+				}
+			}
+			else
+			{
+				for (const int32 Idx : Indices)
+				{
+					UE_LEARNING_ARRAY_VALUE_CHECKF(
+						FMath::IsFinite(View.GetData()[Idx]) && View.GetData()[Idx] != MAX_flt && View.GetData()[Idx] != -MAX_flt,
+						TEXT("Invalid value %f found at flat array index %i"), View.GetData()[Idx], Idx);
+				}
+			}
+#endif
+		}
+
 		/**
 		* Zero the memory of an array view
 		*/
@@ -973,6 +1009,23 @@ namespace UE::Learning
 			return sizeof(int32) + sizeof(int32) * InDimNum + sizeof(InElementType) * Shape.Total();
 		}
 
+		/**
+		* Serialize an array to bytes
+		*/
+		template<uint8 InDimNum, typename InElementType, typename Allocator>
+		inline void SerializeToBytes(int32& InOutOffset, TLearningArrayView<1, uint8> Bytes, const TMultiArrayShape<InDimNum> Shape, const TArray<InElementType, Allocator>& InArray)
+		{
+			Learning::SerializeToBytes(InOutOffset, Bytes, (int32)InDimNum);
+			for (uint8 ShapeIdx = 0; ShapeIdx < InDimNum; ShapeIdx++)
+			{
+				Learning::SerializeToBytes(InOutOffset, Bytes, Shape[ShapeIdx]);
+			}
+
+			UE_LEARNING_CHECK(InOutOffset + sizeof(InElementType) * InArray.Num() <= Bytes.Num());
+			FMemory::Memcpy(&Bytes[InOutOffset], (uint8*)InArray.GetData(), sizeof(InElementType) * InArray.Num());
+			InOutOffset += sizeof(InElementType) * InArray.Num();
+		}
+
 		/**
 		* Serialize an array to bytes
 		*/
@@ -991,7 +1044,30 @@ namespace UE::Learning
 		}
 
 		/**
-		* Serialize an array from bytes
+		* Deserialize an array from bytes
+		*/
+		template<uint8 InDimNum, typename InElementType, typename Allocator>
+		inline void DeserializeFromBytes(int32& InOutOffset, TLearningArrayView<1, const uint8> Bytes, TArray<InElementType, Allocator>& OutArray)
+		{
+			int32 DimNum = INDEX_NONE;
+			Learning::DeserializeFromBytes(InOutOffset, Bytes, DimNum);
+			UE_LEARNING_CHECK(DimNum == InDimNum);
+
+			TLearningArrayShape<InDimNum> Shape;
+			for (uint8 ShapeIdx = 0; ShapeIdx < InDimNum; ShapeIdx++)
+			{
+				Learning::DeserializeFromBytes(InOutOffset, Bytes, Shape[ShapeIdx]);
+			}
+
+			OutArray.SetNumUninitialized(Shape.Total());
+
+			UE_LEARNING_CHECK(InOutOffset + sizeof(InElementType) * OutArray.Num() <= Bytes.Num());
+			FMemory::Memcpy((uint8*)OutArray.GetData(), &Bytes[InOutOffset], sizeof(InElementType) * OutArray.Num());
+			InOutOffset += sizeof(InElementType) * OutArray.Num();
+		}
+
+		/**
+		* Deserialize an array from bytes
 		*/
 		template<uint8 InDimNum, typename InElementType, typename Allocator>
 		inline void DeserializeFromBytes(int32& InOutOffset, TLearningArrayView<1, const uint8> Bytes, TLearningArray<InDimNum, InElementType, Allocator>& OutArray)
diff --git a/Engine/Plugins/Experimental/LearningAgents/Source/Learning/Public/LearningCompletion.h b/Engine/Plugins/Experimental/LearningAgents/Source/Learning/Public/LearningCompletion.h
--- a/Engine/Plugins/Experimental/LearningAgents/Source/Learning/Public/LearningCompletion.h
+++ b/Engine/Plugins/Experimental/LearningAgents/Source/Learning/Public/LearningCompletion.h
@@ -27,19 +27,18 @@ namespace UE::Learning
 	*/
 	struct LEARNING_API FResetInstanceBuffer
 	{
-		void Resize(const int32 InMaxInstanceNum);
+		void Reserve(const int32 InMaxInstanceNum);
 
 		void SetResetInstances(const FIndexSet Instances);
 		void SetResetInstancesFromCompletions(const TLearningArrayView<1, const ECompletionMode> Completions, const FIndexSet Instances);
 		const int32 GetResetInstanceNum() const;
 		const FIndexSet GetResetInstances() const;
+		const TArray<int32>& GetResetInstancesArray() const;
 
 	private:
 
-		int32 MaxInstanceNum = 0;
-		int32 ResetInstanceNum = 0;
 		FIndexSet ResetInstancesSet;
-		TLearningArray<1, int32, TInlineAllocator<1>> ResetInstances;
+		TArray<int32> ResetInstances;
 	};
 
 	namespace Completion
@@ -59,18 +58,16 @@ namespace UE::Learning
 		/**
 		* Set completions for all instances whose episode has reached the max number of steps.
 		*
-		* @param InOutCompletions				Output buffer to write completions to
+		* @param OutCompletions					Output buffer to write completions to
 		* @param EpisodeStepNums				Number of steps taken by each instance
 		* @param EpisodeMaxStepNum				Maximum number of allowed steps
-		* @param EpisodeEndCompletionMode		Completion mode to use for instances who reach the end of the episode
 		* @param Instances						Instances to process
 		*
 		*/
 		LEARNING_API void EvaluateEndOfEpisodeCompletions(
-			TLearningArrayView<1, ECompletionMode> InOutCompletions,
+			TLearningArrayView<1, ECompletionMode> OutCompletions,
 			const TLearningArrayView<1, const int32> EpisodeStepNums,
 			const int32 EpisodeMaxStepNum,
-			const ECompletionMode EpisodeEndCompletionMode,
 			const FIndexSet Instances);
 
 	}
diff --git a/Engine/Plugins/Experimental/LearningAgents/Source/Learning/Public/LearningNeuralNetwork.h b/Engine/Plugins/Experimental/LearningAgents/Source/Learning/Public/LearningNeuralNetwork.h
--- a/Engine/Plugins/Experimental/LearningAgents/Source/Learning/Public/LearningNeuralNetwork.h
+++ b/Engine/Plugins/Experimental/LearningAgents/Source/Learning/Public/LearningNeuralNetwork.h
@@ -6,10 +6,114 @@
 
 #include "Templates/SharedPointer.h"
 
+#include "LearningNeuralNetwork.generated.h"
+
+class UNNEModelData;
+
+namespace UE::NNE
+{
+	class IModelCPU;
+	class IModelInstanceCPU;
+}
+
 namespace UE::Learning
 {
-	struct INeuralNetworkInference;
+	struct FNeuralNetwork;
+	struct FNeuralNetworkInference;
+}
+
+/**
+ * Neural Network Data Object
+ * 
+ * This is the UObject which contains the actual data used by a Neural Network. It stores the raw FileData used to construct the network using NNE,
+ * as well as the input and output sizes and a compatibility hash that can be used to quickly check if two networks may be compatible in terms of 
+ * inputs and outputs.
+ * 
+ * Internally this also stores the various things required to map between NNE style inference and the style of inference used in Learning via 
+ * FNeuralNetwork and FNeuralNetworkInference.
+ */
+UCLASS(BlueprintType)
+class LEARNING_API ULearningNeuralNetworkData : public UObject
+{
+	GENERATED_BODY()
+
+public:
+
+	// Get the FNeuralNetwork object that can be used to do inference.
+	TSharedPtr<UE::Learning::FNeuralNetwork>& GetNetwork();
+
+	/**
+	 * Initialize the ULearningNeuralNetworkData object given some input and output sizes as well as a compatibility hash and FileData.
+	 * 
+	 * @param InInputSize			Network Input Size
+	 * @param InOutputSize			Network Output Size
+	 * @param InCompatibilityHash	Compatibility Hash
+	 * @param InFileData			Network FileData
+	 */
+	void Init(
+		const int32 InInputSize, 
+		const int32 InOutputSize, 
+		const int32 InCompatibilityHash, 
+		const TArrayView<const uint8> InFileData);
+
+	// Initialize this network from another ULearningNeuralNetworkData object
+	void InitFrom(const ULearningNeuralNetworkData* OtherNetworkData);
+
+	// Load a snapshot of this network from the given array of bytes
+	bool LoadFromSnapshot(const TArrayView<const uint8> InBytes);
+
+	// Save a snapshot of this network to the given array of bytes
+	void SaveToSnapshot(TArrayView<uint8> OutBytes) const;
+
+	// Number of bytes required to save or load a snapshot of this network
+	int32 GetSnapshotByteNum() const;
+
+	// If this network is empty or not
+	bool IsEmpty() const;
+
+	// Gets the network input size
+	int32 GetInputSize() const;
+	
+	// Gets the network output size
+	int32 GetOutputSize() const;
 	
+	// Gets the compatibility hash
+	int32 GetCompatibilityHash() const;
+
+private:
+
+	// Calls UpdateNetwork to update the internal representation.
+	virtual void PostLoad() override final;
+
+	// Uploads the FileData to NNE and updates the internal in-memory representation of the network used for inference.
+	void UpdateNetwork();
+
+	// Size of the inputs expected by this network
+	UPROPERTY(VisibleAnywhere, Category = "Network Properties")
+	int32 InputSize = 0;
+
+	// Size of the output produced by this network
+	UPROPERTY(VisibleAnywhere, Category = "Network Properties")
+	int32 OutputSize = 0;
+
+	// Compatibility hash used for testing if inputs to one network are compatible with another
+	UPROPERTY(VisibleAnywhere, Category = "Network Properties")
+	int32 CompatibilityHash = 0;
+
+	// File data used by NNE
+	UPROPERTY()
+	TArray<uint8> FileData;
+
+	// Model Data used by NNE
+	UPROPERTY()
+	TObjectPtr<UNNEModelData> ModelData;
+
+	// Internal in-memory network representation
+	TSharedPtr<UE::Learning::FNeuralNetwork> Network;
+};
+
+namespace UE::Learning
+{
 	/**
 	* Settings object for a neural network instance
 	*/
@@ -23,117 +127,116 @@ namespace UE::Learning
 	};
 
 	/**
-	* Interface for a Neural Network
+	* Neural Network Object
 	*/
-	struct INeuralNetwork
+	struct LEARNING_API FNeuralNetwork
 	{
-		virtual ~INeuralNetwork() {}
-
-		/** Deserialize the network from the raw bytes starting at the given offset. Returns true if successful. */
-		virtual bool DeserializeFromBytes(int32& InOutOffset, const TLearningArrayView<1, const uint8> RawBytes) = 0;
-
-		/** Serialize the network to the raw bytes starting at the given offset. */
-		virtual void SerializeToBytes(int32& InOutOffset, TLearningArrayView<1, uint8> OutRawBytes) const = 0;
-
-		/** Get the number of bytes required to serialize this network. */
-		virtual int32 GetSerializationByteNum() const = 0;
-
 		/** Create a new inference object for this network with the given maximum batch size and inference settings. */
-		virtual TSharedRef<INeuralNetworkInference> CreateInferenceObject(
+		TSharedRef<FNeuralNetworkInference> CreateInferenceObject(
 			const int32 MaxBatchSize,
-			const FNeuralNetworkInferenceSettings& Settings = FNeuralNetworkInferenceSettings()) = 0;
+			const FNeuralNetworkInferenceSettings& Settings = FNeuralNetworkInferenceSettings());
 
-		/** Get the number of inputs expected by this network */
-		virtual int32 GetInputNum() const = 0;
+		bool IsEmpty() const;
+		int32 GetInputSize() const;
+		int32 GetOutputSize() const;
 
-		/** Get the number of outputs expected by this network */
-		virtual int32 GetOutputNum() const = 0;
+		void UpdateModel(const TSharedPtr<NNE::IModelCPU>& InModel, const int32 InInputSize, const int32 InOutputSize);
 
-		/** Get the name of the python class used to represent this type of network during training. */
-		virtual const TCHAR* GetPythonClassName() const = 0;
+	private:
+
+		int32 InputSize = 0;
+		int32 OutputSize = 0;
+		TSharedPtr<NNE::IModelCPU> Model;
+		TArray<TWeakPtr<FNeuralNetworkInference>, TInlineAllocator<64>> InferenceObjects;
 	};
 
 	/**
-	* Interface for a Neural Network Inference
+	* Neural Network Inference Object
 	*/
-	struct INeuralNetworkInference
+	struct LEARNING_API FNeuralNetworkInference
 	{
-		virtual ~INeuralNetworkInference() {}
-
-		virtual void Evaluate(
+		/**
+		 * Constructs a new network inference object.
+		 * 
+		 * @param InModel			NNE Model
+		 * @param MaxBatchSize		Maximum batch size
+		 * @param InInputSize		Network input size
+		 * @param InOutputSize		Network output size
+		 * @param InSettings		Inference settings
+		 */
+		FNeuralNetworkInference(
+			UE::NNE::IModelCPU& InModel,
+			const int32 MaxBatchSize,
+			const int32 InInputSize,
+			const int32 InOutputSize,
+			const FNeuralNetworkInferenceSettings& InSettings = FNeuralNetworkInferenceSettings());
+
+		/**
+		 * Evaluate this network for the given instances.
+		 * 
+		 * Note: this function takes a lock, so although it can be called from multiple threads it will not benefit from multi-threading. This is 
+		 * because to use the NNE interface all the instances that need to be evaluated must be gathered and scattered into internal buffers. To
+		 * compensate for this, this call will do multi-threading internally based on the inference settings used.
+		 * 
+		 * @param Output		The Output Buffer of shape (MaxBatchSize, InOutputSize)
+		 * @param Input			The Input Buffer of shape (MaxBatchSize, InInputSize)
+		 * @param Instances		The batch instances to evaluate
+		 */
+		void Evaluate(
 			TLearningArrayView<2, float> Output,
 			const TLearningArrayView<2, const float> Input,
-			const FIndexSet Instances) = 0;
-	};
-
-	/**
-	* Activation Function for use in basic MLP Neural Network
-	*/
-	enum class EActivationFunction : uint8
-	{
-		// ReLU Activation - Fast to train and evaluate but occasionally causes gradient collapse and untrainable networks.
-		ReLU = 0,
+			const FIndexSet Instances);
 
-		// ELU Activation - Generally performs better than ReLU and is not prone to gradient collapse but slower to evaluate.
-		ELU = 1,
+		// This function will re-build the internal Model Instances used for multi-threading. It should be called whenever the given Model is updated. 
+		void ReloadModelInstances(NNE::IModelCPU& Model);
 
-		// TanH Activation - Smooth activation function that is slower to train and evaluate but sometimes more stable for certain tasks.
-		TanH = 2,
-	};
+	private:
 
-	/**
-	* Basic Implementation for a MLP Neural Network
-	*/
-	struct LEARNING_API FNeuralNetworkMLP : public INeuralNetwork
-	{
-		friend struct FNeuralNetworkInstanceMLP;
-
-		//~ Begin INeuralNetwork Interface
-		virtual bool DeserializeFromBytes(int32& InOutOffset, const TLearningArrayView<1, const uint8> RawBytes) override final;
-		virtual void SerializeToBytes(int32& InOutOffset, TLearningArrayView<1, uint8> OutRawBytes) const override final;
-		virtual int32 GetSerializationByteNum() const override final;
-		virtual TSharedRef<INeuralNetworkInference> CreateInferenceObject(
-			const int32 MaxBatchSize,
-			const FNeuralNetworkInferenceSettings& Settings = FNeuralNetworkInferenceSettings()) override final;
-		virtual int32 GetInputNum() const override final;
-		virtual int32 GetOutputNum() const override final;
-		virtual const TCHAR* GetPythonClassName() const override final { return TEXT("NeuralNetworkMLP"); };
-		//~ End INeuralNetwork Interface
-
-		int32 GetHiddenNum() const;
-		int32 GetLayerNum() const;
-
-		void Resize(
-			const int32 InputNum,
-			const int32 OutputNum,
-			const int32 HiddenNum,
-			const int32 LayerNum);
-
-		EActivationFunction ActivationFunction = EActivationFunction::ELU;
-		TArray<TLearningArray<2, float>, TInlineAllocator<16>> Weights;
-		TArray<TLearningArray<1, float>, TInlineAllocator<16>> Biases;
+		FNeuralNetworkInferenceSettings Settings;
+		FRWLock EvaluationLock;
+		TLearningArray<2, float> InputBuffer;
+		TLearningArray<2, float> OutputBuffer;
+		TArray<TSharedPtr<NNE::IModelInstanceCPU>, TInlineAllocator<64>> ModelInstances;
 	};
 
 	/**
-	* Basic Implementation for a MLP Neural Network Inference
+	* Neural-network based function object.
 	*/
-	struct LEARNING_API FNeuralNetworkInferenceMLP : public INeuralNetworkInference
+	struct LEARNING_API FNeuralNetworkFunction
 	{
-		FNeuralNetworkInferenceMLP(
-			const FNeuralNetworkMLP& InNetwork,
-			const int32 MaxBatchSize,
-			const FNeuralNetworkInferenceSettings& InSettings);
-
-		//~ Begin INeuralNetworkInference Interface
-		virtual void Evaluate(
+		/**
+		 * Constructs a Neural Network Function from the given Neural Network
+		 *
+		 * @param InMaxInstanceNum		Maximum number of instances to evaluate for
+		 * @param InNeuralNetwork		Neural network to use
+		 * @param InInferenceSettings	Inference settings
+		 */
+		FNeuralNetworkFunction(
+			const int32 InMaxInstanceNum,
+			const TSharedPtr<FNeuralNetwork>& InNeuralNetwork,
+			const FNeuralNetworkInferenceSettings& InInferenceSettings = FNeuralNetworkInferenceSettings());
+
+		/**
+		 * Evaluate this network for the given instances.
+		 *
+		 * @param Output		The Output Buffer of shape (MaxInstanceNum, InOutputSize)
+		 * @param Input			The Input Buffer of shape (MaxInstanceNum, InInputSize)
+		 * @param Instances		The instances to evaluate
+		 */
+		void Evaluate(
 			TLearningArrayView<2, float> Output,
 			const TLearningArrayView<2, const float> Input,
-			const FIndexSet Instances) override final;
-		//~ End INeuralNetworkInference Interface
+			const FIndexSet Instances);
 
-		const FNeuralNetworkMLP& Network;
-		FNeuralNetworkInferenceSettings Settings;
-		TLearningArray<2, float> FrontBuffer;
-		TLearningArray<2, float> BackBuffer;
+		/** Sets the NeuralNetwork and re-creates the NeuralNetworkInference object. */
+		void UpdateNeuralNetwork(const TSharedPtr<FNeuralNetwork>& NewNeuralNetwork);
+
+	private:
+
+		int32 MaxInstanceNum = 0;
+		TSharedPtr<FNeuralNetwork> NeuralNetwork;
+		TSharedPtr<FNeuralNetworkInference> NeuralNetworkInference;
+		FNeuralNetworkInferenceSettings InferenceSettings;
 	};
+
 }
diff --git a/Engine/Plugins/Experimental/LearningAgents/Source/Learning/Public/LearningRandom.h b/Engine/Plugins/Experimental/LearningAgents/Source/Learning/Public/LearningRandom.h
--- a/Engine/Plugins/Experimental/LearningAgents/Source/Learning/Public/LearningRandom.h
+++ b/Engine/Plugins/Experimental/LearningAgents/Source/Learning/Public/LearningRandom.h
@@ -118,6 +118,11 @@ namespace UE::Learning::Random
 		TLearningArrayView<1, uint32> Output,
 		const uint32 State);
 
+	LEARNING_API void IntArray(
+		TLearningArrayView<1, uint32> Output,
+		const uint32 State,
+		const FIndexSet Indices);
+
 	LEARNING_API void FloatArray(
 		TLearningArrayView<1, float> Output,
 		const uint32 State);
@@ -149,6 +154,25 @@ namespace UE::Learning::Random
 		const FVector Axis0 = FVector::ForwardVector,
 		const FVector Axis1 = FVector::RightVector);
 
+	LEARNING_API void DistributionIndependantNormal(
+		TLearningArrayView<1, float> Output,
+		const uint32 State,
+		const TLearningArrayView<1, const float> Mean,
+		const TLearningArrayView<1, const float> LogStd,
+		const float Scale);
+
+	LEARNING_API void DistributionMultinoulli(
+		TLearningArrayView<1, float> Output,
+		const uint32 State,
+		const TLearningArrayView<1, const float> Logits,
+		const float Scale);
+
+	LEARNING_API void DistributionBernoulli(
+		TLearningArrayView<1, float> Output,
+		const uint32 State,
+		const TLearningArrayView<1, const float> Logits,
+		const float Scale);
+
 	/////////////////////////////////////////////////////////
 	//
 	// Here we provide a more normal interface which updates
@@ -211,6 +235,11 @@ namespace UE::Learning::Random
 		TLearningArrayView<1, uint32> Output,
 		uint32& State);
 
+	LEARNING_API void SampleIntArray(
+		TLearningArrayView<1, uint32> Output,
+		uint32& State,
+		const FIndexSet Indices);
+
 	LEARNING_API void SampleFloatArray(
 		TLearningArrayView<1, float> Output,
 		uint32& State);
@@ -227,10 +256,6 @@ namespace UE::Learning::Random
 		const float Mean = 0.0f,
 		const float Std = 1.0f);
 
-	LEARNING_API void ResampleStateArray(
-		TLearningArrayView<1, uint32> InOutStates,
-		const FIndexSet Indices);
-
 	LEARNING_API void SamplePlanarClippedGaussianArray(
 		TLearningArrayView<1, FVector> Output,
 		uint32& State,
@@ -245,4 +270,23 @@ namespace UE::Learning::Random
 		uint32& State,
 		const FVector Axis0 = FVector::ForwardVector,
 		const FVector Axis1 = FVector::RightVector);
+
+	LEARNING_API void SampleDistributionIndependantNormal(
+		TLearningArrayView<1, float> Output,
+		uint32& State,
+		const TLearningArrayView<1, const float> Mean,
+		const TLearningArrayView<1, const float> LogStd,
+		const float Scale);
+
+	LEARNING_API void SampleDistributionMultinoulli(
+		TLearningArrayView<1, float> Output,
+		uint32& State,
+		const TLearningArrayView<1, const float> Logits,
+		const float Scale);
+
+	LEARNING_API void SampleDistributionBernoulli(
+		TLearningArrayView<1, float> Output,
+		uint32& State,
+		const TLearningArrayView<1, const float> Logits,
+		const float Scale);
 }
\ No newline at end of file
