LearningAgents: Large restructuring to allow for structured actions and observations.

diff --git a/Engine/Plugins/Experimental/LearningAgents/Source/LearningAgents/Private/LearningAgentsActions.cpp b/Engine/Plugins/Experimental/LearningAgents/Source/LearningAgents/Private/LearningAgentsActions.cpp
--- a/Engine/Plugins/Experimental/LearningAgents/Source/LearningAgents/Private/LearningAgentsActions.cpp
+++ b/Engine/Plugins/Experimental/LearningAgents/Source/LearningAgents/Private/LearningAgentsActions.cpp
@@ -2,677 +2,2490 @@
 
 #include "LearningAgentsActions.h"
 
-#include "LearningAgentsManager.h"
-#include "LearningAgentsInteractor.h"
 #include "LearningArray.h"
-#include "LearningArrayMap.h"
-#include "LearningFeatureObject.h"
 #include "LearningLog.h"
 
-#include "GameFramework/Actor.h"
+#include "Containers/StaticArray.h"
 
-namespace UE::Learning::Agents::Actions::Private
+bool operator==(const FLearningAgentsActionObjectElement& Lhs, const FLearningAgentsActionObjectElement& Rhs)
 {
-	template<typename ActionUObject, typename ActionFObject, typename... InArgTypes>
-	ActionUObject* AddAction(ULearningAgentsInteractor* InInteractor, const FName Name, const TCHAR* FunctionName, InArgTypes&& ...Args)
+	return Lhs.ObjectElement.Index == Rhs.ObjectElement.Index;
+}
+
+uint32 GetTypeHash(const FLearningAgentsActionObjectElement& Element)
+{
+	return (uint32)Element.ObjectElement.Index;
+}
+
+const UE::Learning::Action::FSchema& ULearningAgentsActionSchema::GetActionSchema() const
+{
+	return ActionSchema;
+}
+
+namespace UE::Learning::Agents::Action::Private
+{
+	static inline bool ContainsDuplicates(const TArrayView<const int32> Indices)
+	{
+		TSet<int32, DefaultKeyFuncs<int32>, TInlineSetAllocator<32>> IndicesSet;
+		IndicesSet.Append(Indices);
+		return Indices.Num() != IndicesSet.Num();
+	}
+
+	static inline bool ContainsDuplicates(const TArrayView<const FName> ElementNames)
+	{
+		TSet<FName, DefaultKeyFuncs<FName>, TInlineSetAllocator<32>> ElementNameSet;
+		ElementNameSet.Append(ElementNames);
+		return ElementNames.Num() != ElementNameSet.Num();
+	}
+
+	static inline const TCHAR* GetActionTypeString(const Learning::Action::EType ActionType)
+	{
+		switch (ActionType)
+		{
+		case Learning::Action::EType::Null: return TEXT("Null");
+		case Learning::Action::EType::Continuous: return TEXT("Continuous");
+		case Learning::Action::EType::DiscreteExclusive: return TEXT("DiscreteExclusive");
+		case Learning::Action::EType::DiscreteInclusive: return TEXT("DiscreteInclusive");
+		case Learning::Action::EType::And: return TEXT("Struct");
+		case Learning::Action::EType::OrExclusive: return TEXT("ExclusiveUnion");
+		case Learning::Action::EType::OrInclusive: return TEXT("InclusiveUnion");
+		case Learning::Action::EType::Array: return TEXT("Array");
+		case Learning::Action::EType::Encoding: return TEXT("Encoding");
+		default:
+			UE_LEARNING_NOT_IMPLEMENTED();
+			return TEXT("Unimplemented");
+		}
+	}
+
+	static bool ValidateObjectMatchesSchema(
+		const Learning::Action::FSchema& Schema,
+		const Learning::Action::FSchemaElement SchemaElement,
+		const Learning::Action::FObject& Object,
+		const Learning::Action::FObjectElement ObjectElement,
+		const FString& ObjectName)
+	{
+		// Check Elements are Valid
+
+		if (!Schema.IsValid(SchemaElement))
+		{
+			UE_LOG(LogLearning, Error, TEXT("%s: Invalid Action Schema Element."), *ObjectName);
+			return false;
+		}
+
+		if (!Object.IsValid(ObjectElement))
+		{
+			UE_LOG(LogLearning, Error, TEXT("%s: Invalid Action Object Element."), *ObjectName);
+			return false;
+		}
+
+		// Check Names Match
+
+		const FName ActionSchemaElementName = Schema.GetName(SchemaElement);
+		const FName ActionObjectElementName = Object.GetName(ObjectElement);
+
+		if (ActionSchemaElementName != ActionObjectElementName)
+		{
+			UE_LOG(LogLearning, Warning, TEXT("%s: Action name does not match Schema. Expected '%s', got '%s'."),
+				*ObjectName, *ActionSchemaElementName.ToString(), *ActionObjectElementName.ToString());
+		}
+
+		// Check Types Match
+
+		const Learning::Action::EType ActionSchemaElementType = Schema.GetType(SchemaElement);
+		const Learning::Action::EType ActionObjectElementType = Object.GetType(ObjectElement);
+
+		if (ActionSchemaElementType != ActionObjectElementType)
+		{
+			UE_LOG(LogLearning, Error, TEXT("%s: Action '%s' type does not match Schema. Expected type '%s', got type '%s'."),
+				*ObjectName,
+				*ActionSchemaElementName.ToString(),
+				GetActionTypeString(ActionSchemaElementType),
+				GetActionTypeString(ActionObjectElementType));
+			return false;
+		}
+
+		// Type Specific Checks
+
+		switch (ActionSchemaElementType)
+		{
+		case Learning::Action::EType::Null: return true;
+
+		case Learning::Action::EType::Continuous:
+		{
+			const int32 SchemaElementSize = Schema.GetContinuous(SchemaElement).Num;
+			const int32 ObjectElementSize = Object.GetContinuous(ObjectElement).Values.Num();
+
+			if (SchemaElementSize != ObjectElementSize)
+			{
+				UE_LOG(LogLearning, Error, TEXT("%s: Action '%s' size does not match Schema. Expected '%i', got '%i'."),
+					*ObjectName,
+					*ActionSchemaElementName.ToString(),
+					SchemaElementSize,
+					ObjectElementSize);
+				return false;
+			}
+
+			return true;
+		}
+
+		case Learning::Action::EType::DiscreteExclusive:
+		{
+			const int32 SchemaElementSize = Schema.GetDiscreteExclusive(SchemaElement).Num;
+			const int32 ObjectElementIndex = Object.GetDiscreteExclusive(ObjectElement).DiscreteIndex;
+
+			if (ObjectElementIndex < 0 || ObjectElementIndex >= SchemaElementSize)
+			{
+				UE_LOG(LogLearning, Error, TEXT("%s: Action '%s' index out of range for Schema. Expected '<%i', got '%i'."),
+					*ObjectName,
+					*ActionSchemaElementName.ToString(),
+					SchemaElementSize,
+					ObjectElementIndex);
+				return false;
+			}
+
+			return true;
+		}
+
+		case Learning::Action::EType::DiscreteInclusive:
+		{
+			const int32 SchemaElementSize = Schema.GetDiscreteInclusive(SchemaElement).Num;
+			const TArrayView<const int32> ObjectElementIndices = Object.GetDiscreteInclusive(ObjectElement).DiscreteIndices;
+
+			if (ObjectElementIndices.Num() > SchemaElementSize)
+			{
+				UE_LOG(LogLearning, Error, TEXT("%s: Action '%s' too many indices provided. Expected at most '%i', got '%i'."),
+					*ObjectName,
+					*ActionSchemaElementName.ToString(),
+					SchemaElementSize,
+					ObjectElementIndices.Num());
+				return false;
+			}
+
+			for (int32 SubElementIdx = 0; SubElementIdx < ObjectElementIndices.Num(); SubElementIdx++)
+			{
+				if (ObjectElementIndices[SubElementIdx] < 0 || ObjectElementIndices[SubElementIdx] >= SchemaElementSize)
+				{
+					UE_LOG(LogLearning, Error, TEXT("%s: Action '%s' index out of range for Schema. Expected '<%i', got '%i'."),
+						*ObjectName,
+						*ActionSchemaElementName.ToString(),
+						SchemaElementSize,
+						ObjectElementIndices[SubElementIdx]);
+					return false;
+				}
+			}
+
+			return true;
+		}
+
+		case Learning::Action::EType::And:
+		{
+			const Learning::Action::FSchemaAndParameters SchemaParameters = Schema.GetAnd(SchemaElement);
+			const Learning::Action::FObjectAndParameters ObjectParameters = Object.GetAnd(ObjectElement);
+			UE_LEARNING_CHECK(SchemaParameters.Elements.Num() == SchemaParameters.ElementNames.Num());
+			UE_LEARNING_CHECK(ObjectParameters.Elements.Num() == ObjectParameters.ElementNames.Num());
+
+			if (SchemaParameters.Elements.Num() != ObjectParameters.Elements.Num())
+			{
+				UE_LOG(LogLearning, Error, TEXT("%s: Action '%s' number of sub-elements does not match Schema. Expected '%i', got '%i'."),
+					*ObjectName,
+					*ActionSchemaElementName.ToString(),
+					SchemaParameters.Elements.Num(),
+					ObjectParameters.Elements.Num());
+				return false;
+			}
+
+			for (int32 SchemaElementIdx = 0; SchemaElementIdx < SchemaParameters.Elements.Num(); SchemaElementIdx++)
+			{
+				const int32 ObjectElementIdx = ObjectParameters.ElementNames.Find(SchemaParameters.ElementNames[SchemaElementIdx]);
+
+				if (ObjectElementIdx == INDEX_NONE)
+				{
+					UE_LOG(LogLearning, Error, TEXT("%s: Action '%s' does not include '%s' action required by Schema."),
+						*ObjectName,
+						*ActionSchemaElementName.ToString(),
+						*SchemaParameters.ElementNames[SchemaElementIdx].ToString());
+					return false;
+				}
+
+				if (!ValidateObjectMatchesSchema(
+					Schema,
+					SchemaParameters.Elements[SchemaElementIdx],
+					Object,
+					ObjectParameters.Elements[ObjectElementIdx],
+					ObjectName))
+				{
+					return false;
+				}
+			}
+
+			return true;
+		}
+
+		case Learning::Action::EType::OrExclusive:
+		{
+			const Learning::Action::FSchemaOrExclusiveParameters SchemaParameters = Schema.GetOrExclusive(SchemaElement);
+			const Learning::Action::FObjectOrExclusiveParameters ObjectParameters = Object.GetOrExclusive(ObjectElement);
+			UE_LEARNING_CHECK(SchemaParameters.Elements.Num() == SchemaParameters.ElementNames.Num());
+
+			const int32 SchemaSubElementIdx = SchemaParameters.ElementNames.Find(ObjectParameters.ElementName);
+
+			if (SchemaSubElementIdx == INDEX_NONE)
+			{
+				UE_LOG(LogLearning, Error, TEXT("%s: Action '%s' Schema does not include '%s' action."),
+					*ObjectName,
+					*ActionSchemaElementName.ToString(),
+					*ObjectParameters.ElementName.ToString());
+				return false;
+			}
+
+			return ValidateObjectMatchesSchema(
+				Schema,
+				SchemaParameters.Elements[SchemaSubElementIdx],
+				Object,
+				ObjectParameters.Element,
+				ObjectName);
+		}
+
+		case Learning::Action::EType::OrInclusive:
+		{
+			const Learning::Action::FSchemaOrInclusiveParameters SchemaParameters = Schema.GetOrInclusive(SchemaElement);
+			const Learning::Action::FObjectOrInclusiveParameters ObjectParameters = Object.GetOrInclusive(ObjectElement);
+
+			if (ObjectParameters.Elements.Num() > SchemaParameters.Elements.Num())
+			{
+				UE_LOG(LogLearning, Error, TEXT("%s: Action '%s' too many sub-actions provided. Expected at most '%i', got '%i'."),
+					*ObjectName,
+					*ActionSchemaElementName.ToString(),
+					SchemaParameters.Elements.Num(),
+					ObjectParameters.Elements.Num());
+				return false;
+			}
+
+			for (int32 ObjectSubElementIdx = 0; ObjectSubElementIdx < ObjectParameters.Elements.Num(); ObjectSubElementIdx++)
+			{
+				const int32 SchemaSubElementIdx = SchemaParameters.ElementNames.Find(ObjectParameters.ElementNames[ObjectSubElementIdx]);
+
+				if (SchemaSubElementIdx == INDEX_NONE)
+				{
+					UE_LOG(LogLearning, Error, TEXT("%s: Action '%s' Schema does not include '%s' action."),
+						*ObjectName,
+						*ActionSchemaElementName.ToString(),
+						*ObjectParameters.ElementNames[ObjectSubElementIdx].ToString());
+					return false;
+				}
+
+				if (!ValidateObjectMatchesSchema(
+					Schema,
+					SchemaParameters.Elements[SchemaSubElementIdx],
+					Object,
+					ObjectParameters.Elements[ObjectSubElementIdx],
+					ObjectName))
+				{
+					return false;
+				}
+			}
+
+			return true;
+		}
+
+		case Learning::Action::EType::Array:
+		{
+			const Learning::Action::FSchemaArrayParameters SchemaParameters = Schema.GetArray(SchemaElement);
+			const Learning::Action::FObjectArrayParameters ObjectParameters = Object.GetArray(ObjectElement);
+
+			if (ObjectParameters.Elements.Num() != SchemaParameters.Num)
+			{
+				UE_LOG(LogLearning, Error, TEXT("%s: Action '%s' array incorrect size. Expected '%i' elements, got '%i'."),
+					*ObjectName,
+					*ActionSchemaElementName.ToString(),
+					SchemaParameters.Num,
+					ObjectParameters.Elements.Num());
+				return false;
+			}
+
+			for (int32 ElementIdx = 0; ElementIdx < ObjectParameters.Elements.Num(); ElementIdx++)
+			{
+				if (!ValidateObjectMatchesSchema(
+					Schema,
+					SchemaParameters.Element,
+					Object,
+					ObjectParameters.Elements[ElementIdx],
+					ObjectName))
+				{
+					return false;
+				}
+			}
+
+			return true;
+		}
+
+		case Learning::Action::EType::Encoding:
+		{
+			const Learning::Action::FSchemaEncodingParameters SchemaParameters = Schema.GetEncoding(SchemaElement);
+			const Learning::Action::FObjectEncodingParameters ObjectParameters = Object.GetEncoding(ObjectElement);
+
+			return ValidateObjectMatchesSchema(
+					Schema,
+					SchemaParameters.Element,
+					Object,
+					ObjectParameters.Element,
+					ObjectName);
+		}
+
+		default:
+		{
+			UE_LEARNING_NOT_IMPLEMENTED();
+			return true;
+		}
+		}
+	}
+
+	static void LogAction(
+		const UE::Learning::Action::FObject& Object,
+		const UE::Learning::Action::FObjectElement ObjectElement,
+		const FString& Indentation,
+		const FString& Prefix,
+		const FString& ObjectName)
+	{
+		if (!Object.IsValid(ObjectElement))
+		{
+			UE_LOG(LogLearning, Error, TEXT("%s: Invalid Action Object Element."), *ObjectName);
+			return;
+		}
+
+		const UE::Learning::Action::EType Type = Object.GetType(ObjectElement);
+		const FName Name = Object.GetName(ObjectElement);
+
+		switch (Type)
+		{
+		case UE::Learning::Action::EType::Null:
+		{
+			UE_LOG(LogLearning, Display, TEXT("%s%s \"%s\" (%s)"), *Indentation, *Prefix, *Name.ToString(), GetActionTypeString(Type));
+			return;
+		}
+
+		case UE::Learning::Action::EType::Continuous:
+		{
+			const UE::Learning::Action::FObjectContinuousParameters Parameters = Object.GetContinuous(ObjectElement);
+
+			UE_LOG(LogLearning, Display, TEXT("%s%s \"%s\" (%s) %s"), *Indentation, *Prefix, *Name.ToString(), GetActionTypeString(Type), *UE::Learning::Array::FormatFloat(Parameters.Values));
+			return;
+		}
+
+		case UE::Learning::Action::EType::DiscreteExclusive:
+		{
+			const UE::Learning::Action::FObjectDiscreteExclusiveParameters Parameters = Object.GetDiscreteExclusive(ObjectElement);
+
+			UE_LOG(LogLearning, Display, TEXT("%s%s \"%s\" (%s) %i"), *Indentation, *Prefix, *Name.ToString(), GetActionTypeString(Type), Parameters.DiscreteIndex);
+			return;
+		}
+
+		case UE::Learning::Action::EType::DiscreteInclusive:
+		{
+			const UE::Learning::Action::FObjectDiscreteInclusiveParameters Parameters = Object.GetDiscreteInclusive(ObjectElement);
+
+			UE_LOG(LogLearning, Display, TEXT("%s%s \"%s\" (%s) %s"), *Indentation, *Prefix, *Name.ToString(), GetActionTypeString(Type), *UE::Learning::Array::FormatInt32(Parameters.DiscreteIndices));
+			return;
+		}
+
+		case UE::Learning::Action::EType::And:
+		{
+			const UE::Learning::Action::FObjectAndParameters Parameters = Object.GetAnd(ObjectElement);
+
+			UE_LOG(LogLearning, Display, TEXT("%s%s \"%s\" (%s)"), *Indentation, *Prefix, *Name.ToString(), GetActionTypeString(Type));
+			for (int32 SubElementIdx = 0; SubElementIdx < Parameters.Elements.Num(); SubElementIdx++)
+			{
+				LogAction(Object, Parameters.Elements[SubElementIdx], *(Indentation + TEXT("    ")), FString::Printf(TEXT("| \"%s\": "), *Parameters.ElementNames[SubElementIdx].ToString()), ObjectName);
+			}
+
+			return;
+		}
+
+		case UE::Learning::Action::EType::OrExclusive:
+		{
+			const UE::Learning::Action::FObjectOrExclusiveParameters Parameters = Object.GetOrExclusive(ObjectElement);
+
+			UE_LOG(LogLearning, Display, TEXT("%s%s \"%s\" (%s)"), *Indentation, *Prefix, *Name.ToString(), GetActionTypeString(Type));
+			LogAction(Object, Parameters.Element, *(Indentation + TEXT("    ")), FString::Printf(TEXT("| \"%s\": "), *Parameters.ElementName.ToString()), ObjectName);
+
+			return;
+		}
+
+		case UE::Learning::Action::EType::OrInclusive:
+		{
+			const UE::Learning::Action::FObjectOrInclusiveParameters Parameters = Object.GetOrInclusive(ObjectElement);
+
+			UE_LOG(LogLearning, Display, TEXT("%s%s \"%s\" (%s)"), *Indentation, *Prefix, *Name.ToString(), GetActionTypeString(Type));
+			for (int32 SubElementIdx = 0; SubElementIdx < Parameters.Elements.Num(); SubElementIdx++)
+			{
+				LogAction(Object, Parameters.Elements[SubElementIdx], *(Indentation + TEXT("    ")), FString::Printf(TEXT("| \"%s\": "), *Parameters.ElementNames[SubElementIdx].ToString()), ObjectName);
+			}
+
+			return;
+		}
+
+		case UE::Learning::Action::EType::Array:
+		{
+			const UE::Learning::Action::FObjectArrayParameters Parameters = Object.GetArray(ObjectElement);
+
+			UE_LOG(LogLearning, Display, TEXT("%s%s \"%s\" (%s)"), *Indentation, *Prefix, *Name.ToString(), GetActionTypeString(Type));
+			for (int32 SubElementIdx = 0; SubElementIdx < Parameters.Elements.Num(); SubElementIdx++)
+			{
+				LogAction(Object, Parameters.Elements[SubElementIdx], *(Indentation + TEXT("    ")), FString::Printf(TEXT("| %3i:"), SubElementIdx), ObjectName);
+			}
+
+			return;
+		}
+
+		case UE::Learning::Action::EType::Encoding:
+		{
+			const UE::Learning::Action::FObjectEncodingParameters Parameters = Object.GetEncoding(ObjectElement);
+
+			UE_LOG(LogLearning, Display, TEXT("%s%s \"%s\" (%s)"), *Indentation, *Prefix, *Name.ToString(), GetActionTypeString(Type));
+			LogAction(Object, Parameters.Element, *(Indentation + TEXT("    ")), TEXT("|"), ObjectName);
+
+			return;
+		}
+		}
+	}
+
+	static inline FVector VectorLogSafe(const FVector V, const float Epsilon = UE_SMALL_NUMBER)
+	{
+		return FVector(
+			FMath::Loge(FMath::Max(V.X, Epsilon)),
+			FMath::Loge(FMath::Max(V.Y, Epsilon)),
+			FMath::Loge(FMath::Max(V.Z, Epsilon)));
+	}
+
+	static inline FVector VectorExp(const FVector V)
+	{
+		return FVector(
+			FMath::Exp(V.X),
+			FMath::Exp(V.Y),
+			FMath::Exp(V.Z));
+	}
+
+	static inline void NormalizeProbabilitiesExclusive(TArrayView<float> PriorProbabilities, const FString& ObjectName)
+	{
+		float Total = 0.0f;
+		for (int32 Idx = 0; Idx < PriorProbabilities.Num(); Idx++)
+		{
+			if (PriorProbabilities[Idx] < 0.0f || PriorProbabilities[Idx] > 1.0f)
+			{
+				UE_LOG(LogLearning, Warning, TEXT("%s: Invalid Prior Probability Given (%f), must be in range 0 to 1."), *ObjectName, PriorProbabilities[Idx]);
+			}
+
+			PriorProbabilities[Idx] = FMath::Clamp(PriorProbabilities[Idx], 0.0f, 1.0f);
+			Total += PriorProbabilities[Idx];
+		}
+
+		if (PriorProbabilities.Num() > 0 && FMath::Abs(Total) < UE_SMALL_NUMBER)
+		{
+			UE_LOG(LogLearning, Warning, TEXT("%s: Prior Probabilities are too small. Should sum to 1."), *ObjectName);
+
+			for (int32 Idx = 0; Idx < PriorProbabilities.Num(); Idx++)
+			{
+				PriorProbabilities[Idx] = 1.0f / PriorProbabilities.Num();
+			}
+		}
+		else
+		{
+			for (int32 Idx = 0; Idx < PriorProbabilities.Num(); Idx++)
+			{
+				PriorProbabilities[Idx] /= Total;
+			}
+		}
+	}
+
+	static inline void NormalizeProbabilitiesInclusive(TArrayView<float> PriorProbabilities, const FString& ObjectName)
+	{
+		for (int32 Idx = 0; Idx < PriorProbabilities.Num(); Idx++)
+		{
+			if (PriorProbabilities[Idx] < 0.0f || PriorProbabilities[Idx] > 1.0f)
+			{
+				UE_LOG(LogLearning, Warning, TEXT("%s: Invalid Prior Probability Given (%f), must be in range 0 to 1."), *ObjectName, PriorProbabilities[Idx]);
+			}
+
+			PriorProbabilities[Idx] = FMath::Clamp(PriorProbabilities[Idx], 0.0f, 1.0f);
+		}
+	}
+}
+
+bool ULearningAgentsActionSchema::ValidateObjectMatchesSchema(
+	const FLearningAgentsActionSchemaElement SchemaElement,
+	const ULearningAgentsActionObject* Object,
+	const FLearningAgentsActionObjectElement ObjectElement) const
+{
+	return UE::Learning::Agents::Action::Private::ValidateObjectMatchesSchema(
+		ActionSchema,
+		SchemaElement.SchemaElement,
+		Object->GetActionObject(),
+		ObjectElement.ObjectElement,
+		GetName());
+}
+
+FLearningAgentsActionSchemaElement ULearningAgentsActionSchema::SpecifyNullAction(const FName Name)
+{
+	return { ActionSchema.CreateNull(Name) };
+}
+
+FLearningAgentsActionSchemaElement ULearningAgentsActionSchema::SpecifyContinuousAction(const int32 Size, const FName Name)
+{
+	if (Size < 0)
+	{
+		UE_LOG(LogLearning, Error, TEXT("%s: Invalid Continuous Action Size '%i'."), *GetName(), Size);
+		return FLearningAgentsActionSchemaElement();
+	}
+
+	if (Size == 0)
+	{
+		UE_LOG(LogLearning, Warning, TEXT("%s: Specifying zero-sized Continuous Action."), *GetName());
+	}
+
+	return { ActionSchema.CreateContinuous({ Size }, Name) };
+}
+
+FLearningAgentsActionSchemaElement ULearningAgentsActionSchema::SpecifyExclusiveDiscreteAction(const int32 Size, const TArray<float>& PriorProbabilities, const FName Name)
+{
+	return SpecifyExclusiveDiscreteActionFromArrayView(Size, PriorProbabilities, Name);
+}
+
+FLearningAgentsActionSchemaElement ULearningAgentsActionSchema::SpecifyExclusiveDiscreteActionFromArrayView(const int32 Size, const TArrayView<const float> PriorProbabilities, const FName Name)
+{
+	if (Size < 0)
+	{
+		UE_LOG(LogLearning, Error, TEXT("%s: Invalid DiscreteExclusive Action Size '%i'."), *GetName(), Size);
+		return FLearningAgentsActionSchemaElement();
+	}
+
+	if (Size == 0)
+	{
+		UE_LOG(LogLearning, Warning, TEXT("%s: Specifying zero-sized Exclusive Discrete Action."), *GetName());
+	}
+
+	TArray<float, TInlineAllocator<16>> NormalizedPriorProbabilities;
+	NormalizedPriorProbabilities.Init(1.0f / Size, Size);
+	for (int32 Idx = 0; Idx < PriorProbabilities.Num(); Idx++)
+	{
+		NormalizedPriorProbabilities[Idx] = PriorProbabilities[Idx];
+	}
+	UE::Learning::Agents::Action::Private::NormalizeProbabilitiesExclusive(NormalizedPriorProbabilities, GetName());
+
+	return { ActionSchema.CreateDiscreteExclusive({ Size, MakeArrayView(NormalizedPriorProbabilities) }, Name)};
+}
+
+FLearningAgentsActionSchemaElement ULearningAgentsActionSchema::SpecifyInclusiveDiscreteAction(const int32 Size, const TArray<float>& PriorProbabilities, const FName Name)
+{
+	return SpecifyInclusiveDiscreteActionFromArrayView(Size, PriorProbabilities, Name);
+}
+
+FLearningAgentsActionSchemaElement ULearningAgentsActionSchema::SpecifyInclusiveDiscreteActionFromArrayView(const int32 Size, const TArrayView<const float> PriorProbabilities, const FName Name)
+{
+	if (Size < 0)
+	{
+		UE_LOG(LogLearning, Error, TEXT("%s: Invalid DiscreteInclusive Action Size '%i'."), *GetName(), Size);
+		return FLearningAgentsActionSchemaElement();
+	}
+
+	if (Size == 0)
+	{
+		UE_LOG(LogLearning, Warning, TEXT("%s: Specifying zero-sized Inclusive Discrete Action."), *GetName());
+	}
+
+	TArray<float, TInlineAllocator<16>> NormalizedPriorProbabilities;
+	NormalizedPriorProbabilities.Init(0.5f, Size);
+	for (int32 Idx = 0; Idx < PriorProbabilities.Num(); Idx++)
+	{
+		NormalizedPriorProbabilities[Idx] = PriorProbabilities[Idx];
+	}
+	UE::Learning::Agents::Action::Private::NormalizeProbabilitiesInclusive(NormalizedPriorProbabilities, GetName());
+
+	return { ActionSchema.CreateDiscreteInclusive({ Size, MakeArrayView(NormalizedPriorProbabilities) }, Name) };
+}
+
+FLearningAgentsActionSchemaElement ULearningAgentsActionSchema::SpecifyStructAction(const TMap<FName, FLearningAgentsActionSchemaElement>& Elements, const FName Name)
+{
+	if (Elements.Num() == 0)
+	{
+		UE_LOG(LogLearning, Warning, TEXT("%s: Specifying zero-sized Struct Action."), *GetName());
+	}
+
+	const int32 SubElementNum = Elements.Num();
+
+	TArray<int32, TInlineAllocator<16>> SubElementIndices;
+	TArray<FName, TInlineAllocator<16>> SubElementNames;
+	TArray<FLearningAgentsActionSchemaElement, TInlineAllocator<16>> SubElements;
+	SubElementIndices.Empty(SubElementNum);
+	SubElementNames.Empty(SubElementNum);
+	SubElements.Empty(SubElementNum);
+
+	int32 Index = 0;
+	for (const TPair<FName, FLearningAgentsActionSchemaElement>& Element : Elements)
+	{
+		SubElementIndices.Add(Index);
+		SubElementNames.Add(Element.Key);
+		SubElements.Add(Element.Value);
+		Index++;
+	}
+
+	// Sort Elements According to FName
+
+	SubElementIndices.Sort([SubElementNames](const int32 Lhs, const int32 Rhs)
+	{
+		return SubElementNames[Lhs].ToString().ToLower() < SubElementNames[Rhs].ToString().ToLower();
+	});
+
+	TArray<FName, TInlineAllocator<16>> SortedSubElementNames;
+	TArray<FLearningAgentsActionSchemaElement, TInlineAllocator<16>> SortedSubElements;
+	SortedSubElementNames.SetNumUninitialized(SubElementNum);
+	SortedSubElements.SetNumUninitialized(SubElementNum);
+	for (int32 Idx = 0; Idx < SubElementNum; Idx++)
+	{
+		SortedSubElementNames[Idx] = SubElementNames[SubElementIndices[Idx]];
+		SortedSubElements[Idx] = SubElements[SubElementIndices[Idx]];
+	}
+
+	return SpecifyStructActionFromArrayViews(SortedSubElementNames, SortedSubElements, Name);
+}
+
+FLearningAgentsActionSchemaElement ULearningAgentsActionSchema::SpecifyStructActionFromArrays(const TArray<FName>& ElementNames, const TArray<FLearningAgentsActionSchemaElement>& Elements, const FName Name)
+{
+	return SpecifyStructActionFromArrayViews(ElementNames, Elements, Name);
+}
+
+FLearningAgentsActionSchemaElement ULearningAgentsActionSchema::SpecifyStructActionFromArrayViews(const TArrayView<const FName> ElementNames, const TArrayView<const FLearningAgentsActionSchemaElement> Elements, const FName Name)
+{
+	if (Elements.Num() == 0)
+	{
+		UE_LOG(LogLearning, Warning, TEXT("%s: Specifying zero-sized Struct Action."), *GetName());
+	}
+
+	if (Elements.Num() != ElementNames.Num())
+	{
+		UE_LOG(LogLearning, Error, TEXT("%s: Number of elements (%i) must match number of names (%i)."), *GetName(), Elements.Num(), ElementNames.Num());
+		return FLearningAgentsActionSchemaElement();
+	}
+
+	if (UE::Learning::Agents::Action::Private::ContainsDuplicates(ElementNames))
+	{
+		UE_LOG(LogLearning, Error, TEXT("%s: Element Names contain duplicates."), *GetName());
+		return FLearningAgentsActionSchemaElement();
+	}
+
+	TArray<UE::Learning::Action::FSchemaElement, TInlineAllocator<16>> SubElements;
+	SubElements.Empty(Elements.Num());
+
+	for (const FLearningAgentsActionSchemaElement& Element : Elements)
+	{
+		if (!ActionSchema.IsValid(Element.SchemaElement))
+		{
+			UE_LOG(LogLearning, Error, TEXT("%s: Invalid Action Object."), *GetName());
+			return FLearningAgentsActionSchemaElement();
+		}
+
+		SubElements.Add(Element.SchemaElement);
+	}
+
+	return { ActionSchema.CreateAnd({ ElementNames, SubElements }, Name) };
+}
+
+FLearningAgentsActionSchemaElement ULearningAgentsActionSchema::SpecifyExclusiveUnionAction(const TMap<FName, FLearningAgentsActionSchemaElement>& Elements, const TMap<FName, float>& PriorProbabilities, const FName Name)
+{
+	if (Elements.Num() == 0)
+	{
+		UE_LOG(LogLearning, Warning, TEXT("%s: Specifying zero-sized Exclusive Union Action."), *GetName());
+	}
+
+	const int32 SubElementNum = Elements.Num();
+
+	TArray<int32, TInlineAllocator<16>> SubElementIndices;
+	TArray<FName, TInlineAllocator<16>> SubElementNames;
+	TArray<FLearningAgentsActionSchemaElement, TInlineAllocator<16>> SubElements;
+	TArray<float, TInlineAllocator<16>> SubElementPriorProbabilities;
+	SubElementIndices.Empty(SubElementNum);
+	SubElementNames.Empty(SubElementNum);
+	SubElements.Empty(SubElementNum);
+	SubElementPriorProbabilities.Empty(SubElementNum);
+
+	int32 Index = 0;
+	for (const TPair<FName, FLearningAgentsActionSchemaElement>& Element : Elements)
+	{
+		SubElementIndices.Add(Index);
+		SubElementNames.Add(Element.Key);
+		SubElements.Add(Element.Value);
+		const float* PriorProb = PriorProbabilities.Find(Element.Key);
+		SubElementPriorProbabilities.Add(PriorProb ? *PriorProb : 1.0f / SubElementNum);
+		Index++;
+	}
+
+	// Sort Elements According to FName
+
+	SubElementIndices.Sort([SubElementNames](const int32 Lhs, const int32 Rhs)
+	{
+		return SubElementNames[Lhs].ToString().ToLower() < SubElementNames[Rhs].ToString().ToLower();
+	});
+
+	TArray<FName, TInlineAllocator<16>> SortedSubElementNames;
+	TArray<FLearningAgentsActionSchemaElement, TInlineAllocator<16>> SortedSubElements;
+	TArray<float, TInlineAllocator<16>> SortedSubElementPriorProbabilities;
+	SortedSubElementNames.SetNumUninitialized(SubElementNum);
+	SortedSubElements.SetNumUninitialized(SubElementNum);
+	SortedSubElementPriorProbabilities.SetNumUninitialized(SubElementNum);
+	for (int32 Idx = 0; Idx < SubElementNum; Idx++)
+	{
+		SortedSubElementNames[Idx] = SubElementNames[SubElementIndices[Idx]];
+		SortedSubElements[Idx] = SubElements[SubElementIndices[Idx]];
+		SortedSubElementPriorProbabilities[Idx] = SubElementPriorProbabilities[SubElementIndices[Idx]];
+	}
+
+	return SpecifyExclusiveUnionActionFromArrayViews(SortedSubElementNames, SortedSubElements, SortedSubElementPriorProbabilities, Name);
+}
+
+FLearningAgentsActionSchemaElement ULearningAgentsActionSchema::SpecifyExclusiveUnionActionFromArrays(const TArray<FName>& ElementNames, const TArray<FLearningAgentsActionSchemaElement>& Elements, const TArray<float>& PriorProbabilities, const FName Name)
+{
+	return SpecifyExclusiveUnionActionFromArrayViews(ElementNames, Elements, PriorProbabilities, Name);
+}
+
+FLearningAgentsActionSchemaElement ULearningAgentsActionSchema::SpecifyExclusiveUnionActionFromArrayViews(const TArrayView<const FName> ElementNames, const TArrayView<const FLearningAgentsActionSchemaElement> Elements, const TArrayView<const float> PriorProbabilities, const FName Name)
+{
+	if (Elements.Num() == 0)
+	{
+		UE_LOG(LogLearning, Warning, TEXT("%s: Specifying zero-sized Exclusive Union Action."), *GetName());
+	}
+
+	if (Elements.Num() != ElementNames.Num())
+	{
+		UE_LOG(LogLearning, Error, TEXT("%s: Number of elements (%i) must match number of names (%i)."), *GetName(), Elements.Num(), ElementNames.Num());
+		return FLearningAgentsActionSchemaElement();
+	}
+
+	if (UE::Learning::Agents::Action::Private::ContainsDuplicates(ElementNames))
+	{
+		UE_LOG(LogLearning, Error, TEXT("%s: Element Names contain duplicates."), *GetName());
+		return FLearningAgentsActionSchemaElement();
+	}
+
+	TArray<UE::Learning::Action::FSchemaElement, TInlineAllocator<16>> SubElements;
+	SubElements.Empty(Elements.Num());
+
+	for (const FLearningAgentsActionSchemaElement& Element : Elements)
+	{
+		if (!ActionSchema.IsValid(Element.SchemaElement))
+		{
+			UE_LOG(LogLearning, Error, TEXT("%s: Invalid Action Object."), *GetName());
+			return FLearningAgentsActionSchemaElement();
+		}
+
+		SubElements.Add(Element.SchemaElement);
+	}
+
+	TArray<float, TInlineAllocator<16>> NormalizedPriorProbabilities;
+	NormalizedPriorProbabilities.Init(1.0f / Elements.Num(), Elements.Num());
+	for (int32 Idx = 0; Idx < PriorProbabilities.Num(); Idx++)
+	{
+		NormalizedPriorProbabilities[Idx] = PriorProbabilities[Idx];
+	}
+	UE::Learning::Agents::Action::Private::NormalizeProbabilitiesExclusive(NormalizedPriorProbabilities, GetName());
+
+	return { ActionSchema.CreateOrExclusive({ ElementNames, SubElements, MakeArrayView(NormalizedPriorProbabilities) }, Name)};
+}
+
+FLearningAgentsActionSchemaElement ULearningAgentsActionSchema::SpecifyInclusiveUnionAction(const TMap<FName, FLearningAgentsActionSchemaElement>& Elements, const TMap<FName, float>& PriorProbabilities, const FName Name)
+{
+	if (Elements.Num() == 0)
+	{
+		UE_LOG(LogLearning, Warning, TEXT("%s: Specifying zero-sized Inclusive Union Action."), *GetName());
+	}
+
+	const int32 SubElementNum = Elements.Num();
+
+	TArray<int32, TInlineAllocator<16>> SubElementIndices;
+	TArray<FName, TInlineAllocator<16>> SubElementNames;
+	TArray<FLearningAgentsActionSchemaElement, TInlineAllocator<16>> SubElements;
+	TArray<float, TInlineAllocator<16>> SubElementPriorProbabilities;
+	SubElementIndices.Empty(SubElementNum);
+	SubElementNames.Empty(SubElementNum);
+	SubElements.Empty(SubElementNum);
+	SubElementPriorProbabilities.Empty(SubElementNum);
+
+	int32 Index = 0;
+	for (const TPair<FName, FLearningAgentsActionSchemaElement>& Element : Elements)
+	{
+		SubElementIndices.Add(Index);
+		SubElementNames.Add(Element.Key);
+		SubElements.Add(Element.Value);
+		const float* PriorProb = PriorProbabilities.Find(Element.Key);
+		SubElementPriorProbabilities.Add(PriorProb ? *PriorProb : 1.0f / SubElementNum);
+		Index++;
+	}
+
+	// Sort Elements According to FName
+
+	SubElementIndices.Sort([SubElementNames](const int32 Lhs, const int32 Rhs)
+	{
+		return SubElementNames[Lhs].ToString().ToLower() < SubElementNames[Rhs].ToString().ToLower();
+	});
+
+	TArray<FName, TInlineAllocator<16>> SortedSubElementNames;
+	TArray<FLearningAgentsActionSchemaElement, TInlineAllocator<16>> SortedSubElements;
+	TArray<float, TInlineAllocator<16>> SortedSubElementPriorProbabilities;
+	SortedSubElementNames.SetNumUninitialized(SubElementNum);
+	SortedSubElements.SetNumUninitialized(SubElementNum);
+	SortedSubElementPriorProbabilities.SetNumUninitialized(SubElementNum);
+	for (int32 Idx = 0; Idx < SubElementNum; Idx++)
+	{
+		SortedSubElementNames[Idx] = SubElementNames[SubElementIndices[Idx]];
+		SortedSubElements[Idx] = SubElements[SubElementIndices[Idx]];
+		SortedSubElementPriorProbabilities[Idx] = SubElementPriorProbabilities[SubElementIndices[Idx]];
+	}
+
+	return SpecifyInclusiveUnionActionFromArrayViews(SortedSubElementNames, SortedSubElements, SortedSubElementPriorProbabilities, Name);
+}
+
+FLearningAgentsActionSchemaElement ULearningAgentsActionSchema::SpecifyInclusiveUnionActionFromArrays(const TArray<FName> ElementNames, const TArray<FLearningAgentsActionSchemaElement>& Elements, const TArray<float>& PriorProbabilities, const FName Name)
+{
+	return SpecifyInclusiveUnionActionFromArrayViews(ElementNames, Elements, PriorProbabilities, Name);
+}
+
+FLearningAgentsActionSchemaElement ULearningAgentsActionSchema::SpecifyInclusiveUnionActionFromArrayViews(const TArrayView<const FName> ElementNames, const TArrayView<const FLearningAgentsActionSchemaElement> Elements, const TArrayView<const float> PriorProbabilities, const FName Name)
+{
+	if (Elements.Num() == 0)
+	{
+		UE_LOG(LogLearning, Warning, TEXT("%s: Specifying zero-sized Inclusive Union Action."), *GetName());
+	}
+
+	if (Elements.Num() != ElementNames.Num())
+	{
+		UE_LOG(LogLearning, Error, TEXT("%s: Number of elements (%i) must match number of names (%i)."), *GetName(), Elements.Num(), ElementNames.Num());
+		return FLearningAgentsActionSchemaElement();
+	}
+
+	if (UE::Learning::Agents::Action::Private::ContainsDuplicates(ElementNames))
+	{
+		UE_LOG(LogLearning, Error, TEXT("%s: Element Names contain duplicates."), *GetName());
+		return FLearningAgentsActionSchemaElement();
+	}
+
+	TArray<UE::Learning::Action::FSchemaElement, TInlineAllocator<16>> SubElements;
+	SubElements.Empty(Elements.Num());
+
+	for (const FLearningAgentsActionSchemaElement& Element : Elements)
+	{
+		if (!ActionSchema.IsValid(Element.SchemaElement))
+		{
+			UE_LOG(LogLearning, Error, TEXT("%s: Invalid Action Object."), *GetName());
+			return FLearningAgentsActionSchemaElement();
+		}
+
+		SubElements.Add(Element.SchemaElement);
+	}
+
+	TArray<float, TInlineAllocator<16>> NormalizedPriorProbabilities;
+	NormalizedPriorProbabilities.Init(0.5f, Elements.Num());
+	for (int32 Idx = 0; Idx < PriorProbabilities.Num(); Idx++)
+	{
+		NormalizedPriorProbabilities[Idx] = PriorProbabilities[Idx];
+	}
+	UE::Learning::Agents::Action::Private::NormalizeProbabilitiesInclusive(NormalizedPriorProbabilities, GetName());
+
+	return { ActionSchema.CreateOrInclusive({ ElementNames, SubElements, MakeArrayView(NormalizedPriorProbabilities) }, Name)};
+}
+
+FLearningAgentsActionSchemaElement ULearningAgentsActionSchema::SpecifyStaticArrayAction(const FLearningAgentsActionSchemaElement Element, const int32 Num, const FName Name)
+{
+	if (Num < 0)
+	{
+		UE_LOG(LogLearning, Error, TEXT("%s: Invalid Action Static Array Num %i."), *GetName(), Num);
+		return FLearningAgentsActionSchemaElement();
+	}
+
+	if (Num == 0)
+	{
+		UE_LOG(LogLearning, Warning, TEXT("%s: Specifying zero-sized Static Array Action."), *GetName());
+	}
+
+	if (!ActionSchema.IsValid(Element.SchemaElement))
+	{
+		UE_LOG(LogLearning, Error, TEXT("%s: Invalid Action Object."), *GetName());
+		return FLearningAgentsActionSchemaElement();
+	}
+
+	return { ActionSchema.CreateArray({ Element.SchemaElement, Num }, Name) };
+}
+
+FLearningAgentsActionSchemaElement ULearningAgentsActionSchema::SpecifyPairAction(const FLearningAgentsActionSchemaElement Key, const FLearningAgentsActionSchemaElement Value, const FName Name)
+{
+	return SpecifyStructActionFromArrayViews({ TEXT("Key"), TEXT("Value") }, { Key, Value }, Name);
+}
+
+FLearningAgentsActionSchemaElement ULearningAgentsActionSchema::SpecifyEnumAction(const UEnum* Enum, const TMap<uint8, float>& PriorProbabilities, const FName Name)
+{
+	if (!Enum)
+	{
+		UE_LOG(LogLearning, Error, TEXT("%s: Enum is nullptr."), *GetName());
+		return FLearningAgentsActionSchemaElement();
+	}
+
+	TArray<float, TInlineAllocator<16>> PriorProbabilitiesArray;
+	PriorProbabilitiesArray.Init(1.0f / (Enum->NumEnums() - 1), Enum->NumEnums() - 1);
+	for (const TPair<uint8, float> Prior : PriorProbabilities)
+	{
+		const int32 EnumIndex = Enum->GetIndexByValue(Prior.Key);
+		if (EnumIndex != INDEX_NONE)
+		{
+			PriorProbabilitiesArray[EnumIndex] = Prior.Value;
+		}
+	}
+
+	return SpecifyEnumActionFromArrayView(Enum, PriorProbabilitiesArray, Name);
+}
+
+FLearningAgentsActionSchemaElement ULearningAgentsActionSchema::SpecifyEnumActionFromArray(const UEnum* Enum, const TArray<float>& PriorProbabilities, const FName Name)
+{
+	return SpecifyEnumActionFromArrayView(Enum, PriorProbabilities, Name);
+}
+
+FLearningAgentsActionSchemaElement ULearningAgentsActionSchema::SpecifyEnumActionFromArrayView(const UEnum* Enum, const TArrayView<const float> PriorProbabilities, const FName Name)
+{
+	if (!Enum)
+	{
+		UE_LOG(LogLearning, Error, TEXT("%s: Enum is nullptr."), *GetName());
+		return FLearningAgentsActionSchemaElement();
+	}
+
+	return SpecifyExclusiveDiscreteActionFromArrayView(Enum->NumEnums() - 1, PriorProbabilities, Name);
+}
+
+FLearningAgentsActionSchemaElement ULearningAgentsActionSchema::SpecifyBitmaskAction(const UEnum* Enum, const TMap<uint8, float>& PriorProbabilities, const FName Name)
+{
+	if (!Enum)
+	{
+		UE_LOG(LogLearning, Error, TEXT("%s: Enum is nullptr."), *GetName());
+		return FLearningAgentsActionSchemaElement();
+	}
+
+	if (Enum->NumEnums() - 1 > 32)
+	{
+		UE_LOG(LogLearning, Error, TEXT("%s: Too many values in Enum to use as Bitmask (%i)."), *GetName(), Enum->NumEnums() - 1);
+		return FLearningAgentsActionSchemaElement();
+	}
+
+	TArray<float, TInlineAllocator<16>> PriorProbabilitiesArray;
+	PriorProbabilitiesArray.Init(0.5f, Enum->NumEnums() - 1);
+	for (const TPair<uint8, float> Prior : PriorProbabilities)
+	{
+		const int32 EnumIndex = Enum->GetIndexByValue(Prior.Key);
+		if (EnumIndex != INDEX_NONE)
+		{
+			PriorProbabilitiesArray[EnumIndex] = Prior.Value;
+		}
+	}
+
+	return SpecifyBitmaskActionFromArrayView(Enum, PriorProbabilitiesArray, Name);
+}
+
+FLearningAgentsActionSchemaElement ULearningAgentsActionSchema::SpecifyBitmaskActionFromArray(const UEnum* Enum, const TArray<float>& PriorProbabilities, const FName Name)
+{
+	return SpecifyBitmaskActionFromArrayView(Enum, PriorProbabilities, Name);
+}
+
+FLearningAgentsActionSchemaElement ULearningAgentsActionSchema::SpecifyBitmaskActionFromArrayView(const UEnum* Enum, const TArrayView<const float> PriorProbabilities, const FName Name)
+{
+	if (!Enum)
+	{
+		UE_LOG(LogLearning, Error, TEXT("%s: Enum is nullptr."), *GetName());
+		return FLearningAgentsActionSchemaElement();
+	}
+
+	if (Enum->NumEnums() - 1 > 32)
+	{
+		UE_LOG(LogLearning, Error, TEXT("%s: Too many values in Enum to use as Bitmask (%i)."), *GetName(), Enum->NumEnums() - 1);
+		return FLearningAgentsActionSchemaElement();
+	}
+
+	return SpecifyInclusiveDiscreteActionFromArrayView(Enum->NumEnums() - 1, PriorProbabilities, Name);
+}
+
+FLearningAgentsActionSchemaElement ULearningAgentsActionSchema::SpecifyOptionalAction(const FLearningAgentsActionSchemaElement Element, const float PriorProbability, const FName Name)
+{
+	return SpecifyExclusiveUnionActionFromArrayViews({ TEXT("Null"), TEXT("Valid") }, { SpecifyNullAction(), Element }, { 1.0f - PriorProbability, PriorProbability }, Name);
+}
+
+FLearningAgentsActionSchemaElement ULearningAgentsActionSchema::SpecifyEitherAction(const FLearningAgentsActionSchemaElement A, const FLearningAgentsActionSchemaElement B, const float PriorProbabilityOfA, const FName Name)
+{
+	return SpecifyExclusiveUnionActionFromArrayViews({ TEXT("A"), TEXT("B") }, { A, B }, { 1.0f - PriorProbabilityOfA, PriorProbabilityOfA }, Name);
+}
+
+FLearningAgentsActionSchemaElement ULearningAgentsActionSchema::SpecifyEncodingAction(const FLearningAgentsActionSchemaElement Element, const int32 EncodingSize, const FName Name)
+{
+	if (EncodingSize < 1)
+	{
+		UE_LOG(LogLearning, Error, TEXT("%s: Invalid Action EncodingSize '%i'."), *GetName(), EncodingSize);
+		return FLearningAgentsActionSchemaElement();
+	}
+
+	if (!ActionSchema.IsValid(Element.SchemaElement))
+	{
+		UE_LOG(LogLearning, Error, TEXT("%s: Invalid Action Object."), *GetName());
+		return FLearningAgentsActionSchemaElement();
+	}
+
+	return { ActionSchema.CreateEncoding({ Element.SchemaElement, EncodingSize }, Name) };
+}
+
+FLearningAgentsActionSchemaElement ULearningAgentsActionSchema::SpecifyBoolAction(const float PriorProbability, const FName Name)
+{
+	return SpecifyExclusiveDiscreteActionFromArrayView(2, { 1.0f - PriorProbability, PriorProbability }, Name);
+}
+
+FLearningAgentsActionSchemaElement ULearningAgentsActionSchema::SpecifyFloatAction(const FName Name)
+{
+	return SpecifyContinuousAction(1, Name);
+}
+
+FLearningAgentsActionSchemaElement ULearningAgentsActionSchema::SpecifyTranslationAction(const FName Name)
+{
+	return SpecifyContinuousAction(3, Name);
+}
+
+FLearningAgentsActionSchemaElement ULearningAgentsActionSchema::SpecifyRotationAction(const FName Name)
+{
+	return SpecifyContinuousAction(3, Name);
+}
+
+FLearningAgentsActionSchemaElement ULearningAgentsActionSchema::SpecifyScaleAction(const FName Name)
+{
+	return SpecifyContinuousAction(3, Name);
+}
+
+FLearningAgentsActionSchemaElement ULearningAgentsActionSchema::SpecifyTransformAction(const FName Name)
+{
+	return SpecifyStructActionFromArrayViews(
+		{
+			TEXT("Translation"),
+			TEXT("Rotation"),
+			TEXT("Scale")
+		},
+		{
+			SpecifyTranslationAction(),
+			SpecifyRotationAction(),
+			SpecifyScaleAction()
+		}, 
+		Name);
+}
+
+FLearningAgentsActionSchemaElement ULearningAgentsActionSchema::SpecifyAngleAction(const FName Name)
+{
+	return SpecifyContinuousAction(1, Name);
+}
+
+FLearningAgentsActionSchemaElement ULearningAgentsActionSchema::SpecifyVelocityAction(const FName Name)
+{
+	return SpecifyContinuousAction(3, Name);
+}
+
+FLearningAgentsActionSchemaElement ULearningAgentsActionSchema::SpecifySpeedAction(const FName Name)
+{
+	return SpecifyContinuousAction(1, Name);
+}
+
+const UE::Learning::Action::FObject& ULearningAgentsActionObject::GetActionObject() const
+{
+	return ActionObject;
+}
+
+UE::Learning::Action::FObject& ULearningAgentsActionObject::GetActionObject()
+{
+	return ActionObject;
+}
+
+void ULearningAgentsActionObject::LogAction(const FLearningAgentsActionObjectElement Element)
+{
+	UE::Learning::Agents::Action::Private::LogAction(GetActionObject(), Element.ObjectElement, TEXT(""), TEXT(""), GetName());
+}
+
+FLearningAgentsActionObjectElement ULearningAgentsActionObject::MakeNullAction(const FName Name)
+{
+	return { ActionObject.CreateNull(Name) };
+}
+
+FLearningAgentsActionObjectElement ULearningAgentsActionObject::MakeContinuousAction(const TArray<float>& Values, const FName Name)
+{
+	return MakeContinuousActionFromArrayView(Values, Name);
+}
+
+FLearningAgentsActionObjectElement ULearningAgentsActionObject::MakeContinuousActionFromArrayView(const TArrayView<const float> Values, const FName Name)
+{
+	if (Values.Num() == 0)
+	{
+		UE_LOG(LogLearning, Warning, TEXT("%s: Creating zero-sized Continuous Action."), *GetName());
+	}
+
+	return { ActionObject.CreateContinuous({ Values }, Name) };
+}
+
+FLearningAgentsActionObjectElement ULearningAgentsActionObject::MakeExclusiveDiscreteAction(const int32 Index, const FName Name)
+{
+	if (Index < 0)
+	{
+		UE_LOG(LogLearning, Error, TEXT("%s: Invalid Action Index %i."), *GetName(), Index);
+		return FLearningAgentsActionObjectElement();
+	}
+
+	return { ActionObject.CreateDiscreteExclusive({ Index }, Name) };
+}
+
+FLearningAgentsActionObjectElement ULearningAgentsActionObject::MakeInclusiveDiscreteAction(const TArray<int32>& Indices, const FName Name)
+{
+	return MakeInclusiveDiscreteActionFromArrayView(Indices, Name);
+}
+
+FLearningAgentsActionObjectElement ULearningAgentsActionObject::MakeInclusiveDiscreteActionFromArrayView(const TArrayView<const int32> Indices, const FName Name)
+{
+	if (UE::Learning::Agents::Action::Private::ContainsDuplicates(Indices))
+	{
+		UE_LOG(LogLearning, Error, TEXT("%s: Indices contain duplicates."), *GetName());
+		return FLearningAgentsActionObjectElement();
+	}
+
+	const int32 IndexNum = Indices.Num();
+
+	for (int32 IndexIdx = 0; IndexIdx < IndexNum; IndexIdx++)
+	{
+		if (Indices[IndexIdx] < 0)
+		{
+			UE_LOG(LogLearning, Error, TEXT("%s: Invalid Action Index %i."), *GetName(), Indices[IndexIdx]);
+			return FLearningAgentsActionObjectElement();
+		}
+	}
+
+	return { ActionObject.CreateDiscreteInclusive({ Indices }, Name) };
+}
+
+FLearningAgentsActionObjectElement ULearningAgentsActionObject::MakeStructAction(const TMap<FName, FLearningAgentsActionObjectElement>& Elements, const FName Name)
+{
+	if (Elements.Num() == 0)
+	{
+		UE_LOG(LogLearning, Warning, TEXT("%s: Creating zero-sized Struct Action."), *GetName());
+	}
+
+	const int32 SubElementNum = Elements.Num();
+
+	TArray<FLearningAgentsActionObjectElement, TInlineAllocator<16>> SubElements;
+	TArray<FName, TInlineAllocator<16>> SubElementNames;
+	SubElements.Empty(SubElementNum);
+	SubElementNames.Empty(SubElementNum);
+
+	for (const TPair<FName, FLearningAgentsActionObjectElement>& Element : Elements)
+	{
+		SubElements.Add(Element.Value);
+		SubElementNames.Add(Element.Key);
+	}
+
+	return MakeStructActionFromArrayViews(SubElementNames, SubElements);
+}
+
+FLearningAgentsActionObjectElement ULearningAgentsActionObject::MakeStructActionFromArrays(const TArray<FName>& ElementNames, const TArray<FLearningAgentsActionObjectElement>& Elements, const FName Name)
+{
+	return MakeStructActionFromArrayViews(ElementNames, Elements, Name);
+}
+
+FLearningAgentsActionObjectElement ULearningAgentsActionObject::MakeStructActionFromArrayViews(const TArrayView<const FName> ElementNames, const TArrayView<const FLearningAgentsActionObjectElement> Elements, const FName Name)
+{
+	if (Elements.Num() == 0)
+	{
+		UE_LOG(LogLearning, Warning, TEXT("%s: Creating zero-sized Struct Action."), *GetName());
+	}
+
+	if (Elements.Num() != ElementNames.Num())
+	{
+		UE_LOG(LogLearning, Error, TEXT("%s: Number of elements (%i) must match number of names (%i)."), *GetName(), Elements.Num(), ElementNames.Num());
+		return FLearningAgentsActionObjectElement();
+	}
+
+	if (UE::Learning::Agents::Action::Private::ContainsDuplicates(ElementNames))
+	{
+		UE_LOG(LogLearning, Error, TEXT("%s: Element Names contain duplicates."), *GetName());
+		return FLearningAgentsActionObjectElement();
+	}
+
+	TArray<UE::Learning::Action::FObjectElement, TInlineAllocator<16>> SubElements;
+	SubElements.Empty(Elements.Num());
+
+	for (const FLearningAgentsActionObjectElement& Element : Elements)
+	{
+		if (!ActionObject.IsValid(Element.ObjectElement))
+		{
+			UE_LOG(LogLearning, Error, TEXT("%s: Invalid Action Object."), *GetName());
+			return FLearningAgentsActionObjectElement();
+		}
+
+		SubElements.Add(Element.ObjectElement);
+	}
+
+	return { ActionObject.CreateAnd({ ElementNames, SubElements }, Name) };
+}
+
+FLearningAgentsActionObjectElement ULearningAgentsActionObject::MakeExclusiveUnionAction(const FName ElementName, const FLearningAgentsActionObjectElement Element, const FName Name)
+{
+	if (!ActionObject.IsValid(Element.ObjectElement))
+	{
+		UE_LOG(LogLearning, Error, TEXT("%s: Invalid Action Object."), *GetName());
+		return FLearningAgentsActionObjectElement();
+	}
+
+	return { ActionObject.CreateOrExclusive({ ElementName, Element.ObjectElement }, Name) };
+}
+
+FLearningAgentsActionObjectElement ULearningAgentsActionObject::MakeInclusiveUnionAction(const TMap<FName, FLearningAgentsActionObjectElement>& Elements, const FName Name)
+{
+	const int32 SubElementNum = Elements.Num();
+
+	TArray<FLearningAgentsActionObjectElement, TInlineAllocator<16>> SubElements;
+	TArray<FName, TInlineAllocator<16>> SubElementNames;
+	SubElements.Empty(SubElementNum);
+	SubElementNames.Empty(SubElementNum);
+
+	for (const TPair<FName, FLearningAgentsActionObjectElement>& Element : Elements)
 	{
-		if (!InInteractor)
+		SubElements.Add(Element.Value);
+		SubElementNames.Add(Element.Key);
+	}
+
+	return MakeInclusiveUnionActionFromArrayViews(SubElementNames, SubElements, Name);
+}
+
+FLearningAgentsActionObjectElement ULearningAgentsActionObject::MakeInclusiveUnionActionFromArrays(const TArray<FName>& ElementNames, const TArray<FLearningAgentsActionObjectElement>& Elements, const FName Name)
+{
+	return MakeInclusiveUnionActionFromArrayViews(ElementNames, Elements, Name);
+}
+
+FLearningAgentsActionObjectElement ULearningAgentsActionObject::MakeInclusiveUnionActionFromArrayViews(const TArrayView<const FName> ElementNames, const TArrayView<const FLearningAgentsActionObjectElement> Elements, const FName Name)
+{
+	if (Elements.Num() != ElementNames.Num())
+	{
+		UE_LOG(LogLearning, Error, TEXT("%s: Number of elements (%i) must match number of names (%i)."), *GetName(), Elements.Num(), ElementNames.Num());
+		return FLearningAgentsActionObjectElement();
+	}
+
+	if (UE::Learning::Agents::Action::Private::ContainsDuplicates(ElementNames))
+	{
+		UE_LOG(LogLearning, Error, TEXT("%s: Element Names contain duplicates."), *GetName());
+		return FLearningAgentsActionObjectElement();
+	}
+
+	TArray<UE::Learning::Action::FObjectElement, TInlineAllocator<16>> SubElements;
+	SubElements.Empty(Elements.Num());
+
+	for (const FLearningAgentsActionObjectElement& Element : Elements)
+	{
+		if (!ActionObject.IsValid(Element.ObjectElement))
+		{
+			UE_LOG(LogLearning, Error, TEXT("%s: Invalid Action Object."), *GetName());
+			return FLearningAgentsActionObjectElement();
+		}
+
+		SubElements.Add(Element.ObjectElement);
+	}
+
+	return { ActionObject.CreateOrInclusive({ ElementNames, SubElements }, Name) };
+}
+
+FLearningAgentsActionObjectElement ULearningAgentsActionObject::MakeStaticArrayAction(const TArray<FLearningAgentsActionObjectElement>& Elements, const FName Name)
+{
+	return MakeStaticArrayActionFromArrayView(Elements, Name);
+}
+
+FLearningAgentsActionObjectElement ULearningAgentsActionObject::MakeStaticArrayActionFromArrayView(const TArrayView<const FLearningAgentsActionObjectElement> Elements, const FName Name)
+{
+	if (Elements.Num() == 0)
+	{
+		UE_LOG(LogLearning, Warning, TEXT("%s: Creating zero-sized Static Array Action."), *GetName());
+	}
+
+	TArray<UE::Learning::Action::FObjectElement, TInlineAllocator<16>> SubElements;
+	SubElements.Empty(Elements.Num());
+
+	for (const FLearningAgentsActionObjectElement& Element : Elements)
+	{
+		if (!ActionObject.IsValid(Element.ObjectElement))
 		{
-			UE_LOG(LogLearning, Error, TEXT("%s: InInteractor is nullptr."), FunctionName);
-			return nullptr;
+			UE_LOG(LogLearning, Error, TEXT("%s: Invalid Action Object."), *GetName());
+			return FLearningAgentsActionObjectElement();
 		}
 
-		if (!InInteractor->HasAgentManager())
+		SubElements.Add(Element.ObjectElement);
+	}
+
+	return { ActionObject.CreateArray({ SubElements }, Name) };
+}
+
+FLearningAgentsActionObjectElement ULearningAgentsActionObject::MakePairAction(const FLearningAgentsActionObjectElement Key, const FLearningAgentsActionObjectElement Value, const FName Name)
+{
+	return MakeStructActionFromArrayViews({ TEXT("Key"), TEXT("Value") }, { Key, Value }, Name);
+}
+
+FLearningAgentsActionObjectElement ULearningAgentsActionObject::MakeEnumAction(const UEnum* Enum, const uint8 EnumValue, const FName Name)
+{
+	if (!Enum)
+	{
+		UE_LOG(LogLearning, Error, TEXT("%s: Enum is nullptr."), *GetName());
+		return FLearningAgentsActionObjectElement();
+	}
+
+	const int32 EnumValueIndex = Enum->GetIndexByValue(EnumValue);
+
+	if (EnumValueIndex == INDEX_NONE || EnumValueIndex < 0 || EnumValueIndex >= Enum->NumEnums() - 1)
+	{
+		UE_LOG(LogLearning, Error, TEXT("%s: EnumValue %i not valid for Enum '%s'."), *GetName(), EnumValue , *Enum->GetName());
+		return FLearningAgentsActionObjectElement();
+	}
+
+	return MakeExclusiveDiscreteAction(EnumValueIndex, Name);
+}
+
+FLearningAgentsActionObjectElement ULearningAgentsActionObject::MakeBitmaskAction(const UEnum* Enum, const int32 BitmaskValue, const FName Name)
+{
+	if (!Enum)
+	{
+		UE_LOG(LogLearning, Error, TEXT("%s: Enum is nullptr."), *GetName());
+		return FLearningAgentsActionObjectElement();
+	}
+
+	if (Enum->NumEnums() - 1 > 32)
+	{
+		UE_LOG(LogLearning, Error, TEXT("%s: Too many values in Enum to use as Bitmask (%i)."), *GetName(), Enum->NumEnums() - 1);
+		return FLearningAgentsActionObjectElement();
+	}
+
+	TArray<int32, TInlineAllocator<32>> BitmaskIndices;
+	BitmaskIndices.Empty(Enum->NumEnums() - 1);
+
+	for (int32 BitmaskIdx = 0; BitmaskIdx < Enum->NumEnums() - 1; BitmaskIdx++)
+	{
+		if (BitmaskValue & (1 << BitmaskIdx))
 		{
-			UE_LOG(LogLearning, Error, TEXT("%s: Must be attached to a LearningAgentsManager Actor."), *InInteractor->GetName());
-			return nullptr;
+			BitmaskIndices.Add(BitmaskIdx);
 		}
+	}
+
+	return MakeInclusiveDiscreteActionFromArrayView(BitmaskIndices, Name);
+}
+
+FLearningAgentsActionObjectElement ULearningAgentsActionObject::MakeOptionalAction(const FLearningAgentsActionObjectElement Element, const ELearningAgentsOptionalAction Option, const FName Name)
+{
+	return MakeExclusiveUnionAction(
+		Option == ELearningAgentsOptionalAction::Null ? TEXT("Null") : TEXT("Valid"),
+		Option == ELearningAgentsOptionalAction::Null ? MakeNullAction() : Element,
+		Name);
+}
 
-		const FName UniqueName = MakeUniqueObjectName(InInteractor, ActionUObject::StaticClass(), Name, EUniqueObjectNameOptions::GloballyUnique);
+FLearningAgentsActionObjectElement ULearningAgentsActionObject::MakeOptionalNullAction(const FName Name)
+{
+	return MakeExclusiveUnionAction(TEXT("Null"), MakeNullAction(), Name);
+}
+
+FLearningAgentsActionObjectElement ULearningAgentsActionObject::MakeOptionalValidAction(const FLearningAgentsActionObjectElement Element, const FName Name)
+{
+	return MakeExclusiveUnionAction(TEXT("Valid"), Element, Name);
+}
 
-		ActionUObject* Action = NewObject<ActionUObject>(InInteractor, UniqueName);
-		Action->Init(InInteractor->GetAgentManager()->GetMaxAgentNum());
-		Action->Interactor = InInteractor;
-		Action->FeatureObject = MakeShared<ActionFObject>(
-			Action->GetFName(),
-			InInteractor->GetAgentManager()->GetInstanceData().ToSharedRef(),
-			InInteractor->GetAgentManager()->GetMaxAgentNum(),
-			Forward<InArgTypes>(Args)...);
+FLearningAgentsActionObjectElement ULearningAgentsActionObject::MakeEitherAction(const FLearningAgentsActionObjectElement Element, const ELearningAgentsEitherAction Either, const FName Name)
+{
+	return MakeExclusiveUnionAction(Either == ELearningAgentsEitherAction::A ? TEXT("A") : TEXT("B"), Element, Name);
+}
 
-		// We assume all supported action feature objects can be encoded and decoded
-		UE_LEARNING_CHECK(Action->FeatureObject->IsEncodable() && Action->FeatureObject->IsDecodable());
+FLearningAgentsActionObjectElement ULearningAgentsActionObject::MakeEitherAAction(const FLearningAgentsActionObjectElement A, const FName Name)
+{
+	return MakeExclusiveUnionAction(TEXT("A"), A, Name);
+}
 
-		InInteractor->AddAction(Action, Action->FeatureObject.ToSharedRef());
+FLearningAgentsActionObjectElement ULearningAgentsActionObject::MakeEitherBAction(const FLearningAgentsActionObjectElement B, const FName Name)
+{
+	return MakeExclusiveUnionAction(TEXT("B"), B, Name);
+}
 
-		return Action;
+FLearningAgentsActionObjectElement ULearningAgentsActionObject::MakeEncodingAction(const FLearningAgentsActionObjectElement Element, const FName Name)
+{
+	if (!ActionObject.IsValid(Element.ObjectElement))
+	{
+		UE_LOG(LogLearning, Error, TEXT("%s: Invalid Action Object."), *GetName());
+		return FLearningAgentsActionObjectElement();
 	}
+
+	return { ActionObject.CreateEncoding({ Element.ObjectElement }, Name) };
+}
+
+FLearningAgentsActionObjectElement ULearningAgentsActionObject::MakeBoolAction(const bool bValue, const FName Name)
+{
+	return MakeExclusiveDiscreteAction(bValue ? 1 : 0, Name);
+}
+
+FLearningAgentsActionObjectElement ULearningAgentsActionObject::MakeFloatAction(const float Value, const float FloatScale, const FName Name)
+{
+	return MakeContinuousActionFromArrayView({ Value / FMath::Max(FloatScale, UE_SMALL_NUMBER) }, Name);
 }
 
-//------------------------------------------------------------------
+FLearningAgentsActionObjectElement ULearningAgentsActionObject::MakeTranslationAction(const FVector Translation, const FTransform RelativeTransform, const float TranslationScale, const FName Name)
+{
+	const FVector LocalTranslation = RelativeTransform.InverseTransformPosition(Translation);
+
+	return MakeContinuousActionFromArrayView({
+		(float)LocalTranslation.X / FMath::Max(TranslationScale, UE_SMALL_NUMBER),
+		(float)LocalTranslation.Y / FMath::Max(TranslationScale, UE_SMALL_NUMBER),
+		(float)LocalTranslation.Z / FMath::Max(TranslationScale, UE_SMALL_NUMBER) }, Name);
+}
 
-void ULearningAgentsAction::Init(const int32 MaxAgentNum)
+FLearningAgentsActionObjectElement ULearningAgentsActionObject::MakeRotationAction(const FRotator Rotation, const FRotator RelativeRotation, const float RotationScale, const FName Name)
 {
-	AgentGetIteration.SetNumUninitialized({ MaxAgentNum });
-	AgentSetIteration.SetNumUninitialized({ MaxAgentNum });
-	UE::Learning::Array::Set<1, uint64>(AgentGetIteration, INDEX_NONE);
-	UE::Learning::Array::Set<1, uint64>(AgentSetIteration, INDEX_NONE);
+	return MakeRotationActionFromQuat(FQuat::MakeFromRotator(Rotation), FQuat::MakeFromRotator(RelativeRotation), RotationScale, Name);
 }
 
-void ULearningAgentsAction::OnAgentsAdded(const TArray<int32>& AgentIds)
+FLearningAgentsActionObjectElement ULearningAgentsActionObject::MakeRotationActionFromQuat(const FQuat Rotation, const FQuat RelativeRotation, const float RotationScale, const FName Name)
 {
-	UE::Learning::Array::Set<1, uint64>(AgentGetIteration, 0, AgentIds);
-	UE::Learning::Array::Set<1, uint64>(AgentSetIteration, 0, AgentIds);
+	FQuat LocalRotation = RelativeRotation.Inverse() * Rotation;
+	LocalRotation.EnforceShortestArcWith(FQuat::Identity);
+	const FVector RotationVector = LocalRotation.ToRotationVector();
+
+	return MakeContinuousActionFromArrayView({
+		(float)RotationVector.X / FMath::Max(FMath::DegreesToRadians(RotationScale), UE_SMALL_NUMBER),
+		(float)RotationVector.Y / FMath::Max(FMath::DegreesToRadians(RotationScale), UE_SMALL_NUMBER),
+		(float)RotationVector.Z / FMath::Max(FMath::DegreesToRadians(RotationScale), UE_SMALL_NUMBER),
+		}, Name);
 }
 
-void ULearningAgentsAction::OnAgentsRemoved(const TArray<int32>& AgentIds)
+FLearningAgentsActionObjectElement ULearningAgentsActionObject::MakeScaleAction(const FVector Scale, const FVector RelativeScale, const FName Name)
 {
-	UE::Learning::Array::Set<1, uint64>(AgentGetIteration, INDEX_NONE, AgentIds);
-	UE::Learning::Array::Set<1, uint64>(AgentSetIteration, INDEX_NONE, AgentIds);
+	const FVector LocalLogScale =
+		UE::Learning::Agents::Action::Private::VectorLogSafe(Scale) -
+		UE::Learning::Agents::Action::Private::VectorLogSafe(RelativeScale);
+
+	return MakeContinuousActionFromArrayView({
+		(float)LocalLogScale.X,
+		(float)LocalLogScale.Y,
+		(float)LocalLogScale.Z,
+		}, Name);
 }
 
-void ULearningAgentsAction::OnAgentsReset(const TArray<int32>& AgentIds)
+FLearningAgentsActionObjectElement ULearningAgentsActionObject::MakeTransformAction(const FTransform Transform, const FTransform RelativeTransform, const float TranslationScale, const FName Name)
 {
-	UE::Learning::Array::Set<1, uint64>(AgentGetIteration, 0, AgentIds);
-	UE::Learning::Array::Set<1, uint64>(AgentSetIteration, 0, AgentIds);
+	const FTransform LocalTransform = Transform * RelativeTransform.Inverse();
+
+	return MakeStructActionFromArrayViews(
+		{
+			TEXT("Translation"),
+			TEXT("Rotation"),
+			TEXT("Scale")
+		},
+		{
+			MakeTranslationAction(LocalTransform.GetTranslation(), FTransform::Identity, TranslationScale),
+			MakeRotationActionFromQuat(LocalTransform.GetRotation(), FQuat::Identity),
+			MakeScaleAction(LocalTransform.GetScale3D(), FVector::OneVector)
+		},
+		Name);
 }
 
-uint64 ULearningAgentsAction::GetAgentGetIteration(const int32 AgentId) const
+FLearningAgentsActionObjectElement ULearningAgentsActionObject::MakeAngleAction(const float Angle, const float RelativeAngle, const float AngleScale, const FName Name)
 {
-	return AgentGetIteration[AgentId];
+	return MakeAngleActionRadians(FMath::DegreesToRadians(Angle), FMath::DegreesToRadians(RelativeAngle), FMath::DegreesToRadians(AngleScale), Name);
 }
 
-uint64 ULearningAgentsAction::GetAgentSetIteration(const int32 AgentId) const
+FLearningAgentsActionObjectElement ULearningAgentsActionObject::MakeAngleActionRadians(const float Angle, const float RelativeAngle, const float AngleScale, const FName Name)
 {
-	return AgentSetIteration[AgentId];
+	const float LocalAngle = FMath::FindDeltaAngleRadians(RelativeAngle, Angle);
+	return MakeContinuousActionFromArrayView({ LocalAngle / FMath::Max(AngleScale, UE_SMALL_NUMBER) }, Name);
 }
 
-//------------------------------------------------------------------
+FLearningAgentsActionObjectElement ULearningAgentsActionObject::MakeVelocityAction(const FVector Velocity, const FTransform RelativeTransform, const float VelocityScale, const FName Name)
+{
+	const FVector LocalVelocity = RelativeTransform.InverseTransformVectorNoScale(Velocity);
+
+	return MakeContinuousActionFromArrayView({
+		(float)LocalVelocity.X / FMath::Max(VelocityScale, UE_SMALL_NUMBER),
+		(float)LocalVelocity.Y / FMath::Max(VelocityScale, UE_SMALL_NUMBER),
+		(float)LocalVelocity.Z / FMath::Max(VelocityScale, UE_SMALL_NUMBER),
+		}, Name);
+}
 
-UFloatAction* UFloatAction::AddFloatAction(ULearningAgentsInteractor* InInteractor, const FName Name, const float Scale)
+FLearningAgentsActionObjectElement ULearningAgentsActionObject::MakeSpeedAction(const float Speed, const float SpeedScale, const FName Name)
 {
-	return UE::Learning::Agents::Actions::Private::AddAction<UFloatAction,UE::Learning::FFloatFeature>(InInteractor, Name, TEXT("AddFloatAction"), 1, Scale);
+	return MakeContinuousActionFromArrayView({ Speed / FMath::Max(SpeedScale, UE_SMALL_NUMBER) });
 }
 
-float UFloatAction::GetFloatAction(const int32 AgentId)
+bool ULearningAgentsActionObject::GetNullAction(const FLearningAgentsActionObjectElement Element, const FName Name) const
 {
-	if (!Interactor->HasAgent(AgentId))
+	if (!ActionObject.IsValid(Element.ObjectElement))
+	{
+		UE_LOG(LogLearning, Error, TEXT("%s: Invalid Action Object."), *GetName());
+		return false;
+	}
+
+	if (ActionObject.GetName(Element.ObjectElement) != Name)
 	{
-		UE_LOG(LogLearning, Error, TEXT("%s: AgentId %d not found in the agents set."), *GetName(), AgentId);
-		return 0.0f;
+		UE_LOG(LogLearning, Warning, TEXT("%s: Action name does not match. Action is '%s' but asked for '%s'."), *GetName(), *ActionObject.GetName(Element.ObjectElement).ToString(), *Name.ToString());
 	}
 
-	AgentGetIteration[AgentId]++;
-	return FeatureObject->InstanceData->ConstView(FeatureObject->ValueHandle)[AgentId][0];
+	if (ActionObject.GetType(Element.ObjectElement) != UE::Learning::Action::EType::Null)
+	{
+		UE_LOG(LogLearning, Error, TEXT("%s: Action '%s' type does not match. Action is '%s' but asked for '%s'."),
+			*GetName(),
+			*ActionObject.GetName(Element.ObjectElement).ToString(),
+			UE::Learning::Agents::Action::Private::GetActionTypeString(ActionObject.GetType(Element.ObjectElement)),
+			UE::Learning::Agents::Action::Private::GetActionTypeString(UE::Learning::Action::EType::Null));
+		return false;
+	}
+
+	return true;
 }
 
-void UFloatAction::SetFloatAction(const int32 AgentId, const float Value)
+bool ULearningAgentsActionObject::GetContinuousActionNum(int32& OutNum, const FLearningAgentsActionObjectElement Element, const FName Name) const
 {
-	if (!Interactor->HasAgent(AgentId))
+	if (!ActionObject.IsValid(Element.ObjectElement))
+	{
+		UE_LOG(LogLearning, Error, TEXT("%s: Invalid Action Object."), *GetName());
+		OutNum = 0;
+		return false;
+	}
+
+	if (ActionObject.GetName(Element.ObjectElement) != Name)
+	{
+		UE_LOG(LogLearning, Warning, TEXT("%s: Action name does not match. Action is '%s' but asked for '%s'."), *GetName(), *ActionObject.GetName(Element.ObjectElement).ToString(), *Name.ToString());
+	}
+
+	if (ActionObject.GetType(Element.ObjectElement) != UE::Learning::Action::EType::Continuous)
 	{
-		UE_LOG(LogLearning, Error, TEXT("%s: AgentId %d not found in the agents set."), *GetName(), AgentId);
-		return;
+		UE_LOG(LogLearning, Error, TEXT("%s: Action '%s' type does not match. Action is '%s' but asked for '%s'."),
+			*GetName(),
+			*ActionObject.GetName(Element.ObjectElement).ToString(),
+			UE::Learning::Agents::Action::Private::GetActionTypeString(ActionObject.GetType(Element.ObjectElement)),
+			UE::Learning::Agents::Action::Private::GetActionTypeString(UE::Learning::Action::EType::Continuous));
+		OutNum = 0;
+		return false;
 	}
 
-	FeatureObject->InstanceData->View(FeatureObject->ValueHandle)[AgentId][0] = Value;
-	AgentSetIteration[AgentId]++;
+	OutNum = ActionObject.GetContinuous(Element.ObjectElement).Values.Num();
+	return true;
 }
 
-#if UE_LEARNING_AGENTS_ENABLE_VISUAL_LOG
-void UFloatAction::VisualLog(const UE::Learning::FIndexSet Instances) const
+bool ULearningAgentsActionObject::GetContinuousAction(TArray<float>& OutValues, const FLearningAgentsActionObjectElement Element, const FName Name) const
 {
-	UE_LEARNING_TRACE_CPUPROFILER_EVENT_SCOPE(UFloatAction::VisualLog);
+	int32 OutValueNum = 0;
+	if (!GetContinuousActionNum(OutValueNum, Element, Name))
+	{
+		OutValues.Empty();
+		return false;
+	}
 
-	const TLearningArrayView<2, const float> ValueView = FeatureObject->InstanceData->ConstView(FeatureObject->ValueHandle);
-	const TLearningArrayView<2, const float> FeatureView = FeatureObject->InstanceData->ConstView(FeatureObject->FeatureHandle);
+	OutValues.SetNumUninitialized(OutValueNum);
 
-	for (const int32 Instance : Instances)
+	if (!GetContinuousActionToArrayView(OutValues, Element, Name))
 	{
-		if (const AActor* Actor = Cast<AActor>(Interactor->GetAgent(Instance)))
-		{
-			UE_LEARNING_AGENTS_VLOG_STRING(this, LogLearning, Display,
-				Actor->GetActorLocation(),
-				VisualLogColor.ToFColor(true),
-				TEXT("Agent %i\nScale: [% 6.2f]\nValue: %s\nEncoded: %s"),
-				Instance,
-				FeatureObject->Scale,
-				*UE::Learning::Array::FormatFloat(ValueView[Instance]),
-				*UE::Learning::Array::FormatFloat(FeatureView[Instance]));
-		}
+		OutValues.Empty();
+		return false;
 	}
+
+	return true;
 }
-#endif
 
-UFloatArrayAction* UFloatArrayAction::AddFloatArrayAction(ULearningAgentsInteractor* InInteractor, const FName Name, const int32 Num, const float Scale)
+bool ULearningAgentsActionObject::GetContinuousActionToArrayView(TArrayView<float> OutValues, const FLearningAgentsActionObjectElement Element, const FName Name) const
 {
-	if (Num < 1)
+	if (!ActionObject.IsValid(Element.ObjectElement))
 	{
-		UE_LOG(LogLearning, Error, TEXT("AddFloatArrayAction: Number of elements in array must be at least 1, got %i."), Num);
-		return nullptr;
+		UE_LOG(LogLearning, Error, TEXT("%s: Invalid Action Object."), *GetName());
+		UE::Learning::Array::Zero<1, float>(OutValues);
+		return false;
 	}
 
-	return UE::Learning::Agents::Actions::Private::AddAction<UFloatArrayAction, UE::Learning::FFloatFeature>(InInteractor, Name, TEXT("AddFloatArrayAction"), Num, Scale);
+	if (ActionObject.GetName(Element.ObjectElement) != Name)
+	{
+		UE_LOG(LogLearning, Warning, TEXT("%s: Action name does not match. Action is '%s' but asked for '%s'."), *GetName(), *ActionObject.GetName(Element.ObjectElement).ToString(), *Name.ToString());
+	}
+
+	if (ActionObject.GetType(Element.ObjectElement) != UE::Learning::Action::EType::Continuous)
+	{
+		UE_LOG(LogLearning, Error, TEXT("%s: Action '%s' type does not match. Action is '%s' but asked for '%s'."),
+			*GetName(),
+			*ActionObject.GetName(Element.ObjectElement).ToString(),
+			UE::Learning::Agents::Action::Private::GetActionTypeString(ActionObject.GetType(Element.ObjectElement)),
+			UE::Learning::Agents::Action::Private::GetActionTypeString(UE::Learning::Action::EType::Continuous));
+		UE::Learning::Array::Zero<1, float>(OutValues);
+		return false;
+	}
+
+	const TArrayView<const float> Values = ActionObject.GetContinuous(Element.ObjectElement).Values;
+
+	if (Values.Num() != OutValues.Num())
+	{
+		UE_LOG(LogLearning, Error, TEXT("%s: Action '%s' size does not match. Action is '%i' values but asked for '%i'."),
+			*GetName(), *ActionObject.GetName(Element.ObjectElement).ToString(),
+			Values.Num(), OutValues.Num());
+		UE::Learning::Array::Zero<1, float>(OutValues);
+		return false;
+	}
+
+	UE::Learning::Array::Copy<1, float>(OutValues, Values);
+	return true;
 }
 
-void UFloatArrayAction::GetFloatArrayAction(const int32 AgentId, TArray<float>& OutValues)
+bool ULearningAgentsActionObject::GetExclusiveDiscreteAction(int32& OutIndex, const FLearningAgentsActionObjectElement Element, const FName Name) const
 {
-	if (!Interactor->HasAgent(AgentId))
+	if (!ActionObject.IsValid(Element.ObjectElement))
 	{
-		UE_LOG(LogLearning, Error, TEXT("%s: AgentId %d not found in the agents set."), *GetName(), AgentId);
-		OutValues.Empty();
-		return;
+		UE_LOG(LogLearning, Error, TEXT("%s: Invalid Action Object."), *GetName());
+		OutIndex = 0;
+		return false;
 	}
 
-	const TLearningArrayView<2, const float> View = FeatureObject->InstanceData->ConstView(FeatureObject->ValueHandle);
+	if (ActionObject.GetName(Element.ObjectElement) != Name)
+	{
+		UE_LOG(LogLearning, Warning, TEXT("%s: Action name does not match. Action is '%s' but asked for '%s'."), *GetName(), *ActionObject.GetName(Element.ObjectElement).ToString(), *Name.ToString());
+	}
 
-	AgentGetIteration[AgentId]++;
+	if (ActionObject.GetType(Element.ObjectElement) != UE::Learning::Action::EType::DiscreteExclusive)
+	{
+		UE_LOG(LogLearning, Error, TEXT("%s: Action '%s' type does not match. Action is '%s' but asked for '%s'."),
+			*GetName(),
+			*ActionObject.GetName(Element.ObjectElement).ToString(),
+			UE::Learning::Agents::Action::Private::GetActionTypeString(ActionObject.GetType(Element.ObjectElement)),
+			UE::Learning::Agents::Action::Private::GetActionTypeString(UE::Learning::Action::EType::DiscreteExclusive));
+		OutIndex = 0;
+		return false;
+	}
 
-	OutValues.SetNumUninitialized(View.Num<1>());
-	UE::Learning::Array::Copy<1, float>(OutValues, View[AgentId]);
+	OutIndex = ActionObject.GetDiscreteExclusive(Element.ObjectElement).DiscreteIndex;
+	return true;
 }
 
-void UFloatArrayAction::SetFloatArrayAction(const int32 AgentId, const TArray<float>& Values)
+bool ULearningAgentsActionObject::GetInclusiveDiscreteActionNum(int32& OutNum, const FLearningAgentsActionObjectElement Element, const FName Name) const
 {
-	if (!Interactor->HasAgent(AgentId))
+	if (!ActionObject.IsValid(Element.ObjectElement))
 	{
-		UE_LOG(LogLearning, Error, TEXT("%s: AgentId %d not found in the agents set."), *GetName(), AgentId);
-		return;
+		UE_LOG(LogLearning, Error, TEXT("%s: Invalid Action Object."), *GetName());
+		OutNum = 0;
+		return false;
 	}
 
-	const TLearningArrayView<2, float> View = FeatureObject->InstanceData->View(FeatureObject->ValueHandle);
+	if (ActionObject.GetName(Element.ObjectElement) != Name)
+	{
+		UE_LOG(LogLearning, Warning, TEXT("%s: Action name does not match. Action is '%s' but asked for '%s'."), *GetName(), *ActionObject.GetName(Element.ObjectElement).ToString(), *Name.ToString());
+	}
 
-	if (Values.Num() != View.Num<1>())
+	if (ActionObject.GetType(Element.ObjectElement) != UE::Learning::Action::EType::DiscreteInclusive)
 	{
-		UE_LOG(LogLearning, Error, TEXT("%s: Got wrong number of elements in array. Expected %i, got %i."), *GetName(), View.Num<1>(), Values.Num());
-		return;
+		UE_LOG(LogLearning, Error, TEXT("%s: Action '%s' type does not match. Action is '%s' but asked for '%s'."),
+			*GetName(),
+			*ActionObject.GetName(Element.ObjectElement).ToString(),
+			UE::Learning::Agents::Action::Private::GetActionTypeString(ActionObject.GetType(Element.ObjectElement)),
+			UE::Learning::Agents::Action::Private::GetActionTypeString(UE::Learning::Action::EType::DiscreteInclusive));
+		OutNum = 0;
+		return false;
 	}
 
-	UE::Learning::Array::Copy<1, float>(View[AgentId], Values);
-	AgentSetIteration[AgentId]++;
+	OutNum = ActionObject.GetDiscreteInclusive(Element.ObjectElement).DiscreteIndices.Num();
+	return true;
 }
 
-#if UE_LEARNING_AGENTS_ENABLE_VISUAL_LOG
-void UFloatArrayAction::VisualLog(const UE::Learning::FIndexSet Instances) const
+bool ULearningAgentsActionObject::GetInclusiveDiscreteAction(TArray<int32>& OutIndices, const FLearningAgentsActionObjectElement Element, const FName Name) const
 {
-	UE_LEARNING_TRACE_CPUPROFILER_EVENT_SCOPE(UFloatArrayAction::VisualLog);
+	int32 OutIndexNum = 0;
+	if (!GetInclusiveDiscreteActionNum(OutIndexNum, Element, Name))
+	{
+		OutIndices.Empty();
+		return false;
+	}
 
-	const TLearningArrayView<2, const float> ValueView = FeatureObject->InstanceData->ConstView(FeatureObject->ValueHandle);
-	const TLearningArrayView<2, const float> FeatureView = FeatureObject->InstanceData->ConstView(FeatureObject->FeatureHandle);
+	OutIndices.SetNumUninitialized(OutIndexNum);
 
-	for (const int32 Instance : Instances)
+	if (!GetInclusiveDiscreteActionToArrayView(OutIndices, Element, Name))
 	{
-		if (const AActor* Actor = Cast<AActor>(Interactor->GetAgent(Instance)))
-		{
-			UE_LEARNING_AGENTS_VLOG_STRING(this, LogLearning, Display,
-				Actor->GetActorLocation(),
-				VisualLogColor.ToFColor(true),
-				TEXT("Agent %i\nScale: [% 6.2f]\nValue: %s\nEncoded: %s"),
-				Instance,
-				FeatureObject->Scale,
-				*UE::Learning::Array::FormatFloat(ValueView[Instance]),
-				*UE::Learning::Array::FormatFloat(FeatureView[Instance]));
-		}
+		OutIndices.Empty();
+		return false;
 	}
+
+	return true;
 }
-#endif
 
-//------------------------------------------------------------------
+bool ULearningAgentsActionObject::GetInclusiveDiscreteActionToArrayView(TArrayView<int32> OutIndices, const FLearningAgentsActionObjectElement Element, const FName Name) const
+{
+	if (!ActionObject.IsValid(Element.ObjectElement))
+	{
+		UE_LOG(LogLearning, Error, TEXT("%s: Invalid Action Object."), *GetName());
+		UE::Learning::Array::Zero<1, int32>(OutIndices);
+		return false;
+	}
+
+	if (ActionObject.GetName(Element.ObjectElement) != Name)
+	{
+		UE_LOG(LogLearning, Warning, TEXT("%s: Action name does not match. Action is '%s' but asked for '%s'."), *GetName(), *ActionObject.GetName(Element.ObjectElement).ToString(), *Name.ToString());
+	}
+
+	if (ActionObject.GetType(Element.ObjectElement) != UE::Learning::Action::EType::DiscreteInclusive)
+	{
+		UE_LOG(LogLearning, Error, TEXT("%s: Action '%s' type does not match. Action is '%s' but asked for '%s'."),
+			*GetName(),
+			*ActionObject.GetName(Element.ObjectElement).ToString(),
+			UE::Learning::Agents::Action::Private::GetActionTypeString(ActionObject.GetType(Element.ObjectElement)),
+			UE::Learning::Agents::Action::Private::GetActionTypeString(UE::Learning::Action::EType::DiscreteInclusive));
+		UE::Learning::Array::Zero<1, int32>(OutIndices);
+		return false;
+	}
+
+	const TArrayView<const int32> Indices = ActionObject.GetDiscreteInclusive(Element.ObjectElement).DiscreteIndices;
 
-UVectorAction* UVectorAction::AddVectorAction(ULearningAgentsInteractor* InInteractor, const FName Name, const float Scale)
+	if (Indices.Num() != OutIndices.Num())
+	{
+		UE_LOG(LogLearning, Error, TEXT("%s: Action '%s' size does not match. Action is '%i' elements but asked for '%i'."),
+			*GetName(), *ActionObject.GetName(Element.ObjectElement).ToString(),
+			Indices.Num(), OutIndices.Num());
+		UE::Learning::Array::Zero<1, int32>(OutIndices);
+		return false;
+	}
+
+	UE::Learning::Array::Copy<1, int32>(OutIndices, Indices);
+	return true;
+}
+
+bool ULearningAgentsActionObject::GetStructActionNum(int32& OutNum, const FLearningAgentsActionObjectElement Element, const FName Name) const
 {
-	return UE::Learning::Agents::Actions::Private::AddAction<UVectorAction, UE::Learning::FFloatFeature>(InInteractor, Name, TEXT("AddVectorAction"), 3, Scale);
+	if (!ActionObject.IsValid(Element.ObjectElement))
+	{
+		UE_LOG(LogLearning, Error, TEXT("%s: Invalid Action Object."), *GetName());
+		OutNum = 0;
+		return false;
+	}
+
+	if (ActionObject.GetName(Element.ObjectElement) != Name)
+	{
+		UE_LOG(LogLearning, Warning, TEXT("%s: Action name does not match. Action is '%s' but asked for '%s'."), *GetName(), *ActionObject.GetName(Element.ObjectElement).ToString(), *Name.ToString());
+	}
+
+	if (ActionObject.GetType(Element.ObjectElement) != UE::Learning::Action::EType::And)
+	{
+		UE_LOG(LogLearning, Error, TEXT("%s: Action '%s' type does not match. Action is '%s' but asked for '%s'."),
+			*GetName(),
+			*ActionObject.GetName(Element.ObjectElement).ToString(),
+			UE::Learning::Agents::Action::Private::GetActionTypeString(ActionObject.GetType(Element.ObjectElement)),
+			UE::Learning::Agents::Action::Private::GetActionTypeString(UE::Learning::Action::EType::And));
+		OutNum = 0;
+		return false;
+	}
+	
+	const UE::Learning::Action::FObjectAndParameters Parameters = ActionObject.GetAnd(Element.ObjectElement);
+
+	OutNum = Parameters.Elements.Num();
+	return true;
 }
 
-FVector UVectorAction::GetVectorAction(const int32 AgentId)
+bool ULearningAgentsActionObject::GetStructAction(TMap<FName, FLearningAgentsActionObjectElement>& OutElements, const FLearningAgentsActionObjectElement Element, const FName Name) const
 {
-	if (!Interactor->HasAgent(AgentId))
+	int32 OutElementNum = 0;
+	if (!GetStructActionNum(OutElementNum, Element, Name))
 	{
-		UE_LOG(LogLearning, Error, TEXT("%s: AgentId %d not found in the agents set."), *GetName(), AgentId);
-		return FVector::ZeroVector;
+		OutElements.Empty();
+		return false;
 	}
 
-	const TLearningArrayView<2, const float> View = FeatureObject->InstanceData->ConstView(FeatureObject->ValueHandle);
+	TArray<FName, TInlineAllocator<16>> SubElementNames;
+	TArray<FLearningAgentsActionObjectElement, TInlineAllocator<16>> SubElements;
+	SubElementNames.SetNumUninitialized(OutElementNum);
+	SubElements.SetNumUninitialized(OutElementNum);
+
+	if (!GetStructActionToArrayViews(SubElementNames, SubElements, Element, Name))
+	{
+		OutElements.Empty();
+		return false;
+	}
 
-	AgentGetIteration[AgentId]++;
+	OutElements.Empty(OutElementNum);
+	for (int32 ElementIdx = 0; ElementIdx < OutElementNum; ElementIdx++)
+	{
+		OutElements.Add(SubElementNames[ElementIdx], SubElements[ElementIdx]);
+	}
 
-	return FVector(View[AgentId][0], View[AgentId][1], View[AgentId][2]);
+	return true;
 }
 
-void UVectorAction::SetVectorAction(const int32 AgentId, const FVector InAction)
+bool ULearningAgentsActionObject::GetStructActionToArrays(TArray<FName>& OutElementNames, TArray<FLearningAgentsActionObjectElement>& OutElements, const FLearningAgentsActionObjectElement Element, const FName Name) const
 {
-	if (!Interactor->HasAgent(AgentId))
+	int32 OutElementNum = 0;
+	if (!GetStructActionNum(OutElementNum, Element, Name))
 	{
-		UE_LOG(LogLearning, Error, TEXT("%s: AgentId %d not found in the agents set."), *GetName(), AgentId);
-		return;
+		OutElementNames.Empty();
+		OutElements.Empty();
+		return false;
 	}
 
-	const TLearningArrayView<2, float> View = FeatureObject->InstanceData->View(FeatureObject->ValueHandle);
-	View[AgentId][0] = InAction.X;
-	View[AgentId][1] = InAction.Y;
-	View[AgentId][2] = InAction.Z;
+	OutElementNames.SetNumUninitialized(OutElementNum);
+	OutElements.SetNumUninitialized(OutElementNum);
+
+	if (!GetStructActionToArrayViews(OutElementNames, OutElements, Element, Name))
+	{
+		OutElementNames.Empty();
+		OutElements.Empty();
+		return false;
+	}
 
-	AgentSetIteration[AgentId]++;
+	return true;
 }
 
-#if UE_LEARNING_AGENTS_ENABLE_VISUAL_LOG
-void UVectorAction::VisualLog(const UE::Learning::FIndexSet Instances) const
+bool ULearningAgentsActionObject::GetStructActionToArrayViews(TArrayView<FName> OutElementNames, TArrayView<FLearningAgentsActionObjectElement> OutElements, const FLearningAgentsActionObjectElement Element, const FName Name) const
 {
-	UE_LEARNING_TRACE_CPUPROFILER_EVENT_SCOPE(UVectorAction::VisualLog);
+	if (!ActionObject.IsValid(Element.ObjectElement))
+	{
+		UE_LOG(LogLearning, Error, TEXT("%s: Invalid Action Object."), *GetName());
+		UE::Learning::Array::Set<1, FName>(OutElementNames, NAME_None);
+		UE::Learning::Array::Set<1, FLearningAgentsActionObjectElement>(OutElements, FLearningAgentsActionObjectElement());
+		return false;
+	}
 
-	const TLearningArrayView<2, const float> ValueView = FeatureObject->InstanceData->ConstView(FeatureObject->ValueHandle);
-	const TLearningArrayView<2, const float> FeatureView = FeatureObject->InstanceData->ConstView(FeatureObject->FeatureHandle);
+	if (ActionObject.GetName(Element.ObjectElement) != Name)
+	{
+		UE_LOG(LogLearning, Warning, TEXT("%s: Action name does not match. Action is '%s' but asked for '%s'."), *GetName(), *ActionObject.GetName(Element.ObjectElement).ToString(), *Name.ToString());
+	}
 
-	for (const int32 Instance : Instances)
+	if (ActionObject.GetType(Element.ObjectElement) != UE::Learning::Action::EType::And)
 	{
-		if (const AActor* Actor = Cast<AActor>(Interactor->GetAgent(Instance)))
-		{
-			const FVector Vector(ValueView[Instance][0], ValueView[Instance][1], ValueView[Instance][2]);
+		UE_LOG(LogLearning, Error, TEXT("%s: Action '%s' type does not match. Action is '%s' but asked for '%s'."),
+			*GetName(),
+			*ActionObject.GetName(Element.ObjectElement).ToString(),
+			UE::Learning::Agents::Action::Private::GetActionTypeString(ActionObject.GetType(Element.ObjectElement)),
+			UE::Learning::Agents::Action::Private::GetActionTypeString(UE::Learning::Action::EType::And));
+		UE::Learning::Array::Set<1, FName>(OutElementNames, NAME_None);
+		UE::Learning::Array::Set<1, FLearningAgentsActionObjectElement>(OutElements, FLearningAgentsActionObjectElement());
+		return false;
+	}
 
-			UE_LEARNING_AGENTS_VLOG_ARROW(this, LogLearning, Display,
-				Actor->GetActorLocation(),
-				Actor->GetActorLocation() + Vector,
-				VisualLogColor.ToFColor(true),
-				TEXT(""));
+	const UE::Learning::Action::FObjectAndParameters Parameters = ActionObject.GetAnd(Element.ObjectElement);
 
-			UE_LEARNING_AGENTS_VLOG_STRING(this, LogLearning, Display,
-				Actor->GetActorLocation() + Vector,
-				VisualLogColor.ToFColor(true),
-				TEXT("Vector: [% 6.4f % 6.4f % 6.4f]"),
-				Vector.X, Vector.Y, Vector.Z);
+	if (Parameters.Elements.Num() == 0)
+	{
+		UE_LOG(LogLearning, Warning, TEXT("%s: Getting zero-sized And Action."), *GetName());
+	}
+
+	if (Parameters.Elements.Num() != OutElements.Num())
+	{
+		UE_LOG(LogLearning, Error, TEXT("%s: Action '%s' size does not match. Action is '%i' elements but asked for '%i'."),
+			*GetName(),
+			*ActionObject.GetName(Element.ObjectElement).ToString(),
+			Parameters.Elements.Num(), OutElements.Num());
+		UE::Learning::Array::Set<1, FName>(OutElementNames, NAME_None);
+		UE::Learning::Array::Set<1, FLearningAgentsActionObjectElement>(OutElements, FLearningAgentsActionObjectElement());
+		return false;
+	}
 
-			UE_LEARNING_AGENTS_VLOG_STRING(this, LogLearning, Display,
-				Actor->GetActorLocation(),
-				VisualLogColor.ToFColor(true),
-				TEXT("Agent %i\nScale: [% 6.2f]\nEncoded: [% 6.3f % 6.3f % 6.3f]"),
-				Instance,
-				FeatureObject->Scale,
-				FeatureView[Instance][0], FeatureView[Instance][1], FeatureView[Instance][2]);
+	for (int32 ElementIdx = 0; ElementIdx < Parameters.Elements.Num(); ElementIdx++)
+	{
+		if (!ActionObject.IsValid(Parameters.Elements[ElementIdx]))
+		{
+			UE_LOG(LogLearning, Error, TEXT("%s: Invalid Action Object."), *GetName());
+			UE::Learning::Array::Set<1, FName>(OutElementNames, NAME_None);
+			UE::Learning::Array::Set<1, FLearningAgentsActionObjectElement>(OutElements, FLearningAgentsActionObjectElement());
+			return false;
 		}
+
+		OutElementNames[ElementIdx] = Parameters.ElementNames[ElementIdx];
+		OutElements[ElementIdx] = { Parameters.Elements[ElementIdx] };
 	}
+
+	return true;
 }
-#endif
 
-UVectorArrayAction* UVectorArrayAction::AddVectorArrayAction(ULearningAgentsInteractor* InInteractor, const FName Name, const int32 Num, const float Scale)
+bool ULearningAgentsActionObject::GetExclusiveUnionAction(FName& OutElementName, FLearningAgentsActionObjectElement& OutElement, const FLearningAgentsActionObjectElement Element, const FName Name) const
 {
-	if (Num < 1)
+	if (!ActionObject.IsValid(Element.ObjectElement))
 	{
-		UE_LOG(LogLearning, Error, TEXT("AddVectorArrayAction: Number of elements in array must be at least 1, got %i."), Num);
-		return nullptr;
+		UE_LOG(LogLearning, Error, TEXT("%s: Invalid Action Object."), *GetName());
+		OutElementName = NAME_None;
+		OutElement = FLearningAgentsActionObjectElement();
+		return false;
 	}
 
-	return UE::Learning::Agents::Actions::Private::AddAction<UVectorArrayAction, UE::Learning::FFloatFeature>(InInteractor, Name, TEXT("AddVectorArrayAction"), Num * 3, Scale);
-}
+	if (ActionObject.GetName(Element.ObjectElement) != Name)
+	{
+		UE_LOG(LogLearning, Warning, TEXT("%s: Action name does not match. Action is '%s' but asked for '%s'."), *GetName(), *ActionObject.GetName(Element.ObjectElement).ToString(), *Name.ToString());
+	}
 
-void UVectorArrayAction::GetVectorArrayAction(const int32 AgentId, TArray<FVector>& OutVectors)
-{
-	if (!Interactor->HasAgent(AgentId))
+	if (ActionObject.GetType(Element.ObjectElement) != UE::Learning::Action::EType::OrExclusive)
 	{
-		UE_LOG(LogLearning, Error, TEXT("%s: AgentId %d not found in the agents set."), *GetName(), AgentId);
-		OutVectors.Empty();
-		return;
+		UE_LOG(LogLearning, Error, TEXT("%s: Action '%s' type does not match. Action is '%s' but asked for '%s'."),
+			*GetName(),
+			*ActionObject.GetName(Element.ObjectElement).ToString(),
+			UE::Learning::Agents::Action::Private::GetActionTypeString(ActionObject.GetType(Element.ObjectElement)),
+			UE::Learning::Agents::Action::Private::GetActionTypeString(UE::Learning::Action::EType::OrExclusive));
+		OutElementName = NAME_None;
+		OutElement = FLearningAgentsActionObjectElement();
+		return false;
 	}
 
-	const TLearningArrayView<2, const float> View = FeatureObject->InstanceData->ConstView(FeatureObject->ValueHandle);
+	const UE::Learning::Action::FObjectOrExclusiveParameters Parameters = ActionObject.GetOrExclusive(Element.ObjectElement);
+
+	OutElementName = Parameters.ElementName;
+	OutElement = { Parameters.Element };
+	return true;
+}
 
-	AgentGetIteration[AgentId]++;
+bool ULearningAgentsActionObject::GetInclusiveUnionActionNum(int32& OutNum, const FLearningAgentsActionObjectElement Element, const FName Name) const
+{
+	if (!ActionObject.IsValid(Element.ObjectElement))
+	{
+		UE_LOG(LogLearning, Error, TEXT("%s: Invalid Action Object."), *GetName());
+		OutNum = 0;
+		return false;
+	}
 
-	OutVectors.SetNumUninitialized(View.Num<1>() / 3);
+	if (ActionObject.GetName(Element.ObjectElement) != Name)
+	{
+		UE_LOG(LogLearning, Warning, TEXT("%s: Action name does not match. Action is '%s' but asked for '%s'."), *GetName(), *ActionObject.GetName(Element.ObjectElement).ToString(), *Name.ToString());
+	}
 
-	for (int32 VectorIdx = 0; VectorIdx < View.Num<1>() / 3; VectorIdx++)
+	if (ActionObject.GetType(Element.ObjectElement) != UE::Learning::Action::EType::OrInclusive)
 	{
-		OutVectors[VectorIdx] = FVector(
-			View[AgentId][VectorIdx * 3 + 0], 
-			View[AgentId][VectorIdx * 3 + 1], 
-			View[AgentId][VectorIdx * 3 + 2]);
+		UE_LOG(LogLearning, Error, TEXT("%s: Action '%s' type does not match. Action is '%s' but asked for '%s'."),
+			*GetName(),
+			*ActionObject.GetName(Element.ObjectElement).ToString(),
+			UE::Learning::Agents::Action::Private::GetActionTypeString(ActionObject.GetType(Element.ObjectElement)),
+			UE::Learning::Agents::Action::Private::GetActionTypeString(UE::Learning::Action::EType::OrInclusive));
+		OutNum = 0;
+		return false;
 	}
+
+	const UE::Learning::Action::FObjectOrInclusiveParameters Parameters = ActionObject.GetOrInclusive(Element.ObjectElement);
+
+	OutNum = Parameters.Elements.Num();
+	return true;
 }
 
-void UVectorArrayAction::SetVectorArrayAction(const int32 AgentId, const TArray<FVector>& Vectors)
+bool ULearningAgentsActionObject::GetInclusiveUnionAction(TMap<FName, FLearningAgentsActionObjectElement>& OutElements, const FLearningAgentsActionObjectElement Element, const FName Name) const
 {
-	if (!Interactor->HasAgent(AgentId))
+	int32 OutElementNum = 0;
+	if (!GetInclusiveUnionActionNum(OutElementNum, Element, Name))
 	{
-		UE_LOG(LogLearning, Error, TEXT("%s: AgentId %d not found in the agents set."), *GetName(), AgentId);
-		return;
+		OutElements.Empty();
+		return false;
 	}
 
-	const TLearningArrayView<2, float> View = FeatureObject->InstanceData->View(FeatureObject->ValueHandle);
+	TArray<FName, TInlineAllocator<16>> SubElementNames;
+	TArray<FLearningAgentsActionObjectElement, TInlineAllocator<16>> SubElements;
+	SubElementNames.SetNumUninitialized(OutElementNum);
+	SubElements.SetNumUninitialized(OutElementNum);
 
-	if (Vectors.Num() != View.Num<1>() / 3)
+	if (!GetInclusiveUnionActionToArrayViews(SubElementNames, SubElements, Element, Name))
 	{
-		UE_LOG(LogLearning, Error, TEXT("%s: Got wrong number of elements in array. Expected %i, got %i."), *GetName(), View.Num<1>() / 3, Vectors.Num());
-		return;
+		OutElements.Empty();
+		return false;
 	}
 
-	for (int32 VectorIdx = 0; VectorIdx < Vectors.Num(); VectorIdx++)
+	OutElements.Empty(OutElementNum);
+	for (int32 ElementIdx = 0; ElementIdx < OutElementNum; ElementIdx++)
 	{
-		View[AgentId][VectorIdx * 3 + 0] = Vectors[VectorIdx].X;
-		View[AgentId][VectorIdx * 3 + 1] = Vectors[VectorIdx].Y;
-		View[AgentId][VectorIdx * 3 + 2] = Vectors[VectorIdx].Z;
+		OutElements.Add(SubElementNames[ElementIdx], SubElements[ElementIdx]);
 	}
 
-	AgentSetIteration[AgentId]++;
+	return true;
 }
 
-#if UE_LEARNING_AGENTS_ENABLE_VISUAL_LOG
-void UVectorArrayAction::VisualLog(const UE::Learning::FIndexSet Instances) const
+bool ULearningAgentsActionObject::GetInclusiveUnionActionToArrays(TArray<FName>& OutElementNames, TArray<FLearningAgentsActionObjectElement>& OutElements, const FLearningAgentsActionObjectElement Element, const FName Name) const
 {
-	UE_LEARNING_TRACE_CPUPROFILER_EVENT_SCOPE(UVectorArrayAction::VisualLog);
+	int32 OutElementNum = 0;
+	if (!GetInclusiveUnionActionNum(OutElementNum, Element, Name))
+	{
+		OutElementNames.Empty();
+		OutElements.Empty();
+		return false;
+	}
 
-	const TLearningArrayView<2, const float> ValueView = FeatureObject->InstanceData->ConstView(FeatureObject->ValueHandle);
-	const TLearningArrayView<2, const float> FeatureView = FeatureObject->InstanceData->ConstView(FeatureObject->FeatureHandle);
+	OutElementNames.SetNumUninitialized(OutElementNum);
+	OutElements.SetNumUninitialized(OutElementNum);
 
-	const int32 VectorNum = ValueView.Num<1>() / 3;
+	if (!GetInclusiveUnionActionToArrayViews(OutElementNames, OutElements, Element, Name))
+	{
+		OutElementNames.Empty();
+		OutElements.Empty();
+		return false;
+	}
+
+	return true;
+}
 
-	for (const int32 Instance : Instances)
+bool ULearningAgentsActionObject::GetInclusiveUnionActionToArrayViews(TArrayView<FName> OutElementNames, TArrayView<FLearningAgentsActionObjectElement> OutElements, const FLearningAgentsActionObjectElement Element, const FName Name) const
+{
+	if (!ActionObject.IsValid(Element.ObjectElement))
 	{
-		if (const AActor* Actor = Cast<AActor>(Interactor->GetAgent(Instance)))
-		{
-			for (int32 VectorIdx = 0; VectorIdx < VectorNum; VectorIdx++)
-			{
-				const FVector Vector(ValueView[Instance][VectorIdx * 3 + 0], ValueView[Instance][VectorIdx * 3 + 1], ValueView[Instance][VectorIdx * 3 + 2]);
-				const FVector Offset = UE::Learning::Agents::Debug::GridOffsetForIndex(VectorIdx, VectorNum);
-
-				UE_LEARNING_AGENTS_VLOG_ARROW(this, LogLearning, Display,
-					Actor->GetActorLocation() + Offset,
-					Actor->GetActorLocation() + Offset + Vector,
-					VisualLogColor.ToFColor(true),
-					TEXT(""));
-
-				UE_LEARNING_AGENTS_VLOG_STRING(this, LogLearning, Display,
-					Actor->GetActorLocation() + Offset + Vector,
-					VisualLogColor.ToFColor(true),
-					TEXT("Vector %i: [% 6.4f % 6.4f % 6.4f]"),
-					VectorIdx,
-					Vector.X, Vector.Y, Vector.Z);
-			}
+		UE_LOG(LogLearning, Error, TEXT("%s: Invalid Action Object."), *GetName());
+		UE::Learning::Array::Set<1, FName>(OutElementNames, NAME_None);
+		UE::Learning::Array::Set<1, FLearningAgentsActionObjectElement>(OutElements, FLearningAgentsActionObjectElement());
+		return false;
+	}
 
-			UE_LEARNING_AGENTS_VLOG_STRING(this, LogLearning, Display,
-				Actor->GetActorLocation(),
-				VisualLogColor.ToFColor(true),
-				TEXT("Agent %i\nScale: [% 6.2f]\nEncoded: %s"),
-				Instance,
-				FeatureObject->Scale,
-				*UE::Learning::Array::FormatFloat(FeatureView[Instance]));
-		}
+	if (ActionObject.GetName(Element.ObjectElement) != Name)
+	{
+		UE_LOG(LogLearning, Warning, TEXT("%s: Action name does not match. Action is '%s' but asked for '%s'."), *GetName(), *ActionObject.GetName(Element.ObjectElement).ToString(), *Name.ToString());
 	}
-}
-#endif
 
-//------------------------------------------------------------------
+	if (ActionObject.GetType(Element.ObjectElement) != UE::Learning::Action::EType::OrInclusive)
+	{
+		UE_LOG(LogLearning, Error, TEXT("%s: Action '%s' type does not match. Action is '%s' but asked for '%s'."),
+			*GetName(),
+			*ActionObject.GetName(Element.ObjectElement).ToString(),
+			UE::Learning::Agents::Action::Private::GetActionTypeString(ActionObject.GetType(Element.ObjectElement)),
+			UE::Learning::Agents::Action::Private::GetActionTypeString(UE::Learning::Action::EType::OrInclusive));
+		UE::Learning::Array::Set<1, FName>(OutElementNames, NAME_None);
+		UE::Learning::Array::Set<1, FLearningAgentsActionObjectElement>(OutElements, FLearningAgentsActionObjectElement());
+		return false;
+	}
 
+	const UE::Learning::Action::FObjectOrInclusiveParameters Parameters = ActionObject.GetOrInclusive(Element.ObjectElement);
 
-UPlanarVelocityAction* UPlanarVelocityAction::AddPlanarVelocityAction(ULearningAgentsInteractor* InInteractor, const FName Name, const float Scale, const FVector Axis0, const FVector Axis1)
-{
-	return UE::Learning::Agents::Actions::Private::AddAction<UPlanarVelocityAction, UE::Learning::FPlanarVelocityFeature>(
-		InInteractor, 
-		Name, 
-		TEXT("AddPlanarVelocityAction"), 
-		1, 
-		Scale,
-		Axis0.GetSafeNormal(UE_SMALL_NUMBER, FVector::ForwardVector),
-		Axis1.GetSafeNormal(UE_SMALL_NUMBER, FVector::RightVector));
+	if (Parameters.Elements.Num() != OutElements.Num())
+	{
+		UE_LOG(LogLearning, Error, TEXT("%s: Action '%s' size does not match. Action is '%i' elements but asked for '%i'."),
+			*GetName(),
+			*ActionObject.GetName(Element.ObjectElement).ToString(),
+			Parameters.Elements.Num(), OutElements.Num());
+		UE::Learning::Array::Set<1, FName>(OutElementNames, NAME_None);
+		UE::Learning::Array::Set<1, FLearningAgentsActionObjectElement>(OutElements, FLearningAgentsActionObjectElement());
+		return false;
+	}
+
+	for (int32 ElementIdx = 0; ElementIdx < Parameters.Elements.Num(); ElementIdx++)
+	{
+		if (!ActionObject.IsValid(Parameters.Elements[ElementIdx]))
+		{
+			UE_LOG(LogLearning, Error, TEXT("%s: Invalid Action Object."), *GetName());
+			UE::Learning::Array::Set<1, FName>(OutElementNames, NAME_None);
+			UE::Learning::Array::Set<1, FLearningAgentsActionObjectElement>(OutElements, FLearningAgentsActionObjectElement());
+			return false;
+		}
+
+		OutElementNames[ElementIdx] = Parameters.ElementNames[ElementIdx];
+		OutElements[ElementIdx] = { Parameters.Elements[ElementIdx] };
+	}
+
+	return true;
 }
 
-FVector UPlanarVelocityAction::GetPlanarVelocityAction(const int32 AgentId)
+bool ULearningAgentsActionObject::GetStaticArrayActionNum(int32& OutNum, const FLearningAgentsActionObjectElement Element, const FName Name) const
 {
-	if (!Interactor->HasAgent(AgentId))
+	if (!ActionObject.IsValid(Element.ObjectElement))
 	{
-		UE_LOG(LogLearning, Error, TEXT("%s: AgentId %d not found in the agents set."), *GetName(), AgentId);
-		return FVector::ZeroVector;
+		UE_LOG(LogLearning, Error, TEXT("%s: Invalid Action Object."), *GetName());
+		OutNum = 0;
+		return false;
 	}
 
-	const TLearningArrayView<2, const FVector> View = FeatureObject->InstanceData->ConstView(FeatureObject->VelocityHandle);
+	if (ActionObject.GetName(Element.ObjectElement) != Name)
+	{
+		UE_LOG(LogLearning, Warning, TEXT("%s: Action name does not match. Action is '%s' but asked for '%s'."), *GetName(), *ActionObject.GetName(Element.ObjectElement).ToString(), *Name.ToString());
+	}
 
-	AgentGetIteration[AgentId]++;
+	if (ActionObject.GetType(Element.ObjectElement) != UE::Learning::Action::EType::Array)
+	{
+		UE_LOG(LogLearning, Error, TEXT("%s: Action '%s' type does not match. Action is '%s' but asked for '%s'."),
+			*GetName(),
+			*ActionObject.GetName(Element.ObjectElement).ToString(),
+			UE::Learning::Agents::Action::Private::GetActionTypeString(ActionObject.GetType(Element.ObjectElement)),
+			UE::Learning::Agents::Action::Private::GetActionTypeString(UE::Learning::Action::EType::Array));
+		OutNum = 0;
+		return false;
+	}
 
-	return View[AgentId][0];
+	OutNum = ActionObject.GetArray(Element.ObjectElement).Elements.Num();
+	return true;
 }
 
-void UPlanarVelocityAction::SetPlanarVelocityAction(const int32 AgentId, const FVector Velocity)
+bool ULearningAgentsActionObject::GetStaticArrayAction(TArray<FLearningAgentsActionObjectElement>& OutElements, const FLearningAgentsActionObjectElement Element, const FName Name) const
 {
-	if (!Interactor->HasAgent(AgentId))
+	int32 OutElementNum = 0;
+	if (!GetStaticArrayActionNum(OutElementNum, Element, Name))
 	{
-		UE_LOG(LogLearning, Error, TEXT("%s: AgentId %d not found in the agents set."), *GetName(), AgentId);
-		return;
+		OutElements.Empty();
+		return false;
 	}
 
-	const TLearningArrayView<2, FVector> View = FeatureObject->InstanceData->View(FeatureObject->VelocityHandle);
-	View[AgentId][0] = Velocity;
+	OutElements.SetNumUninitialized(OutElementNum);
 
-	AgentSetIteration[AgentId]++;
+	if (!GetStaticArrayActionToArrayView(OutElements, Element, Name))
+	{
+		OutElements.Empty();
+		return false;
+	}
+
+	return true;
 }
 
-#if UE_LEARNING_AGENTS_ENABLE_VISUAL_LOG
-void UPlanarVelocityAction::VisualLog(const UE::Learning::FIndexSet Instances) const
+bool ULearningAgentsActionObject::GetStaticArrayActionToArrayView(TArrayView<FLearningAgentsActionObjectElement> OutElements, const FLearningAgentsActionObjectElement Element, const FName Name) const
 {
-	UE_LEARNING_TRACE_CPUPROFILER_EVENT_SCOPE(UPlanarVelocityAction::VisualLog);
+	if (!ActionObject.IsValid(Element.ObjectElement))
+	{
+		UE_LOG(LogLearning, Error, TEXT("%s: Invalid Action Object."), *GetName());
+		UE::Learning::Array::Set<1, FLearningAgentsActionObjectElement>(OutElements, FLearningAgentsActionObjectElement());
+		return false;
+	}
 
-	const TLearningArrayView<2, const FVector> VelocityView = FeatureObject->InstanceData->ConstView(FeatureObject->VelocityHandle);
-	const TLearningArrayView<2, const float> FeatureView = FeatureObject->InstanceData->ConstView(FeatureObject->FeatureHandle);
+	if (ActionObject.GetName(Element.ObjectElement) != Name)
+	{
+		UE_LOG(LogLearning, Warning, TEXT("%s: Action name does not match. Action is '%s' but asked for '%s'."), *GetName(), *ActionObject.GetName(Element.ObjectElement).ToString(), *Name.ToString());
+	}
 
-	for (const int32 Instance : Instances)
+	if (ActionObject.GetType(Element.ObjectElement) != UE::Learning::Action::EType::Array)
 	{
-		if (const AActor* Actor = Cast<AActor>(Interactor->GetAgent(Instance)))
-		{
-			const FVector Velocity = VelocityView[Instance][0];
+		UE_LOG(LogLearning, Error, TEXT("%s: Action '%s' type does not match. Action is '%s' but asked for '%s'."), 
+			*GetName(),
+			*ActionObject.GetName(Element.ObjectElement).ToString(),
+			UE::Learning::Agents::Action::Private::GetActionTypeString(ActionObject.GetType(Element.ObjectElement)),
+			UE::Learning::Agents::Action::Private::GetActionTypeString(UE::Learning::Action::EType::Array));
+		UE::Learning::Array::Set<1, FLearningAgentsActionObjectElement>(OutElements, FLearningAgentsActionObjectElement());
+		return false;
+	}
 
-			UE_LEARNING_AGENTS_VLOG_ARROW(this, LogLearning, Display,
-				Actor->GetActorLocation(),
-				Actor->GetActorLocation() + Velocity,
-				VisualLogColor.ToFColor(true),
-				TEXT(""));
+	const TArrayView<const UE::Learning::Action::FObjectElement> SubElements = ActionObject.GetArray(Element.ObjectElement).Elements;
 
-			UE_LEARNING_AGENTS_VLOG_STRING(this, LogLearning, Display,
-				Actor->GetActorLocation() + Velocity,
-				VisualLogColor.ToFColor(true),
-				TEXT("Velocity: [% 6.3f % 6.3f % 6.3f]"),
-				Velocity.X, Velocity.Y, Velocity.Z);
+	if (SubElements.Num() == 0)
+	{
+		UE_LOG(LogLearning, Warning, TEXT("%s: Getting zero-sized Array Action."), *GetName());
+	}
 
-			UE_LEARNING_AGENTS_VLOG_PLANE(this, LogLearning, Display,
-				Actor->GetActorLocation(),
-				FQuat::Identity,
-				FeatureObject->Axis0,
-				FeatureObject->Axis1,
-				VisualLogColor.ToFColor(true),
-				TEXT(""));
+	if (SubElements.Num() != OutElements.Num())
+	{
+		UE_LOG(LogLearning, Error, TEXT("%s: Action '%s' size does not match. Action is '%i' elements but asked for '%i'."),
+			*GetName(), *ActionObject.GetName(Element.ObjectElement).ToString(),
+			SubElements.Num(), OutElements.Num());
+		UE::Learning::Array::Set<1, FLearningAgentsActionObjectElement>(OutElements, FLearningAgentsActionObjectElement());
+		return false;
+	}
 
-			UE_LEARNING_AGENTS_VLOG_TRANSFORM(this, LogLearning, Display,
-				Actor->GetActorLocation(),
-				FQuat::Identity,
-				VisualLogColor.ToFColor(true),
-				TEXT("Agent %i\nScale: [% 6.2f]\nEncoded: %s"),
-				Instance,
-				FeatureObject->Scale,
-				*UE::Learning::Array::FormatFloat(FeatureView[Instance]));
+	for (int32 ElementIdx = 0; ElementIdx < SubElements.Num(); ElementIdx++)
+	{
+		if (!ActionObject.IsValid(SubElements[ElementIdx]))
+		{
+			UE_LOG(LogLearning, Error, TEXT("%s: Invalid Action Object."), *GetName());
+			UE::Learning::Array::Set<1, FLearningAgentsActionObjectElement>(OutElements, FLearningAgentsActionObjectElement());
+			return false;
 		}
+
+		OutElements[ElementIdx] = { SubElements[ElementIdx] };
 	}
-}
-#endif
 
-//------------------------------------------------------------------
+	return true;
+}
 
-URotationAction* URotationAction::AddRotationAction(ULearningAgentsInteractor* InInteractor, const FName Name, const float Scale)
+bool ULearningAgentsActionObject::GetPairAction(FLearningAgentsActionObjectElement& OutKey, FLearningAgentsActionObjectElement& OutValue, const FLearningAgentsActionObjectElement Element, const FName Name) const
 {
-	return UE::Learning::Agents::Actions::Private::AddAction<URotationAction, UE::Learning::FRotationVectorFeature>(InInteractor, Name, TEXT("AddRotationAction"), 1, Scale);
+	TStaticArray<FName, 2> OutElementNames;
+	TStaticArray<FLearningAgentsActionObjectElement, 2> OutElements;
+	if (!GetStructActionToArrayViews(OutElementNames, OutElements, Element, Name))
+	{
+		OutKey = FLearningAgentsActionObjectElement();
+		OutValue = FLearningAgentsActionObjectElement();
+		return false;
+	}
+
+	OutKey = OutElements[MakeArrayView(OutElementNames).Find(TEXT("Key"))];
+	OutValue = OutElements[MakeArrayView(OutElementNames).Find(TEXT("Value"))];
+	return true;
 }
 
-FRotator URotationAction::GetRotationAction(const int32 AgentId)
+bool ULearningAgentsActionObject::GetEnumAction(uint8& OutEnumValue, const UEnum* Enum, const FLearningAgentsActionObjectElement Element, const FName Name) const
 {
-	if (!Interactor->HasAgent(AgentId))
+	if (!Enum)
 	{
-		UE_LOG(LogLearning, Error, TEXT("%s: AgentId %d not found in the agents set."), *GetName(), AgentId);
-		return FRotator::ZeroRotator;
+		UE_LOG(LogLearning, Error, TEXT("%s: Enum is nullptr."), *GetName());
+		OutEnumValue = 0;
+		return false;
 	}
 
-	const TLearningArrayView<2, const FVector> View = FeatureObject->InstanceData->ConstView(FeatureObject->RotationVectorsHandle);
+	int32 OutIndex = 0;
+	if (!GetExclusiveDiscreteAction(OutIndex, Element, Name))
+	{
+		OutEnumValue = 0;
+		return false;
+	}
 
-	AgentGetIteration[AgentId]++;
+	if (OutIndex >= Enum->NumEnums() - 1)
+	{
+		UE_LOG(LogLearning, Error, TEXT("%s: EnumValue out of range for Enum '%s'. Expected %i or less, got %i."), *GetName(), *Enum->GetName(), Enum->NumEnums() - 1, OutIndex);
+		OutEnumValue = 0;
+		return false;
+	}
 
-	return FQuat::MakeFromRotationVector((UE_TWO_PI / 180.0f) * View[AgentId][0]).Rotator();
-}
+	const int32 EnumValue = Enum->GetValueByIndex(OutIndex);
 
-FVector URotationAction::GetRotationActionAsRotationVector(const int32 AgentId)
-{
-	if (!Interactor->HasAgent(AgentId))
+	if (EnumValue == INDEX_NONE)
 	{
-		UE_LOG(LogLearning, Error, TEXT("%s: AgentId %d not found in the agents set."), *GetName(), AgentId);
-		return FVector::ZeroVector;
+		UE_LOG(LogLearning, Error, TEXT("%s: Enum Value not found for index %i."), *GetName(), OutIndex);
+		OutEnumValue = 0;
+		return false;
 	}
 
-	return FeatureObject->InstanceData->ConstView(FeatureObject->RotationVectorsHandle)[AgentId][0];
+	OutEnumValue = (uint8)EnumValue;
+	return true;
 }
 
-FQuat URotationAction::GetRotationActionAsQuat(const int32 AgentId)
+bool ULearningAgentsActionObject::GetBitmaskAction(int32& OutBitmaskValue, const UEnum* Enum, const FLearningAgentsActionObjectElement Element, const FName Name) const
 {
-	if (!Interactor->HasAgent(AgentId))
+	if (!Enum)
+	{
+		UE_LOG(LogLearning, Error, TEXT("%s: Enum is nullptr."), *GetName());
+		OutBitmaskValue = 0;
+		return false;
+	}
+
+	if (Enum->NumEnums() - 1 > 32)
+	{
+		UE_LOG(LogLearning, Error, TEXT("%s: Too many values in Enum to use as Bitmask (%i)."), *GetName(), Enum->NumEnums() - 1);
+		OutBitmaskValue = 0;
+		return false;
+	}
+
+	int32 EnumValueNum;
+	if (!GetInclusiveDiscreteActionNum(EnumValueNum, Element, Name))
+	{
+		OutBitmaskValue = 0;
+		return false;
+	}
+
+	if (EnumValueNum > Enum->NumEnums() - 1)
 	{
-		UE_LOG(LogLearning, Error, TEXT("%s: AgentId %d not found in the agents set."), *GetName(), AgentId);
-		return FQuat::Identity;
+		UE_LOG(LogLearning, Error, TEXT("%s: Too many values for Enum '%s'. Expected %i or less, got %i."), *GetName(), *Enum->GetName(), Enum->NumEnums() - 1, EnumValueNum);
+		OutBitmaskValue = 0;
+		return false;
 	}
 
-	const TLearningArrayView<2, const FVector> View = FeatureObject->InstanceData->ConstView(FeatureObject->RotationVectorsHandle);
+	TArray<int32, TInlineAllocator<32>> OutIndices;
+	OutIndices.SetNumUninitialized(EnumValueNum);
+	if (!GetInclusiveDiscreteActionToArrayView(OutIndices, Element, Name))
+	{
+		OutBitmaskValue = 0;
+		return false;
+	}
+
+	OutBitmaskValue = 0;
+	for (const int32 OutIndex : OutIndices)
+	{
+		OutBitmaskValue |= (1 << OutIndex);
+	}
+	return true;
+}
 
-	AgentGetIteration[AgentId]++;
+bool ULearningAgentsActionObject::GetOptionalAction(ELearningAgentsOptionalAction& OutOption, FLearningAgentsActionObjectElement& OutElement, const FLearningAgentsActionObjectElement Element, const FName Name) const
+{
+	FName OutName = NAME_None;
+	if (!GetExclusiveUnionAction(OutName, OutElement, Element, Name))
+	{
+		OutOption = ELearningAgentsOptionalAction::Null;
+		return false;
+	}
 
-	return FQuat::MakeFromRotationVector((UE_TWO_PI / 180.0f) * View[AgentId][0]);
+	OutOption = OutName == TEXT("Null") ? ELearningAgentsOptionalAction::Null : ELearningAgentsOptionalAction::Valid;
+	return true;
 }
 
-#if UE_LEARNING_AGENTS_ENABLE_VISUAL_LOG
-void URotationAction::VisualLog(const UE::Learning::FIndexSet Instances) const
+bool ULearningAgentsActionObject::GetEitherAction(ELearningAgentsEitherAction& OutEither, FLearningAgentsActionObjectElement& OutElement, const FLearningAgentsActionObjectElement Element, const FName Name) const
 {
-	UE_LEARNING_TRACE_CPUPROFILER_EVENT_SCOPE(URotationAction::VisualLog);
+	FName OutName = NAME_None;
+	if (!GetExclusiveUnionAction(OutName, OutElement, Element, Name))
+	{
+		OutEither = ELearningAgentsEitherAction::A;
+		return false;
+	}
 
-	const TLearningArrayView<2, const FVector> ValueView = FeatureObject->InstanceData->ConstView(FeatureObject->RotationVectorsHandle);
-	const TLearningArrayView<2, const float> FeatureView = FeatureObject->InstanceData->ConstView(FeatureObject->FeatureHandle);
+	OutEither = OutName == TEXT("A") ? ELearningAgentsEitherAction::A : ELearningAgentsEitherAction::B;
+	return true;
+}
 
-	const int32 RotationVectorNum = ValueView.Num<1>();
+bool ULearningAgentsActionObject::GetEncodingAction(FLearningAgentsActionObjectElement& OutElement, const FLearningAgentsActionObjectElement Element, const FName Name) const
+{
+	if (!ActionObject.IsValid(Element.ObjectElement))
+	{
+		UE_LOG(LogLearning, Error, TEXT("%s: Invalid Action Object."), *GetName());
+		OutElement = FLearningAgentsActionObjectElement();
+		return false;
+	}
 
-	for (const int32 Instance : Instances)
+	if (ActionObject.GetName(Element.ObjectElement) != Name)
 	{
-		if (const AActor* Actor = Cast<AActor>(Interactor->GetAgent(Instance)))
-		{
-			const FQuat Rotation = FQuat::MakeFromRotationVector((UE_TWO_PI / 180.0f) * ValueView[Instance][0]);
-			const FRotator Rotator = Rotation.Rotator();
+		UE_LOG(LogLearning, Warning, TEXT("%s: Action name does not match. Action is '%s' but asked for '%s'."), *GetName(), *ActionObject.GetName(Element.ObjectElement).ToString(), *Name.ToString());
+	}
 
-			UE_LEARNING_AGENTS_VLOG_TRANSFORM(this, LogLearning, Display,
-				Actor->GetActorLocation(),
-				Rotation,
-				VisualLogColor.ToFColor(true),
-				TEXT("Agent %i\nScale: [% 6.2f]\nRotation: [% 6.1f % 6.1f % 6.1f]\nEncoded: %s"),
-				Instance,
-				FeatureObject->Scale,
-				Rotator.Pitch, Rotator.Roll, Rotator.Yaw,
-				*UE::Learning::Array::FormatFloat(FeatureView[Instance]));
-		}
+	if (ActionObject.GetType(Element.ObjectElement) != UE::Learning::Action::EType::Encoding)
+	{
+		UE_LOG(LogLearning, Error, TEXT("%s: Action '%s' type does not match. Action is '%s' but asked for '%s'."),
+			*GetName(),
+			*ActionObject.GetName(Element.ObjectElement).ToString(),
+			UE::Learning::Agents::Action::Private::GetActionTypeString(ActionObject.GetType(Element.ObjectElement)),
+			UE::Learning::Agents::Action::Private::GetActionTypeString(UE::Learning::Action::EType::Encoding));
+		OutElement = FLearningAgentsActionObjectElement();
+		return false;
 	}
-}
-#endif
 
+	OutElement = { ActionObject.GetEncoding(Element.ObjectElement).Element };
 
-URotationArrayAction* URotationArrayAction::AddRotationArrayAction(ULearningAgentsInteractor* InInteractor, const FName Name, const int32 RotationNum, const float Scale)
+	return true;
+}
+
+bool ULearningAgentsActionObject::GetBoolAction(bool& bOutValue, const FLearningAgentsActionObjectElement Element, const FName Name) const
 {
-	if (RotationNum < 1)
+	int32 OutIndex = 0;
+	if (!GetExclusiveDiscreteAction(OutIndex, Element, Name))
 	{
-		UE_LOG(LogLearning, Error, TEXT("AddRotationArrayAction: Number of elements in array must be at least 1, got %i."), RotationNum);
-		return nullptr;
+		bOutValue = false;
+		return false;
 	}
 
-	return UE::Learning::Agents::Actions::Private::AddAction<URotationArrayAction, UE::Learning::FRotationVectorFeature>(InInteractor, Name, TEXT("AddRotationArrayAction"), RotationNum, Scale);
+	bOutValue = OutIndex == 1;
+	return true;
 }
 
-void URotationArrayAction::GetRotationArrayAction(const int32 AgentId, TArray<FRotator>& OutRotations)
+bool ULearningAgentsActionObject::GetFloatAction(float& OutValue, const FLearningAgentsActionObjectElement Element, const float FloatScale, const FName Name) const
 {
-	if (!Interactor->HasAgent(AgentId))
+	float OutValuesData;
+	if (!GetContinuousActionToArrayView(MakeArrayView(&OutValuesData, 1), Element, Name))
 	{
-		UE_LOG(LogLearning, Error, TEXT("%s: AgentId %d not found in the agents set."), *GetName(), AgentId);
-		return;
+		OutValue = 0.0f;
+		return false;
 	}
 
-	const TLearningArrayView<2, const FVector> View = FeatureObject->InstanceData->ConstView(FeatureObject->RotationVectorsHandle);
-
-	AgentGetIteration[AgentId]++;
+	OutValue = OutValuesData * FloatScale;
+	return true;
+}
 
-	OutRotations.SetNumUninitialized(View.Num<1>());
-	for (int32 RotationVectorIdx = 0; RotationVectorIdx < View.Num<1>(); RotationVectorIdx++)
+bool ULearningAgentsActionObject::GetTranslationAction(FVector& OutTranslation, const FLearningAgentsActionObjectElement Element, const FTransform RelativeTransform, const float TranslationScale, const FName Name) const
+{
+	TStaticArray<float, 3> OutValues;
+	if (!GetContinuousActionToArrayView(OutValues, Element, Name))
 	{
-		OutRotations[RotationVectorIdx] = FQuat::MakeFromRotationVector((UE_TWO_PI / 180.0f) * View[AgentId][RotationVectorIdx]).Rotator();
+		OutTranslation = FVector::ZeroVector;
+		return false;
 	}
+
+	const FVector LocalTranslation = TranslationScale * FVector(OutValues[0], OutValues[1], OutValues[2]);
+	OutTranslation = RelativeTransform.TransformPosition(LocalTranslation);
+	return true;
 }
 
-void URotationArrayAction::GetRotationArrayActionAsRotationVectors(const int32 AgentId, TArray<FVector>& OutRotationVectors)
+bool ULearningAgentsActionObject::GetRotationAction(FRotator& OutRotation, const FLearningAgentsActionObjectElement Element, const FRotator RelativeRotation, const float RotationScale, const FName Name) const
 {
-	if (!Interactor->HasAgent(AgentId))
+	FQuat OutRotationQuat;
+	if (!GetRotationActionAsQuat(OutRotationQuat, Element, FQuat::MakeFromRotator(RelativeRotation), RotationScale, Name))
 	{
-		UE_LOG(LogLearning, Error, TEXT("%s: AgentId %d not found in the agents set."), *GetName(), AgentId);
-		return;
+		OutRotation = FRotator::ZeroRotator;
+		return false;
 	}
 
-	const TLearningArrayView<2, const FVector> View = FeatureObject->InstanceData->ConstView(FeatureObject->RotationVectorsHandle);
+	OutRotation = OutRotationQuat.Rotator();
+	return true;
+}
+
+bool ULearningAgentsActionObject::GetRotationActionAsQuat(FQuat& OutRotation, const FLearningAgentsActionObjectElement Element, const FQuat RelativeRotation, const float RotationScale, const FName Name) const
+{
+	TStaticArray<float, 3> OutValues;
+	if (!GetContinuousActionToArrayView(OutValues, Element, Name))
+	{
+		OutRotation = FQuat::Identity;
+		return false;
+	}
 
-	AgentGetIteration[AgentId]++;
+	const FVector LocalRotationVector = FMath::DegreesToRadians(RotationScale) * FVector(OutValues[0], OutValues[1], OutValues[2]);
+	OutRotation = RelativeRotation * FQuat::MakeFromRotationVector(LocalRotationVector);
+	return true;
+}
 
-	OutRotationVectors.SetNumUninitialized(View.Num<1>());
-	for (int32 RotationVectorIdx = 0; RotationVectorIdx < View.Num<1>(); RotationVectorIdx++)
+bool ULearningAgentsActionObject::GetScaleAction(FVector& OutScale, const FLearningAgentsActionObjectElement Element, const FVector RelativeScale, const float Scale, const FName Name) const
+{
+	TStaticArray<float, 3> OutValues;
+	if (!GetContinuousActionToArrayView(OutValues, Element, Name))
 	{
-		OutRotationVectors[RotationVectorIdx] = View[AgentId][RotationVectorIdx];
+		OutScale = FVector::OneVector;
+		return false;
 	}
+
+	const FVector LocalScaleVector = UE::Learning::Agents::Action::Private::VectorExp(Scale * FVector(OutValues[0], OutValues[1], OutValues[2]));
+	OutScale = RelativeScale * LocalScaleVector;
+	return true;
 }
 
-void URotationArrayAction::GetRotationArrayActionAsQuats(const int32 AgentId, TArray<FQuat>& OutRotations)
+bool ULearningAgentsActionObject::GetTransformAction(FTransform& OutTransform, const FLearningAgentsActionObjectElement Element, const FTransform RelativeTransform, const float TranslationScale, const float RotationScale, const float ScaleScale, const FName Name) const
 {
-	if (!Interactor->HasAgent(AgentId))
+	TStaticArray<FName, 3> OutElementNames;
+	TStaticArray<FLearningAgentsActionObjectElement, 3> OutElements;
+	if (!GetStructActionToArrayViews(OutElementNames, OutElements, Element, Name))
 	{
-		UE_LOG(LogLearning, Error, TEXT("%s: AgentId %d not found in the agents set."), *GetName(), AgentId);
-		return;
+		OutTransform = FTransform::Identity;
+		return false;
 	}
 
-	const TLearningArrayView<2, const FVector> View = FeatureObject->InstanceData->ConstView(FeatureObject->RotationVectorsHandle);
+	FVector OutTranslation;
+	if (!GetTranslationAction(OutTranslation, OutElements[MakeArrayView(OutElementNames).Find(TEXT("Translation"))], RelativeTransform, TranslationScale))
+	{
+		OutTransform = FTransform::Identity;
+		return false;
+	}
 
-	AgentGetIteration[AgentId]++;
+	FQuat OutRotation;
+	if (!GetRotationActionAsQuat(OutRotation, OutElements[MakeArrayView(OutElementNames).Find(TEXT("Rotation"))], RelativeTransform.GetRotation(), RotationScale))
+	{
+		OutTransform = FTransform::Identity;
+		return false;
+	}
 
-	OutRotations.SetNumUninitialized(View.Num<1>());
-	for (int32 RotationVectorIdx = 0; RotationVectorIdx < View.Num<1>(); RotationVectorIdx++)
+	FVector OutScale;
+	if (!GetScaleAction(OutScale, OutElements[MakeArrayView(OutElementNames).Find(TEXT("Scale"))], RelativeTransform.GetScale3D(), ScaleScale))
 	{
-		OutRotations[RotationVectorIdx] = FQuat::MakeFromRotationVector((UE_TWO_PI / 180.0f) * View[AgentId][RotationVectorIdx]);
+		OutTransform = FTransform::Identity;
+		return false;
 	}
+
+	OutTransform = FTransform(OutRotation, OutTranslation, OutScale);
+	return true;
 }
 
-#if UE_LEARNING_AGENTS_ENABLE_VISUAL_LOG
-void URotationArrayAction::VisualLog(const UE::Learning::FIndexSet Instances) const
+bool ULearningAgentsActionObject::GetAngleAction(float& OutAngle, const FLearningAgentsActionObjectElement Element, const float RelativeAngle, const float AngleScale, const FName Name) const
 {
-	UE_LEARNING_TRACE_CPUPROFILER_EVENT_SCOPE(URotationArrayAction::VisualLog);
+	if (GetAngleActionRadians(OutAngle, Element, FMath::DegreesToRadians(RelativeAngle), FMath::DegreesToRadians(AngleScale), Name))
+	{
+		OutAngle = FMath::RadiansToDegrees(OutAngle);
+		return true;
+	}
+	else
+	{
+		OutAngle = 0.0f;
+		return false;
+	}
+}
 
-	const TLearningArrayView<2, const FVector> ValueView = FeatureObject->InstanceData->ConstView(FeatureObject->RotationVectorsHandle);
-	const TLearningArrayView<2, const float> FeatureView = FeatureObject->InstanceData->ConstView(FeatureObject->FeatureHandle);
+bool ULearningAgentsActionObject::GetAngleActionRadians(float& OutAngle, const FLearningAgentsActionObjectElement Element, const float RelativeAngle, const float AngleScale, const FName Name) const
+{
+	if (!GetContinuousActionToArrayView(MakeArrayView(&OutAngle, 1), Element, Name))
+	{
+		OutAngle = 0.0f;
+		return false;
+	}
 
-	const int32 RotationNum = ValueView.Num<1>();
+	OutAngle = RelativeAngle + AngleScale * OutAngle;
+	return true;
+}
 
-	for (const int32 Instance : Instances)
+bool ULearningAgentsActionObject::GetVelocityAction(FVector& OutVelocity, const FLearningAgentsActionObjectElement Element, const FTransform RelativeTransform, const float VelocityScale, const FName Name) const
+{
+	TStaticArray<float, 3> OutValues;
+	if (!GetContinuousActionToArrayView(OutValues, Element, Name))
 	{
-		if (const AActor* Actor = Cast<AActor>(Interactor->GetAgent(Instance)))
-		{
-			for (int32 RotationIdx = 0; RotationIdx < RotationNum; RotationIdx++)
-			{
-				const FVector Offset = UE::Learning::Agents::Debug::GridOffsetForIndex(RotationIdx, RotationNum);
-				const FQuat Rotation = FQuat::MakeFromRotationVector((UE_TWO_PI / 180.0f) * ValueView[Instance][RotationIdx]);
-				const FRotator Rotator = Rotation.Rotator();
-
-				UE_LEARNING_AGENTS_VLOG_TRANSFORM(this, LogLearning, Display,
-					Actor->GetActorLocation() + Offset,
-					Rotation,
-					VisualLogColor.ToFColor(true),
-					TEXT("Rotation %i: [% 6.3f % 6.3f % 6.3f]"),
-					RotationIdx,
-					Rotator.Pitch, Rotator.Roll, Rotator.Yaw);
-			}
-
-			UE_LEARNING_AGENTS_VLOG_STRING(this, LogLearning, Display,
-				Actor->GetActorLocation(),
-				VisualLogColor.ToFColor(true),
-				TEXT("Agent %i\nScale: [% 6.2f]\nEncoded: %s"),
-				Instance,
-				FeatureObject->Scale,
-				*UE::Learning::Array::FormatFloat(FeatureView[Instance]));
-		}
+		OutVelocity = FVector::OneVector;
+		return false;
 	}
+
+	OutVelocity = RelativeTransform.TransformVector(VelocityScale * FVector(OutValues[0], OutValues[1], OutValues[2]));
+	return true;
 }
-#endif
+
+bool ULearningAgentsActionObject::GetSpeedAction(float& OutSpeed, const FLearningAgentsActionObjectElement Element, const float SpeedScale, const FName Name) const
+{
+	if (!GetContinuousActionToArrayView(MakeArrayView(&OutSpeed, 1), Element, Name))
+	{
+		OutSpeed = 0.0f;
+		return false;
+	}
+
+	OutSpeed = SpeedScale * OutSpeed;
+	return true;
+}
\ No newline at end of file
