LearningAgents: Large restructuring to allow for structured actions and observations.

diff --git a/Engine/Plugins/Experimental/LearningAgents/Source/LearningAgents/Private/LearningAgentsController.cpp b/Engine/Plugins/Experimental/LearningAgents/Source/LearningAgents/Private/LearningAgentsController.cpp
--- a/Engine/Plugins/Experimental/LearningAgents/Source/LearningAgents/Private/LearningAgentsController.cpp
+++ b/Engine/Plugins/Experimental/LearningAgents/Source/LearningAgents/Private/LearningAgentsController.cpp
@@ -5,19 +5,81 @@
 #include "LearningAgentsManager.h"
 #include "LearningAgentsInteractor.h"
 #include "LearningAgentsActions.h"
-#include "LearningFeatureObject.h"
 #include "LearningLog.h"
 
 ULearningAgentsController::ULearningAgentsController() : Super(FObjectInitializer::Get()) {}
 ULearningAgentsController::ULearningAgentsController(FVTableHelper& Helper) : Super(Helper) {}
 ULearningAgentsController::~ULearningAgentsController() = default;
 
-void ULearningAgentsController::SetActions_Implementation(const TArray<int32>& AgentIds)
+void ULearningAgentsController::EvaluateAgentController_Implementation(
+	TArray<FLearningAgentsActionObjectElement>& OutActionObjectElements,
+	ULearningAgentsActionObject* InActionObject,
+	const TArray<FLearningAgentsObservationObjectElement>& InObservationObjectElements,
+	const ULearningAgentsObservationObject* InObservationObject,
+	const TArray<int32>& AgentIds)
 {
-	// Can be overridden to get actions without blueprints
+	UE_LOG(LogLearning, Error, TEXT("%s: EvaluateAgentController function must be overridden!"), *GetName());
+	OutActionObjectElements.Empty();
 }
 
-void ULearningAgentsController::SetupController(ULearningAgentsInteractor* InInteractor)
+void ULearningAgentsController::EvaluateAgentControllerUsingDelegate(
+	TArray<FLearningAgentsActionObjectElement>& OutActionObjectElements,
+	ULearningAgentsActionObject* InActionObject,
+	const TArray<FLearningAgentsObservationObjectElement>& InObservationObjectElements,
+	const ULearningAgentsObservationObject* InObservationObject,
+	const TArray<int32>& AgentIds,
+	const FEvaluateAgentControllerDelegate& Delegate)
+{
+	if (!Delegate.IsBound())
+	{
+		UE_LOG(LogLearning, Error, TEXT("%s: Delegate Not Bound."), *GetName());
+		return;
+	}
+
+	const int32 AgentNum = AgentIds.Num();
+
+	if (AgentNum != InObservationObjectElements.Num())
+	{
+		UE_LOG(LogLearning, Error, TEXT("%s: Not enough Observation Objects. Expected %i, Got %i."), *GetName(), AgentNum, InObservationObjectElements.Num());
+		return;
+	}
+
+	OutActionObjectElements.Empty(AgentNum);
+	for (int32 AgentIdx = 0; AgentIdx < AgentNum; AgentIdx++)
+	{
+		OutActionObjectElements.Add(Delegate.Execute(InActionObject, InObservationObjectElements[AgentIdx], InObservationObject, AgentIds[AgentIdx]));
+	}
+}
+
+ULearningAgentsController* ULearningAgentsController::MakeController(
+	ULearningAgentsManager* InManager, 
+	ULearningAgentsInteractor* InInteractor, 
+	TSubclassOf<ULearningAgentsController> Class,
+	const FName Name)
+{
+	if (!InManager)
+	{
+		UE_LOG(LogLearning, Error, TEXT("MakeController: InManager is nullptr."));
+		return nullptr;
+	}
+
+	if (!Class)
+	{
+		UE_LOG(LogLearning, Error, TEXT("MakeController: Class is nullptr."));
+		return nullptr;
+	}
+
+	const FName UniqueName = MakeUniqueObjectName(InManager, Class, Name, EUniqueObjectNameOptions::GloballyUnique);
+
+	ULearningAgentsController* Controller = NewObject<ULearningAgentsController>(InManager, Class, UniqueName);
+	if (!Controller) { return nullptr; }
+	
+	Controller->SetupController(InManager, InInteractor);
+
+	return Controller->IsSetup() ? Controller : nullptr;
+}
+
+void ULearningAgentsController::SetupController(ULearningAgentsManager* InManager, ULearningAgentsInteractor* InInteractor)
 {
 	if (IsSetup())
 	{
@@ -25,9 +87,9 @@ void ULearningAgentsController::SetupController(ULearningAgentsInteractor* InInt
 		return;
 	}
 
-	if (!Manager)
+	if (!InManager)
 	{
-		UE_LOG(LogLearning, Error, TEXT("%s: Must be attached to a LearningAgentsManager Actor."), *GetName());
+		UE_LOG(LogLearning, Error, TEXT("%s: InManager is nullptr."), *GetName());
 		return;
 	}
 
@@ -43,16 +105,16 @@ void ULearningAgentsController::SetupController(ULearningAgentsInteractor* InInt
 		return;
 	}
 
+	Manager = InManager;
 	Interactor = InInteractor;
-
 	bIsSetup = true;
 
-	OnAgentsAdded(Manager->GetAllAgentIds());
+	InManager->AddListener(this);
 }
 
-void ULearningAgentsController::EncodeActions()
+void ULearningAgentsController::EvaluateController()
 {
-	UE_LEARNING_TRACE_CPUPROFILER_EVENT_SCOPE(ULearningAgentsController::EncodeActions);
+	UE_LEARNING_TRACE_CPUPROFILER_EVENT_SCOPE(ULearningAgentsController::EvaluateController);
 
 	if (!IsSetup())
 	{
@@ -60,80 +122,60 @@ void ULearningAgentsController::EncodeActions()
 		return;
 	}
 
-	const TLearningArrayView<1, uint64> ActionEncodingAgentIteration = Interactor->GetActionEncodingAgentIteration();
-
-	// Run Set Actions Callback
-
-	SetActions(Manager->GetAllAgentIds());
-
-	// Check that all actions have had their setter run
-
-	ValidAgentStatus.SetNumUninitialized(Manager->GetMaxAgentNum());
-	ValidAgentStatus.SetRange(0, Manager->GetMaxAgentNum(), true);
-
-	for (ULearningAgentsAction* ActionObject : Interactor->GetActionObjects())
+	if (Manager->GetAgentNum() == 0)
 	{
-		for (const int32 AgentId : Manager->GetAllAgentSet())
-		{
-			if (ActionObject->GetAgentSetIteration(AgentId) == ActionEncodingAgentIteration[AgentId])
-			{
-				UE_LOG(LogLearning, Warning, TEXT("%s: Action %s for agent with id %i has not been set (got iteration %i, expected iteration %i) and so agent will not have actions encoded."), *GetName(), *ActionObject->GetName(), AgentId, ActionObject->GetAgentSetIteration(AgentId), ActionEncodingAgentIteration[AgentId] + 1);
-				ValidAgentStatus[AgentId] = false;
-				continue;
-			}
-
-			if (ActionObject->GetAgentSetIteration(AgentId) > ActionEncodingAgentIteration[AgentId] + 1)
-			{
-				UE_LOG(LogLearning, Warning, TEXT("%s: Action %s for agent with id %i appears to have been set multiple times (got iteration %i, expected iteration %i) and so agent will not have actions encoded."), *GetName(), *ActionObject->GetName(), AgentId, ActionObject->GetAgentSetIteration(AgentId), ActionEncodingAgentIteration[AgentId] + 1);
-				ValidAgentStatus[AgentId] = false;
-				continue;
-			}
-
-			if (ActionObject->GetAgentSetIteration(AgentId) != ActionEncodingAgentIteration[AgentId] + 1)
-			{
-				UE_LOG(LogLearning, Warning, TEXT("%s: Action %s for agent with id %i does not have a matching iteration number (got iteration %i, expected iteration %i) and so agent will not have actions encoded."), *GetName(), *ActionObject->GetName(), AgentId, ActionObject->GetAgentSetIteration(AgentId), ActionEncodingAgentIteration[AgentId] + 1);
-				ValidAgentStatus[AgentId] = false;
-				continue;
-			}
-		}
+		UE_LOG(LogLearning, Warning, TEXT("%s: No agents added to Manager."), *GetName());
 	}
 
-	ValidAgentIds.Empty(Manager->GetAgentNum());
-
-	for (const int32 AgentId : Manager->GetAllAgentSet())
+	if (Manager->GetAgentNum() != Interactor->ObservationObjectElements.Num())
 	{
-		if (ValidAgentStatus[AgentId]) { ValidAgentIds.Add(AgentId); }
+		UE_LOG(LogLearning, Error, TEXT("%s: Not enough Observation Objects added by GatherAgentObservations. Expected %i, Got %i."), *GetName(), Manager->GetAgentNum(), Interactor->ObservationObjectElements.Num());
+		return;
 	}
 
-	ValidAgentSet = ValidAgentIds;
-	ValidAgentSet.TryMakeSlice();
-
-	// Encode Actions
-
-	Interactor->GetActionFeature().Encode(ValidAgentSet);
+	// Run EvaluateAgentController Callback
 
-	// Increment Action Encoding Iteration
+	Interactor->ActionObject->GetActionObject().Reset();
+	Interactor->ActionObjectElements.Empty(Manager->GetMaxAgentNum());
+	EvaluateAgentController(
+		Interactor->ActionObjectElements, 
+		Interactor->ActionObject, 
+		Interactor->ObservationObjectElements,
+		Interactor->ObservationObject,
+		Manager->GetAllAgentIds());
 
-	for (const int32 AgentId : ValidAgentSet)
+	if (Manager->GetAgentNum() != Interactor->ActionObjectElements.Num())
 	{
-		ActionEncodingAgentIteration[AgentId]++;
+		UE_LOG(LogLearning, Error, TEXT("%s: Not enough Action Objects added by EvaluateAgentController. Expected %i, Got %i."), *GetName(), Manager->GetAgentNum(), Interactor->ActionObjectElements.Num());
+		return;
 	}
 
-	// Visual Logger
+	// Check Action Objects are Valid and if so convert to action vectors
 
-#if UE_LEARNING_AGENTS_ENABLE_VISUAL_LOG
-	for (const ULearningAgentsAction* ActionObject : Interactor->GetActionObjects())
+	for (int32 AgentIdx = 0; AgentIdx < Manager->GetAgentNum(); AgentIdx++)
 	{
-		ActionObject->VisualLog(ValidAgentSet);
+		if (Interactor->ActionSchema->ValidateObjectMatchesSchema(
+			Interactor->ActionSchemaElement, 
+			Interactor->ActionObject, 
+			Interactor->ActionObjectElements[AgentIdx]))
+		{
+			UE::Learning::Action::SetVectorFromObject(
+				Interactor->ActionVectors[Manager->GetAllAgentSet()[AgentIdx]],
+				Interactor->ActionSchema->GetActionSchema(),
+				Interactor->ActionSchemaElement.SchemaElement,
+				Interactor->ActionObject->GetActionObject(),
+				Interactor->ActionObjectElements[AgentIdx].ObjectElement);
+
+			Interactor->ActionVectorIteration[Manager->GetAllAgentSet()[AgentIdx]]++;
+		}
 	}
-#endif
 }
 
 ULearningAgentsInteractor* ULearningAgentsController::GetInteractor(const TSubclassOf<ULearningAgentsInteractor> InteractorClass) const
 {
-	if (!Interactor)
+	if (!IsSetup())
 	{
-		UE_LOG(LogLearning, Error, TEXT("%s: Interactor is nullptr. Did we forget to call Setup on this component?"), *GetName());
+		UE_LOG(LogLearning, Error, TEXT("%s: Setup not complete."), *GetName());
 		return nullptr;
 	}
 
@@ -150,7 +192,7 @@ void ULearningAgentsController::RunController()
 		return;
 	}
 
-	Interactor->EncodeObservations();
-	EncodeActions();
-	Interactor->DecodeActions();
+	Interactor->GatherObservations();
+	EvaluateController();
+	Interactor->ScatterActions();
 }
diff --git a/Engine/Plugins/Experimental/LearningAgents/Source/LearningAgents/Private/LearningAgentsCritic.cpp b/Engine/Plugins/Experimental/LearningAgents/Source/LearningAgents/Private/LearningAgentsCritic.cpp
--- a/Engine/Plugins/Experimental/LearningAgents/Source/LearningAgents/Private/LearningAgentsCritic.cpp
+++ b/Engine/Plugins/Experimental/LearningAgents/Source/LearningAgents/Private/LearningAgentsCritic.cpp
@@ -5,11 +5,8 @@
 #include "LearningAgentsManager.h"
 #include "LearningAgentsInteractor.h"
 #include "LearningAgentsPolicy.h"
-#include "LearningAgentsHelpers.h"
-#include "LearningAgentsNeuralNetworkData.h"
-#include "LearningFeatureObject.h"
 #include "LearningNeuralNetwork.h"
-#include "LearningNeuralNetworkObject.h"
+#include "LearningCritic.h"
 #include "LearningLog.h"
 
 #include "UObject/Package.h"
@@ -17,15 +14,76 @@
 #include "Misc/Paths.h"
 #include "GameFramework/Actor.h"
 
+#include "NNERuntimeBasicCpuBuilder.h"
+
+namespace UE::Learning::Agents::Critic::Private
+{
+	static inline UE::NNE::RuntimeBasic::FModelBuilder::EActivationFunction GetBuilderActivationFunction(const ELearningAgentsActivationFunction ActivationFunction)
+	{
+		switch (ActivationFunction)
+		{
+		case ELearningAgentsActivationFunction::ReLU: return UE::NNE::RuntimeBasic::FModelBuilder::EActivationFunction::ReLU;
+		case ELearningAgentsActivationFunction::ELU: return UE::NNE::RuntimeBasic::FModelBuilder::EActivationFunction::ELU;
+		case ELearningAgentsActivationFunction::TanH: return UE::NNE::RuntimeBasic::FModelBuilder::EActivationFunction::TanH;
+		default:
+			UE_LOG(LogLearning, Error, TEXT("Unknown Activation Function"));
+			return UE::NNE::RuntimeBasic::FModelBuilder::EActivationFunction::ReLU;
+		}
+	}
+}
+
 ULearningAgentsCritic::ULearningAgentsCritic() : Super(FObjectInitializer::Get()) {}
 ULearningAgentsCritic::ULearningAgentsCritic(FVTableHelper& Helper) : Super(Helper) {}
 ULearningAgentsCritic::~ULearningAgentsCritic() = default;
 
+ULearningAgentsCritic* ULearningAgentsCritic::MakeCritic(
+	ULearningAgentsManager* InManager,
+	ULearningAgentsInteractor* InInteractor,
+	ULearningAgentsPolicy* InPolicy,
+	TSubclassOf<ULearningAgentsCritic> Class,
+	const FName Name,
+	ULearningAgentsNeuralNetwork* CriticNeuralNetworkAsset,
+	const bool bReinitializeCriticNetwork,
+	const FLearningAgentsCriticSettings& CriticSettings,
+	const int32 Seed)
+{
+	if (!InManager)
+	{
+		UE_LOG(LogLearning, Error, TEXT("MakeCritic: InManager is nullptr."));
+		return nullptr;
+	}
+
+	if (!Class)
+	{
+		UE_LOG(LogLearning, Error, TEXT("MakeCritic: Class is nullptr."));
+		return nullptr;
+	}
+
+	const FName UniqueName = MakeUniqueObjectName(InManager, Class, Name, EUniqueObjectNameOptions::GloballyUnique);
+
+	ULearningAgentsCritic* Critic = NewObject<ULearningAgentsCritic>(InManager, Class, UniqueName);
+	if (!Critic) { return nullptr; }
+
+	Critic->SetupCritic(
+		InManager, 
+		InInteractor,
+		InPolicy,
+		CriticNeuralNetworkAsset,
+		bReinitializeCriticNetwork,
+		CriticSettings,
+		Seed);
+
+	return Critic->IsSetup() ? Critic : nullptr;
+}
+
 void ULearningAgentsCritic::SetupCritic(
+	ULearningAgentsManager* InManager,
 	ULearningAgentsInteractor* InInteractor, 
 	ULearningAgentsPolicy* InPolicy,
+	ULearningAgentsNeuralNetwork* CriticNeuralNetworkAsset,
+	const bool bReinitializeCriticNetwork,
 	const FLearningAgentsCriticSettings& CriticSettings,
-	ULearningAgentsNeuralNetwork* NeuralNetworkAsset)
+	const int32 Seed)
 {
 	if (IsSetup())
 	{
@@ -33,9 +91,9 @@ void ULearningAgentsCritic::SetupCritic(
 		return;
 	}
 
-	if (!Manager)
+	if (!InManager)
 	{
-		UE_LOG(LogLearning, Error, TEXT("%s: Must be attached to a LearningAgentsManager Actor."), *GetName());
+		UE_LOG(LogLearning, Error, TEXT("%s: InManager is nullptr."), *GetName());
 		return;
 	}
 
@@ -51,8 +109,6 @@ void ULearningAgentsCritic::SetupCritic(
 		return;
 	}
 
-	Interactor = InInteractor;
-
 	if (!InPolicy)
 	{
 		UE_LOG(LogLearning, Error, TEXT("%s: InPolicy is nullptr."), *GetName());
@@ -65,132 +121,97 @@ void ULearningAgentsCritic::SetupCritic(
 		return;
 	}
 
+	Manager = InManager;
+	Interactor = InInteractor;
 	Policy = InPolicy;
 
-	const int32 NetworkInputNum = Interactor->GetObservationFeature().DimNum() + Policy->GetMemoryStateSize();
-	const int32 NetworkOutputNum = 1;
+	const int32 ObservationEncodedVectorSize = Interactor->GetObservationEncodedVectorSize();
+	const int32 MemoryStateSize = Policy->GetMemoryStateSize();
+
+	const int32 CriticHashData[2] = { MemoryStateSize, ObservationEncodedVectorSize };
+	const int32 CriticCompatibilityHash = CityHash32((const char*)CriticHashData, 2 * sizeof(int32));
+
+	// Try to use existing Neural Network Asset
 
-	if (NeuralNetworkAsset)
+	if (CriticNeuralNetworkAsset)
 	{
-		// Use Existing Neural Network Asset
+		CriticNetwork = CriticNeuralNetworkAsset;
 
-		if (NeuralNetworkAsset->NeuralNetworkData)
+		if (CriticNeuralNetworkAsset->NeuralNetworkData && !bReinitializeCriticNetwork)
 		{
-			if (NeuralNetworkAsset->NeuralNetworkData->GetNetworkInterface()->GetInputNum() != NetworkInputNum ||
-				NeuralNetworkAsset->NeuralNetworkData->GetNetworkInterface()->GetOutputNum() != NetworkOutputNum)
+			if (CriticNeuralNetworkAsset->NeuralNetworkData->GetCompatibilityHash() != CriticCompatibilityHash)
 			{
-				UE_LOG(LogLearning, Error, TEXT("%s: Neural Network Asset provided during Setup is incorrect size: Inputs and outputs don't match."), *GetName());
+				UE_LOG(LogLearning, Error, TEXT("%s: Critic Network Asset provided during Setup is incompatible with Schema. Network hash is %i vs Schema hash %i."), *GetName(),
+					CriticNeuralNetworkAsset->NeuralNetworkData->GetCompatibilityHash(),
+					CriticCompatibilityHash);
 				return;
 			}
 
-			Network = NeuralNetworkAsset;
-		}
-		else
-		{
-			Network = NeuralNetworkAsset;
-			Network->NeuralNetworkData = NewObject<ULearningAgentsNeuralNetworkData>(Network);
-			Network->NeuralNetworkData->CreateMLP(
-				NetworkInputNum,
-				NetworkOutputNum,
-				CriticSettings.HiddenLayerSize,
-				CriticSettings.LayerNum,
-				CriticSettings.ActivationFunction);
+			if (CriticNeuralNetworkAsset->NeuralNetworkData->GetInputSize() != ObservationEncodedVectorSize + MemoryStateSize ||
+				CriticNeuralNetworkAsset->NeuralNetworkData->GetOutputSize() != 1)
+			{
+				UE_LOG(LogLearning, Error, TEXT("%s: Critic Network Asset provided during Setup is incorrect size: Got inputs of size %i, expected %i. Got outputs of size %i, expected %i."), *GetName(),
+					CriticNeuralNetworkAsset->NeuralNetworkData->GetInputSize(), ObservationEncodedVectorSize + MemoryStateSize,
+					CriticNeuralNetworkAsset->NeuralNetworkData->GetOutputSize(), 1);
+				return;
+			}
 		}
 	}
-	else
-	{
-		// Create New Neural Network Asset
 
+	if (!CriticNetwork)
+	{
 		const FName UniqueName = MakeUniqueObjectName(this, ULearningAgentsNeuralNetwork::StaticClass(), TEXT("CriticNetwork"), EUniqueObjectNameOptions::GloballyUnique);
-
-		Network = NewObject<ULearningAgentsNeuralNetwork>(this, UniqueName);
-		Network->NeuralNetworkData = NewObject<ULearningAgentsNeuralNetworkData>(Network);
-		Network->NeuralNetworkData->CreateMLP(
-			NetworkInputNum,
-			NetworkOutputNum,
-			CriticSettings.HiddenLayerSize,
-			CriticSettings.LayerNum,
-			CriticSettings.ActivationFunction);
+		CriticNetwork = NewObject<ULearningAgentsNeuralNetwork>(this, UniqueName);
 	}
 
-	// Create Critic Object
-	CriticObject = MakeShared<UE::Learning::FNeuralNetworkCriticFunction>(
-		TEXT("CriticObject"),
-		Manager->GetInstanceData().ToSharedRef(),
-		Manager->GetMaxAgentNum(),
-		Interactor->GetObservationFeature().DimNum(),
-		Policy->GetMemoryStateSize(),
-		Network->NeuralNetworkData->GetNetworkInterface());
-
-	CriticAgentIteration.SetNumUninitialized({ Manager->GetMaxAgentNum() });
-	UE::Learning::Array::Set<1, uint64>(CriticAgentIteration, INDEX_NONE);
-
-	bIsSetup = true;
-
-	OnAgentsAdded(Manager->GetAllAgentIds());
-}
-
-void ULearningAgentsCritic::OnAgentsAdded(const TArray<int32>& AgentIds)
-{
-	if (IsSetup())
+	if (!CriticNetwork->NeuralNetworkData || bReinitializeCriticNetwork)
 	{
-		UE::Learning::Array::Set<1, uint64>(CriticAgentIteration, 0, AgentIds);
+		// Create New Neural Network Asset
 
-		for (ULearningAgentsHelper* Helper : HelperObjects)
+		if (!CriticNetwork->NeuralNetworkData)
 		{
-			Helper->OnAgentsAdded(AgentIds);
+			CriticNetwork->NeuralNetworkData = NewObject<ULearningNeuralNetworkData>(CriticNetwork);
 		}
 
-		AgentsAdded(AgentIds);
-	}
-}
+		UE::NNE::RuntimeBasic::FModelBuilder Builder;
 
-void ULearningAgentsCritic::OnAgentsRemoved(const TArray<int32>& AgentIds)
-{
-	if (IsSetup())
-	{
-		UE::Learning::Array::Set<1, uint64>(CriticAgentIteration, INDEX_NONE, AgentIds);
+		TArray<uint8> FileData;
+		uint32 CriticInputSize, CriticOutputSize;
+		Builder.WriteFileDataAndReset(FileData, CriticInputSize, CriticOutputSize,
+			Builder.MakeMLPWithRandomKaimingWeights(
+				ObservationEncodedVectorSize + MemoryStateSize,
+				1,
+				CriticSettings.HiddenLayerSize,
+				FMath::Max(CriticSettings.LayerNum, 2),
+				UE::Learning::Agents::Critic::Private::GetBuilderActivationFunction(CriticSettings.ActivationFunction)));
 
-		for (ULearningAgentsHelper* Helper : HelperObjects)
-		{
-			Helper->OnAgentsRemoved(AgentIds);
-		}
+		UE_LEARNING_CHECK(CriticInputSize == ObservationEncodedVectorSize + MemoryStateSize);
+		UE_LEARNING_CHECK(CriticOutputSize == 1);
 
-		AgentsRemoved(AgentIds);
+		CriticNetwork->NeuralNetworkData->Init(CriticInputSize, CriticOutputSize, CriticCompatibilityHash, FileData);
+		CriticNetwork->ForceMarkDirty();
 	}
-}
 
-void ULearningAgentsCritic::OnAgentsReset(const TArray<int32>& AgentIds)
-{
-	if (IsSetup())
-	{
-		UE::Learning::Array::Set<1, uint64>(CriticAgentIteration, 0, AgentIds);
+	// Create Critic Object
 
-		for (ULearningAgentsHelper* Helper : HelperObjects)
-		{
-			Helper->OnAgentsReset(AgentIds);
-		}
+	CriticObject = MakeShared<UE::Learning::FNeuralNetworkCritic>(
+		Manager->GetMaxAgentNum(),
+		ObservationEncodedVectorSize,
+		MemoryStateSize,
+		CriticNetwork->NeuralNetworkData->GetNetwork());
 
-		AgentsReset(AgentIds);
-	}
-}
+	Returns.SetNumUninitialized({ Manager->GetMaxAgentNum() });
+	UE::Learning::Array::Set<1, float>(Returns, FLT_MAX);
 
-ULearningAgentsNeuralNetwork* ULearningAgentsCritic::GetNetworkAsset()
-{
-	return Network;
-}
+	ReturnsIteration.SetNumUninitialized({ Manager->GetMaxAgentNum() });
+	UE::Learning::Array::Set<1, uint64>(ReturnsIteration, INDEX_NONE);
 
-UE::Learning::INeuralNetwork& ULearningAgentsCritic::GetCriticNetwork()
-{
-	return *Network->NeuralNetworkData->GetNetworkInterface();
-}
+	bIsSetup = true;
 
-UE::Learning::FNeuralNetworkCriticFunction& ULearningAgentsCritic::GetCriticObject()
-{
-	return *CriticObject;
+	Manager->AddListener(this);
 }
 
-void ULearningAgentsCritic::LoadCriticFromSnapshot(const FFilePath& File)
+void ULearningAgentsCritic::OnAgentsAdded_Implementation(const TArray<int32>& AgentIds)
 {
 	if (!IsSetup())
 	{
@@ -198,10 +219,11 @@ void ULearningAgentsCritic::LoadCriticFromSnapshot(const FFilePath& File)
 		return;
 	}
 
-	Network->LoadNetworkFromSnapshot(File);
+	UE::Learning::Array::Set<1, uint64>(ReturnsIteration, 0, AgentIds);
+	UE::Learning::Array::Set<1, float>(Returns, 0.0f, AgentIds);
 }
 
-void ULearningAgentsCritic::SaveCriticToSnapshot(const FFilePath& File) const
+void ULearningAgentsCritic::OnAgentsRemoved_Implementation(const TArray<int32>& AgentIds)
 {
 	if (!IsSetup())
 	{
@@ -209,10 +231,11 @@ void ULearningAgentsCritic::SaveCriticToSnapshot(const FFilePath& File) const
 		return;
 	}
 
-	Network->SaveNetworkToSnapshot(File);
+	UE::Learning::Array::Set<1, uint64>(ReturnsIteration, INDEX_NONE, AgentIds);
+	UE::Learning::Array::Set<1, float> (Returns, FLT_MAX, AgentIds);
 }
 
-void ULearningAgentsCritic::UseCriticFromAsset(ULearningAgentsNeuralNetwork* NeuralNetworkAsset)
+void ULearningAgentsCritic::OnAgentsReset_Implementation(const TArray<int32>& AgentIds)
 {
 	if (!IsSetup())
 	{
@@ -220,49 +243,24 @@ void ULearningAgentsCritic::UseCriticFromAsset(ULearningAgentsNeuralNetwork* Neu
 		return;
 	}
 
-	if (!NeuralNetworkAsset || !NeuralNetworkAsset->NeuralNetworkData)
-	{
-		UE_LOG(LogLearning, Error, TEXT("%s: Asset is invalid."), *GetName());
-		return;
-	}
-
-	if (NeuralNetworkAsset == Network)
-	{
-		UE_LOG(LogLearning, Error, TEXT("%s: Asset is same as the current network."), *GetName());
-		return;
-	}
-
-	if (NeuralNetworkAsset->NeuralNetworkData->GetNetworkInterface()->GetInputNum() != Network->NeuralNetworkData->GetNetworkInterface()->GetInputNum() ||
-		NeuralNetworkAsset->NeuralNetworkData->GetNetworkInterface()->GetOutputNum() != Network->NeuralNetworkData->GetNetworkInterface()->GetOutputNum())
-	{
-		UE_LOG(LogLearning, Error, TEXT("%s: Failed to use asset as network settings don't match."), *GetName());
-		return;
-	}
-
-	Network = NeuralNetworkAsset;
-	CriticObject->UpdateNeuralNetwork(Network->NeuralNetworkData->GetNetworkInterface());
+	UE::Learning::Array::Set<1, uint64>(ReturnsIteration, 0, AgentIds);
+	UE::Learning::Array::Set<1, float>(Returns, 0.0f, AgentIds);
 }
 
-void ULearningAgentsCritic::LoadCriticFromAsset(ULearningAgentsNeuralNetwork* NeuralNetworkAsset)
+UE::Learning::FNeuralNetworkCritic& ULearningAgentsCritic::GetCriticObject()
 {
-	if (!IsSetup())
-	{
-		UE_LOG(LogLearning, Error, TEXT("%s: Setup not complete."), *GetName());
-		return;
-	}
-
-	Network->LoadNetworkFromAsset(NeuralNetworkAsset);
+	return *CriticObject;
 }
 
-void ULearningAgentsCritic::SaveCriticToAsset(ULearningAgentsNeuralNetwork* NeuralNetworkAsset)
+ULearningAgentsNeuralNetwork* ULearningAgentsCritic::GetCriticNetworkAsset()
 {
 	if (!IsSetup())
 	{
 		UE_LOG(LogLearning, Error, TEXT("%s: Setup not complete."), *GetName());
-		return;
+		return nullptr;
 	}
 
-	Network->SaveNetworkToAsset(NeuralNetworkAsset);
+	return CriticNetwork;
 }
 
 void ULearningAgentsCritic::EvaluateCritic()
@@ -275,16 +273,21 @@ void ULearningAgentsCritic::EvaluateCritic()
 		return;
 	}
 
+	if (Manager->GetAgentNum() == 0)
+	{
+		UE_LOG(LogLearning, Warning, TEXT("%s: No agents added to Manager."), *GetName());
+	}
+
 	// Check Agents actually have encoded observations.
 	// All added agents should already have some memory state even if it is zero.
 
-	ValidAgentIds.Empty(Manager->GetAgentNum());
+	ValidAgentIds.Empty(Manager->GetMaxAgentNum());
 
 	for (const int32 AgentId : Manager->GetAllAgentSet())
 	{
-		if (Interactor->GetObservationEncodingAgentIteration()[AgentId] == 0)
+		if (Policy->ObservationVectorEncodedIteration[AgentId] == 0)
 		{
-			UE_LOG(LogLearning, Warning, TEXT("%s: Agent with id %i has not made observations so critic will not be evaluated for it."), *GetName(), AgentId);
+			UE_LOG(LogLearning, Warning, TEXT("%s: Agent with id %i has not made observations so critic will not be evaluated for it. Was EncodeObservations run without error?"), *GetName(), AgentId);
 			continue;
 		}
 
@@ -294,32 +297,20 @@ void ULearningAgentsCritic::EvaluateCritic()
 	ValidAgentSet = ValidAgentIds;
 	ValidAgentSet.TryMakeSlice();
 
-	// Get views of Observations and Memory State and copy into network input buffers
-
-	TLearningArrayView<2, const float> MemoryStateView = Policy->GetMemoryStateView();
-	TLearningArrayView<2, const float> ObservationsView = Manager->GetInstanceData()->ConstView(Interactor->GetObservationFeature().FeatureHandle);
-	TLearningArrayView<2, float> InputObservationView = Manager->GetInstanceData()->View(CriticObject->InputObservationHandle);
-	TLearningArrayView<2, float> InputMemoryStateView = Manager->GetInstanceData()->View(CriticObject->InputMemoryStateHandle);
-
-	UE::Learning::Array::Copy<2, float>(InputObservationView, ObservationsView, ValidAgentSet);
-	UE::Learning::Array::Copy<2, float>(InputMemoryStateView, MemoryStateView, ValidAgentSet);
-
 	// Evaluate Critic
 
-	CriticObject->Evaluate(ValidAgentSet);
+	CriticObject->Evaluate(
+		Returns,
+		Policy->ObservationVectorsEncoded,
+		Policy->MemoryState,
+		ValidAgentSet);
 
 	// Increment Discounted Return Iteration
 
 	for (const int32 AgentId : ValidAgentSet)
 	{
-		CriticAgentIteration[AgentId]++;
+		ReturnsIteration[AgentId]++;
 	}
-		
-	// Visual Logger
-
-#if UE_LEARNING_AGENTS_ENABLE_VISUAL_LOG
-	VisualLog(ValidAgentSet);
-#endif
 }
 
 float ULearningAgentsCritic::GetEstimatedDiscountedReturn(const int32 AgentId) const
@@ -336,39 +327,12 @@ float ULearningAgentsCritic::GetEstimatedDiscountedReturn(const int32 AgentId) c
 		return 0.0f;
 	}
 
-	if (CriticAgentIteration[AgentId] == 0)
+	if (ReturnsIteration[AgentId] == 0)
 	{
 		UE_LOG(LogLearning, Error, TEXT("%s: Agent with id %d has not yet computed the estimated discounted return. Did you run EvaluateCritic?"), *GetName(), AgentId);
 		return 0.0f;
 	}
 
-	return CriticObject->InstanceData->ConstView(CriticObject->OutputHandle)[AgentId];
+	return Returns[AgentId];
 }
 
-#if UE_LEARNING_AGENTS_ENABLE_VISUAL_LOG
-void ULearningAgentsCritic::VisualLog(const UE::Learning::FIndexSet AgentSet) const
-{
-	UE_LEARNING_TRACE_CPUPROFILER_EVENT_SCOPE(ULearningAgentsCritic::VisualLog);
-
-	const TLearningArrayView<2, const float> InputObservationView = CriticObject->InstanceData->ConstView(CriticObject->InputObservationHandle);
-	const TLearningArrayView<2, const float> InputMemoryStateView = CriticObject->InstanceData->ConstView(CriticObject->InputMemoryStateHandle);
-	const TLearningArrayView<1, const float> OutputView = CriticObject->InstanceData->ConstView(CriticObject->OutputHandle);
-
-	for (const int32 AgentId : AgentSet)
-	{
-		if (const AActor* Actor = Cast<AActor>(GetAgent(AgentId)))
-		{
-			UE_LEARNING_AGENTS_VLOG_STRING(this, LogLearning, Display,
-				Actor->GetActorLocation(),
-				VisualLogColor.ToFColor(true),
-				TEXT("Agent %i\nObservation Input: %s\nObservation Input Stats (Min/Max/Mean/Std): %s\nMemory State Input: %s\nMemory State Input Stats (Min/Max/Mean/Std): %s\nOutput: [% 6.3f]"),
-				AgentId,
-				*UE::Learning::Array::FormatFloat(InputObservationView[AgentId]),
-				*UE::Learning::Agents::Debug::FloatArrayToStatsString(InputObservationView[AgentId]),
-				*UE::Learning::Array::FormatFloat(InputMemoryStateView[AgentId]),
-				*UE::Learning::Agents::Debug::FloatArrayToStatsString(InputMemoryStateView[AgentId]),
-				OutputView[AgentId]);
-		}
-	}
-}
-#endif
diff --git a/Engine/Plugins/Experimental/LearningAgents/Source/LearningAgents/Private/LearningAgentsInteractor.cpp b/Engine/Plugins/Experimental/LearningAgents/Source/LearningAgents/Private/LearningAgentsInteractor.cpp
--- a/Engine/Plugins/Experimental/LearningAgents/Source/LearningAgents/Private/LearningAgentsInteractor.cpp
+++ b/Engine/Plugins/Experimental/LearningAgents/Source/LearningAgents/Private/LearningAgentsInteractor.cpp
@@ -5,19 +5,46 @@
 #include "LearningAgentsManager.h"
 #include "LearningAgentsObservations.h"
 #include "LearningAgentsActions.h"
-#include "LearningAgentsHelpers.h"
+#include "LearningAgentsNeuralNetwork.h"
 
+#include "LearningNeuralNetwork.h"
 #include "LearningArray.h"
-#include "LearningArrayMap.h"
-#include "LearningFeatureObject.h"
 #include "LearningLog.h"
+
 #include "EngineDefines.h"
 
 ULearningAgentsInteractor::ULearningAgentsInteractor() : Super(FObjectInitializer::Get()) {}
 ULearningAgentsInteractor::ULearningAgentsInteractor(FVTableHelper& Helper) : Super(Helper) {}
 ULearningAgentsInteractor::~ULearningAgentsInteractor() = default;
 
-void ULearningAgentsInteractor::SetupInteractor()
+ULearningAgentsInteractor* ULearningAgentsInteractor::MakeInteractor(
+	ULearningAgentsManager* InManager,
+	TSubclassOf<ULearningAgentsInteractor> Class,
+	const FName Name)
+{
+	if (!InManager)
+	{
+		UE_LOG(LogLearning, Error, TEXT("MakeInteractor: InManager is nullptr."));
+		return nullptr;
+	}
+
+	if (!Class)
+	{
+		UE_LOG(LogLearning, Error, TEXT("MakeInteractor: Class is nullptr."));
+		return nullptr;
+	}
+
+	const FName UniqueName = MakeUniqueObjectName(InManager, Class, Name, EUniqueObjectNameOptions::GloballyUnique);
+
+	ULearningAgentsInteractor* Interactor = NewObject<ULearningAgentsInteractor>(InManager, Class, UniqueName);
+	if (!Interactor) { return nullptr; }
+
+	Interactor->SetupInteractor(InManager);
+
+	return Interactor->IsSetup() ? Interactor : nullptr;
+}
+
+void ULearningAgentsInteractor::SetupInteractor(ULearningAgentsManager* InManager)
 {
 	if (IsSetup())
 	{
@@ -25,298 +52,260 @@ void ULearningAgentsInteractor::SetupInteractor()
 		return;
 	}
 
-	if (!Manager)
+	if (!InManager)
 	{
-		UE_LOG(LogLearning, Error, TEXT("%s: Must be attached to a LearningAgentsManager Actor."), *GetName());
+		UE_LOG(LogLearning, Error, TEXT("%s: InManager is nullptr."), *GetName());
 		return;
 	}
 
-	// Setup Observations
-	ObservationObjects.Empty();
-	ObservationFeatures.Empty();
-	SetupObservations();
+	Manager = InManager;
 
-	if (ObservationObjects.Num() == 0)
-	{
-		UE_LOG(LogLearning, Error, TEXT("%s: No observations added to Interactor during SetupObservations."), *GetName());
-		return;
-	}
+	// Specify Observations
+
+	const FName ObservationSchemaUniqueName = MakeUniqueObjectName(this, ULearningAgentsObservationSchema::StaticClass(), TEXT("ObservationSchema"), EUniqueObjectNameOptions::GloballyUnique);
 
-	Observations = MakeShared<UE::Learning::FConcatenateFeature>(*(GetName() + TEXT("Observations")),
-		TLearningArrayView<1, const TSharedRef<UE::Learning::FFeatureObject>>(ObservationFeatures),
-		Manager->GetInstanceData().ToSharedRef(),
-		Manager->GetMaxAgentNum());
+	ObservationSchema = NewObject<ULearningAgentsObservationSchema>(this, ObservationSchemaUniqueName);
+	SpecifyAgentObservations(ObservationSchemaElement, ObservationSchema);
 
-	if (Observations->DimNum() == 0)
+	if (!ObservationSchema->GetObservationSchema().IsValid(ObservationSchemaElement.SchemaElement))
 	{
-		UE_LOG(LogLearning, Error, TEXT("%s: Observation vector is zero-sized - all added observations have no size."), *GetName());
+		UE_LOG(LogLearning, Error, TEXT("%s: Invalid observation provided to Interactor during SpecifyObservations."), *GetName());
 		return;
 	}
 
-	// Setup Actions
-	ActionObjects.Empty();
-	ActionFeatures.Empty();
-	SetupActions();
+	const int32 ObservationVectorSize = ObservationSchema->GetObservationSchema().GetObservationVectorSize(ObservationSchemaElement.SchemaElement);
 
-	if (ActionObjects.Num() == 0)
+	if (ObservationVectorSize == 0)
 	{
-		UE_LOG(LogLearning, Error, TEXT("%s: No actions added to Interactor during SetupActions."), *GetName());
+		UE_LOG(LogLearning, Error, TEXT("%s: Observation vector is zero-sized - specified observations have no size."), *GetName());
 		return;
 	}
 
-	Actions = MakeShared<UE::Learning::FConcatenateFeature>(*(GetName() + TEXT("Actions")),
-		TLearningArrayView<1, const TSharedRef<UE::Learning::FFeatureObject>>(ActionFeatures),
-		Manager->GetInstanceData().ToSharedRef(),
-		Manager->GetMaxAgentNum());
+	const int32 ObservationEncodedVectorSize = ObservationSchema->GetObservationSchema().GetEncodedVectorSize(ObservationSchemaElement.SchemaElement);
 
-	if (Actions->DimNum() == 0)
+	if (ObservationEncodedVectorSize == 0)
 	{
-		UE_LOG(LogLearning, Error, TEXT("%s: Action vector is zero-sized -  all added actions have no size."), *GetName());
+		UE_LOG(LogLearning, Error, TEXT("%s: Observation encoded vector is zero-sized - observations map to empty encoding."), *GetName());
 		return;
 	}
 
-	// Reset Agent iteration
+	ObservationVectors.SetNumUninitialized({ Manager->GetMaxAgentNum(), ObservationVectorSize  });
+	ObservationCompatibilityHash = UE::Learning::Observation::GetSchemaObjectsCompatibilityHash(ObservationSchema->GetObservationSchema(), ObservationSchemaElement.SchemaElement);
 
-	ObservationEncodingAgentIteration.SetNumUninitialized({ Manager->GetMaxAgentNum() });
-	ActionEncodingAgentIteration.SetNumUninitialized({ Manager->GetMaxAgentNum() });
-	UE::Learning::Array::Set<1, uint64>(ObservationEncodingAgentIteration, INDEX_NONE);
-	UE::Learning::Array::Set<1, uint64>(ActionEncodingAgentIteration, INDEX_NONE);
+	const FName ObservationObjectUniqueName = MakeUniqueObjectName(this, ULearningAgentsObservationObject::StaticClass(), TEXT("ObservationObject"), EUniqueObjectNameOptions::GloballyUnique);
 
-	bIsSetup = true;
+	ObservationObject = NewObject<ULearningAgentsObservationObject>(this, ObservationObjectUniqueName);
+	ObservationObjectElements.Empty(Manager->GetMaxAgentNum());
 
-	OnAgentsAdded(Manager->GetAllAgentIds());
-}
+	// Specify Actions
 
-void ULearningAgentsInteractor::OnAgentsAdded(const TArray<int32>& AgentIds)
-{
-	if (IsSetup())
-	{
-		UE::Learning::Array::Set<1, uint64>(ObservationEncodingAgentIteration, 0, AgentIds);
-		UE::Learning::Array::Set<1, uint64>(ActionEncodingAgentIteration, 0, AgentIds);
+	const FName ActionSchemaUniqueName = MakeUniqueObjectName(this, ULearningAgentsActionSchema::StaticClass(), TEXT("ActionSchema"), EUniqueObjectNameOptions::GloballyUnique);
 
-		for (ULearningAgentsHelper* Helper : HelperObjects)
-		{
-			Helper->OnAgentsAdded(AgentIds);
-		}
+	ActionSchema = NewObject<ULearningAgentsActionSchema>(this, ActionSchemaUniqueName);
+	SpecifyAgentActions(ActionSchemaElement, ActionSchema);
 
-		for (ULearningAgentsObservation* ObservationObject : ObservationObjects)
-		{
-			ObservationObject->OnAgentsAdded(AgentIds);
-		}
+	if (!ActionSchema->GetActionSchema().IsValid(ActionSchemaElement.SchemaElement))
+	{
+		UE_LOG(LogLearning, Error, TEXT("%s: Invalid action provided to Interactor during SpecifyActions."), *GetName());
+		return;
+	}
 
-		for (ULearningAgentsAction* ActionObject : ActionObjects)
-		{
-			ActionObject->OnAgentsAdded(AgentIds);
-		}
+	const int32 ActionVectorSize = ActionSchema->GetActionSchema().GetActionVectorSize(ActionSchemaElement.SchemaElement);
 
-		AgentsAdded(AgentIds);
+	if (ActionVectorSize == 0)
+	{
+		UE_LOG(LogLearning, Error, TEXT("%s: Action vector is zero-sized - specified actions have no size."), *GetName());
+		return;
 	}
-}
 
-void ULearningAgentsInteractor::OnAgentsRemoved(const TArray<int32>& AgentIds)
-{
-	if (IsSetup())
+	const int32 ActionEncodedVectorSize = ActionSchema->GetActionSchema().GetEncodedVectorSize(ActionSchemaElement.SchemaElement);
+
+	if (ActionEncodedVectorSize == 0)
 	{
-		UE::Learning::Array::Set<1, uint64>(ObservationEncodingAgentIteration, INDEX_NONE, AgentIds);
-		UE::Learning::Array::Set<1, uint64>(ActionEncodingAgentIteration, INDEX_NONE, AgentIds);
+		UE_LOG(LogLearning, Error, TEXT("%s: Action encoded vector is zero-sized - actions map to empty encoding."), *GetName());
+		return;
+	}
 
-		for (ULearningAgentsHelper* Helper : HelperObjects)
-		{
-			Helper->OnAgentsRemoved(AgentIds);
-		}
+	const int32 ActionDistributionVectorSize = ActionSchema->GetActionSchema().GetActionDistributionVectorSize(ActionSchemaElement.SchemaElement);
 
-		for (ULearningAgentsObservation* ObservationObject : ObservationObjects)
-		{
-			ObservationObject->OnAgentsRemoved(AgentIds);
-		}
+	ActionVectors.SetNumUninitialized({ Manager->GetMaxAgentNum(), ActionVectorSize });
+	ActionCompatibilityHash = UE::Learning::Action::GetSchemaObjectsCompatibilityHash(ActionSchema->GetActionSchema(), ActionSchemaElement.SchemaElement);
 
-		for (ULearningAgentsAction* ActionObject : ActionObjects)
-		{
-			ActionObject->OnAgentsRemoved(AgentIds);
-		}
+	const FName ActionObjectUniqueName = MakeUniqueObjectName(this, ULearningAgentsActionObject::StaticClass(), TEXT("ActionObject"), EUniqueObjectNameOptions::GloballyUnique);
 
-		AgentsRemoved(AgentIds);
-	}
-}
+	ActionObject = NewObject<ULearningAgentsActionObject>(this, ActionObjectUniqueName);
+	ActionObjectElements.Empty(Manager->GetMaxAgentNum());
 
-void ULearningAgentsInteractor::OnAgentsReset(const TArray<int32>& AgentIds)
-{
-	if (IsSetup())
-	{
-		UE::Learning::Array::Set<1, uint64>(ObservationEncodingAgentIteration, 0, AgentIds);
-		UE::Learning::Array::Set<1, uint64>(ActionEncodingAgentIteration, 0, AgentIds);
-		
-		for (ULearningAgentsHelper* Helper : HelperObjects)
-		{
-			Helper->OnAgentsReset(AgentIds);
-		}
+	// Reset Agent iteration
 
-		for (ULearningAgentsObservation* ObservationObject : ObservationObjects)
-		{
-			ObservationObject->OnAgentsReset(AgentIds);
-		}
+	ObservationVectorIteration.SetNumUninitialized({ Manager->GetMaxAgentNum() });
+	ActionVectorIteration.SetNumUninitialized({ Manager->GetMaxAgentNum() });
 
-		for (ULearningAgentsAction* ActionObject : ActionObjects)
-		{
-			ActionObject->OnAgentsReset(AgentIds);
-		}
+	UE::Learning::Array::Set<1, uint64>(ObservationVectorIteration, INDEX_NONE);
+	UE::Learning::Array::Set<1, uint64>(ActionVectorIteration, INDEX_NONE);
 
-		AgentsReset(AgentIds);
-	}
-}
+	bIsSetup = true;
 
-UE::Learning::FFeatureObject& ULearningAgentsInteractor::GetObservationFeature() const
-{
-	return *Observations;
+	Manager->AddListener(this);
 }
 
-UE::Learning::FFeatureObject& ULearningAgentsInteractor::GetActionFeature() const
+void ULearningAgentsInteractor::OnAgentsAdded_Implementation(const TArray<int32>& AgentIds)
 {
-	return *Actions;
-}
+	if (!IsSetup())
+	{
+		UE_LOG(LogLearning, Error, TEXT("%s: Setup not complete."), *GetName());
+		return;
+	}
 
-TConstArrayView<ULearningAgentsObservation*> ULearningAgentsInteractor::GetObservationObjects() const
-{
-	return ObservationObjects;
+	UE::Learning::Array::Set<1, uint64>(ObservationVectorIteration, 0, AgentIds);
+	UE::Learning::Array::Set<1, uint64>(ActionVectorIteration, 0, AgentIds);
 }
 
-TConstArrayView<ULearningAgentsAction*> ULearningAgentsInteractor::GetActionObjects() const
+void ULearningAgentsInteractor::OnAgentsRemoved_Implementation(const TArray<int32>& AgentIds)
 {
-	return ActionObjects;
-}
+	if (!IsSetup())
+	{
+		UE_LOG(LogLearning, Error, TEXT("%s: Setup not complete."), *GetName());
+		return;
+	}
 
-TLearningArrayView<1, uint64> ULearningAgentsInteractor::GetObservationEncodingAgentIteration()
-{
-	return ObservationEncodingAgentIteration;
+	UE::Learning::Array::Set<1, uint64>(ObservationVectorIteration, INDEX_NONE, AgentIds);
+	UE::Learning::Array::Set<1, uint64>(ActionVectorIteration, INDEX_NONE, AgentIds);
 }
 
-TLearningArrayView<1, uint64> ULearningAgentsInteractor::GetActionEncodingAgentIteration()
+void ULearningAgentsInteractor::OnAgentsReset_Implementation(const TArray<int32>& AgentIds)
 {
-	return ActionEncodingAgentIteration;
-}
+	if (!IsSetup())
+	{
+		UE_LOG(LogLearning, Error, TEXT("%s: Setup not complete."), *GetName());
+		return;
+	}
 
-void ULearningAgentsInteractor::SetupObservations_Implementation()
-{
-	// Can be overridden to setup observations without blueprints
+	UE::Learning::Array::Set<1, uint64>(ObservationVectorIteration, 0, AgentIds);
+	UE::Learning::Array::Set<1, uint64>(ActionVectorIteration, 0, AgentIds);
 }
 
-void ULearningAgentsInteractor::SetObservations_Implementation(const TArray<int32>& AgentIds)
+void ULearningAgentsInteractor::SpecifyAgentObservations_Implementation(FLearningAgentsObservationSchemaElement& OutObservationSchemaElement, ULearningAgentsObservationSchema* InObservationSchema)
 {
-	// Can be overridden to set observations without blueprints
+	UE_LOG(LogLearning, Error, TEXT("%s: SpecifyAgentObservations function must be overridden!"), *GetName());
+	OutObservationSchemaElement = FLearningAgentsObservationSchemaElement();
 }
 
-void ULearningAgentsInteractor::AddObservation(TObjectPtr<ULearningAgentsObservation> Object, const TSharedRef<UE::Learning::FFeatureObject>& Feature)
+void ULearningAgentsInteractor::GatherAgentObservations_Implementation(TArray<FLearningAgentsObservationObjectElement>& OutObservationObjectElements, ULearningAgentsObservationObject* InObservationObject, const TArray<int32>& AgentIds)
 {
-	UE_LEARNING_CHECK(!IsSetup());
-	UE_LEARNING_CHECK(Object);
-	ObservationObjects.Add(Object);
-	ObservationFeatures.Add(Feature);
+	UE_LOG(LogLearning, Error, TEXT("%s: GatherAgentObservations function must be overridden!"), *GetName());
+	OutObservationObjectElements.Empty();
 }
 
-void ULearningAgentsInteractor::SetupActions_Implementation()
+void ULearningAgentsInteractor::GatherAgentObservationsUsingDelegate(
+	TArray<FLearningAgentsObservationObjectElement>& OutObservationElements,
+	ULearningAgentsObservationObject* InObservationObject,
+	const TArray<int32>& AgentIds,
+	const FGatherAgentObservationDelegate& Delegate)
 {
-	// Can be overridden to setup actions without blueprints
+	if (!Delegate.IsBound())
+	{
+		UE_LOG(LogLearning, Error, TEXT("%s: Delegate Not Bound."), *GetName());
+		OutObservationElements.Empty();
+		return;
+	}
+
+	OutObservationElements.Empty(AgentIds.Num());
+	for (const int32 AgentId : AgentIds)
+	{
+		OutObservationElements.Add(Delegate.Execute(InObservationObject, AgentId));
+	}
 }
 
-void ULearningAgentsInteractor::GetActions_Implementation(const TArray<int32>& AgentIds)
+void ULearningAgentsInteractor::SpecifyAgentActions_Implementation(FLearningAgentsActionSchemaElement& OutActionSchemaElement, ULearningAgentsActionSchema* InActionSchema)
 {
-	// Can be overridden to get actions without blueprints
+	UE_LOG(LogLearning, Error, TEXT("%s: SpecifyAgentActions function must be overridden!"), *GetName());
+	OutActionSchemaElement = FLearningAgentsActionSchemaElement();
 }
 
-void ULearningAgentsInteractor::AddAction(TObjectPtr<ULearningAgentsAction> Object, const TSharedRef<UE::Learning::FFeatureObject>& Feature)
+void ULearningAgentsInteractor::ScatterAgentActions_Implementation(const TArray<FLearningAgentsActionObjectElement>& InActionObjectElements, const ULearningAgentsActionObject* InActionObject, const TArray<int32>& AgentIds)
 {
-	UE_LEARNING_CHECK(!IsSetup());
-	UE_LEARNING_CHECK(Object);
-	ActionObjects.Add(Object);
-	ActionFeatures.Add(Feature);
+	UE_LOG(LogLearning, Error, TEXT("%s: ScatterAgentActions function must be overridden!"), *GetName());
 }
 
-void ULearningAgentsInteractor::EncodeObservations(const UE::Learning::FIndexSet AgentSet)
+void  ULearningAgentsInteractor::ScatterAgentObservationsUsingDelegate(
+	const TArray<FLearningAgentsActionObjectElement>& InActionObjectElements,
+	const ULearningAgentsActionObject* InActionObject,
+	const TArray<int32>& AgentIds,
+	const FScatterAgentActionDelegate& Delegate)
 {
-	UE_LEARNING_TRACE_CPUPROFILER_EVENT_SCOPE(ULearningAgentsInteractor::EncodeObservations);
-
-	if (!IsSetup())
+	if (!Delegate.IsBound())
 	{
-		UE_LOG(LogLearning, Error, TEXT("%s: Setup not complete."), *GetName());
+		UE_LOG(LogLearning, Error, TEXT("%s: Delegate Not Bound."), *GetName());
 		return;
 	}
 
-	// Run Set Observations Callback
+	const int32 AgentNum = AgentIds.Num();
 
-	ValidAgentIds.Empty(AgentSet.Num());
-	for (const int32 AgentId : AgentSet)
+	if (AgentNum != InActionObjectElements.Num())
 	{
-		ValidAgentIds.Add(AgentId);
+		UE_LOG(LogLearning, Error, TEXT("%s: Not enough Action Objects. Expected %i, Got %i."), *GetName(), AgentNum, InActionObjectElements.Num());
+		return;
 	}
 
-	SetObservations(ValidAgentIds);
-
-	// Check that all observations have had their setter run
-
-	ValidAgentStatus.SetNumUninitialized(Manager->GetMaxAgentNum());
-	ValidAgentStatus.SetRange(0, Manager->GetMaxAgentNum(), true);
-
-	for (ULearningAgentsObservation* ObservationObject : ObservationObjects)
+	for (int32 AgentIdx = 0; AgentIdx < AgentNum; AgentIdx++)
 	{
-		for (const int32 AgentId : AgentSet)
-		{
-			if (ObservationObject->GetAgentIteration(AgentId) <= ObservationEncodingAgentIteration[AgentId])
-			{
-				UE_LOG(LogLearning, Warning, TEXT("%s: Observation %s for agent with id %i has not been set (got iteration %i, expected iteration %i) and so agent will not have observations encoded."), *GetName(), *ObservationObject->GetName(), AgentId, ObservationObject->GetAgentIteration(AgentId), ObservationEncodingAgentIteration[AgentId] + 1);
-				ValidAgentStatus[AgentId] = false;
-				continue;
-			}
+		Delegate.Execute(InActionObjectElements[AgentIdx], InActionObject, AgentIds[AgentIdx]);
+	}
+}
 
-			if (ObservationObject->GetAgentIteration(AgentId) > ObservationEncodingAgentIteration[AgentId] + 1)
-			{
-				UE_LOG(LogLearning, Warning, TEXT("%s: Observation %s for agent with id %i appears to have been set multiple times (got iteration %i, expected iteration %i) and so agent will not have observations encoded."), *GetName(), *ObservationObject->GetName(), AgentId, ObservationObject->GetAgentIteration(AgentId), ObservationEncodingAgentIteration[AgentId] + 1);
-				ValidAgentStatus[AgentId] = false;
-				continue;
-			}
+void ULearningAgentsInteractor::GatherObservations(const UE::Learning::FIndexSet AgentSet, bool bIncrementIteration)
+{
+	UE_LEARNING_TRACE_CPUPROFILER_EVENT_SCOPE(ULearningAgentsInteractor::GatherObservations);
 
-			if (ObservationObject->GetAgentIteration(AgentId) != ObservationEncodingAgentIteration[AgentId] + 1)
-			{
-				UE_LOG(LogLearning, Warning, TEXT("%s: Observation %s for agent with id %i does not have a matching iteration number (got iteration %i, expected iteration %i) and so agent will not have observations encoded."), *GetName(), *ObservationObject->GetName(), AgentId, ObservationObject->GetAgentIteration(AgentId), ObservationEncodingAgentIteration[AgentId] + 1);
-				ValidAgentStatus[AgentId] = false;
-				continue;
-			}
-		}
+	if (!IsSetup())
+	{
+		UE_LOG(LogLearning, Error, TEXT("%s: Setup not complete."), *GetName());
+		return;
 	}
 
-	ValidAgentIds.Empty(Manager->GetAgentNum());
+	// Run GatherAgentObservations Callback
 
+	ValidAgentIds.Empty(Manager->GetMaxAgentNum());
 	for (const int32 AgentId : AgentSet)
 	{
-		if (ValidAgentStatus[AgentId]) { ValidAgentIds.Add(AgentId); }
+		ValidAgentIds.Add(AgentId);
 	}
 
-	ValidAgentSet = ValidAgentIds;
-	ValidAgentSet.TryMakeSlice();
-
-	// Encode Observations
-
-	Observations->Encode(ValidAgentSet);
-
-	// Increment Observation Encoding Iteration
+	ObservationObject->GetObservationObject().Reset();
+	ObservationObjectElements.Empty(Manager->GetMaxAgentNum());
+	GatherAgentObservations(ObservationObjectElements, ObservationObject, ValidAgentIds);
 
-	for (const int32 AgentId : ValidAgentSet)
+	if (ValidAgentIds.Num() != ObservationObjectElements.Num())
 	{
-		ObservationEncodingAgentIteration[AgentId]++;
+		UE_LOG(LogLearning, Error, TEXT("%s: Not enough Observation Objects added by GatherAgentObservations. Expected %i, Got %i."), *GetName(), ValidAgentIds.Num(), ObservationObjectElements.Num());
+		return;
 	}
 
-	// Visual Logger
+	// Check Observation Objects are Valid and if so convert to observation vectors
 
-#if UE_LEARNING_AGENTS_ENABLE_VISUAL_LOG
-	for (const ULearningAgentsObservation* ObservationObject : ObservationObjects)
+	for (int32 AgentIdx = 0; AgentIdx < AgentSet.Num(); AgentIdx++)
 	{
-		ObservationObject->VisualLog(ValidAgentSet);
+		if (ObservationSchema->ValidateObjectMatchesSchema(ObservationSchemaElement, ObservationObject, ObservationObjectElements[AgentIdx]))
+		{
+			UE::Learning::Observation::SetVectorFromObject(
+				ObservationVectors[AgentSet[AgentIdx]],
+				ObservationSchema->GetObservationSchema(),
+				ObservationSchemaElement.SchemaElement,
+				ObservationObject->GetObservationObject(),
+				ObservationObjectElements[AgentIdx].ObjectElement);
+
+			if (bIncrementIteration)
+			{
+				ObservationVectorIteration[AgentSet[AgentIdx]]++;
+			}
+		}
 	}
-#endif
 }
 
-void ULearningAgentsInteractor::DecodeActions(const UE::Learning::FIndexSet AgentSet)
+void ULearningAgentsInteractor::ScatterActions(const UE::Learning::FIndexSet AgentSet)
 {
-	UE_LEARNING_TRACE_CPUPROFILER_EVENT_SCOPE(ULearningAgentsInteractor::DecodeActions);
+	UE_LEARNING_TRACE_CPUPROFILER_EVENT_SCOPE(ULearningAgentsInteractor::ScatterActions);
 
 	if (!IsSetup())
 	{
@@ -324,15 +313,15 @@ void ULearningAgentsInteractor::DecodeActions(const UE::Learning::FIndexSet Agen
 		return;
 	}
 
-	// Check Agents actually have encoded actions.
+	// Check which agents have had their action vector set
 
-	ValidAgentIds.Empty(Manager->GetAgentNum());
+	ValidAgentIds.Empty(Manager->GetMaxAgentNum());
 
 	for (const int32 AgentId : AgentSet)
 	{
-		if (ActionEncodingAgentIteration[AgentId] == 0)
+		if (ActionVectorIteration[AgentId] == 0)
 		{
-			UE_LOG(LogLearning, Warning, TEXT("%s: Agent with id %i does not have an encoded action vector so actions will not be decoded for it. Was EvaluatePolicy or EncodeActions run?"), *GetName(), AgentId);
+			UE_LOG(LogLearning, Warning, TEXT("%s: Agent with id %i does not have an action vector so actions will not be scattered for it. Was DecodeAndSampleActions run without error?"), *GetName(), AgentId);
 			continue;
 		}
 
@@ -342,65 +331,58 @@ void ULearningAgentsInteractor::DecodeActions(const UE::Learning::FIndexSet Agen
 	ValidAgentSet = ValidAgentIds;
 	ValidAgentSet.TryMakeSlice();
 
-	// Decode Actions
-
-	Actions->Decode(ValidAgentSet);
+	// Generate Action Objects
 
-	// Run Get Actions Callback
+	ActionObject->GetActionObject().Reset();
+	ActionObjectElements.Empty(Manager->GetMaxAgentNum());
 
-	GetActions(ValidAgentIds);
+	for (int32 AgentIdx = 0; AgentIdx < ValidAgentSet.Num(); AgentIdx++)
+	{
+		FLearningAgentsActionObjectElement ActionObjectElement;
 
-	// Check that all actions have had their getter run
+		UE::Learning::Action::GetObjectFromVector(
+			ActionObject->GetActionObject(),
+			ActionObjectElement.ObjectElement,
+			ActionSchema->GetActionSchema(),
+			ActionSchemaElement.SchemaElement,
+			ActionVectors[ValidAgentSet[AgentIdx]]);
 
-	for (ULearningAgentsAction* ActionObject : ActionObjects)
-	{
-		for (const int32 AgentId : ValidAgentSet)
-		{
-			if (ActionObject->GetAgentGetIteration(AgentId) == ActionEncodingAgentIteration[AgentId] - 1)
-			{
-				UE_LOG(LogLearning, Warning, TEXT("%s: Action %s for agent with id %i appears to have not been got."), *GetName(), *ActionObject->GetName(), AgentId);
-				continue;
-			}
-				
-			if (ActionObject->GetAgentGetIteration(AgentId) > ActionEncodingAgentIteration[AgentId])
-			{
-				UE_LOG(LogLearning, Warning, TEXT("%s: Action %s for agent with id %i appears to have been got multiple times."), *GetName(), *ActionObject->GetName(), AgentId);
-				continue;
-			}
-				
-			if (ActionObject->GetAgentGetIteration(AgentId) != ActionEncodingAgentIteration[AgentId])
-			{
-				UE_LOG(LogLearning, Warning, TEXT("%s: Action %s for agent with id %i does not have a matching iteration number (got %i, expected %i)."), *GetName(), *ActionObject->GetName(), AgentId, ActionObject->GetAgentGetIteration(AgentId), ActionEncodingAgentIteration[AgentId]);
-				continue;
-			}
-		}
+		ActionObjectElements.Add(ActionObjectElement);
 	}
 
-	// Visual Logger
+	// Scatter Action Objects
 
-#if UE_LEARNING_AGENTS_ENABLE_VISUAL_LOG
-	for (const ULearningAgentsAction* ActionObject : ActionObjects)
-	{
-		ActionObject->VisualLog(ValidAgentSet);
-	}
-#endif
+	ScatterAgentActions(ActionObjectElements, ActionObject, ValidAgentIds);
 }
 
-void ULearningAgentsInteractor::EncodeObservations()
+void ULearningAgentsInteractor::GatherObservations()
 {
-	EncodeObservations(Manager->GetAllAgentSet());
+	if (Manager->GetAgentNum() == 0)
+	{
+		UE_LOG(LogLearning, Warning, TEXT("%s: No agents added to Manager."), *GetName());
+	}
+
+	GatherObservations(Manager->GetAllAgentSet());
 }
 
-void ULearningAgentsInteractor::DecodeActions()
+void ULearningAgentsInteractor::ScatterActions()
 {
-	DecodeActions(Manager->GetAllAgentSet());
+	if (Manager->GetAgentNum() == 0)
+	{
+		UE_LOG(LogLearning, Warning, TEXT("%s: No agents added to Manager."), *GetName());
+	}
+
+	ScatterActions(Manager->GetAllAgentSet());
 }
-void ULearningAgentsInteractor::GetObservationVector(const int32 AgentId, TArray<float>& OutObservationVector) const
+
+
+void ULearningAgentsInteractor::GetObservationVector(TArray<float>& OutObservationVector, int32& OutObservationCompatibilityHash, const int32 AgentId)
 {
 	if (!IsSetup())
 	{
 		UE_LOG(LogLearning, Error, TEXT("%s: Setup not complete."), *GetName());
 		OutObservationVector.Empty();
+		OutObservationCompatibilityHash = 0;
 		return;
 	}
 
@@ -408,26 +390,30 @@ void ULearningAgentsInteractor::GetObservationVector(const int32 AgentId, TArray
 	{
 		UE_LOG(LogLearning, Error, TEXT("%s: AgentId %d not found in the agents set."), *GetName(), AgentId);
 		OutObservationVector.Empty();
+		OutObservationCompatibilityHash = 0;
 		return;
 	}
 
-	if (ObservationEncodingAgentIteration[AgentId] == 0)
+	if (ObservationVectorIteration[AgentId] == 0)
 	{
-		UE_LOG(LogLearning, Error, TEXT("%s: Agent with id %d has not yet computed an observation vector. Have you run EncodeObservations?"), *GetName(), AgentId);
+		UE_LOG(LogLearning, Error, TEXT("%s: Observation vector not set for agent %i."), *GetName(), AgentId);
 		OutObservationVector.Empty();
+		OutObservationCompatibilityHash = 0;
 		return;
 	}
 
-	OutObservationVector.SetNumUninitialized(Observations->DimNum());
-	UE::Learning::Array::Copy<1, float>(OutObservationVector, Observations->FeatureBuffer()[AgentId]);
+	OutObservationVector.SetNumUninitialized(GetObservationVectorSize());
+	UE::Learning::Array::Copy<1, float>(OutObservationVector, ObservationVectors[AgentId]);
+	OutObservationCompatibilityHash = ObservationCompatibilityHash;
 }
 
-void ULearningAgentsInteractor::GetActionVector(const int32 AgentId, TArray<float>& OutActionVector) const
+void ULearningAgentsInteractor::GetActionVector(TArray<float>& OutActionVector, int32& OutActionCompatibilityHash, const int32 AgentId)
 {
 	if (!IsSetup())
 	{
 		UE_LOG(LogLearning, Error, TEXT("%s: Setup not complete."), *GetName());
 		OutActionVector.Empty();
+		OutActionCompatibilityHash = 0;
 		return;
 	}
 
@@ -435,16 +421,188 @@ void ULearningAgentsInteractor::GetActionVector(const int32 AgentId, TArray<floa
 	{
 		UE_LOG(LogLearning, Error, TEXT("%s: AgentId %d not found in the agents set."), *GetName(), AgentId);
 		OutActionVector.Empty();
+		OutActionCompatibilityHash = 0;
 		return;
 	}
 
-	if (ActionEncodingAgentIteration[AgentId] == 0)
+	if (ActionVectorIteration[AgentId] == 0)
 	{
-		UE_LOG(LogLearning, Error, TEXT("%s: Agent with id %d has not yet computed an action vector. Have you run EvaluatePolicy or EncodeActions?"), *GetName(), AgentId);
+		UE_LOG(LogLearning, Error, TEXT("%s: Action vector not set for agent %i."), *GetName(), AgentId);
 		OutActionVector.Empty();
+		OutActionCompatibilityHash = 0;
+		return;
+	}
+
+	OutActionVector.SetNumUninitialized(GetActionVectorSize());
+	UE::Learning::Array::Copy<1, float>(OutActionVector, ActionVectors[AgentId]);
+	OutActionCompatibilityHash = ActionCompatibilityHash;
+}
+
+void ULearningAgentsInteractor::SetObservationVector(const TArray<float>& ObservationVector, const int32 InObservationCompatibilityHash, const int32 AgentId, bool bIncrementIteration)
+{
+	if (!IsSetup())
+	{
+		UE_LOG(LogLearning, Error, TEXT("%s: Setup not complete."), *GetName());
 		return;
 	}
 
-	OutActionVector.SetNumUninitialized(Actions->DimNum());
-	UE::Learning::Array::Copy<1, float>(OutActionVector, Actions->FeatureBuffer()[AgentId]);
+	if (!HasAgent(AgentId))
+	{
+		UE_LOG(LogLearning, Error, TEXT("%s: AgentId %d not found in the agents set."), *GetName(), AgentId);
+		return;
+	}
+
+	if (InObservationCompatibilityHash != ObservationCompatibilityHash)
+	{
+		UE_LOG(LogLearning, Error, TEXT("%s: Observation Compatibility hash incompatible. Got %i, expected %i."), *GetName(), InObservationCompatibilityHash, ObservationCompatibilityHash);
+		return;
+	}
+
+	if (ObservationVector.Num() != GetObservationVectorSize())
+	{
+		UE_LOG(LogLearning, Error, TEXT("%s: Observation Vector size incompatible. Got %i, expected %i."), *GetName(), ObservationVector.Num(), GetObservationVectorSize());
+		return;
+	}
+
+	UE::Learning::Array::Copy<1, float>(ObservationVectors[AgentId], ObservationVector);
+	if (bIncrementIteration) { ObservationVectorIteration[AgentId]++; }
+}
+
+void ULearningAgentsInteractor::SetActionVector(const TArray<float>& ActionVector, const int32 InActionCompatibilityHash, const int32 AgentId, bool bIncrementIteration)
+{
+	if (!IsSetup())
+	{
+		UE_LOG(LogLearning, Error, TEXT("%s: Setup not complete."), *GetName());
+		return;
+	}
+
+	if (!HasAgent(AgentId))
+	{
+		UE_LOG(LogLearning, Error, TEXT("%s: AgentId %d not found in the agents set."), *GetName(), AgentId);
+		return;
+	}
+
+	if (InActionCompatibilityHash != ActionCompatibilityHash)
+	{
+		UE_LOG(LogLearning, Error, TEXT("%s: Action Compatibility hash incompatible. Got %i, expected %i."), *GetName(), InActionCompatibilityHash, ActionCompatibilityHash);
+		return;
+	}
+
+	if (ActionVector.Num() != GetObservationVectorSize())
+	{
+		UE_LOG(LogLearning, Error, TEXT("%s: Action Vector size incompatible. Got %i, expected %i."), *GetName(), ActionVector.Num(), GetActionVectorSize());
+		return;
+	}
+
+	UE::Learning::Array::Copy<1, float>(ActionVectors[AgentId], ActionVector);
+	if (bIncrementIteration) { ActionVectorIteration[AgentId]++; }
+}
+
+bool ULearningAgentsInteractor::HasObservationVector(const int32 AgentId) const
+{
+	if (!IsSetup())
+	{
+		UE_LOG(LogLearning, Error, TEXT("%s: Setup not complete."), *GetName());
+		return false;
+	}
+
+	if (!HasAgent(AgentId))
+	{
+		UE_LOG(LogLearning, Error, TEXT("%s: AgentId %d not found in the agents set."), *GetName(), AgentId);
+		return false;
+	}
+
+	return ObservationVectorIteration[AgentId] > 0;
+}
+
+bool ULearningAgentsInteractor::HasActionVector(const int32 AgentId) const
+{
+	if (!IsSetup())
+	{
+		UE_LOG(LogLearning, Error, TEXT("%s: Setup not complete."), *GetName());
+		return false;
+	}
+
+	if (!HasAgent(AgentId))
+	{
+		UE_LOG(LogLearning, Error, TEXT("%s: AgentId %d not found in the agents set."), *GetName(), AgentId);
+		return false;
+	}
+
+	return ActionVectorIteration[AgentId] > 0;
+}
+
+int32 ULearningAgentsInteractor::GetObservationVectorSize() const
+{
+	if (!IsSetup())
+	{
+		UE_LOG(LogLearning, Error, TEXT("%s: Setup not complete."), *GetName());
+		return 0;
+	}
+
+	return ObservationSchema->GetObservationSchema().GetObservationVectorSize(ObservationSchemaElement.SchemaElement);
+}
+
+int32 ULearningAgentsInteractor::GetObservationEncodedVectorSize() const
+{
+	if (!IsSetup())
+	{
+		UE_LOG(LogLearning, Error, TEXT("%s: Setup not complete."), *GetName());
+		return 0;
+	}
+
+	return ObservationSchema->GetObservationSchema().GetEncodedVectorSize(ObservationSchemaElement.SchemaElement);
+}
+
+int32 ULearningAgentsInteractor::GetActionVectorSize() const
+{
+	if (!IsSetup())
+	{
+		UE_LOG(LogLearning, Error, TEXT("%s: Setup not complete."), *GetName());
+		return 0;
+	}
+
+	return ActionSchema->GetActionSchema().GetActionVectorSize(ActionSchemaElement.SchemaElement);
+}
+
+int32 ULearningAgentsInteractor::GetActionDistributionVectorSize() const
+{
+	if (!IsSetup())
+	{
+		UE_LOG(LogLearning, Error, TEXT("%s: Setup not complete."), *GetName());
+		return 0;
+	}
+
+	return ActionSchema->GetActionSchema().GetActionDistributionVectorSize(ActionSchemaElement.SchemaElement);
+}
+
+int32 ULearningAgentsInteractor::GetActionEncodedVectorSize() const
+{
+	if (!IsSetup())
+	{
+		UE_LOG(LogLearning, Error, TEXT("%s: Setup not complete."), *GetName());
+		return 0;
+	}
+
+	return ActionSchema->GetActionSchema().GetEncodedVectorSize(ActionSchemaElement.SchemaElement);
+}
+
+const UE::Learning::Observation::FSchema& ULearningAgentsInteractor::GetObservationSchema() const
+{
+	return ObservationSchema->GetObservationSchema();
+}
+
+UE::Learning::Observation::FSchemaElement ULearningAgentsInteractor::GetObservationSchemaElement() const
+{
+	return ObservationSchemaElement.SchemaElement;
+}
+
+const UE::Learning::Action::FSchema& ULearningAgentsInteractor::GetActionSchema() const
+{
+	return ActionSchema->GetActionSchema();
+}
+
+UE::Learning::Action::FSchemaElement ULearningAgentsInteractor::GetActionSchemaElement() const
+{
+	return ActionSchemaElement.SchemaElement;
 }
\ No newline at end of file
diff --git a/Engine/Plugins/Experimental/LearningAgents/Source/LearningAgents/Private/LearningAgentsManager.cpp b/Engine/Plugins/Experimental/LearningAgents/Source/LearningAgents/Private/LearningAgentsManager.cpp
--- a/Engine/Plugins/Experimental/LearningAgents/Source/LearningAgents/Private/LearningAgentsManager.cpp
+++ b/Engine/Plugins/Experimental/LearningAgents/Source/LearningAgents/Private/LearningAgentsManager.cpp
@@ -1,20 +1,19 @@
 // Copyright Epic Games, Inc. All Rights Reserved.
 
 #include "LearningAgentsManager.h"
-#include "LearningAgentsManagerComponent.h"
+#include "LearningAgentsManagerListener.h"
 
 #include "LearningArray.h"
-#include "LearningArrayMap.h"
 #include "LearningLog.h"
 
 #include "UObject/ScriptInterface.h"
 #include "EngineDefines.h"
 
-ALearningAgentsManager::ALearningAgentsManager() : Super(FObjectInitializer::Get()) {}
-ALearningAgentsManager::ALearningAgentsManager(FVTableHelper& Helper) : Super(Helper) {}
-ALearningAgentsManager::~ALearningAgentsManager() = default;
+ULearningAgentsManager::ULearningAgentsManager() : Super(FObjectInitializer::Get()) { PrimaryComponentTick.bCanEverTick = true; }
+ULearningAgentsManager::ULearningAgentsManager(FVTableHelper& Helper) : Super(Helper) {}
+ULearningAgentsManager::~ULearningAgentsManager() = default;
 
-void ALearningAgentsManager::PostInitProperties()
+void ULearningAgentsManager::PostInitProperties()
 {
 	Super::PostInitProperties();
 
@@ -23,29 +22,40 @@ void ALearningAgentsManager::PostInitProperties()
 		return;
 	}
 
-	// Allocate Instance Data
-	InstanceData = MakeShared<UE::Learning::FArrayMap>();
-
 	// Pre-populate the vacant ids
 	OnEventAgentIds.Reserve(MaxAgentNum);
 	OccupiedAgentIds.Reserve(MaxAgentNum);
 	VacantAgentIds.Reserve(MaxAgentNum);
 	Agents.Reserve(MaxAgentNum);
-	for (int32 i = MaxAgentNum - 1; i >= 0; i--)
+	for (int32 AgentId = MaxAgentNum - 1; AgentId >= 0; AgentId--)
 	{
-		VacantAgentIds.Push(i);
+		VacantAgentIds.Push(AgentId);
 		Agents.Add(nullptr);
 	}
 
 	UpdateAgentSets();
 }
 
-int32 ALearningAgentsManager::GetMaxAgentNum() const
+void ULearningAgentsManager::TickComponent(float DeltaTime, enum ELevelTick TickType, FActorComponentTickFunction* ThisTickFunction)
+{
+	Super::TickComponent(DeltaTime, TickType, ThisTickFunction);
+
+	// We only want to give agent tick events when the actual level is ticking - not when paused or just the viewport is ticking.
+	if (TickType == ELevelTick::LEVELTICK_All)
+	{
+		for (ULearningAgentsManagerListener* Listener : Listeners)
+		{
+			Listener->OnAgentsManagerTick(OccupiedAgentIds, DeltaTime);
+		}
+	}
+}
+
+int32 ULearningAgentsManager::GetMaxAgentNum() const
 {
 	return MaxAgentNum;
 }
 
-int32 ALearningAgentsManager::AddAgent(UObject* Agent)
+int32 ULearningAgentsManager::AddAgent(UObject* Agent)
 {
 	if (Agent == nullptr)
 	{
@@ -66,19 +76,16 @@ int32 ALearningAgentsManager::AddAgent(UObject* Agent)
 	OccupiedAgentIds.Add(AgentId);
 	UpdateAgentSets();
 
-	OnEventAgentIds.SetNumUninitialized(0, false);
+	OnEventAgentIds.Reset();
 	OnEventAgentIds.Add(AgentId);
 
 	// Call OnAgentsAdded Event
 
 	UE_LOG(LogLearning, Display, TEXT("%s: Adding Agent %s with id %i."), *GetName(), *Agent->GetName(), AgentId);
 
-	TInlineComponentArray<ULearningAgentsManagerComponent*> ManagerComponents;
-	GetComponents<ULearningAgentsManagerComponent>(ManagerComponents);
-
-	for (ULearningAgentsManagerComponent* ManagerComponent : ManagerComponents)
+	for (ULearningAgentsManagerListener* Listener : Listeners)
 	{
-		ManagerComponent->OnAgentsAdded(OnEventAgentIds);
+		Listener->OnAgentsAdded(OnEventAgentIds);
 	}
 
 	// Return new id
@@ -86,9 +93,9 @@ int32 ALearningAgentsManager::AddAgent(UObject* Agent)
 	return AgentId;
 }
 
-void ALearningAgentsManager::AddAgents(TArray<int32>& OutAgentIds, const TArray<UObject*>& InAgents)
+void ULearningAgentsManager::AddAgents(TArray<int32>& OutAgentIds, const TArray<UObject*>& InAgents)
 {
-	OnEventAgentIds.SetNumUninitialized(0, false);
+	OnEventAgentIds.Reset();
 	OutAgentIds.SetNumUninitialized(InAgents.Num());
 
 	for (int32 AgentIdx = 0; AgentIdx < InAgents.Num(); AgentIdx++)
@@ -121,16 +128,13 @@ void ALearningAgentsManager::AddAgents(TArray<int32>& OutAgentIds, const TArray<
 
 	UE_LOG(LogLearning, Display, TEXT("%s: Adding Agents %s."), *GetName(), *UE::Learning::Array::FormatInt32(OnEventAgentIds));
 
-	TInlineComponentArray<ULearningAgentsManagerComponent*> ManagerComponents;
-	GetComponents<ULearningAgentsManagerComponent>(ManagerComponents);
-
-	for (ULearningAgentsManagerComponent* ManagerComponent : ManagerComponents)
+	for (ULearningAgentsManagerListener* Listener : Listeners)
 	{
-		ManagerComponent->OnAgentsAdded(OnEventAgentIds);
+		Listener->OnAgentsAdded(OnEventAgentIds);
 	}
 }
 
-void ALearningAgentsManager::RemoveAgent(const int32 AgentId)
+void ULearningAgentsManager::RemoveAgent(const int32 AgentId)
 {
 	if (AgentId == INDEX_NONE)
 	{
@@ -154,25 +158,22 @@ void ALearningAgentsManager::RemoveAgent(const int32 AgentId)
 	Agents[AgentId] = nullptr;
 	UpdateAgentSets();
 
-	OnEventAgentIds.SetNumUninitialized(0, false);
+	OnEventAgentIds.Reset();
 	OnEventAgentIds.Add(AgentId);
 
 	// Call OnAgentsRemoved Event
 
 	UE_LOG(LogLearning, Display, TEXT("%s: Removing Agent %i."), *GetName(), AgentId);
 
-	TInlineComponentArray<ULearningAgentsManagerComponent*> ManagerComponents;
-	GetComponents<ULearningAgentsManagerComponent>(ManagerComponents);
-
-	for (ULearningAgentsManagerComponent* ManagerComponent : ManagerComponents)
+	for (ULearningAgentsManagerListener* Listener : Listeners)
 	{
-		ManagerComponent->OnAgentsRemoved(OnEventAgentIds);
+		Listener->OnAgentsRemoved(OnEventAgentIds);
 	}
 }
 
-void ALearningAgentsManager::RemoveAgents(const TArray<int32>& AgentIds)
+void ULearningAgentsManager::RemoveAgents(const TArray<int32>& AgentIds)
 {
-	OnEventAgentIds.SetNumUninitialized(0, false);
+	OnEventAgentIds.Reset();
 
 	for (int32 AgentIdx = 0; AgentIdx < AgentIds.Num(); AgentIdx++)
 	{
@@ -205,23 +206,20 @@ void ALearningAgentsManager::RemoveAgents(const TArray<int32>& AgentIds)
 
 	UE_LOG(LogLearning, Display, TEXT("%s: Removing Agents %s."), *GetName(), *UE::Learning::Array::FormatInt32(OnEventAgentIds));
 
-	TInlineComponentArray<ULearningAgentsManagerComponent*> ManagerComponents;
-	GetComponents<ULearningAgentsManagerComponent>(ManagerComponents);
-
-	for (ULearningAgentsManagerComponent* ManagerComponent : ManagerComponents)
+	for (ULearningAgentsManagerListener* Listener : Listeners)
 	{
-		ManagerComponent->OnAgentsRemoved(OnEventAgentIds);
+		Listener->OnAgentsRemoved(OnEventAgentIds);
 	}
 }
 
-void ALearningAgentsManager::RemoveAllAgents()
+void ULearningAgentsManager::RemoveAllAgents()
 {
 	// We need to make a copy as OccupiedAgentIds will be modified during removal.
 	const TArray<int32> RemoveAgentIds = OccupiedAgentIds;
 	RemoveAgents(RemoveAgentIds);
 }
 
-void ALearningAgentsManager::ResetAgent(const int32 AgentId)
+void ULearningAgentsManager::ResetAgent(const int32 AgentId)
 {
 	if (AgentId == INDEX_NONE)
 	{
@@ -235,25 +233,22 @@ void ALearningAgentsManager::ResetAgent(const int32 AgentId)
 		return;
 	}
 
-	OnEventAgentIds.SetNumUninitialized(0, false);
+	OnEventAgentIds.Reset();
 	OnEventAgentIds.Add(AgentId);
 
 	// Call OnAgentsReset Event
 
 	UE_LOG(LogLearning, Display, TEXT("%s: Resetting Agent %i."), *GetName(), AgentId);
 
-	TInlineComponentArray<ULearningAgentsManagerComponent*> ManagerComponents;
-	GetComponents<ULearningAgentsManagerComponent>(ManagerComponents);
-
-	for (ULearningAgentsManagerComponent* ManagerComponent : ManagerComponents)
+	for (ULearningAgentsManagerListener* Listener : Listeners)
 	{
-		ManagerComponent->OnAgentsReset(OnEventAgentIds);
+		Listener->OnAgentsReset(OnEventAgentIds);
 	}
 }
 
-void ALearningAgentsManager::ResetAgents(const TArray<int32>& AgentIds)
+void ULearningAgentsManager::ResetAgents(const TArray<int32>& AgentIds)
 {
-	OnEventAgentIds.SetNumUninitialized(0, false);
+	OnEventAgentIds.Reset();
 
 	for (int32 AgentIdx = 0; AgentIdx < AgentIds.Num(); AgentIdx++)
 	{
@@ -276,34 +271,41 @@ void ALearningAgentsManager::ResetAgents(const TArray<int32>& AgentIds)
 
 	UE_LOG(LogLearning, Display, TEXT("%s: Resetting Agents %s."), *GetName(), *UE::Learning::Array::FormatInt32(OnEventAgentIds));
 
-	TInlineComponentArray<ULearningAgentsManagerComponent*> ManagerComponents;
-	GetComponents<ULearningAgentsManagerComponent>(ManagerComponents);
-
-	for (ULearningAgentsManagerComponent* ManagerComponent : ManagerComponents)
+	for (ULearningAgentsManagerListener* Listener : Listeners)
 	{
-		ManagerComponent->OnAgentsReset(OnEventAgentIds);
+		Listener->OnAgentsReset(OnEventAgentIds);
 	}
 }
 
-void ALearningAgentsManager::ResetAllAgents()
+void ULearningAgentsManager::ResetAllAgents()
 {
 	ResetAgents(OccupiedAgentIds);
 }
 
-UObject* ALearningAgentsManager::GetAgent(const int32 AgentId, const TSubclassOf<UObject> AgentClass) const
+UObject* ULearningAgentsManager::GetAgent(const int32 AgentId, const TSubclassOf<UObject> AgentClass) const
 {
 	if (!HasAgent(AgentId))
 	{
 		UE_LOG(LogLearning, Error, TEXT("%s: AgentId %d not found. Be sure to only use AgentIds returned by AddAgent and check that the agent has not be removed."), *GetName(), AgentId);
 		return nullptr;
 	}
-	else
+	
+	if (Agents[AgentId] == nullptr)
 	{
-		return Agents[AgentId];
+		UE_LOG(LogLearning, Warning, TEXT("%s: AgentId %d object is nullptr."), *GetName(), AgentId);
+		return nullptr;
 	}
+	
+	if (!Agents[AgentId].IsA(AgentClass))
+	{
+		UE_LOG(LogLearning, Error, TEXT("%s: AgentId %d cast invalid. Agent was of class '%s', not '%s'."), *GetName(), AgentId, *Agents[AgentId]->GetClass()->GetName(), *AgentClass->GetName());
+		return nullptr;
+	}
+	
+	return Agents[AgentId];
 }
 
-void ALearningAgentsManager::GetAgents(TArray<UObject*>& OutAgents, const TArray<int32>& AgentIds, const TSubclassOf<UObject> AgentClass) const
+void ULearningAgentsManager::GetAgents(TArray<UObject*>& OutAgents, const TArray<int32>& AgentIds, const TSubclassOf<UObject> AgentClass) const
 {
 	OutAgents.SetNumUninitialized(AgentIds.Num());
 
@@ -313,21 +315,34 @@ void ALearningAgentsManager::GetAgents(TArray<UObject*>& OutAgents, const TArray
 		{
 			UE_LOG(LogLearning, Error, TEXT("%s: AgentId %d not found. Be sure to only use AgentIds returned by AddAgent and check that the agent has not be removed."), *GetName(), AgentIds[AgentIdIdx]);
 			OutAgents[AgentIdIdx] = nullptr;
+			continue;
+		}
+		
+		if (Agents[AgentIds[AgentIdIdx]] == nullptr)
+		{
+			UE_LOG(LogLearning, Warning, TEXT("%s: AgentId %d object is nullptr."), *GetName(), AgentIds[AgentIdIdx]);
+			OutAgents[AgentIdIdx] = nullptr;
+			continue;
 		}
-		else
+		
+		if (!Agents[AgentIds[AgentIdIdx]].IsA(AgentClass))
 		{
-			OutAgents[AgentIdIdx] = Agents[AgentIds[AgentIdIdx]];
+			UE_LOG(LogLearning, Error, TEXT("%s: AgentId %d cast invalid. Agent was of class '%s', not '%s'."), *GetName(), AgentIds[AgentIdIdx], *Agents[AgentIds[AgentIdIdx]]->GetClass()->GetName(), *AgentClass->GetName());
+			OutAgents[AgentIdIdx] = nullptr;
+			continue;
 		}
+		
+		OutAgents[AgentIdIdx] = Agents[AgentIds[AgentIdIdx]];
 	}
 }
 
-void ALearningAgentsManager::GetAllAgents(TArray<UObject*>& OutAgents, TArray<int32>& OutAgentIds, const TSubclassOf<UObject> AgentClass) const
+void ULearningAgentsManager::GetAllAgents(TArray<UObject*>& OutAgents, TArray<int32>& OutAgentIds, const TSubclassOf<UObject> AgentClass) const
 {
 	OutAgentIds = OccupiedAgentIds;
 	GetAgents(OutAgents, OutAgentIds, AgentClass);
 }
 
-int32 ALearningAgentsManager::GetAgentId(UObject* Agent) const
+int32 ULearningAgentsManager::GetAgentId(UObject* Agent) const
 {
 	const int32 AgentId = Agents.Find(Agent);
 
@@ -342,7 +357,7 @@ int32 ALearningAgentsManager::GetAgentId(UObject* Agent) const
 	}
 }
 
-void ALearningAgentsManager::GetAgentIds(TArray<int32>& OutAgentIds, const TArray<UObject*>& InAgents) const
+void ULearningAgentsManager::GetAgentIds(TArray<int32>& OutAgentIds, const TArray<UObject*>& InAgents) const
 {
 	OutAgentIds.SetNumUninitialized(InAgents.Num());
 
@@ -362,70 +377,63 @@ void ALearningAgentsManager::GetAgentIds(TArray<int32>& OutAgentIds, const TArra
 	}
 }
 
-int32 ALearningAgentsManager::GetAgentNum() const
+int32 ULearningAgentsManager::GetAgentNum() const
 {
 	return OccupiedAgentIds.Num();
 }
 
-bool ALearningAgentsManager::HasAgentObject(UObject* Agent) const
+bool ULearningAgentsManager::HasAgentObject(UObject* Agent) const
 {
 	return Agents.Contains(Agent);
 }
 
-bool ALearningAgentsManager::HasAgent(const int32 AgentId) const
+bool ULearningAgentsManager::HasAgent(const int32 AgentId) const
 {
 	return OccupiedAgentSet.Contains(AgentId);
 }
 
-void ALearningAgentsManager::AddManagerAsTickPrerequisiteOfAgents(const TArray<AActor*>& InAgents)
-{
-	for (AActor* Agent : InAgents)
-	{
-		Agent->AddTickPrerequisiteActor(this);
-	}
-}
-
-void ALearningAgentsManager::AddAgentsAsTickPrerequisiteOfManager(const TArray<AActor*>& InAgents)
-{
-	for (AActor* Agent : InAgents)
-	{
-		this->AddTickPrerequisiteActor(Agent);
-	}
-}
-
-const UObject* ALearningAgentsManager::GetAgent(const int32 AgentId) const
+const UObject* ULearningAgentsManager::GetAgent(const int32 AgentId) const
 {
 	UE_LEARNING_CHECK(HasAgent(AgentId));
 	return Agents[AgentId];
 }
 
-UObject* ALearningAgentsManager::GetAgent(const int32 AgentId)
+UObject* ULearningAgentsManager::GetAgent(const int32 AgentId)
 {
 	UE_LEARNING_CHECK(HasAgent(AgentId));
 	return Agents[AgentId];
 }
 
-const TArray<int32>& ALearningAgentsManager::GetAllAgentIds() const
+const TArray<int32>& ULearningAgentsManager::GetAllAgentIds() const
 {
 	return OccupiedAgentIds;
 }
 
-UE::Learning::FIndexSet ALearningAgentsManager::GetAllAgentSet() const
+UE::Learning::FIndexSet ULearningAgentsManager::GetAllAgentSet() const
 {
 	return OccupiedAgentSet;
 }
 
-const TSharedPtr<UE::Learning::FArrayMap>& ALearningAgentsManager::GetInstanceData() const
+TConstArrayView<TObjectPtr<UObject>> ULearningAgentsManager::GetAgents() const
 {
-	return InstanceData;
+	return Agents;
 }
 
-TConstArrayView<TObjectPtr<UObject>> ALearningAgentsManager::GetAgents() const
+void ULearningAgentsManager::AddListener(ULearningAgentsManagerListener* Listener)
 {
-	return Agents;
+	Listeners.Add(Listener);
+	OnEventAgentIds = OccupiedAgentIds;
+	Listener->OnAgentsAdded(OnEventAgentIds);
+}
+
+void ULearningAgentsManager::RemoveListener(ULearningAgentsManagerListener* Listener)
+{
+	OnEventAgentIds = OccupiedAgentIds;
+	Listener->OnAgentsRemoved(OnEventAgentIds);
+	Listeners.Remove(Listener);
 }
 
-void ALearningAgentsManager::UpdateAgentSets()
+void ULearningAgentsManager::UpdateAgentSets()
 {
 	// Keep OccupiedAgentIds sorted in ascending order for better memory access patterns.
 	OccupiedAgentIds.Sort();
diff --git a/Engine/Plugins/Experimental/LearningAgents/Source/LearningAgents/Private/LearningAgentsNeuralNetwork.cpp b/Engine/Plugins/Experimental/LearningAgents/Source/LearningAgents/Private/LearningAgentsNeuralNetwork.cpp
--- a/Engine/Plugins/Experimental/LearningAgents/Source/LearningAgents/Private/LearningAgentsNeuralNetwork.cpp
+++ b/Engine/Plugins/Experimental/LearningAgents/Source/LearningAgents/Private/LearningAgentsNeuralNetwork.cpp
@@ -2,8 +2,6 @@
 
 #include "LearningAgentsNeuralNetwork.h"
 
-#include "LearningAgentsNeuralNetworkData.h"
-
 #include "LearningLog.h"
 #include "LearningArray.h"
 #include "LearningNeuralNetwork.h"
@@ -17,8 +15,11 @@ ULearningAgentsNeuralNetwork::~ULearningAgentsNeuralNetwork() = default;
 
 void ULearningAgentsNeuralNetwork::ResetNetwork()
 {
-	NeuralNetworkData->ConditionalBeginDestroy();
-	NeuralNetworkData = nullptr;
+	if (NeuralNetworkData)
+	{
+		NeuralNetworkData->ConditionalBeginDestroy();
+		NeuralNetworkData = nullptr;
+	}
 	
 	ForceMarkDirty();
 }
@@ -29,39 +30,11 @@ void ULearningAgentsNeuralNetwork::LoadNetworkFromSnapshot(const FFilePath& File
 
 	if (FFileHelper::LoadFileToArray(NetworkData, *File.FilePath))
 	{
-		ULearningAgentsNeuralNetworkData* TempNeuralNetworkData = NewObject<ULearningAgentsNeuralNetworkData>(this);
-		int32 Offset = 0;
-		bool bSuccess = TempNeuralNetworkData->GetNetworkInterface()->DeserializeFromBytes(Offset, NetworkData);
-		
-		if (!bSuccess)
+		if (!NeuralNetworkData->LoadFromSnapshot(NetworkData))
 		{
 			UE_LOG(LogLearning, Error, TEXT("%s: Failed to load network. Invalid Format: \"%s\""), *GetName(), *File.FilePath);
 			return;
 		}
-		
-		NetworkData.Empty();
-
-		if (NeuralNetworkData)
-		{
-			// If we already have a neural network check settings match
-
-			if (TempNeuralNetworkData->GetNetworkInterface()->GetInputNum() != NeuralNetworkData->GetNetworkInterface()->GetInputNum() ||
-				TempNeuralNetworkData->GetNetworkInterface()->GetOutputNum() != NeuralNetworkData->GetNetworkInterface()->GetOutputNum())
-			{
-				UE_LOG(LogLearning, Error, TEXT("%s: Failed to load network from snapshot as settings don't match."), *GetName());
-				return;
-			}
-			else
-			{
-				NeuralNetworkData->CopyFrom(TempNeuralNetworkData);
-			}
-		}
-		else
-		{
-			// Otherwise use loaded neural network as-is
-
-			NeuralNetworkData = TempNeuralNetworkData;
-		}
 
 		ForceMarkDirty();
 	}
@@ -80,11 +53,8 @@ void ULearningAgentsNeuralNetwork::SaveNetworkToSnapshot(const FFilePath& File)
 	}
 
 	TArray<uint8> NetworkData;
-	NetworkData.SetNumUninitialized(NeuralNetworkData->GetNetworkInterface()->GetSerializationByteNum());
-
-	int32 Offset = 0;
-	NeuralNetworkData->GetNetworkInterface()->SerializeToBytes(Offset, NetworkData);
-	UE_LEARNING_CHECK(Offset == NetworkData.Num());
+	NetworkData.SetNumUninitialized(NeuralNetworkData->GetSnapshotByteNum());
+	NeuralNetworkData->SaveToSnapshot(NetworkData);
 
 	if (!FFileHelper::SaveArrayToFile(NetworkData, *File.FilePath))
 	{
@@ -106,21 +76,12 @@ void ULearningAgentsNeuralNetwork::LoadNetworkFromAsset(ULearningAgentsNeuralNet
 		return;
 	}
 
-	if (NeuralNetworkData)
-	{
-		if (NeuralNetworkAsset->NeuralNetworkData->GetNetworkInterface()->GetInputNum() != NeuralNetworkData->GetNetworkInterface()->GetInputNum() ||
-			NeuralNetworkAsset->NeuralNetworkData->GetNetworkInterface()->GetOutputNum() != NeuralNetworkData->GetNetworkInterface()->GetOutputNum())
-		{
-			UE_LOG(LogLearning, Error, TEXT("%s: Failed to load network from asset as settings don't match."), *GetName());
-			return;
-		}
-	}
-	else
+	if (!NeuralNetworkData)
 	{
-		NeuralNetworkData = NewObject<ULearningAgentsNeuralNetworkData>(this);
+		NeuralNetworkData = NewObject<ULearningNeuralNetworkData>(this);
 	}
 
-	NeuralNetworkData->CopyFrom(NeuralNetworkAsset->NeuralNetworkData);
+	NeuralNetworkData->InitFrom(NeuralNetworkAsset->NeuralNetworkData);
 	ForceMarkDirty();
 }
 
@@ -140,10 +101,10 @@ void ULearningAgentsNeuralNetwork::SaveNetworkToAsset(ULearningAgentsNeuralNetwo
 
 	if (!NeuralNetworkAsset->NeuralNetworkData)
 	{
-		NeuralNetworkAsset->NeuralNetworkData = NewObject<ULearningAgentsNeuralNetworkData>(NeuralNetworkAsset);
+		NeuralNetworkAsset->NeuralNetworkData = NewObject<ULearningNeuralNetworkData>(NeuralNetworkAsset);
 	}
 
-	NeuralNetworkAsset->NeuralNetworkData->CopyFrom(NeuralNetworkData);
+	NeuralNetworkAsset->NeuralNetworkData->InitFrom(NeuralNetworkData);
 	NeuralNetworkAsset->ForceMarkDirty();
 }
 
