LearningAgents: Large restructuring to allow for structured actions and observations.

diff --git a/Engine/Plugins/Experimental/LearningAgents/Source/LearningAgents/Private/LearningAgentsObservations.cpp b/Engine/Plugins/Experimental/LearningAgents/Source/LearningAgents/Private/LearningAgentsObservations.cpp
--- a/Engine/Plugins/Experimental/LearningAgents/Source/LearningAgents/Private/LearningAgentsObservations.cpp
+++ b/Engine/Plugins/Experimental/LearningAgents/Source/LearningAgents/Private/LearningAgentsObservations.cpp
@@ -2,2437 +2,3210 @@
 
 #include "LearningAgentsObservations.h"
 
-#include "LearningAgentsManager.h"
-#include "LearningAgentsInteractor.h"
 #include "LearningArray.h"
-#include "LearningArrayMap.h"
-#include "LearningFeatureObject.h"
 #include "LearningLog.h"
 
-#include "GameFramework/Actor.h"
+#include "Containers/StaticArray.h"
 
-namespace UE::Learning::Agents::Observations::Private
+#include "Runtime/Engine/Public/DrawDebugHelpers.h"
+#include "Engine/World.h"
+#include "Components/SplineComponent.h"
+
+bool operator==(const FLearningAgentsObservationObjectElement& Lhs, const FLearningAgentsObservationObjectElement& Rhs)
+{
+	return Lhs.ObjectElement.Index == Rhs.ObjectElement.Index;
+}
+
+uint32 GetTypeHash(const FLearningAgentsObservationObjectElement& Element)
+{
+	return (uint32)Element.ObjectElement.Index;
+}
+
+const UE::Learning::Observation::FSchema& ULearningAgentsObservationSchema::GetObservationSchema() const
+{
+	return ObservationSchema;
+}
+
+namespace UE::Learning::Agents::Observation::Private
 {
-	template<typename ObservationUObject, typename ObservationFObject, typename... InArgTypes>
-	ObservationUObject* AddObservation(ULearningAgentsInteractor* InInteractor, const FName Name, const TCHAR* FunctionName, InArgTypes&& ...Args)
+	static inline bool ContainsDuplicates(const TArrayView<const int32> Indices)
+	{
+		TSet<int32, DefaultKeyFuncs<int32>, TInlineSetAllocator<32>> IndicesSet;
+		IndicesSet.Append(Indices);
+		return Indices.Num() != IndicesSet.Num();
+	}
+
+	static inline bool ContainsDuplicates(const TArrayView<const FName> ElementNames)
+	{
+		TSet<FName, DefaultKeyFuncs<FName>, TInlineSetAllocator<32>> ElementNameSet;
+		ElementNameSet.Append(ElementNames);
+		return ElementNames.Num() != ElementNameSet.Num();
+	}
+
+	static inline const TCHAR* GetObservationTypeString(const Learning::Observation::EType ObservationType)
+	{
+		switch (ObservationType)
+		{
+		case  Learning::Observation::EType::Null: return TEXT("Null");
+		case  Learning::Observation::EType::Continuous: return TEXT("Continuous");
+		case  Learning::Observation::EType::And: return TEXT("Struct");
+		case  Learning::Observation::EType::OrExclusive: return TEXT("ExclusiveUnion");
+		case  Learning::Observation::EType::OrInclusive: return TEXT("InclusiveUnion");
+		case  Learning::Observation::EType::Array: return TEXT("StaticArray");
+		case  Learning::Observation::EType::Set: return TEXT("Set");
+		case  Learning::Observation::EType::Encoding: return TEXT("Encoding");
+		default:
+			UE_LEARNING_NOT_IMPLEMENTED();
+			return TEXT("Unimplemented");
+		}
+	}
+
+	static bool ValidateObjectMatchesSchema(
+		const Learning::Observation::FSchema& Schema,
+		const Learning::Observation::FSchemaElement SchemaElement,
+		const Learning::Observation::FObject& Object,
+		const Learning::Observation::FObjectElement ObjectElement,
+		const FString& ObjectName)
+	{
+		// Check Elements are Valid
+
+		if (!Schema.IsValid(SchemaElement))
+		{
+			UE_LOG(LogLearning, Error, TEXT("%s: Invalid Observation Schema Element."), *ObjectName);
+			return false;
+		}
+
+		if (!Object.IsValid(ObjectElement))
+		{
+			UE_LOG(LogLearning, Error, TEXT("%s: Invalid Observation Object Element."), *ObjectName);
+			return false;
+		}
+
+		// Check Names Match
+
+		const FName ObservationSchemaElementName = Schema.GetName(SchemaElement);
+		const FName ObservationObjectElementName = Object.GetName(ObjectElement);
+
+		if (ObservationSchemaElementName != ObservationObjectElementName)
+		{
+			UE_LOG(LogLearning, Warning, TEXT("%s: Observation name does not match Schema. Expected '%s', got '%s'."),
+				*ObjectName, *ObservationSchemaElementName.ToString(), *ObservationObjectElementName.ToString());
+		}
+
+		// Check Types Match
+
+		const Learning::Observation::EType ObservationSchemaElementType = Schema.GetType(SchemaElement);
+		const Learning::Observation::EType ObservationObjectElementType = Object.GetType(ObjectElement);
+
+		if (ObservationSchemaElementType != ObservationObjectElementType)
+		{
+			UE_LOG(LogLearning, Error, TEXT("%s: Observation '%s' type does not match Schema. Expected type '%s', got type '%s'."),
+				*ObjectName,
+				*ObservationSchemaElementName.ToString(),
+				GetObservationTypeString(ObservationSchemaElementType),
+				GetObservationTypeString(ObservationObjectElementType));
+			return false;
+		}
+
+		// Type Specific Checks
+
+		switch (ObservationSchemaElementType)
+		{
+		case Learning::Observation::EType::Null: return true;
+		
+		case Learning::Observation::EType::Continuous:
+		{
+			const int32 SchemaElementSize = Schema.GetContinuous(SchemaElement).Num;
+			const int32 ObjectElementSize = Object.GetContinuous(ObjectElement).Values.Num();
+
+			if (SchemaElementSize != ObjectElementSize)
+			{
+				UE_LOG(LogLearning, Error, TEXT("%s: Observation '%s' size does not match Schema. Expected '%i', got '%i'."),
+					*ObjectName,
+					*ObservationSchemaElementName.ToString(),
+					SchemaElementSize,
+					ObjectElementSize);
+				return false;
+			}
+
+			return true;
+		}
+
+		case Learning::Observation::EType::And:
+		{
+			const Learning::Observation::FSchemaAndParameters SchemaParameters = Schema.GetAnd(SchemaElement);
+			const Learning::Observation::FObjectAndParameters ObjectParameters = Object.GetAnd(ObjectElement);
+			UE_LEARNING_CHECK(SchemaParameters.Elements.Num() == SchemaParameters.ElementNames.Num());
+			UE_LEARNING_CHECK(ObjectParameters.Elements.Num() == ObjectParameters.ElementNames.Num());
+
+			if (SchemaParameters.Elements.Num() != ObjectParameters.Elements.Num())
+			{
+				UE_LOG(LogLearning, Error, TEXT("%s: Observation '%s' number of sub-elements does not match Schema. Expected '%i', got '%i'."),
+					*ObjectName,
+					*ObservationSchemaElementName.ToString(),
+					SchemaParameters.Elements.Num(),
+					ObjectParameters.Elements.Num());
+				return false;
+			}
+
+			for (int32 SchemaElementIdx = 0; SchemaElementIdx < SchemaParameters.Elements.Num(); SchemaElementIdx++)
+			{
+				const int32 ObjectElementIdx = ObjectParameters.ElementNames.Find(SchemaParameters.ElementNames[SchemaElementIdx]);
+
+				if (ObjectElementIdx == INDEX_NONE)
+				{
+					UE_LOG(LogLearning, Error, TEXT("%s: Observation '%s' does not include '%s' observation required by Schema."),
+						*ObjectName,
+						*ObservationSchemaElementName.ToString(),
+						*SchemaParameters.ElementNames[SchemaElementIdx].ToString());
+					return false;
+				}
+
+				if (!ValidateObjectMatchesSchema(
+					Schema,
+					SchemaParameters.Elements[SchemaElementIdx],
+					Object,
+					ObjectParameters.Elements[ObjectElementIdx],
+					ObjectName))
+				{
+					return false;
+				}
+			}
+
+			return true;
+		}
+
+		case Learning::Observation::EType::OrExclusive:
+		{
+			const Learning::Observation::FSchemaOrExclusiveParameters SchemaParameters = Schema.GetOrExclusive(SchemaElement);
+			const Learning::Observation::FObjectOrExclusiveParameters ObjectParameters = Object.GetOrExclusive(ObjectElement);
+			UE_LEARNING_CHECK(SchemaParameters.Elements.Num() == SchemaParameters.ElementNames.Num());
+
+			const int32 SchemaSubElementIdx = SchemaParameters.ElementNames.Find(ObjectParameters.ElementName);
+
+			if (SchemaSubElementIdx == INDEX_NONE)
+			{
+				UE_LOG(LogLearning, Error, TEXT("%s: Observation '%s' Schema does not include '%s' observation."),
+					*ObjectName,
+					*ObservationSchemaElementName.ToString(),
+					*ObjectParameters.ElementName.ToString());
+				return false;
+			}
+
+			return ValidateObjectMatchesSchema(
+				Schema,
+				SchemaParameters.Elements[SchemaSubElementIdx],
+				Object,
+				ObjectParameters.Element,
+				ObjectName);
+		}
+
+		case Learning::Observation::EType::OrInclusive:
+		{
+			const Learning::Observation::FSchemaOrInclusiveParameters SchemaParameters = Schema.GetOrInclusive(SchemaElement);
+			const Learning::Observation::FObjectOrInclusiveParameters ObjectParameters = Object.GetOrInclusive(ObjectElement);
+
+			if (ObjectParameters.Elements.Num() > SchemaParameters.Elements.Num())
+			{
+				UE_LOG(LogLearning, Error, TEXT("%s: Observation '%s' too many sub-observations provided. Expected at most '%i', got '%i'."),
+					*ObjectName,
+					*ObservationSchemaElementName.ToString(),
+					SchemaParameters.Elements.Num(),
+					ObjectParameters.Elements.Num());
+				return false;
+			}
+
+			for (int32 ObjectSubElementIdx = 0; ObjectSubElementIdx < ObjectParameters.Elements.Num(); ObjectSubElementIdx++)
+			{
+				const int32 SchemaSubElementIdx = SchemaParameters.ElementNames.Find(ObjectParameters.ElementNames[ObjectSubElementIdx]);
+
+				if (SchemaSubElementIdx == INDEX_NONE)
+				{
+					UE_LOG(LogLearning, Error, TEXT("%s: Observation '%s' Schema does not include '%s' observation."),
+						*ObjectName,
+						*ObservationSchemaElementName.ToString(),
+						*ObjectParameters.ElementNames[ObjectSubElementIdx].ToString());
+					return false;
+				}
+
+				if (!ValidateObjectMatchesSchema(
+					Schema,
+					SchemaParameters.Elements[SchemaSubElementIdx],
+					Object,
+					ObjectParameters.Elements[ObjectSubElementIdx],
+					ObjectName))
+				{
+					return false;
+				}
+			}
+
+			return true;
+		}
+
+		case Learning::Observation::EType::Array:
+		{
+			const Learning::Observation::FSchemaArrayParameters SchemaParameters = Schema.GetArray(SchemaElement);
+			const Learning::Observation::FObjectArrayParameters ObjectParameters = Object.GetArray(ObjectElement);
+
+			if (ObjectParameters.Elements.Num() != SchemaParameters.Num)
+			{
+				UE_LOG(LogLearning, Error, TEXT("%s: Observation '%s' array incorrect size. Expected '%i' elements, got '%i'."),
+					*ObjectName,
+					*ObservationSchemaElementName.ToString(),
+					SchemaParameters.Num,
+					ObjectParameters.Elements.Num());
+				return false;
+			}
+
+			for (int32 ElementIdx = 0; ElementIdx < ObjectParameters.Elements.Num(); ElementIdx++)
+			{
+				if (!ValidateObjectMatchesSchema(
+					Schema,
+					SchemaParameters.Element,
+					Object,
+					ObjectParameters.Elements[ElementIdx],
+					ObjectName))
+				{
+					return false;
+				}
+			}
+
+			return true;
+		}
+
+		case Learning::Observation::EType::Set:
+		{
+			const Learning::Observation::FSchemaSetParameters SchemaParameters = Schema.GetSet(SchemaElement);
+			const Learning::Observation::FObjectSetParameters ObjectParameters = Object.GetSet(ObjectElement);
+
+			if (ObjectParameters.Elements.Num() > SchemaParameters.MaxNum)
+			{
+				UE_LOG(LogLearning, Error, TEXT("%s: Observation '%s' set too large. Expected at most '%i' elements, got '%i'."),
+					*ObjectName,
+					*ObservationSchemaElementName.ToString(),
+					SchemaParameters.MaxNum,
+					ObjectParameters.Elements.Num());
+				return false;
+			}
+
+			for (int32 ElementIdx = 0; ElementIdx < ObjectParameters.Elements.Num(); ElementIdx++)
+			{
+				if (!ValidateObjectMatchesSchema(
+					Schema,
+					SchemaParameters.Element,
+					Object,
+					ObjectParameters.Elements[ElementIdx],
+					ObjectName))
+				{
+					return false;
+				}
+			}
+
+			return true;
+		}
+
+		case Learning::Observation::EType::Encoding:
+		{
+			const Learning::Observation::FSchemaEncodingParameters SchemaParameters = Schema.GetEncoding(SchemaElement);
+			const Learning::Observation::FObjectEncodingParameters ObjectParameters = Object.GetEncoding(ObjectElement);
+
+			return ValidateObjectMatchesSchema(
+				Schema,
+				SchemaParameters.Element,
+				Object,
+				ObjectParameters.Element,
+				ObjectName);
+		}
+
+		default:
+		{
+			UE_LEARNING_NOT_IMPLEMENTED();
+			return true;
+		}
+		}
+	}
+
+	static void LogObservation(
+		const UE::Learning::Observation::FObject& Object, 
+		const UE::Learning::Observation::FObjectElement ObjectElement,
+		const FString& Indentation,
+		const FString& Prefix,
+		const FString& ObjectName)
 	{
-		if (!InInteractor)
+		if (!Object.IsValid(ObjectElement))
 		{
-			UE_LOG(LogLearning, Error, TEXT("%s: InInteractor is nullptr."), FunctionName);
-			return nullptr;
+			UE_LOG(LogLearning, Error, TEXT("%s: Invalid Observation Object Element."), *ObjectName);
+			return;
 		}
 
-		if (!InInteractor->HasAgentManager())
+		const UE::Learning::Observation::EType Type = Object.GetType(ObjectElement);
+		const FName Name = Object.GetName(ObjectElement);
+
+		switch (Type)
+		{
+		case UE::Learning::Observation::EType::Null:
+		{
+			UE_LOG(LogLearning, Display, TEXT("%s%s \"%s\" (%s)"), *Indentation, *Prefix, *Name.ToString(), GetObservationTypeString(Type));
+			return;
+		}
+
+		case UE::Learning::Observation::EType::Continuous:
+		{
+			const UE::Learning::Observation::FObjectContinuousParameters Parameters = Object.GetContinuous(ObjectElement);
+
+			UE_LOG(LogLearning, Display, TEXT("%s%s \"%s\" (%s) %s"), *Indentation, *Prefix, *Name.ToString(), GetObservationTypeString(Type), *UE::Learning::Array::FormatFloat(Parameters.Values));
+			return;
+		}
+
+		case UE::Learning::Observation::EType::And:
+		{
+			const UE::Learning::Observation::FObjectAndParameters Parameters = Object.GetAnd(ObjectElement);
+
+			UE_LOG(LogLearning, Display, TEXT("%s%s \"%s\" (%s)"), *Indentation, *Prefix, *Name.ToString(), GetObservationTypeString(Type));
+			for (int32 SubElementIdx = 0; SubElementIdx < Parameters.Elements.Num(); SubElementIdx++)
+			{
+				LogObservation(Object, Parameters.Elements[SubElementIdx], *(Indentation + TEXT("    ")), FString::Printf(TEXT("| \"%s\": "), *Parameters.ElementNames[SubElementIdx].ToString()), ObjectName);
+			}
+
+			return;
+		}
+
+		case UE::Learning::Observation::EType::OrExclusive:
+		{
+			const UE::Learning::Observation::FObjectOrExclusiveParameters Parameters = Object.GetOrExclusive(ObjectElement);
+
+			UE_LOG(LogLearning, Display, TEXT("%s%s \"%s\" (%s)"), *Indentation, *Prefix, *Name.ToString(), GetObservationTypeString(Type));
+			LogObservation(Object, Parameters.Element, *(Indentation + TEXT("    ")), FString::Printf(TEXT("| \"%s\": "), *Parameters.ElementName.ToString()), ObjectName);
+
+			return;
+		}
+
+		case UE::Learning::Observation::EType::OrInclusive:
+		{
+			const UE::Learning::Observation::FObjectOrInclusiveParameters Parameters = Object.GetOrInclusive(ObjectElement);
+
+			UE_LOG(LogLearning, Display, TEXT("%s%s \"%s\" (%s)"), *Indentation, *Prefix, *Name.ToString(), GetObservationTypeString(Type));
+			for (int32 SubElementIdx = 0; SubElementIdx < Parameters.Elements.Num(); SubElementIdx++)
+			{
+				LogObservation(Object, Parameters.Elements[SubElementIdx], *(Indentation + TEXT("    ")), FString::Printf(TEXT("| \"%s\": "), *Parameters.ElementNames[SubElementIdx].ToString()), ObjectName);
+			}
+
+			return;
+		}
+
+		case UE::Learning::Observation::EType::Array:
+		{
+			const UE::Learning::Observation::FObjectArrayParameters Parameters = Object.GetArray(ObjectElement);
+
+			UE_LOG(LogLearning, Display, TEXT("%s%s \"%s\" (%s)"), *Indentation, *Prefix, *Name.ToString(), GetObservationTypeString(Type));
+			for (int32 SubElementIdx = 0; SubElementIdx < Parameters.Elements.Num(); SubElementIdx++)
+			{
+				LogObservation(Object, Parameters.Elements[SubElementIdx], *(Indentation + TEXT("    ")), FString::Printf(TEXT("| %3i:"), SubElementIdx), ObjectName);
+			}
+
+			return;
+		}
+
+		case UE::Learning::Observation::EType::Set:
+		{
+			const UE::Learning::Observation::FObjectSetParameters Parameters = Object.GetSet(ObjectElement);
+
+			UE_LOG(LogLearning, Display, TEXT("%s%s \"%s\" (%s)"), *Indentation, *Prefix, *Name.ToString(), GetObservationTypeString(Type));
+			for (int32 SubElementIdx = 0; SubElementIdx < Parameters.Elements.Num(); SubElementIdx++)
+			{
+				LogObservation(Object, Parameters.Elements[SubElementIdx], *(Indentation + TEXT("    ")), FString::Printf(TEXT("| %3i:"), SubElementIdx), ObjectName);
+			}
+
+			return;
+		}
+
+		case UE::Learning::Observation::EType::Encoding:
+		{
+			const UE::Learning::Observation::FObjectEncodingParameters Parameters = Object.GetEncoding(ObjectElement);
+
+			UE_LOG(LogLearning, Display, TEXT("%s%s \"%s\" (%s)"), *Indentation, *Prefix, *Name.ToString(), GetObservationTypeString(Type));
+			LogObservation(Object, Parameters.Element, *(Indentation + TEXT("    ")), TEXT("|"), ObjectName);
+
+			return;
+		}
+		}
+	}
+
+	static inline FVector VectorLogSafe(const FVector V, const float Epsilon = UE_SMALL_NUMBER)
+	{
+		return FVector(
+			FMath::Loge(FMath::Max(V.X, Epsilon)),
+			FMath::Loge(FMath::Max(V.Y, Epsilon)),
+			FMath::Loge(FMath::Max(V.Z, Epsilon)));
+	}
+
+	static inline FVector VectorExp(const FVector V)
+	{
+		return FVector(
+			FMath::Exp(V.X),
+			FMath::Exp(V.Y),
+			FMath::Exp(V.Z));
+	}
+}
+
+FTransform ULearningAgentsObservationFunctions::ProjectTransformOntoGroundPlane(const FTransform Transform, const FVector LocalForwardVector, const float GroundPlaneHeight)
+{
+	FVector Position = Transform.GetLocation();
+	Position.Z = GroundPlaneHeight;
+
+	const FVector Direction = (FVector(1.0f, 1.0f, 0.0f) * Transform.TransformVectorNoScale(LocalForwardVector)).GetSafeNormal(UE_SMALL_NUMBER, FVector::ForwardVector);
+
+	return FTransform(FQuat::FindBetweenNormals(FVector::ForwardVector, Direction), Position, Transform.GetScale3D());
+}
+
+bool ULearningAgentsObservationSchema::ValidateObjectMatchesSchema(
+	const FLearningAgentsObservationSchemaElement SchemaElement,
+	const ULearningAgentsObservationObject* Object,
+	const FLearningAgentsObservationObjectElement ObjectElement) const
+{
+	return UE::Learning::Agents::Observation::Private::ValidateObjectMatchesSchema(
+		ObservationSchema,
+		SchemaElement.SchemaElement,
+		Object->GetObservationObject(),
+		ObjectElement.ObjectElement,
+		GetName());
+}
+
+FLearningAgentsObservationSchemaElement ULearningAgentsObservationSchema::SpecifyNullObservation(const FName Name)
+{
+	return { ObservationSchema.CreateNull(Name) };
+}
+
+FLearningAgentsObservationSchemaElement ULearningAgentsObservationSchema::SpecifyContinuousObservation(const int32 Size, const FName Name)
+{
+	if (Size < 0)
+	{
+		UE_LOG(LogLearning, Error, TEXT("%s: Invalid Continuous Observation Size '%i'."), *GetName(), Size);
+		return FLearningAgentsObservationSchemaElement();
+	}
+	
+	if (Size == 0)
+	{
+		UE_LOG(LogLearning, Warning, TEXT("%s: Specifying zero-sized Continuous Observation."), *GetName());
+	}
+
+	return { ObservationSchema.CreateContinuous({ Size }, Name) };
+}
+
+FLearningAgentsObservationSchemaElement ULearningAgentsObservationSchema::SpecifyExclusiveDiscreteObservation(const int32 Size, const FName Name)
+{
+	return SpecifyContinuousObservation(Size, Name);
+}
+
+FLearningAgentsObservationSchemaElement ULearningAgentsObservationSchema::SpecifyInclusiveDiscreteObservation(const int32 Size, const FName Name)
+{
+	return SpecifyContinuousObservation(Size, Name);
+}
+
+FLearningAgentsObservationSchemaElement ULearningAgentsObservationSchema::SpecifyIndexObservation(const int32 Size, const FName Name)
+{
+	return SpecifyContinuousObservation(Size, Name);
+}
+
+FLearningAgentsObservationSchemaElement ULearningAgentsObservationSchema::SpecifyStructObservation(const TMap<FName, FLearningAgentsObservationSchemaElement>& Elements, const FName Name)
+{
+	if (Elements.Num() == 0)
+	{
+		UE_LOG(LogLearning, Warning, TEXT("%s: Specifying zero-sized Struct Observation."), *GetName());
+	}
+
+	const int32 SubElementNum = Elements.Num();
+
+	TArray<int32, TInlineAllocator<16>> SubElementIndices;
+	TArray<FName, TInlineAllocator<16>> SubElementNames;
+	TArray<FLearningAgentsObservationSchemaElement, TInlineAllocator<16>> SubElements;
+	SubElementIndices.Empty(Elements.Num());
+	SubElementNames.Empty(Elements.Num());
+	SubElements.Empty(Elements.Num());
+
+	int32 Index = 0;
+	for (const TPair<FName, FLearningAgentsObservationSchemaElement>& Element : Elements)
+	{
+		SubElementIndices.Add(Index);
+		SubElementNames.Add(Element.Key);
+		SubElements.Add(Element.Value);
+		Index++;
+	}
+
+	// Sort Elements According to FName
+
+	SubElementIndices.Sort([SubElementNames](const int32 Lhs, const int32 Rhs)
+	{
+		return SubElementNames[Lhs].ToString().ToLower() < SubElementNames[Rhs].ToString().ToLower();
+	});
+
+	TArray<FName, TInlineAllocator<16>> SortedSubElementNames;
+	TArray<FLearningAgentsObservationSchemaElement, TInlineAllocator<16>> SortedSubElements;
+	SortedSubElementNames.SetNumUninitialized(SubElementNum);
+	SortedSubElements.SetNumUninitialized(SubElementNum);
+	for (int32 Idx = 0; Idx < SubElementNum; Idx++)
+	{
+		SortedSubElementNames[Idx] = SubElementNames[SubElementIndices[Idx]];
+		SortedSubElements[Idx] = SubElements[SubElementIndices[Idx]];
+	}
+
+	return SpecifyStructObservationFromArrayViews(SortedSubElementNames, SortedSubElements, Name);
+}
+
+FLearningAgentsObservationSchemaElement ULearningAgentsObservationSchema::SpecifyStructObservationFromArrays(const TArray<FName>& ElementNames, const TArray<FLearningAgentsObservationSchemaElement>& Elements, const FName Name)
+{
+	return SpecifyStructObservationFromArrayViews(ElementNames, Elements, Name);
+}
+
+FLearningAgentsObservationSchemaElement ULearningAgentsObservationSchema::SpecifyStructObservationFromArrayViews(const TArrayView<const FName> ElementNames, const TArrayView<const FLearningAgentsObservationSchemaElement> Elements, const FName Name)
+{
+	if (Elements.Num() == 0)
+	{
+		UE_LOG(LogLearning, Warning, TEXT("%s: Specifying zero-sized Struct Observation."), *GetName());
+	}
+
+	if (Elements.Num() != ElementNames.Num())
+	{
+		UE_LOG(LogLearning, Error, TEXT("%s: Number of elements (%i) must match number of names (%i)."), *GetName(), Elements.Num(), ElementNames.Num());
+		return FLearningAgentsObservationSchemaElement();
+	}
+
+	if (UE::Learning::Agents::Observation::Private::ContainsDuplicates(ElementNames))
+	{
+		UE_LOG(LogLearning, Error, TEXT("%s: Element Names contain duplicates."), *GetName());
+		return FLearningAgentsObservationSchemaElement();
+	}
+
+	TArray<UE::Learning::Observation::FSchemaElement, TInlineAllocator<16>> SubElements;
+	SubElements.Empty(Elements.Num());
+
+	for (const FLearningAgentsObservationSchemaElement& Element : Elements)
+	{
+		if (!ObservationSchema.IsValid(Element.SchemaElement))
+		{
+			UE_LOG(LogLearning, Error, TEXT("%s: Invalid Observation Object."), *GetName());
+			return FLearningAgentsObservationSchemaElement();
+		}
+
+		SubElements.Add(Element.SchemaElement);
+	}
+
+	return { ObservationSchema.CreateAnd({ ElementNames, SubElements }, Name) };
+}
+
+FLearningAgentsObservationSchemaElement ULearningAgentsObservationSchema::SpecifyExclusiveUnionObservation(const TMap<FName, FLearningAgentsObservationSchemaElement>& Elements, const int32 EncodingSize, const FName Name)
+{
+	if (EncodingSize < 1)
+	{
+		UE_LOG(LogLearning, Error, TEXT("%s: Invalid Observation EncodingSize '%i'."), *GetName(), EncodingSize);
+		return FLearningAgentsObservationSchemaElement();
+	}
+
+	if (Elements.Num() == 0)
+	{
+		UE_LOG(LogLearning, Warning, TEXT("%s: Specifying zero-sized Exclusive Union Observation."), *GetName());
+	}
+
+	const int32 SubElementNum = Elements.Num();
+
+	TArray<int32, TInlineAllocator<16>> SubElementIndices;
+	TArray<FName, TInlineAllocator<16>> SubElementNames;
+	TArray<FLearningAgentsObservationSchemaElement, TInlineAllocator<16>> SubElements;
+	SubElementIndices.Empty(Elements.Num());
+	SubElementNames.Empty(Elements.Num());
+	SubElements.Empty(Elements.Num());
+
+	int32 Index = 0;
+	for (const TPair<FName, FLearningAgentsObservationSchemaElement>& Element : Elements)
+	{
+		SubElementIndices.Add(Index);
+		SubElementNames.Add(Element.Key);
+		SubElements.Add(Element.Value);
+		Index++;
+	}
+
+	// Sort Elements According to FName
+
+	SubElementIndices.Sort([SubElementNames](const int32 Lhs, const int32 Rhs)
+	{
+		return SubElementNames[Lhs].ToString().ToLower() < SubElementNames[Rhs].ToString().ToLower();
+	});
+
+	TArray<FName, TInlineAllocator<16>> SortedSubElementNames;
+	TArray<FLearningAgentsObservationSchemaElement, TInlineAllocator<16>> SortedSubElements;
+	SortedSubElementNames.SetNumUninitialized(SubElementNum);
+	SortedSubElements.SetNumUninitialized(SubElementNum);
+	for (int32 Idx = 0; Idx < SubElementNum; Idx++)
+	{
+		SortedSubElementNames[Idx] = SubElementNames[SubElementIndices[Idx]];
+		SortedSubElements[Idx] = SubElements[SubElementIndices[Idx]];
+	}
+
+	return SpecifyExclusiveUnionObservationFromArrayViews(SortedSubElementNames, SortedSubElements, EncodingSize, Name);
+}
+
+FLearningAgentsObservationSchemaElement ULearningAgentsObservationSchema::SpecifyExclusiveUnionObservationFromArrays(const TArray<FName>& ElementNames, const TArray<FLearningAgentsObservationSchemaElement>& Elements, const int32 EncodingSize, const FName Name)
+{
+	return SpecifyExclusiveUnionObservationFromArrayViews(ElementNames, Elements, EncodingSize, Name);
+}
+
+FLearningAgentsObservationSchemaElement ULearningAgentsObservationSchema::SpecifyExclusiveUnionObservationFromArrayViews(const TArrayView<const FName> ElementNames, const TArrayView<const FLearningAgentsObservationSchemaElement> Elements, const int32 EncodingSize, const FName Name)
+{
+	if (EncodingSize < 1)
+	{
+		UE_LOG(LogLearning, Error, TEXT("%s: Invalid Observation EncodingSize '%i'."), *GetName(), EncodingSize);
+		return FLearningAgentsObservationSchemaElement();
+	}
+
+	if (Elements.Num() == 0)
+	{
+		UE_LOG(LogLearning, Warning, TEXT("%s: Specifying zero-sized Exclusive Union Observation."), *GetName());
+	}
+
+	if (Elements.Num() != ElementNames.Num())
+	{
+		UE_LOG(LogLearning, Error, TEXT("%s: Number of elements (%i) must match number of names (%i)."), *GetName(), Elements.Num(), ElementNames.Num());
+		return FLearningAgentsObservationSchemaElement();
+	}
+
+	if (UE::Learning::Agents::Observation::Private::ContainsDuplicates(ElementNames))
+	{
+		UE_LOG(LogLearning, Error, TEXT("%s: Element Names contain duplicates."), *GetName());
+		return FLearningAgentsObservationSchemaElement();
+	}
+
+	TArray<UE::Learning::Observation::FSchemaElement, TInlineAllocator<16>> SubElements;
+	SubElements.Empty(Elements.Num());
+
+	for (const FLearningAgentsObservationSchemaElement& Element : Elements)
+	{
+		if (!ObservationSchema.IsValid(Element.SchemaElement))
+		{
+			UE_LOG(LogLearning, Error, TEXT("%s: Invalid Observation Object."), *GetName());
+			return FLearningAgentsObservationSchemaElement();
+		}
+
+		SubElements.Add(Element.SchemaElement);
+	}
+
+	return { ObservationSchema.CreateOrExclusive({ ElementNames, SubElements, EncodingSize }, Name) };
+}
+
+FLearningAgentsObservationSchemaElement ULearningAgentsObservationSchema::SpecifyInclusiveUnionObservation(const TMap<FName, FLearningAgentsObservationSchemaElement>& Elements, const int32 AttentionEncodingSize, const int32 AttentionHeadNum, const int32 ValueEncodingSize, const FName Name)
+{
+	if (AttentionEncodingSize < 1 || AttentionHeadNum < 1 || ValueEncodingSize < 1)
+	{
+		UE_LOG(LogLearning, Error, TEXT("%s: Invalid Observation Parameters: AttentionEncodingSize: %i, AttentionHeadNum: %i, ValueEncodingSize: %i."), *GetName(), AttentionEncodingSize, AttentionHeadNum, ValueEncodingSize);
+		return FLearningAgentsObservationSchemaElement();
+	}
+
+	if (Elements.Num() == 0)
+	{
+		UE_LOG(LogLearning, Warning, TEXT("%s: Specifying zero-sized Inclusive Union Observation."), *GetName());
+	}
+
+	const int32 SubElementNum = Elements.Num();
+
+	TArray<int32, TInlineAllocator<16>> SubElementIndices;
+	TArray<FName, TInlineAllocator<16>> SubElementNames;
+	TArray<FLearningAgentsObservationSchemaElement, TInlineAllocator<16>> SubElements;
+	SubElementIndices.Empty(SubElementNum);
+	SubElementNames.Empty(SubElementNum);
+	SubElements.Empty(SubElementNum);
+
+	int32 Index = 0;
+	for (const TPair<FName, FLearningAgentsObservationSchemaElement>& Element : Elements)
+	{
+		SubElementIndices.Add(Index);
+		SubElementNames.Add(Element.Key);
+		SubElements.Add(Element.Value);
+		Index++;
+	}
+
+	// Sort Elements According to FName
+
+	SubElementIndices.Sort([SubElementNames](const int32 Lhs, const int32 Rhs)
+	{
+		return SubElementNames[Lhs].ToString().ToLower() < SubElementNames[Rhs].ToString().ToLower();
+	});
+
+	TArray<FName, TInlineAllocator<16>> SortedSubElementNames;
+	TArray<FLearningAgentsObservationSchemaElement, TInlineAllocator<16>> SortedSubElements;
+	SortedSubElementNames.SetNumUninitialized(SubElementNum);
+	SortedSubElements.SetNumUninitialized(SubElementNum);
+	for (int32 Idx = 0; Idx < SubElementNum; Idx++)
+	{
+		SortedSubElementNames[Idx] = SubElementNames[SubElementIndices[Idx]];
+		SortedSubElements[Idx] = SubElements[SubElementIndices[Idx]];
+	}
+
+	return SpecifyInclusiveUnionObservationFromArrayViews(SortedSubElementNames, SortedSubElements, AttentionEncodingSize, AttentionHeadNum, ValueEncodingSize, Name);
+}
+
+FLearningAgentsObservationSchemaElement ULearningAgentsObservationSchema::SpecifyInclusiveUnionObservationFromArrays(const TArray<FName>& ElementNames, const TArray<FLearningAgentsObservationSchemaElement>& Elements, const int32 AttentionEncodingSize, const int32 AttentionHeadNum, const int32 ValueEncodingSize, const FName Name)
+{
+	return SpecifyInclusiveUnionObservationFromArrayViews(ElementNames, Elements, AttentionEncodingSize, AttentionHeadNum, ValueEncodingSize, Name);
+}
+
+FLearningAgentsObservationSchemaElement ULearningAgentsObservationSchema::SpecifyInclusiveUnionObservationFromArrayViews(const TArrayView<const FName> ElementNames, const TArrayView<const FLearningAgentsObservationSchemaElement> Elements, const int32 AttentionEncodingSize, const int32 AttentionHeadNum, const int32 ValueEncodingSize, const FName Name)
+{
+	if (AttentionEncodingSize < 1 || AttentionHeadNum < 1 || ValueEncodingSize < 1)
+	{
+		UE_LOG(LogLearning, Error, TEXT("%s: Invalid Observation Parameters: AttentionEncodingSize: %i, AttentionHeadNum: %i, ValueEncodingSize: %i."), *GetName(), AttentionEncodingSize, AttentionHeadNum, ValueEncodingSize);
+		return FLearningAgentsObservationSchemaElement();
+	}
+
+	if (Elements.Num() == 0)
+	{
+		UE_LOG(LogLearning, Warning, TEXT("%s: Specifying zero-sized Inclusive Union Observation."), *GetName());
+	}
+
+	if (Elements.Num() != ElementNames.Num())
+	{
+		UE_LOG(LogLearning, Error, TEXT("%s: Number of elements (%i) must match number of names (%i)."), *GetName(), Elements.Num(), ElementNames.Num());
+		return FLearningAgentsObservationSchemaElement();
+	}
+
+	if (UE::Learning::Agents::Observation::Private::ContainsDuplicates(ElementNames))
+	{
+		UE_LOG(LogLearning, Error, TEXT("%s: Element Names contain duplicates."), *GetName());
+		return FLearningAgentsObservationSchemaElement();
+	}
+
+	TArray<UE::Learning::Observation::FSchemaElement, TInlineAllocator<16>> SubElements;
+	SubElements.Empty(Elements.Num());
+
+	for (const FLearningAgentsObservationSchemaElement& Element : Elements)
+	{
+		if (!ObservationSchema.IsValid(Element.SchemaElement))
+		{
+			UE_LOG(LogLearning, Error, TEXT("%s: Invalid Observation Object."), *GetName());
+			return FLearningAgentsObservationSchemaElement();
+		}
+
+		SubElements.Add(Element.SchemaElement);
+	}
+
+	return { ObservationSchema.CreateOrInclusive({ ElementNames, SubElements, AttentionEncodingSize, AttentionHeadNum, ValueEncodingSize }, Name) };
+}
+
+FLearningAgentsObservationSchemaElement ULearningAgentsObservationSchema::SpecifyStaticArrayObservation(const FLearningAgentsObservationSchemaElement Element, const int32 Num, const FName Name)
+{
+	if (Num < 0)
+	{
+		UE_LOG(LogLearning, Error, TEXT("%s: Invalid Observation Array Num %i."), *GetName(), Num);
+		return FLearningAgentsObservationSchemaElement();
+	}
+
+	if (Num == 0)
+	{
+		UE_LOG(LogLearning, Warning, TEXT("%s: Specifying zero-sized Static Array Observation."), *GetName());
+	}
+
+	if (!ObservationSchema.IsValid(Element.SchemaElement))
+	{
+		UE_LOG(LogLearning, Error, TEXT("%s: Invalid Observation Object."), *GetName());
+		return FLearningAgentsObservationSchemaElement();
+	}
+
+	return { ObservationSchema.CreateArray({ Element.SchemaElement, Num }, Name) };
+}
+
+FLearningAgentsObservationSchemaElement ULearningAgentsObservationSchema::SpecifySetObservation(const FLearningAgentsObservationSchemaElement Element, const int32 MaxNum, const int32 AttentionEncodingSize, const int32 AttentionHeadNum, const int32 ValueEncodingSize, const FName Name)
+{
+	if (MaxNum < 0)
+	{
+		UE_LOG(LogLearning, Error, TEXT("%s: Invalid Observation Set MaxNum %i."), *GetName(), MaxNum);
+		return FLearningAgentsObservationSchemaElement();
+	}
+
+	if (MaxNum == 0)
+	{
+		UE_LOG(LogLearning, Warning, TEXT("%s: Specifying zero-sized Set Observation."), *GetName());
+	}
+
+	if (AttentionEncodingSize < 1 || AttentionHeadNum < 1 || ValueEncodingSize < 1)
+	{
+		UE_LOG(LogLearning, Error, TEXT("%s: Invalid Observation Parameters: AttentionEncodingSize: %i, AttentionHeadNum: %i, ValueEncodingSize: %i."), *GetName(), AttentionEncodingSize, AttentionHeadNum, ValueEncodingSize);
+		return FLearningAgentsObservationSchemaElement();
+	}
+
+	if (!ObservationSchema.IsValid(Element.SchemaElement))
+	{
+		UE_LOG(LogLearning, Error, TEXT("%s: Invalid Observation Object."), *GetName());
+		return FLearningAgentsObservationSchemaElement();
+	}
+
+	return { ObservationSchema.CreateSet({ Element.SchemaElement, MaxNum, AttentionEncodingSize, AttentionHeadNum, ValueEncodingSize }, Name) };
+}
+
+FLearningAgentsObservationSchemaElement ULearningAgentsObservationSchema::SpecifyPairObservation(const FLearningAgentsObservationSchemaElement Element0, const FLearningAgentsObservationSchemaElement Element1, const FName Name)
+{
+	return SpecifyStructObservationFromArrayViews({ TEXT("Key"), TEXT("Value") }, { Element0, Element1 }, Name);
+}
+
+FLearningAgentsObservationSchemaElement ULearningAgentsObservationSchema::SpecifyArrayObservation(const FLearningAgentsObservationSchemaElement Element, const int32 MaxNum, const int32 AttentionEncodingSize, const int32 AttentionHeadNum, const int32 ValueEncodingSize, const FName Name)
+{
+	return SpecifySetObservation(SpecifyPairObservation(SpecifyIndexObservation(MaxNum), Element), MaxNum, AttentionEncodingSize, AttentionHeadNum, ValueEncodingSize);
+}
+
+FLearningAgentsObservationSchemaElement ULearningAgentsObservationSchema::SpecifyMapObservation(const FLearningAgentsObservationSchemaElement KeyElement, const FLearningAgentsObservationSchemaElement ValueElement, const int32 MaxNum, const int32 AttentionEncodingSize, const int32 AttentionHeadNum, const int32 ValueEncodingSize, const FName Name)
+{
+	return SpecifySetObservation(SpecifyPairObservation(KeyElement, ValueElement), MaxNum, AttentionEncodingSize, AttentionHeadNum, ValueEncodingSize);
+}
+
+FLearningAgentsObservationSchemaElement ULearningAgentsObservationSchema::SpecifyEnumObservation(const UEnum* Enum, const FName Name)
+{
+	if (!Enum)
+	{
+		UE_LOG(LogLearning, Error, TEXT("%s: Enum is nullptr."), *GetName());
+		return FLearningAgentsObservationSchemaElement();
+	}
+
+	return SpecifyContinuousObservation(Enum->NumEnums() - 1, Name);
+}
+
+FLearningAgentsObservationSchemaElement ULearningAgentsObservationSchema::SpecifyBitmaskObservation(const UEnum* Enum, const FName Name)
+{
+	if (!Enum)
+	{
+		UE_LOG(LogLearning, Error, TEXT("%s: Enum is nullptr."), *GetName());
+		return FLearningAgentsObservationSchemaElement();
+	}
+
+	if (Enum->NumEnums() - 1 > 32)
+	{
+		UE_LOG(LogLearning, Error, TEXT("%s: Too many values in Enum to use as Bitmask (%i)."), *GetName(), Enum->NumEnums() - 1);
+		return FLearningAgentsObservationSchemaElement();
+	}
+
+	return SpecifyContinuousObservation(Enum->NumEnums() - 1, Name);
+}
+
+FLearningAgentsObservationSchemaElement ULearningAgentsObservationSchema::SpecifyOptionalObservation(const FLearningAgentsObservationSchemaElement Element, const int32 EncodingSize, const FName Name)
+{
+	return SpecifyExclusiveUnionObservationFromArrayViews({ TEXT("Null"), TEXT("Valid") }, { SpecifyNullObservation(), Element }, EncodingSize, Name);
+}
+
+FLearningAgentsObservationSchemaElement ULearningAgentsObservationSchema::SpecifyEitherObservation(const FLearningAgentsObservationSchemaElement A, const FLearningAgentsObservationSchemaElement B, const int32 EncodingSize, const FName Name)
+{
+	return SpecifyExclusiveUnionObservationFromArrayViews({ TEXT("A"), TEXT("B") }, { A, B }, EncodingSize, Name);
+}
+
+FLearningAgentsObservationSchemaElement ULearningAgentsObservationSchema::SpecifyEncodingObservation(const FLearningAgentsObservationSchemaElement Element, const int32 EncodingSize, const FName Name)
+{
+	if (EncodingSize < 1)
+	{
+		UE_LOG(LogLearning, Error, TEXT("%s: Invalid Observation EncodingSize '%i'."), *GetName(), EncodingSize);
+		return FLearningAgentsObservationSchemaElement();
+	}
+
+	if (!ObservationSchema.IsValid(Element.SchemaElement))
+	{
+		UE_LOG(LogLearning, Error, TEXT("%s: Invalid Observation Object."), *GetName());
+		return FLearningAgentsObservationSchemaElement();
+	}
+
+	return { ObservationSchema.CreateEncoding({ Element.SchemaElement, EncodingSize }, Name) };
+}
+
+FLearningAgentsObservationSchemaElement ULearningAgentsObservationSchema::SpecifyBoolObservation(const FName Name)
+{
+	return SpecifyContinuousObservation(1, Name);
+}
+
+FLearningAgentsObservationSchemaElement ULearningAgentsObservationSchema::SpecifyFloatObservation(const FName Name)
+{
+	return SpecifyContinuousObservation(1, Name);
+}
+
+FLearningAgentsObservationSchemaElement ULearningAgentsObservationSchema::SpecifyTranslationObservation(const FName Name)
+{
+	return SpecifyContinuousObservation(3, Name);
+}
+
+FLearningAgentsObservationSchemaElement ULearningAgentsObservationSchema::SpecifyRotationObservation(const FName Name)
+{
+	return SpecifyContinuousObservation(6, Name);
+}
+
+FLearningAgentsObservationSchemaElement ULearningAgentsObservationSchema::SpecifyScaleObservation(const FName Name)
+{
+	return SpecifyContinuousObservation(3, Name);
+}
+
+FLearningAgentsObservationSchemaElement ULearningAgentsObservationSchema::SpecifyTransformObservation(const FName Name)
+{
+	return SpecifyStructObservationFromArrayViews(
+		{
+			TEXT("Translation"),
+			TEXT("Rotation"),
+			TEXT("Scale")
+		},
+		{
+			SpecifyTranslationObservation(),
+			SpecifyRotationObservation(),
+			SpecifyScaleObservation()
+		}, 
+		Name);
+}
+
+FLearningAgentsObservationSchemaElement ULearningAgentsObservationSchema::SpecifyAngleObservation(const FName Name)
+{
+	return SpecifyContinuousObservation(2, Name);
+}
+
+FLearningAgentsObservationSchemaElement ULearningAgentsObservationSchema::SpecifyVelocityObservation(const FName Name)
+{
+	return SpecifyContinuousObservation(3, Name);
+}
+
+FLearningAgentsObservationSchemaElement ULearningAgentsObservationSchema::SpecifyDirectionObservation(const FName Name)
+{
+	return SpecifyContinuousObservation(3, Name);
+}
+
+FLearningAgentsObservationSchemaElement ULearningAgentsObservationSchema::SpecifyLocationAlongSplineObservation(const FName Name)
+{
+	return SpecifyTranslationObservation(Name);
+}
+
+FLearningAgentsObservationSchemaElement ULearningAgentsObservationSchema::SpecifyProportionAlongSplineObservation(const FName Name)
+{
+	return SpecifyExclusiveUnionObservationFromArrayViews(
+		{
+			TEXT("Angle"),
+			TEXT("Proportion")
+		},
+		{
+			SpecifyAngleObservation(),
+			SpecifyFloatObservation()
+		},
+		8,
+		Name);
+}
+
+FLearningAgentsObservationSchemaElement ULearningAgentsObservationSchema::SpecifyDirectionAlongSplineObservation(const FName Name)
+{
+	return SpecifyDirectionObservation(Name);
+}
+
+FLearningAgentsObservationSchemaElement ULearningAgentsObservationSchema::SpecifyPropertiesAlongSplineObservation(const FName Name)
+{
+	return SpecifyStructObservationFromArrayViews(
+		{
+			TEXT("Location"),
+			TEXT("Proportion"),
+			TEXT("Direction")
+		},
+		{
+			SpecifyLocationAlongSplineObservation(),
+			SpecifyProportionAlongSplineObservation(),
+			SpecifyDirectionAlongSplineObservation(),
+		},
+		Name);
+}
+
+FLearningAgentsObservationSchemaElement ULearningAgentsObservationSchema::SpecifyProportionAlongRayObservation(const FName Name)
+{
+	return SpecifyFloatObservation(Name);
+}
+
+FLearningAgentsObservationSchemaElement ULearningAgentsObservationSchema::SpecifyProportionAlongRaysObservation(const int32 Num, const FName Name)
+{
+	return SpecifyStaticArrayObservation(SpecifyProportionAlongRayObservation(), Num, Name);
+}
+
+const UE::Learning::Observation::FObject& ULearningAgentsObservationObject::GetObservationObject() const
+{
+	return ObservationObject;
+}
+
+UE::Learning::Observation::FObject& ULearningAgentsObservationObject::GetObservationObject()
+{
+	return ObservationObject;
+}
+
+void ULearningAgentsObservationObject::LogObservation(const FLearningAgentsObservationObjectElement Element)
+{
+	UE::Learning::Agents::Observation::Private::LogObservation(GetObservationObject(), Element.ObjectElement, TEXT(""), TEXT(""), GetName());
+}
+
+FLearningAgentsObservationObjectElement ULearningAgentsObservationObject::MakeNullObservation(const FName Name)
+{
+	return { ObservationObject.CreateNull(Name) };
+}
+
+FLearningAgentsObservationObjectElement ULearningAgentsObservationObject::MakeContinuousObservation(const TArray<float>& Values, const FName Name)
+{
+	return MakeContinuousObservationFromArrayView(Values, Name);
+}
+
+FLearningAgentsObservationObjectElement ULearningAgentsObservationObject::MakeContinuousObservationFromArrayView(const TArrayView<const float> Values, const FName Name)
+{
+	if (Values.Num() == 0)
+	{
+		UE_LOG(LogLearning, Warning, TEXT("%s: Creating zero-sized Continuous Observation."), *GetName());
+	}
+
+	return { ObservationObject.CreateContinuous({ Values }, Name) };
+}
+
+FLearningAgentsObservationObjectElement ULearningAgentsObservationObject::MakeExclusiveDiscreteObservation(const int32 DiscreteIndex, const int32 Size, const FName Name)
+{
+	if (DiscreteIndex < 0 || DiscreteIndex >= Size)
+	{
+		UE_LOG(LogLearning, Error, TEXT("%s: Discrete index out of range: Got %i, expected <= %i."), *GetName(), DiscreteIndex, Size);
+		return FLearningAgentsObservationObjectElement();
+	}
+
+	TArray<float, TInlineAllocator<32>> Values;
+	Values.Init(0.0f, Size);
+	Values[DiscreteIndex] = 1.0f;
+	return MakeContinuousObservationFromArrayView(Values, Name);
+}
+
+FLearningAgentsObservationObjectElement ULearningAgentsObservationObject::MakeInclusiveDiscreteObservation(const TArray<int32>& DiscreteIndices, const int32 Size, const FName Name)
+{
+	return MakeInclusiveDiscreteObservationFromArrayView(DiscreteIndices, Size, Name);
+}
+
+FLearningAgentsObservationObjectElement ULearningAgentsObservationObject::MakeInclusiveDiscreteObservationFromArrayView(const TArrayView<const int32> DiscreteIndices, const int32 Size, const FName Name)
+{
+	if (UE::Learning::Agents::Observation::Private::ContainsDuplicates(DiscreteIndices))
+	{
+		UE_LOG(LogLearning, Error, TEXT("%s: Indices contain duplicates."), *GetName());
+		return FLearningAgentsObservationObjectElement();
+	}
+
+	TArray<float, TInlineAllocator<32>> Values;
+	Values.Init(0.0f, Size);
+
+	for (int32 Idx = 0; Idx < DiscreteIndices.Num(); Idx++)
+	{
+		if (DiscreteIndices[Idx] < 0 || DiscreteIndices[Idx] >= Size)
 		{
-			UE_LOG(LogLearning, Error, TEXT("%s: Must be attached to a LearningAgentsManager Actor."), *InInteractor->GetName());
-			return nullptr;
+			UE_LOG(LogLearning, Error, TEXT("%s: Discrete index out of range: Got %i, expected <= %i."), *GetName(), DiscreteIndices[Idx], Size);
+			return FLearningAgentsObservationObjectElement();
 		}
 
-		const FName UniqueName = MakeUniqueObjectName(InInteractor, ObservationUObject::StaticClass(), Name, EUniqueObjectNameOptions::GloballyUnique);
-
-		ObservationUObject* Observation = NewObject<ObservationUObject>(InInteractor, UniqueName);
+		Values[DiscreteIndices[Idx]] = 1.0f;
+	}
 
-		Observation->Init(InInteractor->GetAgentManager()->GetMaxAgentNum());
-		Observation->Interactor = InInteractor;
-		Observation->FeatureObject = MakeShared<ObservationFObject>(
-			Observation->GetFName(),
-			InInteractor->GetAgentManager()->GetInstanceData().ToSharedRef(),
-			InInteractor->GetAgentManager()->GetMaxAgentNum(),
-			Forward<InArgTypes>(Args)...);
+	return MakeContinuousObservationFromArrayView(Values, Name);
+}
 
-		// We assume all supported observation feature objects can be encoded
-		UE_LEARNING_CHECK(Observation->FeatureObject->IsEncodable());
+FLearningAgentsObservationObjectElement ULearningAgentsObservationObject::MakeIndexObservation(const int32 Index, const int32 Size, const FName Name)
+{
+	if (Index < 0 || Index >= Size)
+	{
+		UE_LOG(LogLearning, Error, TEXT("%s: Discrete index out of range: Got %i, expected <= %i."), *GetName(), Index, Size);
+		return FLearningAgentsObservationObjectElement();
+	}
 
-		InInteractor->AddObservation(Observation, Observation->FeatureObject.ToSharedRef());
+	TArray<float, TInlineAllocator<32>> Values;
+	Values.Init(0.0f, Size);
 
-		return Observation;
+	for (int32 Idx = 0; Idx < Index; Idx++)
+	{
+		Values[Idx] = 1.0f;
 	}
-}
-
-//------------------------------------------------------------------
 
-void ULearningAgentsObservation::Init(const int32 MaxAgentNum)
-{
-	AgentIteration.SetNumUninitialized({ MaxAgentNum });
-	UE::Learning::Array::Set<1, uint64>(AgentIteration, INDEX_NONE);
+	return MakeContinuousObservationFromArrayView(Values, Name);
 }
 
-void ULearningAgentsObservation::OnAgentsAdded(const TArray<int32>& AgentIds)
+FLearningAgentsObservationObjectElement ULearningAgentsObservationObject::MakeStructObservation(const TMap<FName, FLearningAgentsObservationObjectElement>& Elements, const FName Name)
 {
-	UE::Learning::Array::Set<1, uint64>(AgentIteration, 0, AgentIds);
-}
+	if (Elements.Num() == 0)
+	{
+		UE_LOG(LogLearning, Warning, TEXT("%s: Creating zero-sized Struct Observation."), *GetName());
+	}
 
-void ULearningAgentsObservation::OnAgentsRemoved(const TArray<int32>& AgentIds)
-{
-	UE::Learning::Array::Set<1, uint64>(AgentIteration, INDEX_NONE, AgentIds);
-}
+	const int32 SubElementNum = Elements.Num();
 
-void ULearningAgentsObservation::OnAgentsReset(const TArray<int32>& AgentIds)
-{
-	UE::Learning::Array::Set<1, uint64>(AgentIteration, 0, AgentIds);
-}
+	TArray<FName, TInlineAllocator<16>> SubElementNames;
+	TArray<FLearningAgentsObservationObjectElement, TInlineAllocator<16>> SubElements;
+	SubElementNames.Empty(SubElementNum);
+	SubElements.Empty(SubElementNum);
 
-uint64 ULearningAgentsObservation::GetAgentIteration(const int32 AgentId) const
-{
-	return AgentIteration[AgentId];
-}
+	for (const TPair<FName, FLearningAgentsObservationObjectElement>& Element : Elements)
+	{
+		SubElementNames.Add(Element.Key);
+		SubElements.Add(Element.Value);
+	}
 
-//------------------------------------------------------------------
+	return MakeStructObservationFromArrayViews(SubElementNames, SubElements);
+}
 
-UFloatObservation* UFloatObservation::AddFloatObservation(ULearningAgentsInteractor* InInteractor, const FName Name, const float Scale)
+FLearningAgentsObservationObjectElement ULearningAgentsObservationObject::MakeStructObservationFromArrays(const TArray<FName>& ElementNames, const TArray<FLearningAgentsObservationObjectElement>& Elements, const FName Name)
 {
-	return UE::Learning::Agents::Observations::Private::AddObservation<UFloatObservation, UE::Learning::FFloatFeature>(InInteractor, Name, TEXT("AddFloatObservation"), 1, Scale);
+	return MakeStructObservationFromArrayViews(ElementNames, Elements, Name);
 }
 
-void UFloatObservation::SetFloatObservation(const int32 AgentId, const float Value)
+FLearningAgentsObservationObjectElement ULearningAgentsObservationObject::MakeStructObservationFromArrayViews(const TArrayView<const FName> ElementNames, const TArrayView<const FLearningAgentsObservationObjectElement> Elements, const FName Name)
 {
-	if (!Interactor->HasAgent(AgentId))
+	if (Elements.Num() == 0)
 	{
-		UE_LOG(LogLearning, Error, TEXT("%s: AgentId %d not found in the agents set."), *GetName(), AgentId);
-		return;
+		UE_LOG(LogLearning, Warning, TEXT("%s: Creating zero-sized Struct Observation."), *GetName());
 	}
 
-	FeatureObject->InstanceData->View(FeatureObject->ValueHandle)[AgentId][0] = Value;
-	AgentIteration[AgentId]++;
-}
+	if (Elements.Num() != ElementNames.Num())
+	{
+		UE_LOG(LogLearning, Error, TEXT("%s: Number of elements (%i) must match number of names (%i)."), *GetName(), Elements.Num(), ElementNames.Num());
+		return FLearningAgentsObservationObjectElement();
+	}
 
-#if UE_LEARNING_AGENTS_ENABLE_VISUAL_LOG
-void UFloatObservation::VisualLog(const UE::Learning::FIndexSet Instances) const
-{
-	UE_LEARNING_TRACE_CPUPROFILER_EVENT_SCOPE(UFloatObservation::VisualLog);
+	if (UE::Learning::Agents::Observation::Private::ContainsDuplicates(ElementNames))
+	{
+		UE_LOG(LogLearning, Error, TEXT("%s: Element Names contain duplicates."), *GetName());
+		return FLearningAgentsObservationObjectElement();
+	}
 
-	const TLearningArrayView<2, const float> ValueView = FeatureObject->InstanceData->ConstView(FeatureObject->ValueHandle);
-	const TLearningArrayView<2, const float> FeatureView = FeatureObject->InstanceData->ConstView(FeatureObject->FeatureHandle);
+	TArray<UE::Learning::Observation::FObjectElement, TInlineAllocator<16>> SubElements;
+	SubElements.Empty(Elements.Num());
 
-	for (const int32 Instance : Instances)
+	for (const FLearningAgentsObservationObjectElement& Element : Elements)
 	{
-		if (const AActor* Actor = Cast<AActor>(Interactor->GetAgent(Instance)))
+		if (!ObservationObject.IsValid(Element.ObjectElement))
 		{
-			UE_LEARNING_AGENTS_VLOG_STRING(this, LogLearning, Display,
-				Actor->GetActorLocation(),
-				VisualLogColor.ToFColor(true),
-				TEXT("Agent %i\nScale: [% 6.2f]\nValue: %s\nEncoded: %s"),
-				Instance,
-				FeatureObject->Scale,
-				*UE::Learning::Array::FormatFloat(ValueView[Instance]),
-				*UE::Learning::Array::FormatFloat(FeatureView[Instance]));
+			UE_LOG(LogLearning, Error, TEXT("%s: Invalid Observation Object."), *GetName());
+			return FLearningAgentsObservationObjectElement();
 		}
+
+		SubElements.Add(Element.ObjectElement);
 	}
-}
-#endif
 
-//------------------------------------------------------------------
+	return { ObservationObject.CreateAnd({ ElementNames, SubElements }, Name)};
+}
 
-UFloatArrayObservation* UFloatArrayObservation::AddFloatArrayObservation(ULearningAgentsInteractor* InInteractor, const FName Name, const int32 Num, const float Scale)
+FLearningAgentsObservationObjectElement ULearningAgentsObservationObject::MakeExclusiveUnionObservation(const FName ElementName, const FLearningAgentsObservationObjectElement Element, const FName Name)
 {
-	if (Num < 1)
+	if (!ObservationObject.IsValid(Element.ObjectElement))
 	{
-		UE_LOG(LogLearning, Error, TEXT("AddFloatArrayObservation: Number of elements in array must be at least 1, got %i."), Num);
-		return nullptr;
+		UE_LOG(LogLearning, Error, TEXT("%s: Invalid Observation Object."), *GetName());
+		return FLearningAgentsObservationObjectElement();
 	}
 
-	return UE::Learning::Agents::Observations::Private::AddObservation<UFloatArrayObservation, UE::Learning::FFloatFeature>(InInteractor, Name, TEXT("AddFloatArrayObservation"), Num, Scale);
+	return { ObservationObject.CreateOrExclusive({ ElementName, Element.ObjectElement }, Name) };
 }
 
-void UFloatArrayObservation::SetFloatArrayObservation(const int32 AgentId, const TArray<float>& Values)
+FLearningAgentsObservationObjectElement ULearningAgentsObservationObject::MakeInclusiveUnionObservation(const TMap<FName, FLearningAgentsObservationObjectElement>& Elements, const FName Name)
 {
-	if (!Interactor->HasAgent(AgentId))
-	{
-		UE_LOG(LogLearning, Error, TEXT("%s: AgentId %d not found in the agents set."), *GetName(), AgentId);
-		return;
-	}
+	const int32 SubElementNum = Elements.Num();
 
-	const TLearningArrayView<2, float> View = FeatureObject->InstanceData->View(FeatureObject->ValueHandle);
+	TArray<FName, TInlineAllocator<16>> SubElementNames;
+	TArray<FLearningAgentsObservationObjectElement, TInlineAllocator<16>> SubElements;
+	SubElementNames.Empty(SubElementNum);
+	SubElements.Empty(SubElementNum);
 
-	if (Values.Num() != View.Num<1>())
+	for (const TPair<FName, FLearningAgentsObservationObjectElement>& Element : Elements)
 	{
-		UE_LOG(LogLearning, Error, TEXT("%s: Got wrong number of elements in array. Expected %i, got %i."), *GetName(), View.Num<1>(), Values.Num());
-		return;
+		SubElementNames.Add(Element.Key);
+		SubElements.Add(Element.Value);
 	}
 
-	UE::Learning::Array::Copy<1, float>(View[AgentId], Values);
-	AgentIteration[AgentId]++;
+	return MakeInclusiveUnionObservationFromArrayViews(SubElementNames, SubElements, Name);
 }
 
-#if UE_LEARNING_AGENTS_ENABLE_VISUAL_LOG
-void UFloatArrayObservation::VisualLog(const UE::Learning::FIndexSet Instances) const
+FLearningAgentsObservationObjectElement ULearningAgentsObservationObject::MakeInclusiveUnionObservationFromArrays(const TArray<FName>& ElementNames, const TArray<FLearningAgentsObservationObjectElement>& Elements, const FName Name)
 {
-	UE_LEARNING_TRACE_CPUPROFILER_EVENT_SCOPE(UFloatArrayObservation::VisualLog);
+	return MakeInclusiveUnionObservationFromArrayViews(ElementNames, Elements, Name);
+}
+
+FLearningAgentsObservationObjectElement ULearningAgentsObservationObject::MakeInclusiveUnionObservationFromArrayViews(const TArrayView<const FName> ElementNames, const TArrayView<const FLearningAgentsObservationObjectElement> Elements, const FName Name)
+{
+	if (Elements.Num() != ElementNames.Num())
+	{
+		UE_LOG(LogLearning, Error, TEXT("%s: Number of elements (%i) must match number of names (%i)."), *GetName(), Elements.Num(), ElementNames.Num());
+		return FLearningAgentsObservationObjectElement();
+	}
+
+	if (UE::Learning::Agents::Observation::Private::ContainsDuplicates(ElementNames))
+	{
+		UE_LOG(LogLearning, Error, TEXT("%s: Element Names contain duplicates."), *GetName());
+		return FLearningAgentsObservationObjectElement();
+	}
 
-	const TLearningArrayView<2, const float> ValueView = FeatureObject->InstanceData->ConstView(FeatureObject->ValueHandle);
-	const TLearningArrayView<2, const float> FeatureView = FeatureObject->InstanceData->ConstView(FeatureObject->FeatureHandle);
+	TArray<UE::Learning::Observation::FObjectElement, TInlineAllocator<16>> SubElements;
+	SubElements.Empty(Elements.Num());
 
-	for (const int32 Instance : Instances)
+	for (const FLearningAgentsObservationObjectElement& Element : Elements)
 	{
-		if (const AActor* Actor = Cast<AActor>(Interactor->GetAgent(Instance)))
+		if (!ObservationObject.IsValid(Element.ObjectElement))
 		{
-			UE_LEARNING_AGENTS_VLOG_STRING(this, LogLearning, Display,
-				Actor->GetActorLocation(),
-				VisualLogColor.ToFColor(true),
-				TEXT("Agent %i\nScale: [% 6.2f]\nValue: %s\nEncoded: %s"),
-				Instance,
-				FeatureObject->Scale,
-				*UE::Learning::Array::FormatFloat(ValueView[Instance]),
-				*UE::Learning::Array::FormatFloat(FeatureView[Instance]));
+			UE_LOG(LogLearning, Error, TEXT("%s: Invalid Observation Object."), *GetName());
+			return FLearningAgentsObservationObjectElement();
 		}
+
+		SubElements.Add(Element.ObjectElement);
 	}
-}
-#endif
 
-//------------------------------------------------------------------
+	return { ObservationObject.CreateOrInclusive({ ElementNames, SubElements }, Name) };
+}
 
-UVectorObservation* UVectorObservation::AddVectorObservation(ULearningAgentsInteractor* InInteractor, const FName Name, const float Scale)
+FLearningAgentsObservationObjectElement ULearningAgentsObservationObject::MakeStaticArrayObservation(const TArray<FLearningAgentsObservationObjectElement>& Elements, const FName Name)
 {
-	return UE::Learning::Agents::Observations::Private::AddObservation<UVectorObservation, UE::Learning::FFloatFeature>(InInteractor, Name, TEXT("AddVectorObservation"), 3, Scale);
+	return MakeStaticArrayObservationFromArrayView(Elements, Name);
 }
 
-void UVectorObservation::SetVectorObservation(const int32 AgentId, const FVector Vector)
+FLearningAgentsObservationObjectElement ULearningAgentsObservationObject::MakeStaticArrayObservationFromArrayView(const TArrayView<const FLearningAgentsObservationObjectElement> Elements, const FName Name)
 {
-	if (!Interactor->HasAgent(AgentId))
+	if (Elements.Num() == 0)
 	{
-		UE_LOG(LogLearning, Error, TEXT("%s: AgentId %d not found in the agents set."), *GetName(), AgentId);
-		return;
+		UE_LOG(LogLearning, Warning, TEXT("%s: Creating zero-sized Static Array Observation."), *GetName());
 	}
 
-	const TLearningArrayView<2, float> View = FeatureObject->InstanceData->View(FeatureObject->ValueHandle);
+	TArray<UE::Learning::Observation::FObjectElement, TInlineAllocator<16>> SubElements;
+	SubElements.Empty(Elements.Num());
+
+	for (const FLearningAgentsObservationObjectElement& Element : Elements)
+	{
+		if (!ObservationObject.IsValid(Element.ObjectElement))
+		{
+			UE_LOG(LogLearning, Error, TEXT("%s: Invalid Observation Object."), *GetName());
+			return FLearningAgentsObservationObjectElement();
+		}
+
+		SubElements.Add(Element.ObjectElement);
+	}
 
-	View[AgentId][0] = Vector.X;
-	View[AgentId][1] = Vector.Y;
-	View[AgentId][2] = Vector.Z;
-	AgentIteration[AgentId]++;
+	return { ObservationObject.CreateArray({ SubElements }, Name) };
 }
 
-#if UE_LEARNING_AGENTS_ENABLE_VISUAL_LOG
-void UVectorObservation::VisualLog(const UE::Learning::FIndexSet Instances) const
+FLearningAgentsObservationObjectElement ULearningAgentsObservationObject::MakeSetObservation(const TSet<FLearningAgentsObservationObjectElement>& Elements, const FName Name)
 {
-	UE_LEARNING_TRACE_CPUPROFILER_EVENT_SCOPE(UVectorObservation::VisualLog);
-
-	const TLearningArrayView<2, const float> ValueView = FeatureObject->InstanceData->ConstView(FeatureObject->ValueHandle);
-	const TLearningArrayView<2, const float> FeatureView = FeatureObject->InstanceData->ConstView(FeatureObject->FeatureHandle);
+	TArray<UE::Learning::Observation::FObjectElement, TInlineAllocator<16>> SubElements;
+	SubElements.Empty(Elements.Num());
 
-	for (const int32 Instance : Instances)
+	for (const FLearningAgentsObservationObjectElement& Element : Elements)
 	{
-		if (const AActor* Actor = Cast<AActor>(Interactor->GetAgent(Instance)))
+		if (!ObservationObject.IsValid(Element.ObjectElement))
 		{
-			const FVector Vector(ValueView[Instance][0], ValueView[Instance][1], ValueView[Instance][2]);
-
-			UE_LEARNING_AGENTS_VLOG_ARROW(this, LogLearning, Display,
-				Actor->GetActorLocation(),
-				Actor->GetActorLocation() + Vector,
-				VisualLogColor.ToFColor(true),
-				TEXT(""));
-
-			UE_LEARNING_AGENTS_VLOG_STRING(this, LogLearning, Display,
-				Actor->GetActorLocation() + Vector,
-				VisualLogColor.ToFColor(true),
-				TEXT("Vector: [% 6.4f % 6.4f % 6.4f]"),
-				Vector.X, Vector.Y, Vector.Z);
-
-			UE_LEARNING_AGENTS_VLOG_STRING(this, LogLearning, Display,
-				Actor->GetActorLocation(),
-				VisualLogColor.ToFColor(true),
-				TEXT("Agent %i\nScale: [% 6.2f]\nEncoded: %s"),
-				Instance,
-				FeatureObject->Scale,
-				*UE::Learning::Array::FormatFloat(FeatureView[Instance]));
+			UE_LOG(LogLearning, Error, TEXT("%s: Invalid Observation Object."), *GetName());
+			return FLearningAgentsObservationObjectElement();
 		}
+
+		SubElements.Add(Element.ObjectElement);
 	}
+
+	return { ObservationObject.CreateSet({ SubElements }, Name) };
 }
-#endif
 
-UVectorArrayObservation* UVectorArrayObservation::AddVectorArrayObservation(ULearningAgentsInteractor* InInteractor, const FName Name, const int32 Num, const float Scale)
+FLearningAgentsObservationObjectElement ULearningAgentsObservationObject::MakeSetObservationFromArray(const TArray<FLearningAgentsObservationObjectElement>& Elements, const FName Name)
 {
-	if (Num < 1)
+	return MakeSetObservationFromArrayView(Elements, Name);
+}
+
+FLearningAgentsObservationObjectElement ULearningAgentsObservationObject::MakeSetObservationFromArrayView(const TArrayView<const FLearningAgentsObservationObjectElement> Elements, const FName Name)
+{
+	TArray<UE::Learning::Observation::FObjectElement, TInlineAllocator<16>> SubElements;
+	SubElements.Empty(Elements.Num());
+
+	for (const FLearningAgentsObservationObjectElement& Element : Elements)
 	{
-		UE_LOG(LogLearning, Error, TEXT("AddVectorArrayObservation: Number of elements in array must be at least 1, got %i."), Num);
-		return nullptr;
+		if (!ObservationObject.IsValid(Element.ObjectElement))
+		{
+			UE_LOG(LogLearning, Error, TEXT("%s: Invalid Observation Object."), *GetName());
+			return FLearningAgentsObservationObjectElement();
+		}
+
+		SubElements.Add(Element.ObjectElement);
 	}
 
-	return UE::Learning::Agents::Observations::Private::AddObservation<UVectorArrayObservation, UE::Learning::FFloatFeature>(InInteractor, Name, TEXT("AddVectorArrayObservation"), Num * 3, Scale);
+	return { ObservationObject.CreateSet({ SubElements }, Name) };
 }
 
-void UVectorArrayObservation::SetVectorArrayObservation(const int32 AgentId, const TArray<FVector>& Vectors)
+FLearningAgentsObservationObjectElement ULearningAgentsObservationObject::MakePairObservation(const FLearningAgentsObservationObjectElement Key, const FLearningAgentsObservationObjectElement Value, const FName Name)
 {
-	if (!Interactor->HasAgent(AgentId))
-	{
-		UE_LOG(LogLearning, Error, TEXT("%s: AgentId %d not found in the agents set."), *GetName(), AgentId);
-		return;
-	}
+	return MakeStructObservationFromArrayViews({ TEXT("Key"), TEXT("Value") }, { Key, Value }, Name);
+}
+
+FLearningAgentsObservationObjectElement ULearningAgentsObservationObject::MakeArrayObservation(const TArray<FLearningAgentsObservationObjectElement>& Elements, const FName Name)
+{
+	return MakeArrayObservationFromArrayView(Elements, Name);
+}
 
-	const TLearningArrayView<2, float> View = FeatureObject->InstanceData->View(FeatureObject->ValueHandle);
+FLearningAgentsObservationObjectElement ULearningAgentsObservationObject::MakeArrayObservationFromArrayView(const TArrayView<const FLearningAgentsObservationObjectElement> Elements, const FName Name)
+{
+	TArray<FLearningAgentsObservationObjectElement, TInlineAllocator<16>> SubElements;
+	SubElements.Empty(Elements.Num());
 
-	if (Vectors.Num() != View.Num<1>() / 3)
+	for (int32 ElementIdx = 0; ElementIdx < Elements.Num(); ElementIdx++)
 	{
-		UE_LOG(LogLearning, Error, TEXT("%s: Got wrong number of elements in array. Expected %i, got %i."), *GetName(), View.Num<1>() / 3, Vectors.Num());
-		return;
-	}
+		const FLearningAgentsObservationObjectElement Element = Elements[ElementIdx];
 
-	for (int32 VectorIdx = 0; VectorIdx < Vectors.Num(); VectorIdx++)
-	{
-		View[AgentId][VectorIdx * 3 + 0] = Vectors[VectorIdx].X;
-		View[AgentId][VectorIdx * 3 + 1] = Vectors[VectorIdx].Y;
-		View[AgentId][VectorIdx * 3 + 2] = Vectors[VectorIdx].Z;
+		if (!ObservationObject.IsValid(Element.ObjectElement))
+		{
+			UE_LOG(LogLearning, Error, TEXT("%s: Invalid Observation Object."), *GetName());
+			return FLearningAgentsObservationObjectElement();
+		}
+
+		SubElements.Add(MakePairObservation(MakeIndexObservation(ElementIdx, Elements.Num()), Element));
 	}
 
-	AgentIteration[AgentId]++;
+	return MakeSetObservationFromArrayView(SubElements);
 }
 
-#if UE_LEARNING_AGENTS_ENABLE_VISUAL_LOG
-void UVectorArrayObservation::VisualLog(const UE::Learning::FIndexSet Instances) const
+FLearningAgentsObservationObjectElement ULearningAgentsObservationObject::MakeMapObservation(const TMap<FLearningAgentsObservationObjectElement, FLearningAgentsObservationObjectElement>& Map, const FName Name)
 {
-	UE_LEARNING_TRACE_CPUPROFILER_EVENT_SCOPE(UVectorArrayObservation::VisualLog);
-
-	const TLearningArrayView<2, const float> ValueView = FeatureObject->InstanceData->ConstView(FeatureObject->ValueHandle);
-	const TLearningArrayView<2, const float> FeatureView = FeatureObject->InstanceData->ConstView(FeatureObject->FeatureHandle);
+	TArray<FLearningAgentsObservationObjectElement, TInlineAllocator<16>> SubElements;
+	SubElements.Empty(Map.Num());
 
-	const int32 VectorNum = ValueView.Num<1>() / 3;
-
-	for (const int32 Instance : Instances)
+	for (TPair<FLearningAgentsObservationObjectElement, FLearningAgentsObservationObjectElement> Item : Map)
 	{
-		if (const AActor* Actor = Cast<AActor>(Interactor->GetAgent(Instance)))
+		if (!ObservationObject.IsValid(Item.Key.ObjectElement) || !ObservationObject.IsValid(Item.Value.ObjectElement))
 		{
-			for (int32 VectorIdx = 0; VectorIdx < VectorNum; VectorIdx++)
-			{
-				const FVector Vector(ValueView[Instance][VectorIdx * 3 + 0], ValueView[Instance][VectorIdx * 3 + 1], ValueView[Instance][VectorIdx * 3 + 2]);
-				const FVector Offset = UE::Learning::Agents::Debug::GridOffsetForIndex(VectorIdx, VectorNum);
-
-				UE_LEARNING_AGENTS_VLOG_ARROW(this, LogLearning, Display,
-					Actor->GetActorLocation() + Offset,
-					Actor->GetActorLocation() + Offset + Vector,
-					VisualLogColor.ToFColor(true),
-					TEXT(""));
-
-				UE_LEARNING_AGENTS_VLOG_STRING(this, LogLearning, Display,
-					Actor->GetActorLocation() + Offset + Vector,
-					VisualLogColor.ToFColor(true),
-					TEXT("Vector %i: [% 6.4f % 6.4f % 6.4f]"),
-					VectorIdx,
-					Vector.X, Vector.Y, Vector.Z);
-			}
-
-			UE_LEARNING_AGENTS_VLOG_STRING(this, LogLearning, Display,
-				Actor->GetActorLocation(),
-				VisualLogColor.ToFColor(true),
-				TEXT("Agent %i\nScale: [% 6.2f]\nEncoded: %s"),
-				Instance,
-				FeatureObject->Scale,
-				*UE::Learning::Array::FormatFloat(FeatureView[Instance]));
+			UE_LOG(LogLearning, Error, TEXT("%s: Invalid Observation Object."), *GetName());
+			return FLearningAgentsObservationObjectElement();
 		}
+
+		SubElements.Add(MakePairObservation(Item.Key, Item.Value));
 	}
+
+	return MakeSetObservationFromArrayView(SubElements);
 }
-#endif
 
-//------------------------------------------------------------------
+FLearningAgentsObservationObjectElement ULearningAgentsObservationObject::MakeMapObservationFromArrays(const TArray<FLearningAgentsObservationObjectElement>& Keys, const TArray<FLearningAgentsObservationObjectElement>& Values, const FName Name)
+{
+	return MakeMapObservationFromArrayViews(Keys, Values, Name);
+}
 
-UEnumObservation* UEnumObservation::AddEnumObservation(ULearningAgentsInteractor* InInteractor, const UEnum* EnumType, const FName Name)
+FLearningAgentsObservationObjectElement ULearningAgentsObservationObject::MakeMapObservationFromArrayViews(const TArrayView<const FLearningAgentsObservationObjectElement> Keys, const TArrayView<const FLearningAgentsObservationObjectElement> Values, const FName Name)
 {
-	if (!EnumType)
+	if (Keys.Num() != Values.Num())
 	{
-		UE_LOG(LogLearning, Error, TEXT("AddEnumObservation: Invalid Enum."));
-		return nullptr;
+		UE_LOG(LogLearning, Error, TEXT("%s: Number of keys (%i) must match number of values (%i)."), *GetName(), Keys.Num(), Values.Num());
+		return FLearningAgentsObservationObjectElement();
 	}
 
-	int32 NumEnums = EnumType->NumEnums() - 1;
-	if (NumEnums < 1)
+	TArray<FLearningAgentsObservationObjectElement, TInlineAllocator<16>> SubElements;
+	SubElements.Empty(Keys.Num());
+
+	for (int32 ElementIdx = 0; ElementIdx < Keys.Num(); ElementIdx++)
 	{
-		UE_LOG(LogLearning, Error, TEXT("AddEnumObservation: Enum requires at least one entry to be used as an observation."));
-		return nullptr;
+		if(!ObservationObject.IsValid(Keys[ElementIdx].ObjectElement) || !ObservationObject.IsValid(Values[ElementIdx].ObjectElement))
+		{
+			UE_LOG(LogLearning, Error, TEXT("%s: Invalid Observation Object."), *GetName());
+			return FLearningAgentsObservationObjectElement();
+		}
+
+		SubElements.Add(MakePairObservation(Keys[ElementIdx], Values[ElementIdx]));
 	}
 
-	UEnumObservation* Observation = UE::Learning::Agents::Observations::Private::AddObservation<UEnumObservation, UE::Learning::FFloatFeature>(InInteractor, Name, TEXT("AddEnumObservation"), NumEnums);
-	if (Observation) { Observation->Enum = EnumType; }
-	
-	return Observation;
+	return MakeSetObservationFromArrayView(SubElements);
 }
 
-void UEnumObservation::SetEnumObservation(const int32 AgentId, const uint8 EnumValue)
+FLearningAgentsObservationObjectElement ULearningAgentsObservationObject::MakeEnumObservation(const UEnum* Enum, const uint8 EnumValue, const FName Name)
 {
-	if (!Interactor->HasAgent(AgentId))
+	if (!Enum)
 	{
-		UE_LOG(LogLearning, Error, TEXT("%s: AgentId %d not found in the agents set."), *GetName(), AgentId);
-		return;
+		UE_LOG(LogLearning, Error, TEXT("%s: Enum is nullptr."), *GetName());
+		return FLearningAgentsObservationObjectElement();
 	}
 
-	const int32 Index = Enum->GetIndexByValue(EnumValue);
+	const int32 EnumValueIndex = Enum->GetIndexByValue(EnumValue);
 
-	if (Index == INDEX_NONE)
+	if (EnumValueIndex == INDEX_NONE || EnumValueIndex < 0 || EnumValueIndex >= Enum->NumEnums() - 1)
 	{
-		UE_LOG(LogLearning, Error, TEXT("%s: Invalid enum value: %i."), *GetName(), EnumValue);
-		return;
+		UE_LOG(LogLearning, Error, TEXT("%s: EnumValue %i not valid for Enum '%s'."), *GetName(), EnumValue , *Enum->GetName());
+		return FLearningAgentsObservationObjectElement();
 	}
 
-	const TLearningArrayView<2, float> EnumView = FeatureObject->InstanceData->View(FeatureObject->ValueHandle);
-	UE::Learning::Array::Zero(EnumView[AgentId]);
-	EnumView[AgentId][Index] = 1.0f;
+	TArray<float, TInlineAllocator<32>> OneHot;
+	OneHot.Init(0.0f, Enum->NumEnums() - 1);
+	OneHot[EnumValueIndex] = 1.0f;
 
-	AgentIteration[AgentId]++;
+	return MakeContinuousObservationFromArrayView(OneHot, Name);
 }
 
-#if UE_LEARNING_AGENTS_ENABLE_VISUAL_LOG
-void UEnumObservation::VisualLog(const UE::Learning::FIndexSet Instances) const
+FLearningAgentsObservationObjectElement ULearningAgentsObservationObject::MakeBitmaskObservation(const UEnum* Enum, const int32 BitmaskValue, const FName Name)
 {
-	UE_LEARNING_TRACE_CPUPROFILER_EVENT_SCOPE(UEnumObservation::VisualLog);
-
-	const TLearningArrayView<2, const float> ValueView = FeatureObject->InstanceData->ConstView(FeatureObject->ValueHandle);
-	const TLearningArrayView<2, const float> FeatureView = FeatureObject->InstanceData->ConstView(FeatureObject->FeatureHandle);
+	if (!Enum)
+	{
+		UE_LOG(LogLearning, Error, TEXT("%s: Enum is nullptr."), *GetName());
+		return FLearningAgentsObservationObjectElement();
+	}
 
-	const int32 EnumNum = ValueView.Num<1>();
-		
-	for (const int32 Instance : Instances)
+	if (Enum->NumEnums() - 1 > 32)
 	{
-		if (const AActor* Actor = Cast<AActor>(Interactor->GetAgent(Instance)))
-		{
-			int32 EnumEntryIdx = INDEX_NONE;
+		UE_LOG(LogLearning, Error, TEXT("%s: Too many values in Enum to use as Bitmask (%i)."), *GetName(), Enum->NumEnums() - 1);
+		return FLearningAgentsObservationObjectElement();
+	}
 
-			for (int32 EnumIdx = 0; EnumIdx < EnumNum; EnumIdx++)
-			{
-				if (ValueView[Instance][EnumIdx] != 0.0f)
-				{
-					EnumEntryIdx = EnumIdx;
-					break;
-				}
-			}
+	TArray<float, TInlineAllocator<32>> OneHot;
+	OneHot.Init(0.0f, Enum->NumEnums() - 1);
 
-			if (EnumEntryIdx != INDEX_NONE)
-			{
-				UE_LEARNING_AGENTS_VLOG_STRING(this, LogLearning, Display,
-					Actor->GetActorLocation(),
-					VisualLogColor.ToFColor(true),
-					TEXT("Agent %i\nValue: %i\nIndex: %i\nName: \"%s\"\nEncoded: %s"),
-					Instance,
-					Enum->GetValueByIndex(EnumEntryIdx),
-					EnumEntryIdx,
-					*Enum->GetDisplayNameTextByIndex(EnumEntryIdx).ToString(),
-					*UE::Learning::Array::FormatFloat(FeatureView[Instance]));
-			}
-			else
-			{
-				UE_LOG(LogLearning, Error, TEXT("Invalid Enum encoding."));
-			}
+	for (int32 EnumIdx = 0; EnumIdx < Enum->NumEnums() - 1; EnumIdx++)
+	{
+		if (BitmaskValue & (1 << EnumIdx))
+		{
+			OneHot[EnumIdx] = 1.0f;
 		}
 	}
+
+	return MakeContinuousObservationFromArrayView(OneHot, Name);
 }
-#endif
 
-UEnumArrayObservation* UEnumArrayObservation::AddEnumArrayObservation(ULearningAgentsInteractor* InInteractor, const UEnum* EnumType, const FName Name, const int32 EnumNum)
+FLearningAgentsObservationObjectElement ULearningAgentsObservationObject::MakeOptionalObservation(const FLearningAgentsObservationObjectElement Element, const ELearningAgentsOptionalObservation Option, const FName Name)
 {
-	if (!EnumType)
-	{
-		UE_LOG(LogLearning, Error, TEXT("AddEnumArrayObservation: Invalid Enum."));
-		return nullptr;
-	}
+	return MakeExclusiveUnionObservation(
+		Option == ELearningAgentsOptionalObservation::Null ? TEXT("Null") : TEXT("Valid"),
+		Option == ELearningAgentsOptionalObservation::Null ? MakeNullObservation() : Element,
+		Name);
+}
 
-	int32 NumEnums = EnumType->NumEnums() - 1;
-	if (NumEnums < 1)
-	{
-		UE_LOG(LogLearning, Error, TEXT("AddEnumArrayObservation: Enum requires at least one entry to be used as an observation."));
-		return nullptr;
-	}
+FLearningAgentsObservationObjectElement ULearningAgentsObservationObject::MakeOptionalNullObservation(const FName Name)
+{
+	return MakeExclusiveUnionObservation(TEXT("Null"), MakeNullObservation(), Name);
+}
 
-	if (EnumNum < 1)
-	{
-		UE_LOG(LogLearning, Error, TEXT("AddEnumArrayObservation: Number of elements in array must be at least 1, got %i."), EnumNum);
-		return nullptr;
-	}
+FLearningAgentsObservationObjectElement ULearningAgentsObservationObject::MakeOptionalValidObservation(const FLearningAgentsObservationObjectElement Element, const FName Name)
+{
+	return MakeExclusiveUnionObservation(TEXT("Valid"), Element, Name);
+}
+
+FLearningAgentsObservationObjectElement ULearningAgentsObservationObject::MakeEitherObservation(const FLearningAgentsObservationObjectElement Element, const ELearningAgentsEitherObservation Either, const FName Name)
+{
+	return MakeExclusiveUnionObservation(Either == ELearningAgentsEitherObservation::A ? TEXT("A") : TEXT("B"), Element, Name);
+}
 
-	UEnumArrayObservation* Observation = UE::Learning::Agents::Observations::Private::AddObservation<UEnumArrayObservation, UE::Learning::FFloatFeature>(InInteractor, Name, TEXT("AddEnumArrayObservation"), EnumNum * NumEnums);
-	if (Observation) { Observation->Enum = EnumType; }
+FLearningAgentsObservationObjectElement ULearningAgentsObservationObject::MakeEitherAObservation(const FLearningAgentsObservationObjectElement A, const FName Name)
+{
+	return MakeExclusiveUnionObservation(TEXT("A"), A, Name);
+}
 
-	return Observation;
+FLearningAgentsObservationObjectElement ULearningAgentsObservationObject::MakeEitherBObservation(const FLearningAgentsObservationObjectElement B, const FName Name)
+{
+	return MakeExclusiveUnionObservation(TEXT("B"), B, Name);
 }
 
-void UEnumArrayObservation::SetEnumArrayObservation(const int32 AgentId, const TArray<uint8>& EnumValues)
+FLearningAgentsObservationObjectElement ULearningAgentsObservationObject::MakeEncodingObservation(const FLearningAgentsObservationObjectElement Element, const FName Name)
 {
-	if (!Interactor->HasAgent(AgentId))
+	if (!ObservationObject.IsValid(Element.ObjectElement))
 	{
-		UE_LOG(LogLearning, Error, TEXT("%s: AgentId %d not found in the agents set."), *GetName(), AgentId);
-		return;
+		UE_LOG(LogLearning, Error, TEXT("%s: Invalid Observation Object."), *GetName());
+		return FLearningAgentsObservationObjectElement();
 	}
 
-	const TLearningArrayView<2, float> EnumView = FeatureObject->InstanceData->View(FeatureObject->ValueHandle);
+	return { ObservationObject.CreateEncoding({ Element.ObjectElement }, Name) };
+}
 
-	const int32 EnumEntryNum = Enum->NumEnums();
+FLearningAgentsObservationObjectElement ULearningAgentsObservationObject::MakeBoolObservation(const bool bValue, const FName Name)
+{
+	return MakeContinuousObservationFromArrayView({ bValue ? 1.0f : -1.0f }, Name);
+}
 
-	if (EnumValues.Num() != EnumView.Num<1>() / EnumEntryNum)
-	{
-		UE_LOG(LogLearning, Error, TEXT("%s: Got wrong number of elements in array. Expected %i, got %i."), *GetName(), EnumView.Num<1>() / EnumEntryNum, EnumValues.Num());
-		return;
-	}
+FLearningAgentsObservationObjectElement ULearningAgentsObservationObject::MakeFloatObservation(const float Value, const float FloatScale, const FName Name)
+{
+	return MakeContinuousObservationFromArrayView({ Value / FMath::Max(FloatScale, UE_SMALL_NUMBER) }, Name);
+}
 
-	for (int32 EnumIdx = 0; EnumIdx < EnumValues.Num(); EnumIdx++)
-	{
-		const TLearningArrayView<1, float> EnumViewSlice = EnumView[AgentId].Slice(EnumIdx * EnumEntryNum, EnumEntryNum);
-		UE::Learning::Array::Zero(EnumViewSlice);
+FLearningAgentsObservationObjectElement ULearningAgentsObservationObject::MakeTranslationObservation(const FVector Translation, const FTransform RelativeTransform, const float TranslationScale, const FName Name)
+{
+	const FVector LocalTranslation = RelativeTransform.InverseTransformPosition(Translation);
 
-		const int32 Index = Enum->GetIndexByValue(EnumValues[EnumIdx]);
+	return MakeContinuousObservationFromArrayView({
+		(float)LocalTranslation.X / FMath::Max(TranslationScale, UE_SMALL_NUMBER),
+		(float)LocalTranslation.Y / FMath::Max(TranslationScale, UE_SMALL_NUMBER),
+		(float)LocalTranslation.Z / FMath::Max(TranslationScale, UE_SMALL_NUMBER),
+		}, Name);
+}
 
-		if (Index == INDEX_NONE)
-		{
-			UE_LOG(LogLearning, Error, TEXT("%s: Invalid enum value: %i"), *GetName(), EnumValues[EnumIdx]);
-			continue;
-		}
+FLearningAgentsObservationObjectElement ULearningAgentsObservationObject::MakeRotationObservation(const FRotator Rotation, const FRotator RelativeRotation, const FName Name)
+{
+	return MakeRotationObservationFromQuat(FQuat::MakeFromRotator(Rotation), FQuat::MakeFromRotator(RelativeRotation), Name);
+}
 
-		EnumViewSlice[Index] = 1.0f;
-	}
+FLearningAgentsObservationObjectElement ULearningAgentsObservationObject::MakeRotationObservationFromQuat(const FQuat Rotation, const FQuat RelativeRotation, const FName Name)
+{
+	const FQuat LocalRotation = RelativeRotation.Inverse() * Rotation;
+	const FVector LocalAxisForward = LocalRotation.GetForwardVector();
+	const FVector LocalAxisRight = LocalRotation.GetRightVector();
 
-	AgentIteration[AgentId]++;
+	return MakeContinuousObservationFromArrayView({
+		(float)LocalAxisForward.X,
+		(float)LocalAxisForward.Y,
+		(float)LocalAxisForward.Z,
+		(float)LocalAxisRight.X,
+		(float)LocalAxisRight.Y,
+		(float)LocalAxisRight.Z,
+		}, Name);
 }
 
-#if UE_LEARNING_AGENTS_ENABLE_VISUAL_LOG
-void UEnumArrayObservation::VisualLog(const UE::Learning::FIndexSet Instances) const
+FLearningAgentsObservationObjectElement ULearningAgentsObservationObject::MakeScaleObservation(const FVector Scale, const FVector RelativeScale, const FName Name)
 {
-	UE_LEARNING_TRACE_CPUPROFILER_EVENT_SCOPE(UEnumArrayObservation::VisualLog);
+	const FVector LocalLogScale = 
+		UE::Learning::Agents::Observation::Private::VectorLogSafe(Scale) - 
+		UE::Learning::Agents::Observation::Private::VectorLogSafe(RelativeScale);
 
-	const TLearningArrayView<2, const float> ValueView = FeatureObject->InstanceData->ConstView(FeatureObject->ValueHandle);
-	const TLearningArrayView<2, const float> FeatureView = FeatureObject->InstanceData->ConstView(FeatureObject->FeatureHandle);
+	return MakeContinuousObservationFromArrayView({
+		(float)LocalLogScale.X,
+		(float)LocalLogScale.Y,
+		(float)LocalLogScale.Z,
+		}, Name);
+}
 
-	for (const int32 Instance : Instances)
-	{
-		if (const AActor* Actor = Cast<AActor>(Interactor->GetAgent(Instance)))
+FLearningAgentsObservationObjectElement ULearningAgentsObservationObject::MakeTransformObservation(const FTransform Transform, const FTransform RelativeTransform, const float TranslationScale, const FName Name)
+{
+	const FTransform LocalTransform = Transform * RelativeTransform.Inverse();
+	
+	return MakeStructObservationFromArrayViews(
 		{
-			UE_LEARNING_AGENTS_VLOG_STRING(this, LogLearning, Display,
-				Actor->GetActorLocation(),
-				VisualLogColor.ToFColor(true),
-				TEXT("Agent %i\nEncoded: %s"),
-				Instance,
-				*UE::Learning::Array::FormatFloat(FeatureView[Instance]));
-
-		}
-	}
+			TEXT("Translation"),
+			TEXT("Rotation"),
+			TEXT("Scale")
+		},
+		{
+			MakeTranslationObservation(LocalTransform.GetTranslation(), FTransform::Identity, TranslationScale),
+			MakeRotationObservationFromQuat(LocalTransform.GetRotation(), FQuat::Identity),
+			MakeScaleObservation(LocalTransform.GetScale3D(), FVector::OneVector)
+		},
+		Name);
 }
-#endif
 
-//------------------------------------------------------------------
+FLearningAgentsObservationObjectElement ULearningAgentsObservationObject::MakeAngleObservation(const float Angle, const float RelativeAngle, const FName Name)
+{
+	return MakeAngleObservationRadians(FMath::DegreesToRadians(Angle), FMath::DegreesToRadians(RelativeAngle), Name);
+}
 
-UTimeObservation* UTimeObservation::AddTimeObservation(ULearningAgentsInteractor* InInteractor, const FName Name, const float Scale)
+FLearningAgentsObservationObjectElement ULearningAgentsObservationObject::MakeAngleObservationRadians(const float Angle, const float RelativeAngle, const FName Name)
 {
-	return UE::Learning::Agents::Observations::Private::AddObservation<UTimeObservation, UE::Learning::FTimeFeature>(InInteractor, Name, TEXT("AddTimeObservation"), 1, Scale);
+	const float LocalAngle = FMath::FindDeltaAngleRadians(RelativeAngle, Angle);
+
+	return MakeContinuousObservationFromArrayView({
+		(float)FMath::Sin(Angle),
+		(float)FMath::Cos(Angle),
+		}, Name);
 }
 
-void UTimeObservation::SetTimeObservation(const int32 AgentId, const float Time, const float RelativeTime)
+FLearningAgentsObservationObjectElement ULearningAgentsObservationObject::MakeVelocityObservation(const FVector Velocity, const FTransform RelativeTransform, const float VelocityScale, const FName Name)
 {
-	if (!Interactor->HasAgent(AgentId))
-	{
-		UE_LOG(LogLearning, Error, TEXT("%s: AgentId %d not found in the agents set."), *GetName(), AgentId);
-		return;
-	}
+	const FVector LocalVelocity = RelativeTransform.InverseTransformVectorNoScale(Velocity);
 
-	FeatureObject->InstanceData->View(FeatureObject->TimeHandle)[AgentId][0] = Time;
-	FeatureObject->InstanceData->View(FeatureObject->RelativeTimeHandle)[AgentId] = RelativeTime;
-	AgentIteration[AgentId]++;
+	return MakeContinuousObservationFromArrayView({
+		(float)LocalVelocity.X / FMath::Max(VelocityScale, UE_SMALL_NUMBER),
+		(float)LocalVelocity.Y / FMath::Max(VelocityScale, UE_SMALL_NUMBER),
+		(float)LocalVelocity.Z / FMath::Max(VelocityScale, UE_SMALL_NUMBER),
+		}, Name);
 }
 
-#if UE_LEARNING_AGENTS_ENABLE_VISUAL_LOG
-void UTimeObservation::VisualLog(const UE::Learning::FIndexSet Instances) const
+FLearningAgentsObservationObjectElement ULearningAgentsObservationObject::MakeDirectionObservation(const FVector Direction, const FTransform RelativeTransform, const FName Name)
 {
-	UE_LEARNING_TRACE_CPUPROFILER_EVENT_SCOPE(UTimeObservation::VisualLog);
+	const FVector LocalDirection = RelativeTransform.InverseTransformVectorNoScale(Direction).GetSafeNormal(UE_SMALL_NUMBER, FVector::ForwardVector);
 
-	const TLearningArrayView<2, const float> TimeView = FeatureObject->InstanceData->ConstView(FeatureObject->TimeHandle);
-	const TLearningArrayView<1, const float> RelativeTimeView = FeatureObject->InstanceData->ConstView(FeatureObject->RelativeTimeHandle);
-	const TLearningArrayView<2, const float> FeatureView = FeatureObject->InstanceData->ConstView(FeatureObject->FeatureHandle);
+	return MakeContinuousObservationFromArrayView({
+		(float)LocalDirection.X,
+		(float)LocalDirection.Y,
+		(float)LocalDirection.Z,
+		}, Name);
+}
 
-	for (const int32 Instance : Instances)
+FLearningAgentsObservationObjectElement ULearningAgentsObservationObject::MakeLocationAlongSplineObservation(const USplineComponent* SplineComponent, const float DistanceAlongSpline, const FTransform RelativeTransform, const float TranslationScale, const FName Name)
+{
+	if (!SplineComponent)
 	{
-		if (const AActor* Actor = Cast<AActor>(Interactor->GetAgent(Instance)))
-		{
-			const float Time = TimeView[Instance][0];
-			const float RelativeTime = RelativeTimeView[Instance];
-
-			UE_LEARNING_AGENTS_VLOG_STRING(this, LogLearning, Display,
-				Actor->GetActorLocation(),
-				VisualLogColor.ToFColor(true),
-				TEXT("Agent %i\nScale: [% 6.2f]\nRelative Time: [% 6.3f]\nTime: [% 6.3f]\nEncoded: %s"),
-				Instance,
-				FeatureObject->Scale,
-				RelativeTime,
-				Time,
-				*UE::Learning::Array::FormatFloat(FeatureView[Instance]));
-		}
+		UE_LOG(LogLearning, Error, TEXT("MakeLocationAlongSplineObservation: SplineComponent was nullptr."));
+		return FLearningAgentsObservationObjectElement();
 	}
+
+	return MakeTranslationObservation(SplineComponent->GetLocationAtDistanceAlongSpline(DistanceAlongSpline, ESplineCoordinateSpace::World), RelativeTransform, TranslationScale, Name);
 }
-#endif
 
-UTimeArrayObservation* UTimeArrayObservation::AddTimeArrayObservation(ULearningAgentsInteractor* InInteractor, const FName Name, const int32 TimeNum, const float Scale)
+FLearningAgentsObservationObjectElement ULearningAgentsObservationObject::MakeProportionAlongSplineObservation(const USplineComponent* SplineComponent, const float DistanceAlongSpline, const FName Name)
 {
-	if (TimeNum < 1)
+	if (!SplineComponent)
 	{
-		UE_LOG(LogLearning, Error, TEXT("AddTimeArrayObservation: Number of elements in array must be at least 1, got %i."), TimeNum);
-		return nullptr;
+		UE_LOG(LogLearning, Error, TEXT("MakeProportionAlongSplineObservation: SplineComponent was nullptr."));
+		return FLearningAgentsObservationObjectElement();
 	}
 
-	return UE::Learning::Agents::Observations::Private::AddObservation<UTimeArrayObservation, UE::Learning::FTimeFeature>(InInteractor, Name, TEXT("AddTimeArrayObservation"), TimeNum, Scale);
+	if (SplineComponent->IsClosedLoop())
+	{
+		const float TotalDistance = SplineComponent->GetSplineLength();
+		const float WrapDistance = FMath::Wrap(DistanceAlongSpline, 0.0f, TotalDistance);
+		const float Proportion = WrapDistance / FMath::Max(TotalDistance, UE_SMALL_NUMBER);
+		const float Angle = FMath::Wrap(UE_TWO_PI * Proportion, -UE_PI, UE_PI);
+		return MakeExclusiveUnionObservation(TEXT("Angle"), MakeAngleObservation(Angle), Name);
+	}
+	else
+	{
+		const float Proportion = FMath::Clamp(DistanceAlongSpline / FMath::Max(SplineComponent->GetSplineLength(), UE_SMALL_NUMBER), 0.0f, 1.0f);
+		return MakeExclusiveUnionObservation(TEXT("Proportion"), MakeFloatObservation(Proportion), Name);
+	}
 }
 
-void UTimeArrayObservation::SetTimeArrayObservation(const int32 AgentId, const TArray<float>& Times, const float RelativeTime)
+FLearningAgentsObservationObjectElement ULearningAgentsObservationObject::MakeDirectionAlongSplineObservation(const USplineComponent* SplineComponent, const float DistanceAlongSpline, const FTransform RelativeTransform, const FName Name)
 {
-	if (!Interactor->HasAgent(AgentId))
+	if (!SplineComponent)
 	{
-		UE_LOG(LogLearning, Error, TEXT("%s: AgentId %d not found in the agents set."), *GetName(), AgentId);
-		return;
+		UE_LOG(LogLearning, Error, TEXT("MakeDirectionAlongSplineObservation: SplineComponent was nullptr."));
+		return FLearningAgentsObservationObjectElement();
 	}
 
-	const TLearningArrayView<2, float> TimeView = FeatureObject->InstanceData->View(FeatureObject->TimeHandle);
-	const TLearningArrayView<1, float> RelativeTimeView = FeatureObject->InstanceData->View(FeatureObject->RelativeTimeHandle);
+	return MakeDirectionObservation(SplineComponent->GetDirectionAtDistanceAlongSpline(DistanceAlongSpline, ESplineCoordinateSpace::World), RelativeTransform, Name);
+}
 
-	if (Times.Num() != TimeView.Num<1>())
+FLearningAgentsObservationObjectElement ULearningAgentsObservationObject::MakePropertiesAlongSplineObservation(const USplineComponent* SplineComponent, const float DistanceAlongSpline, const FTransform RelativeTransform, const float TranslationScale, const FName Name)
+{
+	if (!SplineComponent)
 	{
-		UE_LOG(LogLearning, Error, TEXT("%s: Got wrong number of elements in array. Expected %i, got %i."), *GetName(), TimeView.Num<1>(), Times.Num());
-		return;
+		UE_LOG(LogLearning, Error, TEXT("MakePropertiesAlongSplineObservation: SplineComponent was nullptr."));
+		return FLearningAgentsObservationObjectElement();
 	}
 
-	UE::Learning::Array::Copy<1, float>(TimeView[AgentId], Times);
-	RelativeTimeView[AgentId] = RelativeTime;
-	AgentIteration[AgentId]++;
+	return MakeStructObservationFromArrayViews(
+		{
+			TEXT("Location"),
+			TEXT("Proportion"),
+			TEXT("Direction")
+		},
+		{
+			MakeLocationAlongSplineObservation(SplineComponent, DistanceAlongSpline, RelativeTransform, TranslationScale),
+			MakeProportionAlongSplineObservation(SplineComponent, DistanceAlongSpline),
+			MakeDirectionAlongSplineObservation(SplineComponent, DistanceAlongSpline, RelativeTransform),
+		}, Name);
 }
 
-#if UE_LEARNING_AGENTS_ENABLE_VISUAL_LOG
-void UTimeArrayObservation::VisualLog(const UE::Learning::FIndexSet Instances) const
+FLearningAgentsObservationObjectElement ULearningAgentsObservationObject::MakeProportionAlongRayObservation(const FVector RayStart, const FVector RayEnd, const FTransform RayTransform, const ECollisionChannel CollisionChannel, const FName Name)
 {
-	UE_LEARNING_TRACE_CPUPROFILER_EVENT_SCOPE(UTimeArrayObservation::VisualLog);
+	const FVector RayStartWorld = RayTransform.TransformPosition(RayStart);
+	const FVector RayEndWorld = RayTransform.TransformPosition(RayEnd);
+	const float RayDistance = FVector::Distance(RayStartWorld, RayEndWorld);
+
 
-	const TLearningArrayView<2, const float> TimeView = FeatureObject->InstanceData->ConstView(FeatureObject->TimeHandle);
-	const TLearningArrayView<1, const float> RelativeTimeView = FeatureObject->InstanceData->ConstView(FeatureObject->RelativeTimeHandle);
-	const TLearningArrayView<2, const float> FeatureView = FeatureObject->InstanceData->ConstView(FeatureObject->FeatureHandle);
+	FCollisionObjectQueryParams ObjectQueryParams;
+	ObjectQueryParams.AddObjectTypesToQuery(CollisionChannel);
 
-	for (const int32 Instance : Instances)
+	FHitResult TraceHit;
+	const bool bHit = GetWorld()->LineTraceSingleByObjectType(TraceHit, RayStartWorld, RayEndWorld, ObjectQueryParams);
+
+	if (bHit)
 	{
-		if (const AActor* Actor = Cast<AActor>(Interactor->GetAgent(Instance)))
-		{
-			UE_LEARNING_AGENTS_VLOG_STRING(this, LogLearning, Display,
-				Actor->GetActorLocation(),
-				VisualLogColor.ToFColor(true),
-				TEXT("Agent %i\nScale: [% 6.2f]\nRelative Time: [% 6.3f]\nTimes: %s\nEncoded: %s"),
-				Instance,
-				FeatureObject->Scale,
-				RelativeTimeView[Instance],
-				*UE::Learning::Array::FormatFloat(TimeView[Instance]),
-				*UE::Learning::Array::FormatFloat(FeatureView[Instance]));
-		}
+		//DrawDebugLine(GetWorld(), RayStartWorld, RayStartWorld + TraceHit.Time * (RayEndWorld - RayStartWorld), FColor::Red, false, -1.0f, 0, 5.0f);
+		//DrawDebugLine(GetWorld(), RayStartWorld + TraceHit.Time * (RayEndWorld - RayStartWorld), RayEndWorld, FColor::Blue, false, -1.0f, 0, 5.0f);
+		return MakeFloatObservation(1.0f - TraceHit.Time, 1.0f, Name);
+	}
+	else
+	{
+		//DrawDebugLine(GetWorld(), RayStartWorld, RayEndWorld, FColor::Red, false, -1.0f, 0, 5.0f);
+		return MakeFloatObservation(0.0f, 1.0f, Name);
 	}
 }
-#endif
-
-//------------------------------------------------------------------
 
-UAngleObservation* UAngleObservation::AddAngleObservation(ULearningAgentsInteractor* InInteractor, const FName Name, const float Scale)
+FLearningAgentsObservationObjectElement ULearningAgentsObservationObject::MakeProportionAlongRaysObservation(const TArray<FVector>& RayStarts, const TArray<FVector>& RayEnds, const FTransform RayTransform, const ECollisionChannel CollisionChannel, const FName Name)
 {
-	return UE::Learning::Agents::Observations::Private::AddObservation<UAngleObservation, UE::Learning::FAngleFeature>(InInteractor, Name, TEXT("AddAngleObservation"), 1, Scale);
-}
+	if (RayStarts.Num() != RayEnds.Num())
+	{
+		UE_LOG(LogLearning, Error, TEXT("MakeProportionAlongRaysObservation: Different number of RayStarts and RayEnds (%i vs %i)."), RayStarts.Num(), RayEnds.Num());
+		return FLearningAgentsObservationObjectElement();
+	}
 
-void UAngleObservation::SetAngleObservation(const int32 AgentId, const float Angle, const float RelativeAngle)
-{
-	if (!Interactor->HasAgent(AgentId))
+	const int32 RayNum = RayStarts.Num();
+
+	TArray<FLearningAgentsObservationObjectElement, TInlineAllocator<32>> RayElements;
+	RayElements.Reserve(RayNum);
+
+	for (int32 RayIdx = 0; RayIdx < RayNum; RayIdx++)
 	{
-		UE_LOG(LogLearning, Error, TEXT("%s: AgentId %d not found in the agents set."), *GetName(), AgentId);
-		return;
+		RayElements.Add(MakeProportionAlongRayObservation(RayStarts[RayIdx], RayEnds[RayIdx], RayTransform, CollisionChannel));
 	}
 
-	FeatureObject->InstanceData->View(FeatureObject->AngleHandle)[AgentId][0] = FMath::DegreesToRadians(Angle);
-	FeatureObject->InstanceData->View(FeatureObject->RelativeAngleHandle)[AgentId] = FMath::DegreesToRadians(RelativeAngle);
-	AgentIteration[AgentId]++;
+	return MakeStaticArrayObservationFromArrayView(RayElements, Name);
 }
 
-#if UE_LEARNING_AGENTS_ENABLE_VISUAL_LOG
-void UAngleObservation::VisualLog(const UE::Learning::FIndexSet Instances) const
+bool ULearningAgentsObservationObject::GetNullObservation(const FLearningAgentsObservationObjectElement Element, const FName Name) const
 {
-	UE_LEARNING_TRACE_CPUPROFILER_EVENT_SCOPE(UAngleObservation::VisualLog);
-
-	const TLearningArrayView<2, const float> AngleView = FeatureObject->InstanceData->ConstView(FeatureObject->AngleHandle);
-	const TLearningArrayView<1, const float> RelativeAngleView = FeatureObject->InstanceData->ConstView(FeatureObject->RelativeAngleHandle);
-	const TLearningArrayView<2, const float> FeatureView = FeatureObject->InstanceData->ConstView(FeatureObject->FeatureHandle);
+	if (!ObservationObject.IsValid(Element.ObjectElement))
+	{
+		UE_LOG(LogLearning, Error, TEXT("%s: Invalid Observation Object."), *GetName());
+		return false;
+	}
 
-	for (const int32 Instance : Instances)
+	if (ObservationObject.GetName(Element.ObjectElement) != Name)
 	{
-		if (const AActor* Actor = Cast<AActor>(Interactor->GetAgent(Instance)))
-		{
-			const float Angle = AngleView[Instance][0];
-			const float RelativeAngle = RelativeAngleView[Instance];
-
-			UE_LEARNING_AGENTS_VLOG_ANGLE(this, LogLearning, Display,
-				Angle,
-				0.0f,
-				Actor->GetActorLocation(),
-				50.0f,
-				VisualLogColor.ToFColor(true),
-				TEXT("Agent %i\nScale: [% 6.2f]\nRelative Angle: [% 6.1f]\nAngle: [% 6.1f]\nEncoded: %s"),
-				Instance,
-				FeatureObject->Scale,
-				FMath::RadiansToDegrees(RelativeAngle),
-				FMath::RadiansToDegrees(Angle),
-				*UE::Learning::Array::FormatFloat(FeatureView[Instance]));
-		}
+		UE_LOG(LogLearning, Warning, TEXT("%s: Observation name does not match. Observation is '%s' but asked for '%s'."), *GetName(), *ObservationObject.GetName(Element.ObjectElement).ToString(), *Name.ToString());
 	}
-}
-#endif
 
-UAngleArrayObservation* UAngleArrayObservation::AddAngleArrayObservation(ULearningAgentsInteractor* InInteractor, const FName Name, const int32 AngleNum, const float Scale)
-{
-	if (AngleNum < 1)
+	if (ObservationObject.GetType(Element.ObjectElement) != UE::Learning::Observation::EType::Null)
 	{
-		UE_LOG(LogLearning, Error, TEXT("AddAngleArrayObservation: Number of elements in array must be at least 1, got %i."), AngleNum);
-		return nullptr;
+		UE_LOG(LogLearning, Error, TEXT("%s: Observation '%s' type does not match. Observation is '%s' but asked for '%s'."),
+			*GetName(),
+			*ObservationObject.GetName(Element.ObjectElement).ToString(),
+			UE::Learning::Agents::Observation::Private::GetObservationTypeString(ObservationObject.GetType(Element.ObjectElement)),
+			UE::Learning::Agents::Observation::Private::GetObservationTypeString(UE::Learning::Observation::EType::Null));
+		return false;
 	}
 
-	return UE::Learning::Agents::Observations::Private::AddObservation<UAngleArrayObservation, UE::Learning::FAngleFeature>(InInteractor, Name, TEXT("AddAngleArrayObservation"), AngleNum, Scale);
+	return true;
 }
 
-void UAngleArrayObservation::SetAngleArrayObservation(const int32 AgentId, const TArray<float>& Angles, const float RelativeAngle)
+bool ULearningAgentsObservationObject::GetContinuousObservationNum(int32& OutNum, const FLearningAgentsObservationObjectElement Element, const FName Name) const
 {
-	if (!Interactor->HasAgent(AgentId))
+	if (!ObservationObject.IsValid(Element.ObjectElement))
 	{
-		UE_LOG(LogLearning, Error, TEXT("%s: AgentId %d not found in the agents set."), *GetName(), AgentId);
-		return;
+		UE_LOG(LogLearning, Error, TEXT("%s: Invalid Observation Object."), *GetName());
+		OutNum = 0;
+		return false;
 	}
 
-	const TLearningArrayView<2, float> AngleView = FeatureObject->InstanceData->View(FeatureObject->AngleHandle);
-	const TLearningArrayView<1, float> RelativeAngleView = FeatureObject->InstanceData->View(FeatureObject->RelativeAngleHandle);
-
-	if (Angles.Num() != AngleView.Num<1>())
+	if (ObservationObject.GetName(Element.ObjectElement) != Name)
 	{
-		UE_LOG(LogLearning, Error, TEXT("%s: Got wrong number of elements in array. Expected %i, got %i."), *GetName(), AngleView.Num<1>(), Angles.Num());
-		return;
+		UE_LOG(LogLearning, Warning, TEXT("%s: Observation name does not match. Observation is '%s' but asked for '%s'."), *GetName(), *ObservationObject.GetName(Element.ObjectElement).ToString(), *Name.ToString());
 	}
 
-	for (int32 AngleIdx = 0; AngleIdx < Angles.Num(); AngleIdx++)
+	if (ObservationObject.GetType(Element.ObjectElement) != UE::Learning::Observation::EType::Continuous)
 	{
-		AngleView[AgentId][AngleIdx] = FMath::DegreesToRadians(Angles[AngleIdx]);
+		UE_LOG(LogLearning, Error, TEXT("%s: Observation '%s' type does not match. Observation is '%s' but asked for '%s'."),
+			*GetName(),
+			*ObservationObject.GetName(Element.ObjectElement).ToString(),
+			UE::Learning::Agents::Observation::Private::GetObservationTypeString(ObservationObject.GetType(Element.ObjectElement)),
+			UE::Learning::Agents::Observation::Private::GetObservationTypeString(UE::Learning::Observation::EType::Continuous));
+		OutNum = 0;
+		return false;
 	}
-	
-	RelativeAngleView[AgentId] = FMath::DegreesToRadians(RelativeAngle);
-	AgentIteration[AgentId]++;
+
+	OutNum = ObservationObject.GetContinuous(Element.ObjectElement).Values.Num();
+	return true;
 }
 
-#if UE_LEARNING_AGENTS_ENABLE_VISUAL_LOG
-void UAngleArrayObservation::VisualLog(const UE::Learning::FIndexSet Instances) const
+bool ULearningAgentsObservationObject::GetContinuousObservation(TArray<float>& OutValues, const FLearningAgentsObservationObjectElement Element, const FName Name) const
 {
-	UE_LEARNING_TRACE_CPUPROFILER_EVENT_SCOPE(UAngleArrayObservation::VisualLog);
+	int32 OutValueNum = 0;
+	if (!GetContinuousObservationNum(OutValueNum, Element, Name))
+	{
+		OutValues.Empty();
+		return false;
+	}
 
-	const TLearningArrayView<2, const float> AngleView = FeatureObject->InstanceData->ConstView(FeatureObject->AngleHandle);
-	const TLearningArrayView<1, const float> RelativeAngleView = FeatureObject->InstanceData->ConstView(FeatureObject->RelativeAngleHandle);
-	const TLearningArrayView<2, const float> FeatureView = FeatureObject->InstanceData->ConstView(FeatureObject->FeatureHandle);
+	OutValues.SetNumUninitialized(OutValueNum);
+
+	if (!GetContinuousObservationToArrayView(OutValues, Element, Name))
+	{
+		OutValues.Empty();
+		return false;
+	}
 
-	const int32 AngleNum = AngleView.Num<1>();
+	return true;
+}
 
-	for (const int32 Instance : Instances)
+bool ULearningAgentsObservationObject::GetContinuousObservationToArrayView(TArrayView<float> OutValues, const FLearningAgentsObservationObjectElement Element, const FName Name) const
+{
+	if (!ObservationObject.IsValid(Element.ObjectElement))
 	{
-		if (const AActor* Actor = Cast<AActor>(Interactor->GetAgent(Instance)))
-		{
-			const float RelativeAngle = RelativeAngleView[Instance];
+		UE_LOG(LogLearning, Error, TEXT("%s: Invalid Observation Object."), *GetName());
+		UE::Learning::Array::Zero<1, float>(OutValues);
+		return false;
+	}
 
-			for (int32 AngleIdx = 0; AngleIdx < AngleNum; AngleIdx++)
-			{
-				const float Angle = AngleView[Instance][AngleIdx];
-				const FVector Offset = UE::Learning::Agents::Debug::GridOffsetForIndex(AngleIdx, AngleNum);
-
-				UE_LEARNING_AGENTS_VLOG_ANGLE(this, LogLearning, Display,
-					Angle,
-					RelativeAngle,
-					Actor->GetActorLocation() + Offset,
-					5.0f,
-					VisualLogColor.ToFColor(true),
-					TEXT("Angle %i: [% 6.4f]"),
-					AngleIdx,
-					FMath::RadiansToDegrees(Angle));
-			}
+	if (ObservationObject.GetName(Element.ObjectElement) != Name)
+	{
+		UE_LOG(LogLearning, Warning, TEXT("%s: Observation name does not match. Observation is '%s' but asked for '%s'."), *GetName(), *ObservationObject.GetName(Element.ObjectElement).ToString(), *Name.ToString());
+	}
 
-			UE_LEARNING_AGENTS_VLOG_STRING(this, LogLearning, Display,
-				Actor->GetActorLocation() + FVector(0.0f, 0.0f, 20.0f),
-				VisualLogColor.ToFColor(true),
-				TEXT("Agent %i\nScale: [% 6.2f]\nRelative Angle: [% 6.1f]\nEncoded: %s"),
-				Instance,
-				FeatureObject->Scale,
-				FMath::RadiansToDegrees(RelativeAngle),
-				*UE::Learning::Array::FormatFloat(FeatureView[Instance]));
-		}
+	if (ObservationObject.GetType(Element.ObjectElement) != UE::Learning::Observation::EType::Continuous)
+	{
+		UE_LOG(LogLearning, Error, TEXT("%s: Observation '%s' type does not match. Observation is '%s' but asked for '%s'."),
+			*GetName(),
+			*ObservationObject.GetName(Element.ObjectElement).ToString(),
+			UE::Learning::Agents::Observation::Private::GetObservationTypeString(ObservationObject.GetType(Element.ObjectElement)),
+			UE::Learning::Agents::Observation::Private::GetObservationTypeString(UE::Learning::Observation::EType::Continuous));
+		UE::Learning::Array::Zero<1, float>(OutValues);
+		return false;
 	}
-}
-#endif
 
-//------------------------------------------------------------------
+	const TArrayView<const float> Values = ObservationObject.GetContinuous(Element.ObjectElement).Values;
 
-URotationObservation* URotationObservation::AddRotationObservation(ULearningAgentsInteractor* InInteractor, const FName Name, const float Scale)
-{
-	return UE::Learning::Agents::Observations::Private::AddObservation<URotationObservation, UE::Learning::FRotationFeature>(InInteractor, Name, TEXT("AddRotationObservation"), 1, Scale);
-}
+	if (Values.Num() != OutValues.Num())
+	{
+		UE_LOG(LogLearning, Error, TEXT("%s: Observation '%s' size does not match. Observation is '%i' values but asked for '%i'."),
+			*GetName(), *ObservationObject.GetName(Element.ObjectElement).ToString(),
+			Values.Num(), OutValues.Num());
+		UE::Learning::Array::Zero<1, float>(OutValues);
+		return false;
+	}
 
-void URotationObservation::SetRotationObservation(const int32 AgentId, const FRotator Rotation, const FRotator RelativeRotation)
-{
-	SetRotationObservationFromQuat(AgentId, Rotation.Quaternion(), RelativeRotation.Quaternion());
+	UE::Learning::Array::Copy<1, float>(OutValues, Values);
+	return true;
 }
 
-void URotationObservation::SetRotationObservationFromQuat(const int32 AgentId, const FQuat Rotation, const FQuat RelativeRotation)
+bool ULearningAgentsObservationObject::GetExclusiveDiscreteObservation(int32& OutIndex, const FLearningAgentsObservationObjectElement Element, const FName Name) const
 {
-	if (!Interactor->HasAgent(AgentId))
+	int32 DiscreteValueNum;
+	if (!GetContinuousObservationNum(DiscreteValueNum, Element, Name))
+	{
+		OutIndex = INDEX_NONE;
+		return false;
+	}
+
+	TArray<float, TInlineAllocator<32>> OneHot;
+	OneHot.SetNumUninitialized(DiscreteValueNum);
+	if (!GetContinuousObservationToArrayView(OneHot, Element, Name))
 	{
-		UE_LOG(LogLearning, Error, TEXT("%s: AgentId %d not found in the agents set."), *GetName(), AgentId);
-		return;
+		OutIndex = INDEX_NONE;
+		return false;
+	}
+
+	OutIndex = INDEX_NONE;
+	for (int32 Idx = 0; Idx < DiscreteValueNum; Idx++)
+	{
+		if (OneHot[Idx])
+		{
+			OutIndex = Idx;
+			break;
+		}
 	}
 
-	FeatureObject->InstanceData->View(FeatureObject->RotationHandle)[AgentId][0] = Rotation;
-	FeatureObject->InstanceData->View(FeatureObject->RelativeRotationHandle)[AgentId] = RelativeRotation;
-	AgentIteration[AgentId]++;
+	return OutIndex != INDEX_NONE;
 }
 
-#if UE_LEARNING_AGENTS_ENABLE_VISUAL_LOG
-void URotationObservation::VisualLog(const UE::Learning::FIndexSet Instances) const
+bool ULearningAgentsObservationObject::GetInclusiveDiscreteObservationNum(int32& OutNum, const FLearningAgentsObservationObjectElement Element, const FName Name) const
 {
-	UE_LEARNING_TRACE_CPUPROFILER_EVENT_SCOPE(URotationObservation::VisualLog);
+	int32 DiscreteValueNum;
+	if (!GetContinuousObservationNum(DiscreteValueNum, Element, Name))
+	{
+		OutNum = 0;
+		return false;
+	}
 
-	const TLearningArrayView<2, const FQuat> RotationView = FeatureObject->InstanceData->ConstView(FeatureObject->RotationHandle);
-	const TLearningArrayView<1, const FQuat> RelativeRotationView = FeatureObject->InstanceData->ConstView(FeatureObject->RelativeRotationHandle);
-	const TLearningArrayView<2, const float> FeatureView = FeatureObject->InstanceData->ConstView(FeatureObject->FeatureHandle);
+	TArray<float, TInlineAllocator<32>> OneHot;
+	OneHot.SetNumUninitialized(DiscreteValueNum);
+	if (!GetContinuousObservationToArrayView(OneHot, Element, Name))
+	{
+		OutNum = 0;
+		return false;
+	}
 
-	for (const int32 Instance : Instances)
+	OutNum = 0;
+	for (int32 Idx = 0; Idx < DiscreteValueNum; Idx++)
 	{
-		if (const AActor* Actor = Cast<AActor>(Interactor->GetAgent(Instance)))
-		{
-			const FQuat Rotation = RotationView[Instance][0];
-			const FRotator Rotator = Rotation.Rotator();
-			const FQuat RelativeRotation = RelativeRotationView[Instance];
-			const FRotator RelativeRotator = RelativeRotation.Rotator();
-
-			UE_LEARNING_AGENTS_VLOG_TRANSFORM(this, LogLearning, Display,
-				Actor->GetActorLocation() + FVector(20.0f, 0.0f, 0.0f),
-				Rotation,
-				VisualLogColor.ToFColor(true),
-				TEXT("Relative Rotation: [% 6.1f % 6.1f % 6.1f]"),
-				RelativeRotator.Pitch, RelativeRotator.Roll, RelativeRotator.Yaw);
-
-			UE_LEARNING_AGENTS_VLOG_TRANSFORM(this, LogLearning, Display,
-				Actor->GetActorLocation(),
-				Rotation,
-				VisualLogColor.ToFColor(true),
-				TEXT("Agent %i\nScale: [% 6.2f]\nRotation: [% 6.1f % 6.1f % 6.1f]\nEncoded: %s"),
-				Instance,
-				FeatureObject->Scale,
-				Rotator.Pitch, Rotator.Roll, Rotator.Yaw,
-				*UE::Learning::Array::FormatFloat(FeatureView[Instance]));
-		}
+		if (OneHot[Idx]) { OutNum++; }
 	}
+
+	return true;
 }
-#endif
 
-URotationArrayObservation* URotationArrayObservation::AddRotationArrayObservation(ULearningAgentsInteractor* InInteractor, const FName Name, const int32 RotationNum, const float Scale)
+bool ULearningAgentsObservationObject::GetInclusiveDiscreteObservation(TArray<int32>& OutIndices, const FLearningAgentsObservationObjectElement Element, const FName Name) const
 {
-	if (RotationNum < 1)
+	int32 OutIndicesNum;
+	if (!GetInclusiveDiscreteObservationNum(OutIndicesNum, Element, Name))
 	{
-		UE_LOG(LogLearning, Error, TEXT("AddRotationArrayObservation: Number of elements in array must be at least 1, got %i."), RotationNum);
-		return nullptr;
+		OutIndices.Empty();
+		return false;
 	}
 
-	return UE::Learning::Agents::Observations::Private::AddObservation<URotationArrayObservation, UE::Learning::FRotationFeature>(InInteractor, Name, TEXT("AddRotationArrayObservation"), RotationNum, Scale);
+	OutIndices.SetNumUninitialized(OutIndicesNum);
+	return GetInclusiveDiscreteObservationToArrayView(OutIndices, Element, Name);
 }
 
-void URotationArrayObservation::SetRotationArrayObservation(const int32 AgentId, const TArray<FRotator>& Rotations, const FRotator RelativeRotation)
+bool ULearningAgentsObservationObject::GetInclusiveDiscreteObservationToArrayView(TArrayView<int32> OutIndices, const FLearningAgentsObservationObjectElement Element, const FName Name) const
 {
-	if (!Interactor->HasAgent(AgentId))
+	int32 DiscreteValueNum;
+	if (!GetContinuousObservationNum(DiscreteValueNum, Element, Name))
 	{
-		UE_LOG(LogLearning, Error, TEXT("%s: AgentId %d not found in the agents set."), *GetName(), AgentId);
-		return;
+		return false;
 	}
 
-	const TLearningArrayView<2, FQuat> RotationView = FeatureObject->InstanceData->View(FeatureObject->RotationHandle);
-	const TLearningArrayView<1, FQuat> RelativeRotationView = FeatureObject->InstanceData->View(FeatureObject->RelativeRotationHandle);
+	TArray<float, TInlineAllocator<32>> OneHot;
+	OneHot.SetNumUninitialized(DiscreteValueNum);
+	if (!GetContinuousObservationToArrayView(OneHot, Element, Name))
+	{
+		return false;
+	}
+
+	int32 IndicesNum = 0;
+	for (int32 Idx = 0; Idx < DiscreteValueNum; Idx++)
+	{
+		if (OneHot[Idx]) { IndicesNum++; }
+	}
 
-	if (Rotations.Num() != RotationView.Num<1>())
+	if (IndicesNum != OutIndices.Num())
 	{
-		UE_LOG(LogLearning, Error, TEXT("%s: Got wrong number of elements in array. Expected %i, got %i."), *GetName(), RotationView.Num<1>(), Rotations.Num());
-		return;
+		UE_LOG(LogLearning, Error, TEXT("%s: Observation '%s' size does not match. Observation is '%i' indices but asked for '%i'."),
+			*GetName(), *ObservationObject.GetName(Element.ObjectElement).ToString(),
+			IndicesNum, OutIndices.Num());
+		UE::Learning::Array::Zero<1, int32>(OutIndices);
+		return false;
 	}
 
-	for (int32 RotationIdx = 0; RotationIdx < Rotations.Num(); RotationIdx++)
+	int32 Offset = 0;
+	for (int32 Idx = 0; Idx < DiscreteValueNum; Idx++)
 	{
-		RotationView[AgentId][RotationIdx] = Rotations[RotationIdx].Quaternion();
+		if (OneHot[Idx]) { OutIndices[Offset] = Idx; Offset++; }
 	}
 
-	RelativeRotationView[AgentId] = RelativeRotation.Quaternion();
-	AgentIteration[AgentId]++;
+	return true;
 }
 
-void URotationArrayObservation::SetRotationArrayObservationFromQuats(const int32 AgentId, const TArray<FQuat>& Rotations, const FQuat RelativeRotation)
+bool ULearningAgentsObservationObject::GetIndexObservation(int32& OutIndex, const FLearningAgentsObservationObjectElement Element, const FName Name) const
 {
-	if (!Interactor->HasAgent(AgentId))
+	int32 DiscreteValueNum;
+	if (!GetContinuousObservationNum(DiscreteValueNum, Element, Name))
 	{
-		UE_LOG(LogLearning, Error, TEXT("%s: AgentId %d not found in the agents set."), *GetName(), AgentId);
-		return;
+		OutIndex = INDEX_NONE;
+		return false;
 	}
 
-	const TLearningArrayView<2, FQuat> RotationView = FeatureObject->InstanceData->View(FeatureObject->RotationHandle);
-	const TLearningArrayView<1, FQuat> RelativeRotationView = FeatureObject->InstanceData->View(FeatureObject->RelativeRotationHandle);
+	TArray<float, TInlineAllocator<32>> OneHot;
+	OneHot.SetNumUninitialized(DiscreteValueNum);
+	if (!GetContinuousObservationToArrayView(OneHot, Element, Name))
+	{
+		OutIndex = INDEX_NONE;
+		return false;
+	}
 
-	if (Rotations.Num() != RotationView.Num<1>())
+	OutIndex = 0;
+	for (int32 Idx = 0; Idx < DiscreteValueNum; Idx++)
 	{
-		UE_LOG(LogLearning, Error, TEXT("%s: Got wrong number of elements in array. Expected %i, got %i."), *GetName(), RotationView.Num<1>(), Rotations.Num());
-		return;
+		if (OneHot[Idx]) { OutIndex = Idx; }
 	}
 
-	UE::Learning::Array::Copy<1, FQuat>(RotationView[AgentId], Rotations);
-	RelativeRotationView[AgentId] = RelativeRotation;
-	AgentIteration[AgentId]++;
+	return true;
 }
 
-#if UE_LEARNING_AGENTS_ENABLE_VISUAL_LOG
-void URotationArrayObservation::VisualLog(const UE::Learning::FIndexSet Instances) const
+bool ULearningAgentsObservationObject::GetStructObservationNum(int32& OutNum, const FLearningAgentsObservationObjectElement Element, const FName Name) const
 {
-	UE_LEARNING_TRACE_CPUPROFILER_EVENT_SCOPE(URotationArrayObservation::VisualLog);
-
-	const TLearningArrayView<2, const FQuat> RotationView = FeatureObject->InstanceData->ConstView(FeatureObject->RotationHandle);
-	const TLearningArrayView<1, const FQuat> RelativeRotationView = FeatureObject->InstanceData->ConstView(FeatureObject->RelativeRotationHandle);
-	const TLearningArrayView<2, const float> FeatureView = FeatureObject->InstanceData->ConstView(FeatureObject->FeatureHandle);
+	if (!ObservationObject.IsValid(Element.ObjectElement))
+	{
+		UE_LOG(LogLearning, Error, TEXT("%s: Invalid Observation Object."), *GetName());
+		OutNum = 0;
+		return false;
+	}
 
-	const int32 RotationNum = RotationView.Num<1>();
+	if (ObservationObject.GetName(Element.ObjectElement) != Name)
+	{
+		UE_LOG(LogLearning, Warning, TEXT("%s: Observation name does not match. Observation is '%s' but asked for '%s'."), *GetName(), *ObservationObject.GetName(Element.ObjectElement).ToString(), *Name.ToString());
+	}
 
-	for (const int32 Instance : Instances)
+	if (ObservationObject.GetType(Element.ObjectElement) != UE::Learning::Observation::EType::And)
 	{
-		if (const AActor* Actor = Cast<AActor>(Interactor->GetAgent(Instance)))
-		{
-			const FQuat RelativeRotation = RelativeRotationView[Instance];
-			const FRotator RelativeRotator = RelativeRotation.Rotator();
+		UE_LOG(LogLearning, Error, TEXT("%s: Observation '%s' type does not match. Observation is '%s' but asked for '%s'."),
+			*GetName(),
+			*ObservationObject.GetName(Element.ObjectElement).ToString(),
+			UE::Learning::Agents::Observation::Private::GetObservationTypeString(ObservationObject.GetType(Element.ObjectElement)),
+			UE::Learning::Agents::Observation::Private::GetObservationTypeString(UE::Learning::Observation::EType::And));
+		OutNum = 0;
+		return false;
+	}
 
-			for (int32 RotationIdx = 0; RotationIdx < RotationNum; RotationIdx++)
-			{
-				const FQuat Rotation = RotationView[Instance][RotationIdx];
-				const FRotator Rotator = Rotation.Rotator();
-				const FVector Offset = UE::Learning::Agents::Debug::GridOffsetForIndex(RotationIdx, RotationNum);
-
-				UE_LEARNING_AGENTS_VLOG_TRANSFORM(this, LogLearning, Display,
-					Actor->GetActorLocation() + Offset,
-					Rotation,
-					VisualLogColor.ToFColor(true),
-					TEXT("Rotation %i: [% 6.3f % 6.3f % 6.3f]"),
-					RotationIdx,
-					Rotator.Pitch, Rotator.Roll, Rotator.Yaw);
-			}
+	const UE::Learning::Observation::FObjectAndParameters Parameters = ObservationObject.GetAnd(Element.ObjectElement);
 
-			UE_LEARNING_AGENTS_VLOG_TRANSFORM(this, LogLearning, Display,
-				Actor->GetActorLocation(),
-				RelativeRotation,
-				VisualLogColor.ToFColor(true),
-				TEXT("Agent %i\nScale: [% 6.2f]\nRelative Rotation: [% 6.3f % 6.3f % 6.3f]\nEncoded: %s"),
-				Instance,
-				FeatureObject->Scale,
-				RelativeRotator.Pitch, RelativeRotator.Roll, RelativeRotator.Yaw,
-				*UE::Learning::Array::FormatFloat(FeatureView[Instance]));
-		}
-	}
+	OutNum = Parameters.Elements.Num();
+	return true;
 }
-#endif
 
+bool ULearningAgentsObservationObject::GetStructObservation(TMap<FName, FLearningAgentsObservationObjectElement>& OutElements, const FLearningAgentsObservationObjectElement Element, const FName Name) const
+{
+	int32 OutElementNum = 0;
+	if (!GetStructObservationNum(OutElementNum, Element, Name))
+	{
+		OutElements.Empty();
+		return false;
+	}
 
-//------------------------------------------------------------------
+	TArray<FName, TInlineAllocator<16>> SubElementNames;
+	TArray<FLearningAgentsObservationObjectElement, TInlineAllocator<16>> SubElements;
+	SubElementNames.SetNumUninitialized(OutElementNum);
+	SubElements.SetNumUninitialized(OutElementNum);
 
-UDirectionObservation* UDirectionObservation::AddDirectionObservation(ULearningAgentsInteractor* InInteractor, const FName Name, const float Scale)
-{
-	return UE::Learning::Agents::Observations::Private::AddObservation<UDirectionObservation, UE::Learning::FDirectionFeature>(InInteractor, Name, TEXT("AddDirectionObservation"), 1, Scale);
-}
+	if (!GetStructObservationToArrayViews(SubElementNames, SubElements, Element, Name))
+	{
+		OutElements.Empty();
+		return false;
+	}
 
-void UDirectionObservation::SetDirectionObservation(const int32 AgentId, const FVector Direction, const FRotator RelativeRotation)
-{
-	if (!Interactor->HasAgent(AgentId))
+	OutElements.Empty(OutElementNum);
+	for (int32 ElementIdx = 0; ElementIdx < OutElementNum; ElementIdx++)
 	{
-		UE_LOG(LogLearning, Error, TEXT("%s: AgentId %d not found in the agents set."), *GetName(), AgentId);
-		return;
+		OutElements.Add(SubElementNames[ElementIdx], SubElements[ElementIdx]);
 	}
 
-	FeatureObject->InstanceData->View(FeatureObject->DirectionHandle)[AgentId][0] = Direction;
-	FeatureObject->InstanceData->View(FeatureObject->RelativeRotationHandle)[AgentId] = RelativeRotation.Quaternion();
-	AgentIteration[AgentId]++;
+	return true;
 }
 
-#if UE_LEARNING_AGENTS_ENABLE_VISUAL_LOG
-void UDirectionObservation::VisualLog(const UE::Learning::FIndexSet Instances) const
+bool ULearningAgentsObservationObject::GetStructObservationToArrays(TArray<FName>& OutElementNames, TArray<FLearningAgentsObservationObjectElement>& OutElements, const FLearningAgentsObservationObjectElement Element, const FName Name) const
 {
-	UE_LEARNING_TRACE_CPUPROFILER_EVENT_SCOPE(UDirectionObservation::VisualLog);
-
-	const TLearningArrayView<2, const FVector> DirectionView = FeatureObject->InstanceData->ConstView(FeatureObject->DirectionHandle);
-	const TLearningArrayView<1, const FQuat> RelativeRotationView = FeatureObject->InstanceData->ConstView(FeatureObject->RelativeRotationHandle);
-	const TLearningArrayView<2, const float> FeatureView = FeatureObject->InstanceData->ConstView(FeatureObject->FeatureHandle);
-
-	for (const int32 Instance : Instances)
+	int32 OutElementNum = 0;
+	if (!GetStructObservationNum(OutElementNum, Element, Name))
 	{
-		if (const AActor* Actor = Cast<AActor>(Interactor->GetAgent(Instance)))
-		{
-			const FVector Direction = DirectionView[Instance][0];
-			const FQuat RelativeRotation = RelativeRotationView[Instance];
-			const FVector LocalDirection = RelativeRotation.UnrotateVector(Direction);
-
-			UE_LEARNING_AGENTS_VLOG_ARROW(this, LogLearning, Display,
-				Actor->GetActorLocation(),
-				Actor->GetActorLocation() + 100.0f * Direction,
-				VisualLogColor.ToFColor(true),
-				TEXT(""));
-
-			UE_LEARNING_AGENTS_VLOG_STRING(this, LogLearning, Display,
-				Actor->GetActorLocation() + 100.0f * Direction,
-				VisualLogColor.ToFColor(true),
-				TEXT("Direction: [% 6.3f % 6.3f % 6.3f]\nLocal Direction: [% 6.3f % 6.3f % 6.3f]"),
-				Direction.X, Direction.Y, Direction.Z,
-				LocalDirection.X, LocalDirection.Y, LocalDirection.Z);
-
-			UE_LEARNING_AGENTS_VLOG_TRANSFORM(this, LogLearning, Display,
-				Actor->GetActorLocation(),
-				RelativeRotation,
-				VisualLogColor.ToFColor(true),
-				TEXT("Agent %i\nScale: [% 6.2f]\nEncoded: %s"),
-				Instance,
-				FeatureObject->Scale,
-				*UE::Learning::Array::FormatFloat(FeatureView[Instance]));
-		}
+		OutElementNames.Empty();
+		OutElements.Empty();
+		return false;
 	}
-}
-#endif
 
-UDirectionArrayObservation* UDirectionArrayObservation::AddDirectionArrayObservation(ULearningAgentsInteractor* InInteractor, const FName Name, const int32 DirectionNum, const float Scale)
-{
-	if (DirectionNum < 1)
+	OutElementNames.SetNumUninitialized(OutElementNum);
+	OutElements.SetNumUninitialized(OutElementNum);
+
+	if (!GetStructObservationToArrayViews(OutElementNames, OutElements, Element, Name))
 	{
-		UE_LOG(LogLearning, Error, TEXT("AddDirectionArrayObservation: Number of elements in array must be at least 1, got %i."), DirectionNum);
-		return nullptr;
+		OutElementNames.Empty();
+		OutElements.Empty();
+		return false;
 	}
 
-	return UE::Learning::Agents::Observations::Private::AddObservation<UDirectionArrayObservation, UE::Learning::FDirectionFeature>(InInteractor, Name, TEXT("AddDirectionArrayObservation"), DirectionNum, Scale);
+	return true;
 }
 
-void UDirectionArrayObservation::SetDirectionArrayObservation(const int32 AgentId, const TArray<FVector>& Directions, const FRotator RelativeRotation)
+bool ULearningAgentsObservationObject::GetStructObservationToArrayViews(TArrayView<FName> OutElementNames, TArrayView<FLearningAgentsObservationObjectElement> OutElements, const FLearningAgentsObservationObjectElement Element, const FName Name) const
 {
-	if (!Interactor->HasAgent(AgentId))
+	if (!ObservationObject.IsValid(Element.ObjectElement))
 	{
-		UE_LOG(LogLearning, Error, TEXT("%s: AgentId %d not found in the agents set."), *GetName(), AgentId);
-		return;
+		UE_LOG(LogLearning, Error, TEXT("%s: Invalid Observation Object."), *GetName());
+		UE::Learning::Array::Set<1, FName>(OutElementNames, NAME_None);
+		UE::Learning::Array::Set<1, FLearningAgentsObservationObjectElement>(OutElements, FLearningAgentsObservationObjectElement());
+		return false;
 	}
 
-	const TLearningArrayView<2, FVector> DirectionView = FeatureObject->InstanceData->View(FeatureObject->DirectionHandle);
-	const TLearningArrayView<1, FQuat> RelativeRotationView = FeatureObject->InstanceData->View(FeatureObject->RelativeRotationHandle);
-
-	if (Directions.Num() != DirectionView.Num<1>())
+	if (ObservationObject.GetName(Element.ObjectElement) != Name)
 	{
-		UE_LOG(LogLearning, Error, TEXT("%s: Got wrong number of elements in array. Expected %i, got %i."), *GetName(), DirectionView.Num<1>(), Directions.Num());
-		return;
+		UE_LOG(LogLearning, Warning, TEXT("%s: Observation name does not match. Observation is '%s' but asked for '%s'."), *GetName(), *ObservationObject.GetName(Element.ObjectElement).ToString(), *Name.ToString());
 	}
 
-	UE::Learning::Array::Copy<1, FVector>(DirectionView[AgentId], Directions);
-	RelativeRotationView[AgentId] = RelativeRotation.Quaternion();
-	AgentIteration[AgentId]++;
-}
+	if (ObservationObject.GetType(Element.ObjectElement) != UE::Learning::Observation::EType::And)
+	{
+		UE_LOG(LogLearning, Error, TEXT("%s: Observation '%s' type does not match. Observation is '%s' but asked for '%s'."),
+			*GetName(),
+			*ObservationObject.GetName(Element.ObjectElement).ToString(),
+			UE::Learning::Agents::Observation::Private::GetObservationTypeString(ObservationObject.GetType(Element.ObjectElement)),
+			UE::Learning::Agents::Observation::Private::GetObservationTypeString(UE::Learning::Observation::EType::And));
+		UE::Learning::Array::Set<1, FName>(OutElementNames, NAME_None);
+		UE::Learning::Array::Set<1, FLearningAgentsObservationObjectElement>(OutElements, FLearningAgentsObservationObjectElement());
+		return false;
+	}
 
-#if UE_LEARNING_AGENTS_ENABLE_VISUAL_LOG
-void UDirectionArrayObservation::VisualLog(const UE::Learning::FIndexSet Instances) const
-{
-	UE_LEARNING_TRACE_CPUPROFILER_EVENT_SCOPE(UDirectionArrayObservation::VisualLog);
+	const UE::Learning::Observation::FObjectAndParameters Parameters = ObservationObject.GetAnd(Element.ObjectElement);
 
-	const TLearningArrayView<2, const FVector> DirectionView = FeatureObject->InstanceData->ConstView(FeatureObject->DirectionHandle);
-	const TLearningArrayView<1, const FQuat> RelativeRotationView = FeatureObject->InstanceData->ConstView(FeatureObject->RelativeRotationHandle);
-	const TLearningArrayView<2, const float> FeatureView = FeatureObject->InstanceData->ConstView(FeatureObject->FeatureHandle);
+	if (Parameters.Elements.Num() == 0)
+	{
+		UE_LOG(LogLearning, Warning, TEXT("%s: Getting zero-sized And Observation."), *GetName());
+	}
 
-	const int32 DirectionNum = DirectionView.Num<1>();
+	if (Parameters.Elements.Num() != OutElements.Num())
+	{
+		UE_LOG(LogLearning, Error, TEXT("%s: Observation '%s' size does not match. Observation is '%i' elements but asked for '%i'."),
+			*GetName(),
+			*ObservationObject.GetName(Element.ObjectElement).ToString(),
+			Parameters.Elements.Num(), OutElements.Num());
+		UE::Learning::Array::Set<1, FName>(OutElementNames, NAME_None);
+		UE::Learning::Array::Set<1, FLearningAgentsObservationObjectElement>(OutElements, FLearningAgentsObservationObjectElement());
+		return false;
+	}
 
-	for (const int32 Instance : Instances)
+	for (int32 ElementIdx = 0; ElementIdx < Parameters.Elements.Num(); ElementIdx++)
 	{
-		if (const AActor* Actor = Cast<AActor>(Interactor->GetAgent(Instance)))
+		if (!ObservationObject.IsValid(Parameters.Elements[ElementIdx]))
 		{
-			const FQuat RelativeRotation = RelativeRotationView[Instance];
-
-			for (int32 DirectionIdx = 0; DirectionIdx < DirectionNum; DirectionIdx++)
-			{
-				const FVector Direction = DirectionView[Instance][0];
-				const FVector LocalDirection = RelativeRotation.UnrotateVector(Direction);
-				const FVector Offset = UE::Learning::Agents::Debug::GridOffsetForIndex(DirectionIdx, DirectionNum);
-
-				UE_LEARNING_AGENTS_VLOG_ARROW(this, LogLearning, Display,
-					Actor->GetActorLocation() + Offset,
-					Actor->GetActorLocation() + Offset + 10.0f * Direction,
-					VisualLogColor.ToFColor(true),
-					TEXT(""));
-
-				UE_LEARNING_AGENTS_VLOG_STRING(this, LogLearning, Display,
-					Actor->GetActorLocation() + Offset + 10.0f * Direction,
-					VisualLogColor.ToFColor(true),
-					TEXT("Direction %i: [% 6.3f % 6.3f % 6.3f]\nLocal Direction: [% 6.3f % 6.3f % 6.3f]"),
-					DirectionIdx,
-					Direction.X, Direction.Y, Direction.Z,
-					LocalDirection.X, LocalDirection.Y, LocalDirection.Z);
-			}
-
-			UE_LEARNING_AGENTS_VLOG_TRANSFORM(this, LogLearning, Display,
-				Actor->GetActorLocation(),
-				RelativeRotation,
-				VisualLogColor.ToFColor(true),
-				TEXT("Agent %i\nScale: [% 6.2f]\nEncoded: %s"),
-				Instance,
-				FeatureObject->Scale,
-				*UE::Learning::Array::FormatFloat(FeatureView[Instance]));
+			UE_LOG(LogLearning, Error, TEXT("%s: Invalid Observation Object."), *GetName());
+			UE::Learning::Array::Set<1, FName>(OutElementNames, NAME_None);
+			UE::Learning::Array::Set<1, FLearningAgentsObservationObjectElement>(OutElements, FLearningAgentsObservationObjectElement());
+			return false;
 		}
+
+		OutElementNames[ElementIdx] = Parameters.ElementNames[ElementIdx];
+		OutElements[ElementIdx] = { Parameters.Elements[ElementIdx] };
 	}
-}
-#endif
 
-UPlanarDirectionObservation* UPlanarDirectionObservation::AddPlanarDirectionObservation(ULearningAgentsInteractor* InInteractor, const FName Name, const float Scale, const FVector Axis0, const FVector Axis1)
-{
-	return UE::Learning::Agents::Observations::Private::AddObservation<UPlanarDirectionObservation, UE::Learning::FPlanarDirectionFeature>(
-		InInteractor, 
-		Name, 
-		TEXT("AddPlanarDirectionObservation"), 
-		1, 
-		Scale, 
-		Axis0.GetSafeNormal(UE_SMALL_NUMBER, FVector::ForwardVector), 
-		Axis1.GetSafeNormal(UE_SMALL_NUMBER, FVector::RightVector));
+	return true;
 }
 
-void UPlanarDirectionObservation::SetPlanarDirectionObservation(const int32 AgentId, const FVector Direction, const FRotator RelativeRotation)
+bool ULearningAgentsObservationObject::GetExclusiveUnionObservation(FName& OutElementName, FLearningAgentsObservationObjectElement& OutElement, const FLearningAgentsObservationObjectElement Element, const FName Name) const
 {
-	if (!Interactor->HasAgent(AgentId))
+	if (!ObservationObject.IsValid(Element.ObjectElement))
 	{
-		UE_LOG(LogLearning, Error, TEXT("%s: AgentId %d not found in the agents set."), *GetName(), AgentId);
-		return;
+		UE_LOG(LogLearning, Error, TEXT("%s: Invalid Observation Object."), *GetName());
+		OutElementName = NAME_None;
+		OutElement = FLearningAgentsObservationObjectElement();
+		return false;
 	}
 
-	FeatureObject->InstanceData->View(FeatureObject->DirectionHandle)[AgentId][0] = Direction;
-	FeatureObject->InstanceData->View(FeatureObject->RelativeRotationHandle)[AgentId] = RelativeRotation.Quaternion();
-	AgentIteration[AgentId]++;
-}
-
-#if UE_LEARNING_AGENTS_ENABLE_VISUAL_LOG
-void UPlanarDirectionObservation::VisualLog(const UE::Learning::FIndexSet Instances) const
-{
-	UE_LEARNING_TRACE_CPUPROFILER_EVENT_SCOPE(UPlanarDirectionObservation::VisualLog);
-
-	const TLearningArrayView<2, const FVector> DirectionView = FeatureObject->InstanceData->ConstView(FeatureObject->DirectionHandle);
-	const TLearningArrayView<1, const FQuat> RelativeRotationView = FeatureObject->InstanceData->ConstView(FeatureObject->RelativeRotationHandle);
-	const TLearningArrayView<2, const float> FeatureView = FeatureObject->InstanceData->ConstView(FeatureObject->FeatureHandle);
-
-	for (const int32 Instance : Instances)
+	if (ObservationObject.GetName(Element.ObjectElement) != Name)
 	{
-		if (const AActor* Actor = Cast<AActor>(Interactor->GetAgent(Instance)))
-		{
-			const FVector Direction = DirectionView[Instance][0];
-			const FQuat RelativeRotation = RelativeRotationView[Instance];
-			const FVector LocalDirection = RelativeRotation.UnrotateVector(Direction);
-
-			UE_LEARNING_AGENTS_VLOG_ARROW(this, LogLearning, Display,
-				Actor->GetActorLocation(),
-				Actor->GetActorLocation() + 100.0f * Direction,
-				VisualLogColor.ToFColor(true),
-				TEXT(""));
-
-			UE_LEARNING_AGENTS_VLOG_STRING(this, LogLearning, Display,
-				Actor->GetActorLocation() + 100.0f * Direction,
-				VisualLogColor.ToFColor(true),
-				TEXT("Direction: [% 6.3f % 6.3f % 6.3f]\nLocal Direction: [% 6.3f % 6.3f % 6.3f]"),
-				Direction.X, Direction.Y, Direction.Z,
-				LocalDirection.X, LocalDirection.Y, LocalDirection.Z);
-
-			UE_LEARNING_AGENTS_VLOG_PLANE(this, LogLearning, Display,
-				Actor->GetActorLocation(),
-				RelativeRotation,
-				FeatureObject->Axis0,
-				FeatureObject->Axis1,
-				VisualLogColor.ToFColor(true),
-				TEXT(""));
-
-			UE_LEARNING_AGENTS_VLOG_TRANSFORM(this, LogLearning, Display,
-				Actor->GetActorLocation(),
-				RelativeRotation,
-				VisualLogColor.ToFColor(true),
-				TEXT("Agent %i\nScale: [% 6.2f]\nEncoded: %s"),
-				Instance,
-				FeatureObject->Scale,
-				*UE::Learning::Array::FormatFloat(FeatureView[Instance]));
-		}
+		UE_LOG(LogLearning, Warning, TEXT("%s: Observation name does not match. Observation is '%s' but asked for '%s'."), *GetName(), *ObservationObject.GetName(Element.ObjectElement).ToString(), *Name.ToString());
 	}
-}
-#endif
 
-UPlanarDirectionArrayObservation* UPlanarDirectionArrayObservation::AddPlanarDirectionArrayObservation(ULearningAgentsInteractor* InInteractor, const FName Name, const int32 DirectionNum, const float Scale, const FVector Axis0, const FVector Axis1)
-{
-	if (DirectionNum < 1)
+	if (ObservationObject.GetType(Element.ObjectElement) != UE::Learning::Observation::EType::OrExclusive)
 	{
-		UE_LOG(LogLearning, Error, TEXT("AddPlanarDirectionArrayObservation: Number of elements in array must be at least 1, got %i."), DirectionNum);
-		return nullptr;
+		UE_LOG(LogLearning, Error, TEXT("%s: Observation '%s' type does not match. Observation is '%s' but asked for '%s'."),
+			*GetName(),
+			*ObservationObject.GetName(Element.ObjectElement).ToString(),
+			UE::Learning::Agents::Observation::Private::GetObservationTypeString(ObservationObject.GetType(Element.ObjectElement)),
+			UE::Learning::Agents::Observation::Private::GetObservationTypeString(UE::Learning::Observation::EType::OrExclusive));
+		OutElementName = NAME_None;
+		OutElement = FLearningAgentsObservationObjectElement();
+		return false;
 	}
 
-	return UE::Learning::Agents::Observations::Private::AddObservation<UPlanarDirectionArrayObservation, UE::Learning::FPlanarDirectionFeature>(
-		InInteractor, 
-		Name, 
-		TEXT("AddPlanarDirectionArrayObservation"), 
-		DirectionNum, 
-		Scale, 
-		Axis0.GetSafeNormal(UE_SMALL_NUMBER, FVector::ForwardVector),
-		Axis1.GetSafeNormal(UE_SMALL_NUMBER, FVector::RightVector));
+	const UE::Learning::Observation::FObjectOrExclusiveParameters Parameters = ObservationObject.GetOrExclusive(Element.ObjectElement);
+	OutElementName = Parameters.ElementName;
+	OutElement = { Parameters.Element };
+	return true;
 }
 
-void UPlanarDirectionArrayObservation::SetPlanarDirectionArrayObservation(const int32 AgentId, const TArray<FVector>& Directions, const FRotator RelativeRotation)
+bool ULearningAgentsObservationObject::GetInclusiveUnionObservationNum(int32& OutNum, const FLearningAgentsObservationObjectElement Element, const FName Name) const
 {
-	if (!Interactor->HasAgent(AgentId))
+	if (!ObservationObject.IsValid(Element.ObjectElement))
 	{
-		UE_LOG(LogLearning, Error, TEXT("%s: AgentId %d not found in the agents set."), *GetName(), AgentId);
-		return;
+		UE_LOG(LogLearning, Error, TEXT("%s: Invalid Observation Object."), *GetName());
+		OutNum = 0;
+		return false;
 	}
 
-	const TLearningArrayView<2, FVector> DirectionView = FeatureObject->InstanceData->View(FeatureObject->DirectionHandle);
-	const TLearningArrayView<1, FQuat> RelativeRotationView = FeatureObject->InstanceData->View(FeatureObject->RelativeRotationHandle);
+	if (ObservationObject.GetName(Element.ObjectElement) != Name)
+	{
+		UE_LOG(LogLearning, Warning, TEXT("%s: Observation name does not match. Observation is '%s' but asked for '%s'."), *GetName(), *ObservationObject.GetName(Element.ObjectElement).ToString(), *Name.ToString());
+	}
 
-	if (Directions.Num() != DirectionView.Num<1>())
+	if (ObservationObject.GetType(Element.ObjectElement) != UE::Learning::Observation::EType::OrInclusive)
 	{
-		UE_LOG(LogLearning, Error, TEXT("%s: Got wrong number of elements in array. Expected %i, got %i."), *GetName(), DirectionView.Num<1>(), Directions.Num());
-		return;
+		UE_LOG(LogLearning, Error, TEXT("%s: Observation '%s' type does not match. Observation is '%s' but asked for '%s'."),
+			*GetName(),
+			*ObservationObject.GetName(Element.ObjectElement).ToString(),
+			UE::Learning::Agents::Observation::Private::GetObservationTypeString(ObservationObject.GetType(Element.ObjectElement)),
+			UE::Learning::Agents::Observation::Private::GetObservationTypeString(UE::Learning::Observation::EType::OrInclusive));
+		OutNum = 0;
+		return false;
 	}
 
-	UE::Learning::Array::Copy<1, FVector>(DirectionView[AgentId], Directions);
-	RelativeRotationView[AgentId] = RelativeRotation.Quaternion();
-	AgentIteration[AgentId]++;
+	const UE::Learning::Observation::FObjectOrInclusiveParameters Parameters = ObservationObject.GetOrInclusive(Element.ObjectElement);
+
+	OutNum = Parameters.Elements.Num();
+	return true;
 }
 
-#if UE_LEARNING_AGENTS_ENABLE_VISUAL_LOG
-void UPlanarDirectionArrayObservation::VisualLog(const UE::Learning::FIndexSet Instances) const
+bool ULearningAgentsObservationObject::GetInclusiveUnionObservation(TMap<FName, FLearningAgentsObservationObjectElement>& OutElements, const FLearningAgentsObservationObjectElement Element, const FName Name) const
 {
-	UE_LEARNING_TRACE_CPUPROFILER_EVENT_SCOPE(UPlanarDirectionArrayObservation::VisualLog);
-
-	const TLearningArrayView<2, const FVector> DirectionView = FeatureObject->InstanceData->ConstView(FeatureObject->DirectionHandle);
-	const TLearningArrayView<1, const FQuat> RelativeRotationView = FeatureObject->InstanceData->ConstView(FeatureObject->RelativeRotationHandle);
-	const TLearningArrayView<2, const float> FeatureView = FeatureObject->InstanceData->ConstView(FeatureObject->FeatureHandle);
-
-	const int32 DirectionNum = DirectionView.Num<1>();
-	
-	for (const int32 Instance : Instances)
+	int32 OutElementNum = 0;
+	if (!GetInclusiveUnionObservationNum(OutElementNum, Element, Name))
 	{
-		if (const AActor* Actor = Cast<AActor>(Interactor->GetAgent(Instance)))
-		{
-			const FQuat RelativeRotation = RelativeRotationView[Instance];
-
-			for (int32 DirectionIdx = 0; DirectionIdx < DirectionNum; DirectionIdx++)
-			{
-				const FVector Direction = DirectionView[Instance][0];
-				const FVector LocalDirection = RelativeRotation.UnrotateVector(Direction);
-				const FVector Offset = UE::Learning::Agents::Debug::GridOffsetForIndex(DirectionIdx, DirectionNum);
-
-				UE_LEARNING_AGENTS_VLOG_ARROW(this, LogLearning, Display,
-					Actor->GetActorLocation() + Offset,
-					Actor->GetActorLocation() + Offset + 10.0f * Direction,
-					VisualLogColor.ToFColor(true),
-					TEXT(""));
-
-				UE_LEARNING_AGENTS_VLOG_STRING(this, LogLearning, Display,
-					Actor->GetActorLocation() + Offset + 10.0f * Direction,
-					VisualLogColor.ToFColor(true),
-					TEXT("Direction %i: [% 6.3f % 6.3f % 6.3f]\nLocal Direction: [% 6.3f % 6.3f % 6.3f]"),
-					DirectionIdx,
-					Direction.X, Direction.Y, Direction.Z,
-					LocalDirection.X, LocalDirection.Y, LocalDirection.Z);
-			}
-
-			UE_LEARNING_AGENTS_VLOG_PLANE(this, LogLearning, Display,
-				Actor->GetActorLocation() + FVector(0.0f, 0.0f, 20.0f),
-				RelativeRotation,
-				FeatureObject->Axis0,
-				FeatureObject->Axis1,
-				VisualLogColor.ToFColor(true),
-				TEXT(""));
-
-			UE_LEARNING_AGENTS_VLOG_TRANSFORM(this, LogLearning, Display,
-				Actor->GetActorLocation(),
-				RelativeRotation,
-				VisualLogColor.ToFColor(true),
-				TEXT("Agent %i\nScale: [% 6.2f]\nEncoded: %s"),
-				Instance,
-				FeatureObject->Scale,
-				*UE::Learning::Array::FormatFloat(FeatureView[Instance]));
-		}
+		OutElements.Empty();
+		return false;
 	}
-}
-#endif
 
+	TArray<FName, TInlineAllocator<16>> SubElementNames;
+	TArray<FLearningAgentsObservationObjectElement, TInlineAllocator<16>> SubElements;
+	SubElementNames.SetNumUninitialized(OutElementNum);
+	SubElements.SetNumUninitialized(OutElementNum);
 
+	if (!GetInclusiveUnionObservationToArrayViews(SubElementNames, SubElements, Element, Name))
+	{
+		OutElements.Empty();
+		return false;
+	}
 
-//------------------------------------------------------------------
+	OutElements.Empty(OutElementNum);
+	for (int32 ElementIdx = 0; ElementIdx < OutElementNum; ElementIdx++)
+	{
+		OutElements.Add(SubElementNames[ElementIdx], SubElements[ElementIdx]);
+	}
 
-UPositionObservation* UPositionObservation::AddPositionObservation(ULearningAgentsInteractor* InInteractor, const FName Name, const float Scale)
-{
-	return UE::Learning::Agents::Observations::Private::AddObservation<UPositionObservation, UE::Learning::FPositionFeature>(InInteractor, Name, TEXT("AddPositionObservation"), 1, Scale);
+	return true;
 }
 
-void UPositionObservation::SetPositionObservation(const int32 AgentId, const FVector Position, const FVector RelativePosition, const FRotator RelativeRotation)
+bool ULearningAgentsObservationObject::GetInclusiveUnionObservationToArrays(TArray<FName>& OutElementNames, TArray<FLearningAgentsObservationObjectElement>& OutElements, const FLearningAgentsObservationObjectElement Element, const FName Name) const
 {
-	if (!Interactor->HasAgent(AgentId))
+	int32 OutElementNum = 0;
+	if (!GetInclusiveUnionObservationNum(OutElementNum, Element, Name))
+	{
+		OutElementNames.Empty();
+		OutElements.Empty();
+		return false;
+	}
+
+	OutElementNames.SetNumUninitialized(OutElementNum);
+	OutElements.SetNumUninitialized(OutElementNum);
+
+	if (!GetInclusiveUnionObservationToArrayViews(OutElementNames, OutElements, Element, Name))
 	{
-		UE_LOG(LogLearning, Error, TEXT("%s: AgentId %d not found in the agents set."), *GetName(), AgentId);
-		return;
+		OutElementNames.Empty();
+		OutElements.Empty();
+		return false;
 	}
 
-	FeatureObject->InstanceData->View(FeatureObject->PositionHandle)[AgentId][0] = Position;
-	FeatureObject->InstanceData->View(FeatureObject->RelativePositionHandle)[AgentId] = RelativePosition;
-	FeatureObject->InstanceData->View(FeatureObject->RelativeRotationHandle)[AgentId] = RelativeRotation.Quaternion();
-	AgentIteration[AgentId]++;
+	return true;
 }
 
-#if UE_LEARNING_AGENTS_ENABLE_VISUAL_LOG
-void UPositionObservation::VisualLog(const UE::Learning::FIndexSet Instances) const
+bool ULearningAgentsObservationObject::GetInclusiveUnionObservationToArrayViews(TArrayView<FName> OutElementNames, TArrayView<FLearningAgentsObservationObjectElement> OutElements, const FLearningAgentsObservationObjectElement Element, const FName Name) const
 {
-	UE_LEARNING_TRACE_CPUPROFILER_EVENT_SCOPE(UPositionObservation::VisualLog);
-
-	const TLearningArrayView<2, const FVector> PositionView = FeatureObject->InstanceData->ConstView(FeatureObject->PositionHandle);
-	const TLearningArrayView<1, const FVector> RelativePositionView = FeatureObject->InstanceData->ConstView(FeatureObject->RelativePositionHandle);
-	const TLearningArrayView<1, const FQuat> RelativeRotationView = FeatureObject->InstanceData->ConstView(FeatureObject->RelativeRotationHandle);
-	const TLearningArrayView<2, const float> FeatureView = FeatureObject->InstanceData->ConstView(FeatureObject->FeatureHandle);
+	if (!ObservationObject.IsValid(Element.ObjectElement))
+	{
+		UE_LOG(LogLearning, Error, TEXT("%s: Invalid Observation Object."), *GetName());
+		UE::Learning::Array::Set<1, FName>(OutElementNames, NAME_None);
+		UE::Learning::Array::Set<1, FLearningAgentsObservationObjectElement>(OutElements, FLearningAgentsObservationObjectElement());
+		return false;
+	}
 
-	for (const int32 Instance : Instances)
+	if (ObservationObject.GetName(Element.ObjectElement) != Name)
 	{
-		const FVector Position = PositionView[Instance][0];
-		const FVector RelativePosition = RelativePositionView[Instance];
-		const FQuat RelativeRotation = RelativeRotationView[Instance];
-		const FVector LocalPosition = RelativeRotation.UnrotateVector(Position - RelativePosition);
+		UE_LOG(LogLearning, Warning, TEXT("%s: Observation name does not match. Observation is '%s' but asked for '%s'."), *GetName(), *ObservationObject.GetName(Element.ObjectElement).ToString(), *Name.ToString());
+	}
 
-		UE_LEARNING_AGENTS_VLOG_LOCATION(this, LogLearning, Display,
-			Position,
-			10,
-			VisualLogColor.ToFColor(true),
-			TEXT("Position: [% 6.1f % 6.1f % 6.1f]\nLocal Position: [% 6.1f % 6.1f % 6.1f]"),
-			Position.X, Position.Y, Position.Z,
-			LocalPosition.X, LocalPosition.Y, LocalPosition.Z);
+	if (ObservationObject.GetType(Element.ObjectElement) != UE::Learning::Observation::EType::OrInclusive)
+	{
+		UE_LOG(LogLearning, Error, TEXT("%s: Observation '%s' type does not match. Observation is '%s' but asked for '%s'."),
+			*GetName(),
+			*ObservationObject.GetName(Element.ObjectElement).ToString(),
+			UE::Learning::Agents::Observation::Private::GetObservationTypeString(ObservationObject.GetType(Element.ObjectElement)),
+			UE::Learning::Agents::Observation::Private::GetObservationTypeString(UE::Learning::Observation::EType::OrInclusive));
+		UE::Learning::Array::Set<1, FName>(OutElementNames, NAME_None);
+		UE::Learning::Array::Set<1, FLearningAgentsObservationObjectElement>(OutElements, FLearningAgentsObservationObjectElement());
+		return false;
+	}
 
-		UE_LEARNING_AGENTS_VLOG_SEGMENT(this, LogLearning, Display,
-			RelativePosition,
-			Position,
-			VisualLogColor.ToFColor(true),
-			TEXT(""));
+	const UE::Learning::Observation::FObjectOrInclusiveParameters Parameters = ObservationObject.GetOrInclusive(Element.ObjectElement);
 
-		UE_LEARNING_AGENTS_VLOG_TRANSFORM(this, LogLearning, Display,
-			RelativePosition,
-			RelativeRotation,
-			VisualLogColor.ToFColor(true),
-			TEXT("Agent %i\nScale: [% 6.2f]\nEncoded: %s"),
-			Instance,
-			FeatureObject->Scale,
-			*UE::Learning::Array::FormatFloat(FeatureView[Instance]));
+	if (Parameters.Elements.Num() != OutElements.Num())
+	{
+		UE_LOG(LogLearning, Error, TEXT("%s: Observation '%s' size does not match. Observation is '%i' elements but asked for '%i'."),
+			*GetName(),
+			*ObservationObject.GetName(Element.ObjectElement).ToString(),
+			Parameters.Elements.Num(), OutElements.Num());
+		UE::Learning::Array::Set<1, FName>(OutElementNames, NAME_None);
+		UE::Learning::Array::Set<1, FLearningAgentsObservationObjectElement>(OutElements, FLearningAgentsObservationObjectElement());
+		return false;
 	}
-}
-#endif
 
-UPositionArrayObservation* UPositionArrayObservation::AddPositionArrayObservation(ULearningAgentsInteractor* InInteractor, const FName Name, const int32 PositionNum, const float Scale)
-{
-	if (PositionNum < 1)
+	for (int32 ElementIdx = 0; ElementIdx < Parameters.Elements.Num(); ElementIdx++)
 	{
-		UE_LOG(LogLearning, Error, TEXT("AddPositionArrayObservation: Number of elements in array must be at least 1, got %i."), PositionNum);
-		return nullptr;
+		if (!ObservationObject.IsValid(Parameters.Elements[ElementIdx]))
+		{
+			UE_LOG(LogLearning, Error, TEXT("%s: Invalid Observation Object."), *GetName());
+			UE::Learning::Array::Set<1, FLearningAgentsObservationObjectElement>(OutElements, FLearningAgentsObservationObjectElement());
+			UE::Learning::Array::Set<1, FName>(OutElementNames, NAME_None);
+			return false;
+		}
+
+		OutElementNames[ElementIdx] = Parameters.ElementNames[ElementIdx];
+		OutElements[ElementIdx] = { Parameters.Elements[ElementIdx] };
 	}
 
-	return UE::Learning::Agents::Observations::Private::AddObservation<UPositionArrayObservation, UE::Learning::FPositionFeature>(InInteractor, Name, TEXT("AddPositionArrayObservation"), PositionNum, Scale);
+	return true;
 }
 
-void UPositionArrayObservation::SetPositionArrayObservation(const int32 AgentId, const TArray<FVector>& Positions, const FVector RelativePosition, const FRotator RelativeRotation)
+bool ULearningAgentsObservationObject::GetStaticArrayObservationNum(int32& OutNum, const FLearningAgentsObservationObjectElement Element, const FName Name) const
 {
-	if (!Interactor->HasAgent(AgentId))
+	if (!ObservationObject.IsValid(Element.ObjectElement))
 	{
-		UE_LOG(LogLearning, Error, TEXT("%s: AgentId %d not found in the agents set."), *GetName(), AgentId);
-		return;
+		UE_LOG(LogLearning, Error, TEXT("%s: Invalid Observation Object."), *GetName());
+		OutNum = 0;
+		return false;
 	}
 
-	const TLearningArrayView<2, FVector> PositionView = FeatureObject->InstanceData->View(FeatureObject->PositionHandle);
-	const TLearningArrayView<1, FVector> RelativePositionView = FeatureObject->InstanceData->View(FeatureObject->RelativePositionHandle);
-	const TLearningArrayView<1, FQuat> RelativeRotationView = FeatureObject->InstanceData->View(FeatureObject->RelativeRotationHandle);
+	if (ObservationObject.GetName(Element.ObjectElement) != Name)
+	{
+		UE_LOG(LogLearning, Warning, TEXT("%s: Observation name does not match. Observation is '%s' but asked for '%s'."), *GetName(), *ObservationObject.GetName(Element.ObjectElement).ToString(), *Name.ToString());
+	}
 
-	if (Positions.Num() != PositionView.Num<1>())
+	if (ObservationObject.GetType(Element.ObjectElement) != UE::Learning::Observation::EType::Array)
 	{
-		UE_LOG(LogLearning, Error, TEXT("%s: Got wrong number of elements in array. Expected %i, got %i."), *GetName(), PositionView.Num<1>(), Positions.Num());
-		return;
+		UE_LOG(LogLearning, Error, TEXT("%s: Observation '%s' type does not match. Observation is '%s' but asked for '%s'."),
+			*GetName(),
+			*ObservationObject.GetName(Element.ObjectElement).ToString(),
+			UE::Learning::Agents::Observation::Private::GetObservationTypeString(ObservationObject.GetType(Element.ObjectElement)),
+			UE::Learning::Agents::Observation::Private::GetObservationTypeString(UE::Learning::Observation::EType::Array));
+		OutNum = 0;
+		return false;
 	}
 
-	RelativePositionView[AgentId] = RelativePosition;
-	RelativeRotationView[AgentId] = RelativeRotation.Quaternion();
-	UE::Learning::Array::Copy<1, FVector>(PositionView[AgentId], Positions);
-	AgentIteration[AgentId]++;
+	OutNum = ObservationObject.GetArray(Element.ObjectElement).Elements.Num();
+	return true;
 }
 
-#if UE_LEARNING_AGENTS_ENABLE_VISUAL_LOG
-void UPositionArrayObservation::VisualLog(const UE::Learning::FIndexSet Instances) const
+bool ULearningAgentsObservationObject::GetStaticArrayObservation(TArray<FLearningAgentsObservationObjectElement>& OutElements, const FLearningAgentsObservationObjectElement Element, const FName Name) const
 {
-	UE_LEARNING_TRACE_CPUPROFILER_EVENT_SCOPE(UPositionArrayObservation::VisualLog);
-
-	const TLearningArrayView<2, const FVector> PositionView = FeatureObject->InstanceData->ConstView(FeatureObject->PositionHandle);
-	const TLearningArrayView<1, const FVector> RelativePositionView = FeatureObject->InstanceData->ConstView(FeatureObject->RelativePositionHandle);
-	const TLearningArrayView<1, const FQuat> RelativeRotationView = FeatureObject->InstanceData->ConstView(FeatureObject->RelativeRotationHandle);
-	const TLearningArrayView<2, const float> FeatureView = FeatureObject->InstanceData->ConstView(FeatureObject->FeatureHandle);
-
-	const int32 PositionNum = PositionView.Num<1>();
-
-	for (const int32 Instance : Instances)
+	int32 OutElementNum = 0;
+	if (!GetStaticArrayObservationNum(OutElementNum, Element, Name))
 	{
-		const FVector RelativePosition = RelativePositionView[Instance];
-		const FQuat RelativeRotation = RelativeRotationView[Instance];
+		OutElements.Empty();
+		return false;
+	}
 
-		for (int32 PositionIdx = 0; PositionIdx < PositionNum; PositionIdx++)
-		{
-			const FVector LocalPosition = RelativeRotation.UnrotateVector(PositionView[Instance][PositionIdx] - RelativePosition);
-
-			UE_LEARNING_AGENTS_VLOG_LOCATION(this, LogLearning, Display,
-				PositionView[Instance][PositionIdx],
-				10,
-				VisualLogColor.ToFColor(true),
-				TEXT("Position: [% 6.1f % 6.1f % 6.1f]\nLocal Position: [% 6.1f % 6.1f % 6.1f]"),
-				PositionView[Instance][PositionIdx].X,
-				PositionView[Instance][PositionIdx].Y,
-				PositionView[Instance][PositionIdx].Z,
-				LocalPosition.X, LocalPosition.Y, LocalPosition.Z);
-
-			UE_LEARNING_AGENTS_VLOG_SEGMENT(this, LogLearning, Display,
-				RelativePosition,
-				PositionView[Instance][PositionIdx],
-				VisualLogColor.ToFColor(true),
-				TEXT(""));
-		}
+	OutElements.SetNumUninitialized(OutElementNum);
 
-		UE_LEARNING_AGENTS_VLOG_TRANSFORM(this, LogLearning, Display,
-			RelativePosition,
-			RelativeRotation,
-			VisualLogColor.ToFColor(true),
-			TEXT("Agent %i\nScale: [% 6.2f]\nEncoded: %s"),
-			Instance,
-			FeatureObject->Scale,
-			*UE::Learning::Array::FormatFloat(FeatureView[Instance]));
+	if (!GetStaticArrayObservationToArrayView(OutElements, Element, Name))
+	{
+		OutElements.Empty();
+		return false;
 	}
-}
-#endif
 
-UScalarPositionObservation* UScalarPositionObservation::AddScalarPositionObservation(ULearningAgentsInteractor* InInteractor, const FName Name, const float Scale)
-{
-	return UE::Learning::Agents::Observations::Private::AddObservation<UScalarPositionObservation, UE::Learning::FScalarPositionFeature>(InInteractor, Name, TEXT("AddScalarPositionObservation"), 1, Scale);
+	return true;
 }
 
-void UScalarPositionObservation::SetScalarPositionObservation(const int32 AgentId, const float Position, const float RelativePosition)
+bool ULearningAgentsObservationObject::GetStaticArrayObservationToArrayView(TArrayView<FLearningAgentsObservationObjectElement> OutElements, const FLearningAgentsObservationObjectElement Element, const FName Name) const
 {
-	if (!Interactor->HasAgent(AgentId))
+	if (!ObservationObject.IsValid(Element.ObjectElement))
 	{
-		UE_LOG(LogLearning, Error, TEXT("%s: AgentId %d not found in the agents set."), *GetName(), AgentId);
-		return;
+		UE_LOG(LogLearning, Error, TEXT("%s: Invalid Observation Object."), *GetName());
+		UE::Learning::Array::Set<1, FLearningAgentsObservationObjectElement>(OutElements, FLearningAgentsObservationObjectElement());
+		return false;
 	}
 
-	FeatureObject->InstanceData->View(FeatureObject->PositionHandle)[AgentId][0] = Position;
-	FeatureObject->InstanceData->View(FeatureObject->RelativePositionHandle)[AgentId] = RelativePosition;
-	AgentIteration[AgentId]++;
-}
+	if (ObservationObject.GetName(Element.ObjectElement) != Name)
+	{
+		UE_LOG(LogLearning, Warning, TEXT("%s: Observation name does not match. Observation is '%s' but asked for '%s'."), *GetName(), *ObservationObject.GetName(Element.ObjectElement).ToString(), *Name.ToString());
+	}
 
-void UScalarPositionObservation::SetScalarPositionObservationWithAxis(const int32 AgentId, const FVector Position, const FVector RelativePosition, const FVector Axis)
-{
-	SetScalarPositionObservation(Position.Dot(Axis), RelativePosition.Dot(Axis));
-}
+	if (ObservationObject.GetType(Element.ObjectElement) != UE::Learning::Observation::EType::Array)
+	{
+		UE_LOG(LogLearning, Error, TEXT("%s: Observation '%s' type does not match. Observation is '%s' but asked for '%s'."),
+			*GetName(),
+			*ObservationObject.GetName(Element.ObjectElement).ToString(),
+			UE::Learning::Agents::Observation::Private::GetObservationTypeString(ObservationObject.GetType(Element.ObjectElement)),
+			UE::Learning::Agents::Observation::Private::GetObservationTypeString(UE::Learning::Observation::EType::Array));
+		UE::Learning::Array::Set<1, FLearningAgentsObservationObjectElement>(OutElements, FLearningAgentsObservationObjectElement());
+		return false;
+	}
 
-#if UE_LEARNING_AGENTS_ENABLE_VISUAL_LOG
-void UScalarPositionObservation::VisualLog(const UE::Learning::FIndexSet Instances) const
-{
-	UE_LEARNING_TRACE_CPUPROFILER_EVENT_SCOPE(UScalarPositionObservation::VisualLog);
+	const UE::Learning::Observation::FObjectArrayParameters Parameters = ObservationObject.GetArray(Element.ObjectElement);
+
+	if (Parameters.Elements.Num() == 0)
+	{
+		UE_LOG(LogLearning, Warning, TEXT("%s: Getting zero-sized Static Array Observation."), *GetName());
+	}
 
-	const TLearningArrayView<2, const float> PositionView = FeatureObject->InstanceData->ConstView(FeatureObject->PositionHandle);
-	const TLearningArrayView<1, const float> RelativePositionView = FeatureObject->InstanceData->ConstView(FeatureObject->RelativePositionHandle);
-	const TLearningArrayView<2, const float> FeatureView = FeatureObject->InstanceData->ConstView(FeatureObject->FeatureHandle);
+	if (Parameters.Elements.Num() != OutElements.Num())
+	{
+		UE_LOG(LogLearning, Error, TEXT("%s: Observation '%s' size does not match. Observation is '%i' elements but asked for '%i'."),
+			*GetName(), *ObservationObject.GetName(Element.ObjectElement).ToString(),
+			Parameters.Elements.Num(), OutElements.Num());
+		UE::Learning::Array::Set<1, FLearningAgentsObservationObjectElement>(OutElements, FLearningAgentsObservationObjectElement());
+		return false;
+	}
 
-	for (const int32 Instance : Instances)
+	for (int32 ElementIdx = 0; ElementIdx < Parameters.Elements.Num(); ElementIdx++)
 	{
-		if (const AActor* Actor = Cast<AActor>(Interactor->GetAgent(Instance)))
+		if (!ObservationObject.IsValid(Parameters.Elements[ElementIdx]))
 		{
-			UE_LEARNING_AGENTS_VLOG_STRING(this, LogLearning, Display,
-				Actor->GetActorLocation(),
-				VisualLogColor.ToFColor(true),
-				TEXT("Agent %i\nScale: [% 6.2f]\nRelative Position [% 6.1f]\nLocal Position: [% 6.1f]\nPosition: [% 6.1f]\nEncoded: %s"),
-				Instance,
-				FeatureObject->Scale,
-				RelativePositionView[Instance],
-				PositionView[Instance][0] - RelativePositionView[Instance],
-				PositionView[Instance][0],
-				*UE::Learning::Array::FormatFloat(FeatureView[Instance]));
+			UE_LOG(LogLearning, Error, TEXT("%s: Invalid Observation Object."), *GetName());
+			UE::Learning::Array::Set<1, FLearningAgentsObservationObjectElement>(OutElements, FLearningAgentsObservationObjectElement());
+			return false;
 		}
-	}
-}
-#endif
 
-UScalarPositionArrayObservation* UScalarPositionArrayObservation::AddScalarPositionArrayObservation(ULearningAgentsInteractor* InInteractor, const FName Name, const int32 PositionNum, const float Scale)
-{
-	if (PositionNum < 1)
-	{
-		UE_LOG(LogLearning, Error, TEXT("AddScalarPositionArrayObservation: Number of elements in array must be at least 1, got %i."), PositionNum);
-		return nullptr;
+		OutElements[ElementIdx] = { Parameters.Elements[ElementIdx] };
 	}
 
-	return UE::Learning::Agents::Observations::Private::AddObservation<UScalarPositionArrayObservation, UE::Learning::FScalarPositionFeature>(InInteractor, Name, TEXT("AddScalarPositionArrayObservation"), PositionNum, Scale);
+	return true;
 }
 
-void UScalarPositionArrayObservation::SetScalarPositionArrayObservation(const int32 AgentId, const TArray<float>& Positions, const float RelativePosition)
+bool ULearningAgentsObservationObject::GetSetObservationNum(int32& OutNum, const FLearningAgentsObservationObjectElement Element, const FName Name) const
 {
-	if (!Interactor->HasAgent(AgentId))
+	if (!ObservationObject.IsValid(Element.ObjectElement))
 	{
-		UE_LOG(LogLearning, Error, TEXT("%s: AgentId %d not found in the agents set."), *GetName(), AgentId);
-		return;
+		UE_LOG(LogLearning, Error, TEXT("%s: Invalid Observation Object."), *GetName());
+		OutNum = 0;
+		return false;
 	}
 
-	const TLearningArrayView<2, float> PositionView = FeatureObject->InstanceData->View(FeatureObject->PositionHandle);
-	const TLearningArrayView<1, float> RelativePositionView = FeatureObject->InstanceData->View(FeatureObject->RelativePositionHandle);
+	if (ObservationObject.GetName(Element.ObjectElement) != Name)
+	{
+		UE_LOG(LogLearning, Warning, TEXT("%s: Observation name does not match. Observation is '%s' but asked for '%s'."), *GetName(), *ObservationObject.GetName(Element.ObjectElement).ToString(), *Name.ToString());
+	}
 
-	if (Positions.Num() != PositionView.Num<1>())
+	if (ObservationObject.GetType(Element.ObjectElement) != UE::Learning::Observation::EType::Set)
 	{
-		UE_LOG(LogLearning, Error, TEXT("%s: Got wrong number of elements in array. Expected %i, got %i."), *GetName(), PositionView.Num<1>(), Positions.Num());
-		return;
+		UE_LOG(LogLearning, Error, TEXT("%s: Observation '%s' type does not match. Observation is '%s' but asked for '%s'."),
+			*GetName(),
+			*ObservationObject.GetName(Element.ObjectElement).ToString(),
+			UE::Learning::Agents::Observation::Private::GetObservationTypeString(ObservationObject.GetType(Element.ObjectElement)),
+			UE::Learning::Agents::Observation::Private::GetObservationTypeString(UE::Learning::Observation::EType::Set));
+		OutNum = 0;
+		return false;
 	}
 
-	RelativePositionView[AgentId] = RelativePosition;
-	UE::Learning::Array::Copy<1, float>(PositionView[AgentId], Positions);
-	AgentIteration[AgentId]++;
+	OutNum = ObservationObject.GetSet(Element.ObjectElement).Elements.Num();
+	return true;
 }
 
-void UScalarPositionArrayObservation::SetScalarPositionArrayObservationWithAxis(const int32 AgentId, const TArray<FVector>& Positions, const FVector RelativePosition, const FVector Axis)
+bool ULearningAgentsObservationObject::GetSetObservation(TSet<FLearningAgentsObservationObjectElement>& OutElements, const FLearningAgentsObservationObjectElement Element, const FName Name) const
 {
-	if (!Interactor->HasAgent(AgentId))
+	int32 OutElementNum = 0;
+	if (!GetSetObservationNum(OutElementNum, Element, Name))
 	{
-		UE_LOG(LogLearning, Error, TEXT("%s: AgentId %d not found in the agents set."), *GetName(), AgentId);
-		return;
+		OutElements.Empty();
+		return false;
 	}
 
-	const TLearningArrayView<2, float> PositionView = FeatureObject->InstanceData->View(FeatureObject->PositionHandle);
-	const TLearningArrayView<1, float> RelativePositionView = FeatureObject->InstanceData->View(FeatureObject->RelativePositionHandle);
-
-	if (Positions.Num() != PositionView.Num<1>())
+	if (!ObservationObject.IsValid(Element.ObjectElement))
 	{
-		UE_LOG(LogLearning, Error, TEXT("%s: Got wrong number of elements in array. Expected %i, got %i."), *GetName(), PositionView.Num<1>(), Positions.Num());
-		return;
+		UE_LOG(LogLearning, Error, TEXT("%s: Invalid Observation Object."), *GetName());
+		OutElements.Empty();
+		return false;
 	}
 
-	RelativePositionView[AgentId] = RelativePosition.Dot(Axis);
-	for (int32 PositionIdx = 0; PositionIdx < PositionView.Num<1>(); PositionIdx++)
+	if (ObservationObject.GetName(Element.ObjectElement) != Name)
 	{
-		PositionView[AgentId][PositionIdx] = Positions[PositionIdx].Dot(Axis);
+		UE_LOG(LogLearning, Warning, TEXT("%s: Observation name does not match. Observation is '%s' but asked for '%s'."), *GetName(), *ObservationObject.GetName(Element.ObjectElement).ToString(), *Name.ToString());
 	}
 
-	AgentIteration[AgentId]++;
-}
-
-#if UE_LEARNING_AGENTS_ENABLE_VISUAL_LOG
-void UScalarPositionArrayObservation::VisualLog(const UE::Learning::FIndexSet Instances) const
-{
-	UE_LEARNING_TRACE_CPUPROFILER_EVENT_SCOPE(UScalarPositionArrayObservation::VisualLog);
+	if (ObservationObject.GetType(Element.ObjectElement) != UE::Learning::Observation::EType::Set)
+	{
+		UE_LOG(LogLearning, Error, TEXT("%s: Observation '%s' type does not match. Observation is '%s' but asked for '%s'."),
+			*GetName(),
+			*ObservationObject.GetName(Element.ObjectElement).ToString(),
+			UE::Learning::Agents::Observation::Private::GetObservationTypeString(ObservationObject.GetType(Element.ObjectElement)),
+			UE::Learning::Agents::Observation::Private::GetObservationTypeString(UE::Learning::Observation::EType::Set));
+		OutElements.Empty();
+		return false;
+	}
 
-	const TLearningArrayView<2, const float> PositionView = FeatureObject->InstanceData->ConstView(FeatureObject->PositionHandle);
-	const TLearningArrayView<1, const float> RelativePositionView = FeatureObject->InstanceData->ConstView(FeatureObject->RelativePositionHandle);
-	const TLearningArrayView<2, const float> FeatureView = FeatureObject->InstanceData->ConstView(FeatureObject->FeatureHandle);
+	const UE::Learning::Observation::FObjectSetParameters Parameters = ObservationObject.GetSet(Element.ObjectElement);
 
-	const int32 PositionNum = PositionView.Num<1>();
+	if (Parameters.Elements.Num() != OutElements.Num())
+	{
+		UE_LOG(LogLearning, Error, TEXT("%s: Observation '%s' size does not match. Observation is '%i' elements but asked for '%i'."),
+			*GetName(), *ObservationObject.GetName(Element.ObjectElement).ToString(),
+			Parameters.Elements.Num(), OutElements.Num());
+		OutElements.Empty();
+		return false;
+	}
 
-	for (const int32 Instance : Instances)
+	OutElements.Empty(Parameters.Elements.Num());
+	for (int32 ElementIdx = 0; ElementIdx < Parameters.Elements.Num(); ElementIdx++)
 	{
-		if (const AActor* Actor = Cast<AActor>(Interactor->GetAgent(Instance)))
+		if (!ObservationObject.IsValid(Parameters.Elements[ElementIdx]))
 		{
-			UE_LEARNING_AGENTS_VLOG_STRING(this, LogLearning, Display,
-				Actor->GetActorLocation(),
-				VisualLogColor.ToFColor(true),
-				TEXT("Agent %i\nScale: [% 6.2f]\nRelative Position: [% 6.2f]\nPositions: %s\nEncoded: %s"),
-				Instance,
-				FeatureObject->Scale,
-				RelativePositionView[Instance],
-				*UE::Learning::Array::FormatFloat(PositionView[Instance]),
-				*UE::Learning::Array::FormatFloat(FeatureView[Instance]));
+			UE_LOG(LogLearning, Error, TEXT("%s: Invalid Observation Object."), *GetName());
+			OutElements.Empty();
+			return false;
 		}
+
+		OutElements.Add({ Parameters.Elements[ElementIdx] });
 	}
-}
-#endif
 
-UPlanarPositionObservation* UPlanarPositionObservation::AddPlanarPositionObservation(ULearningAgentsInteractor* InInteractor, const FName Name, const float Scale, const FVector Axis0, const FVector Axis1)
-{
-	return UE::Learning::Agents::Observations::Private::AddObservation<UPlanarPositionObservation, UE::Learning::FPlanarPositionFeature>(
-		InInteractor, 
-		Name, 
-		TEXT("AddPlanarPositionObservation"), 
-		1, 
-		Scale, 
-		Axis0.GetSafeNormal(UE_SMALL_NUMBER, FVector::ForwardVector),
-		Axis1.GetSafeNormal(UE_SMALL_NUMBER, FVector::RightVector));
+	return true;
 }
 
-void UPlanarPositionObservation::SetPlanarPositionObservation(const int32 AgentId, const FVector Position, const FVector RelativePosition, const FRotator RelativeRotation)
+bool ULearningAgentsObservationObject::GetSetObservationToArray(TArray<FLearningAgentsObservationObjectElement>& OutElements, const FLearningAgentsObservationObjectElement Element, const FName Name) const
 {
-	if (!Interactor->HasAgent(AgentId))
+	int32 OutElementNum = 0;
+	if (!GetSetObservationNum(OutElementNum, Element, Name))
+	{
+		OutElements.Empty();
+		return false;
+	}
+
+	OutElements.SetNumUninitialized(OutElementNum);
+
+	if (!GetSetObservationToArrayView(OutElements, Element, Name))
 	{
-		UE_LOG(LogLearning, Error, TEXT("%s: AgentId %d not found in the agents set."), *GetName(), AgentId);
-		return;
+		OutElements.Empty();
+		return false;
 	}
 
-	FeatureObject->InstanceData->View(FeatureObject->PositionHandle)[AgentId][0] = Position;
-	FeatureObject->InstanceData->View(FeatureObject->RelativePositionHandle)[AgentId] = RelativePosition;
-	FeatureObject->InstanceData->View(FeatureObject->RelativeRotationHandle)[AgentId] = RelativeRotation.Quaternion();
-	AgentIteration[AgentId]++;
+	return true;
 }
 
-#if UE_LEARNING_AGENTS_ENABLE_VISUAL_LOG
-void UPlanarPositionObservation::VisualLog(const UE::Learning::FIndexSet Instances) const
+bool ULearningAgentsObservationObject::GetSetObservationToArrayView(TArrayView<FLearningAgentsObservationObjectElement> OutElements, const FLearningAgentsObservationObjectElement Element, const FName Name) const
 {
-	UE_LEARNING_TRACE_CPUPROFILER_EVENT_SCOPE(UPlanarPositionObservation::VisualLog);
+	if (!ObservationObject.IsValid(Element.ObjectElement))
+	{
+		UE_LOG(LogLearning, Error, TEXT("%s: Invalid Observation Object."), *GetName());
+		UE::Learning::Array::Set<1, FLearningAgentsObservationObjectElement>(OutElements, FLearningAgentsObservationObjectElement());
+		return false;
+	}
 
-	const TLearningArrayView<2, const FVector> PositionView = FeatureObject->InstanceData->ConstView(FeatureObject->PositionHandle);
-	const TLearningArrayView<1, const FVector> RelativePositionView = FeatureObject->InstanceData->ConstView(FeatureObject->RelativePositionHandle);
-	const TLearningArrayView<1, const FQuat> RelativeRotationView = FeatureObject->InstanceData->ConstView(FeatureObject->RelativeRotationHandle);
-	const TLearningArrayView<2, const float> FeatureView = FeatureObject->InstanceData->ConstView(FeatureObject->FeatureHandle);
+	if (ObservationObject.GetName(Element.ObjectElement) != Name)
+	{
+		UE_LOG(LogLearning, Warning, TEXT("%s: Observation name does not match. Observation is '%s' but asked for '%s'."), *GetName(), *ObservationObject.GetName(Element.ObjectElement).ToString(), *Name.ToString());
+	}
 
-	for (const int32 Instance : Instances)
+	if (ObservationObject.GetType(Element.ObjectElement) != UE::Learning::Observation::EType::Set)
 	{
-		const FVector Position = PositionView[Instance][0];
-		const FVector RelativePosition = RelativePositionView[Instance];
-		const FQuat RelativeRotation = RelativeRotationView[Instance];
-		const FVector LocalPosition = RelativeRotation.UnrotateVector(Position - RelativePosition);
+		UE_LOG(LogLearning, Error, TEXT("%s: Observation '%s' type does not match. Observation is '%s' but asked for '%s'."),
+			*GetName(),
+			*ObservationObject.GetName(Element.ObjectElement).ToString(),
+			UE::Learning::Agents::Observation::Private::GetObservationTypeString(ObservationObject.GetType(Element.ObjectElement)),
+			UE::Learning::Agents::Observation::Private::GetObservationTypeString(UE::Learning::Observation::EType::Set));
+		UE::Learning::Array::Set<1, FLearningAgentsObservationObjectElement>(OutElements, FLearningAgentsObservationObjectElement());
+		return false;
+	}
 
-		UE_LEARNING_AGENTS_VLOG_LOCATION(this, LogLearning, Display,
-			Position,
-			10,
-			VisualLogColor.ToFColor(true),
-			TEXT("Position: [% 6.1f % 6.1f % 6.1f]\nLocal Position: [% 6.1f % 6.1f % 6.1f]"),
-			Position.X, Position.Y, Position.Z,
-			LocalPosition.X, LocalPosition.Y, LocalPosition.Z);
+	const UE::Learning::Observation::FObjectSetParameters Parameters = ObservationObject.GetSet(Element.ObjectElement);
 
-		UE_LEARNING_AGENTS_VLOG_SEGMENT(this, LogLearning, Display,
-			RelativePosition,
-			Position,
-			VisualLogColor.ToFColor(true),
-			TEXT(""));
+	if (Parameters.Elements.Num() != OutElements.Num())
+	{
+		UE_LOG(LogLearning, Error, TEXT("%s: Observation '%s' size does not match. Observation is '%i' elements but asked for '%i'."),
+			*GetName(), *ObservationObject.GetName(Element.ObjectElement).ToString(),
+			Parameters.Elements.Num(), OutElements.Num());
+		UE::Learning::Array::Set<1, FLearningAgentsObservationObjectElement>(OutElements, FLearningAgentsObservationObjectElement());
+		return false;
+	}
 
-		UE_LEARNING_AGENTS_VLOG_PLANE(this, LogLearning, Display,
-			RelativePosition,
-			RelativeRotation,
-			FeatureObject->Axis0,
-			FeatureObject->Axis1,
-			VisualLogColor.ToFColor(true),
-			TEXT(""));
+	for (int32 ElementIdx = 0; ElementIdx < Parameters.Elements.Num(); ElementIdx++)
+	{
+		if (!ObservationObject.IsValid(Parameters.Elements[ElementIdx]))
+		{
+			UE_LOG(LogLearning, Error, TEXT("%s: Invalid Observation Object."), *GetName());
+			UE::Learning::Array::Set<1, FLearningAgentsObservationObjectElement>(OutElements, FLearningAgentsObservationObjectElement());
+			return false;
+		}
 
-		UE_LEARNING_AGENTS_VLOG_TRANSFORM(this, LogLearning, Display,
-			RelativePosition,
-			RelativeRotation,
-			VisualLogColor.ToFColor(true),
-			TEXT("Agent %i\nScale: [% 6.2f]\nEncoded: %s"),
-			Instance,
-			FeatureObject->Scale,
-			*UE::Learning::Array::FormatFloat(FeatureView[Instance]));
+		OutElements[ElementIdx] = { Parameters.Elements[ElementIdx] };
 	}
+
+	return true;
 }
-#endif
 
-UPlanarPositionArrayObservation* UPlanarPositionArrayObservation::AddPlanarPositionArrayObservation(ULearningAgentsInteractor* InInteractor, const FName Name, const int32 PositionNum, const float Scale, const FVector Axis0, const FVector Axis1)
+bool ULearningAgentsObservationObject::GetPairObservation(FLearningAgentsObservationObjectElement& OutKey, FLearningAgentsObservationObjectElement& OutValue, const FLearningAgentsObservationObjectElement Element, const FName Name) const
 {
-	if (PositionNum < 1)
+	TStaticArray<FName, 2> OutElementNames;
+	TStaticArray<FLearningAgentsObservationObjectElement, 2> OutElements;
+	if (!GetStructObservationToArrayViews(OutElementNames, OutElements, Element, Name))
 	{
-		UE_LOG(LogLearning, Error, TEXT("AddPlanarPositionArrayObservation: Number of elements in array must be at least 1, got %i."), PositionNum);
-		return nullptr;
+		OutKey = FLearningAgentsObservationObjectElement();
+		OutValue = FLearningAgentsObservationObjectElement();
+		return false;
 	}
 
-	return UE::Learning::Agents::Observations::Private::AddObservation<UPlanarPositionArrayObservation, UE::Learning::FPlanarPositionFeature>(
-		InInteractor, 
-		Name, 
-		TEXT("AddPlanarPositionArrayObservation"), 
-		PositionNum, 
-		Scale, 
-		Axis0.GetSafeNormal(UE_SMALL_NUMBER, FVector::ForwardVector),
-		Axis1.GetSafeNormal(UE_SMALL_NUMBER, FVector::RightVector));
+	OutKey = OutElements[MakeArrayView(OutElementNames).Find(TEXT("Key"))];
+	OutValue = OutElements[MakeArrayView(OutElementNames).Find(TEXT("Value"))];
+	return true;
 }
 
-void UPlanarPositionArrayObservation::SetPlanarPositionArrayObservation(const int32 AgentId, const TArray<FVector>& Positions, const FVector RelativePosition, const FRotator RelativeRotation)
+bool ULearningAgentsObservationObject::GetArrayObservationNum(int32& OutNum, const FLearningAgentsObservationObjectElement Element, const FName Name) const
 {
-	if (!Interactor->HasAgent(AgentId))
+	return GetSetObservationNum(OutNum, Element, Name);
+}
+
+bool ULearningAgentsObservationObject::GetArrayObservation(TArray<FLearningAgentsObservationObjectElement>& OutElements, const FLearningAgentsObservationObjectElement Element, const FName Name) const
+{
+	int32 OutElementNum = 0;
+	if (!GetArrayObservationNum(OutElementNum, Element, Name))
 	{
-		UE_LOG(LogLearning, Error, TEXT("%s: AgentId %d not found in the agents set."), *GetName(), AgentId);
-		return;
+		OutElements.Empty();
+		return false;
 	}
 
-	const TLearningArrayView<2, FVector> PositionView = FeatureObject->InstanceData->View(FeatureObject->PositionHandle);
-	const TLearningArrayView<1, FVector> RelativePositionView = FeatureObject->InstanceData->View(FeatureObject->RelativePositionHandle);
-	const TLearningArrayView<1, FQuat> RelativeRotationView = FeatureObject->InstanceData->View(FeatureObject->RelativeRotationHandle);
+	OutElements.SetNumUninitialized(OutElementNum);
 
-	if (Positions.Num() != PositionView.Num<1>())
+	if (!GetArrayObservationToArrayView(OutElements, Element, Name))
 	{
-		UE_LOG(LogLearning, Error, TEXT("%s: Got wrong number of elements in array. Expected %i, got %i."), *GetName(), PositionView.Num<1>(), Positions.Num());
-		return;
+		OutElements.Empty();
+		return false;
 	}
 
-	RelativePositionView[AgentId] = RelativePosition;
-	RelativeRotationView[AgentId] = RelativeRotation.Quaternion();
-	UE::Learning::Array::Copy<1, FVector>(PositionView[AgentId], Positions);
-	AgentIteration[AgentId]++;
+	return true;
 }
 
-#if UE_LEARNING_AGENTS_ENABLE_VISUAL_LOG
-void UPlanarPositionArrayObservation::VisualLog(const UE::Learning::FIndexSet Instances) const
+bool ULearningAgentsObservationObject::GetArrayObservationToArrayView(TArrayView<FLearningAgentsObservationObjectElement> OutElements, const FLearningAgentsObservationObjectElement Element, const FName Name) const
 {
-	UE_LEARNING_TRACE_CPUPROFILER_EVENT_SCOPE(UPlanarPositionArrayObservation::VisualLog);
-
-	const TLearningArrayView<2, const FVector> PositionView = FeatureObject->InstanceData->ConstView(FeatureObject->PositionHandle);
-	const TLearningArrayView<1, const FVector> RelativePositionView = FeatureObject->InstanceData->ConstView(FeatureObject->RelativePositionHandle);
-	const TLearningArrayView<1, const FQuat> RelativeRotationView = FeatureObject->InstanceData->ConstView(FeatureObject->RelativeRotationHandle);
-	const TLearningArrayView<2, const float> FeatureView = FeatureObject->InstanceData->ConstView(FeatureObject->FeatureHandle);
-
-	const int32 PositionNum = PositionView.Num<1>();
-
-	for (const int32 Instance : Instances)
+	TArray<FLearningAgentsObservationObjectElement, TInlineAllocator<16>> Pairs;
+	Pairs.SetNumUninitialized(OutElements.Num());
+	if (!GetSetObservationToArrayView(Pairs, Element, Name))
 	{
-		const FVector RelativePosition = RelativePositionView[Instance];
-		const FQuat RelativeRotation = RelativeRotationView[Instance];
+		UE::Learning::Array::Set<1, FLearningAgentsObservationObjectElement>(OutElements, FLearningAgentsObservationObjectElement());
+		return false;
+	}
 
-		for (int32 PositionIdx = 0; PositionIdx < PositionNum; PositionIdx++)
+	for (int32 PairIdx = 0; PairIdx < Pairs.Num(); PairIdx++)
+	{
+		FLearningAgentsObservationObjectElement Key, Value;
+		if (!GetPairObservation(Pairs[PairIdx], Key, Value))
 		{
-			const FVector LocalPosition = RelativeRotation.UnrotateVector(PositionView[Instance][PositionIdx] - RelativePosition);
-
-			UE_LEARNING_AGENTS_VLOG_LOCATION(this, LogLearning, Display,
-				PositionView[Instance][PositionIdx],
-				10,
-				VisualLogColor.ToFColor(true),
-				TEXT("Position: [% 6.1f % 6.1f % 6.1f]\nLocal Position: [% 6.1f % 6.1f % 6.1f]"),
-				PositionView[Instance][PositionIdx].X, 
-				PositionView[Instance][PositionIdx].Y, 
-				PositionView[Instance][PositionIdx].Z,
-				LocalPosition.X, LocalPosition.Y, LocalPosition.Z);
-
-			UE_LEARNING_AGENTS_VLOG_SEGMENT(this, LogLearning, Display,
-				RelativePosition,
-				PositionView[Instance][PositionIdx],
-				VisualLogColor.ToFColor(true),
-				TEXT(""));
+			UE::Learning::Array::Set<1, FLearningAgentsObservationObjectElement>(OutElements, FLearningAgentsObservationObjectElement());
+			return false;
 		}
 
-		UE_LEARNING_AGENTS_VLOG_PLANE(this, LogLearning, Display,
-			RelativePosition,
-			RelativeRotation,
-			FeatureObject->Axis0,
-			FeatureObject->Axis1,
-			VisualLogColor.ToFColor(true),
-			TEXT(""));
-
+		int32 Index = INDEX_NONE;
+		if (!GetIndexObservation(Index, Key))
+		{
+			UE::Learning::Array::Set<1, FLearningAgentsObservationObjectElement>(OutElements, FLearningAgentsObservationObjectElement());
+			return false;
+		}
 
-		UE_LEARNING_AGENTS_VLOG_TRANSFORM(this, LogLearning, Display,
-			RelativePosition,
-			RelativeRotation,
-			VisualLogColor.ToFColor(true),
-			TEXT("Agent %i\nScale: [% 6.2f]\nEncoded: %s"),
-			Instance,
-			FeatureObject->Scale,
-			*UE::Learning::Array::FormatFloat(FeatureView[Instance]));
+		OutElements[Index] = Value;
 	}
-}
-#endif
 
-//------------------------------------------------------------------
+	return true;
+}
 
-UVelocityObservation* UVelocityObservation::AddVelocityObservation(ULearningAgentsInteractor* InInteractor, const FName Name, const float Scale)
+bool ULearningAgentsObservationObject::GetMapObservationNum(int32& OutNum, const FLearningAgentsObservationObjectElement Element, const FName Name) const
 {
-	return UE::Learning::Agents::Observations::Private::AddObservation<UVelocityObservation, UE::Learning::FVelocityFeature>(InInteractor, Name, TEXT("AddVelocityObservation"), 1, Scale);
+	return GetSetObservationNum(OutNum, Element, Name);
 }
 
-void UVelocityObservation::SetVelocityObservation(const int32 AgentId, const FVector Velocity, const FRotator RelativeRotation)
+bool ULearningAgentsObservationObject::GetMapObservation(TMap<FLearningAgentsObservationObjectElement, FLearningAgentsObservationObjectElement>& OutElements, const FLearningAgentsObservationObjectElement Element, const FName Name) const
 {
-	if (!Interactor->HasAgent(AgentId))
+	int32 OutElementNum = 0;
+	if (!GetMapObservationNum(OutElementNum, Element, Name))
 	{
-		UE_LOG(LogLearning, Error, TEXT("%s: AgentId %d not found in the agents set."), *GetName(), AgentId);
-		return;
+		OutElements.Empty();
+		return false;
 	}
 
-	FeatureObject->InstanceData->View(FeatureObject->VelocityHandle)[AgentId][0] = Velocity;
-	FeatureObject->InstanceData->View(FeatureObject->RelativeRotationHandle)[AgentId] = RelativeRotation.Quaternion();
-	AgentIteration[AgentId]++;
-}
-
-#if UE_LEARNING_AGENTS_ENABLE_VISUAL_LOG
-void UVelocityObservation::VisualLog(const UE::Learning::FIndexSet Instances) const
-{
-	UE_LEARNING_TRACE_CPUPROFILER_EVENT_SCOPE(UVelocityObservation::VisualLog);
-
-	const TLearningArrayView<2, const FVector> VelocityView = FeatureObject->InstanceData->ConstView(FeatureObject->VelocityHandle);
-	const TLearningArrayView<1, const FQuat> RelativeRotationView = FeatureObject->InstanceData->ConstView(FeatureObject->RelativeRotationHandle);
-	const TLearningArrayView<2, const float> FeatureView = FeatureObject->InstanceData->ConstView(FeatureObject->FeatureHandle);
+	TArray<FLearningAgentsObservationObjectElement, TInlineAllocator<16>> Pairs;
+	Pairs.SetNumUninitialized(OutElementNum);
+	if (!GetSetObservationToArrayView(Pairs, Element, Name))
+	{
+		OutElements.Empty();
+		return false;
+	}
 
-	for (const int32 Instance : Instances)
+	OutElements.Empty(OutElementNum);
+	for (int32 PairIdx = 0; PairIdx < OutElementNum; PairIdx++)
 	{
-		if (const AActor* Actor = Cast<AActor>(Interactor->GetAgent(Instance)))
+		FLearningAgentsObservationObjectElement Key, Value;
+		if (!GetPairObservation(Pairs[PairIdx], Key, Value))
 		{
-			const FVector Velocity = VelocityView[Instance][0];
-			const FQuat RelativeRotation = RelativeRotationView[Instance];
-			const FVector LocalVelocity = RelativeRotation.UnrotateVector(Velocity);
-
-			UE_LEARNING_AGENTS_VLOG_ARROW(this, LogLearning, Display,
-				Actor->GetActorLocation(),
-				Actor->GetActorLocation() + Velocity,
-				VisualLogColor.ToFColor(true),
-				TEXT(""));
-
-			UE_LEARNING_AGENTS_VLOG_STRING(this, LogLearning, Display,
-				Actor->GetActorLocation() + Velocity,
-				VisualLogColor.ToFColor(true),
-				TEXT("Velocity: [% 6.3f % 6.3f % 6.3f]\nLocal Velocity: [% 6.3f % 6.3f % 6.3f]"),
-				Velocity.X, Velocity.Y, Velocity.Z,
-				LocalVelocity.X, LocalVelocity.Y, LocalVelocity.Z);
-
-			UE_LEARNING_AGENTS_VLOG_TRANSFORM(this, LogLearning, Display,
-				Actor->GetActorLocation(),
-				RelativeRotation,
-				VisualLogColor.ToFColor(true),
-				TEXT("Agent %i\nScale: [% 6.2f]\nEncoded: %s"),
-				Instance,
-				FeatureObject->Scale,
-				*UE::Learning::Array::FormatFloat(FeatureView[Instance]));
+			OutElements.Empty();
+			return false;
 		}
-	}
-}
-#endif
 
-UVelocityArrayObservation* UVelocityArrayObservation::AddVelocityArrayObservation(ULearningAgentsInteractor* InInteractor, const FName Name, const int32 VelocityNum, const float Scale)
-{
-	if (VelocityNum < 1)
-	{
-		UE_LOG(LogLearning, Error, TEXT("AddVelocityArrayObservation: Number of elements in array must be at least 1, got %i."), VelocityNum);
-		return nullptr;
+		OutElements.Add(Key, Value);
 	}
 
-	return UE::Learning::Agents::Observations::Private::AddObservation<UVelocityArrayObservation, UE::Learning::FVelocityFeature>(InInteractor, Name, TEXT("AddVelocityArrayObservation"), VelocityNum, Scale);
+	return true;
 }
 
-void UVelocityArrayObservation::SetVelocityArrayObservation(const int32 AgentId, const TArray<FVector>& Velocities, const FRotator RelativeRotation)
+bool ULearningAgentsObservationObject::GetMapObservationToArrays(TArray<FLearningAgentsObservationObjectElement>& OutKeys, TArray<FLearningAgentsObservationObjectElement>& OutValues, const FLearningAgentsObservationObjectElement Element, const FName Name) const
 {
-	if (!Interactor->HasAgent(AgentId))
+	int32 OutElementNum = 0;
+	if (!GetMapObservationNum(OutElementNum, Element, Name))
 	{
-		UE_LOG(LogLearning, Error, TEXT("%s: AgentId %d not found in the agents set."), *GetName(), AgentId);
-		return;
+		OutKeys.Empty();
+		OutValues.Empty();
+		return false;
 	}
 
-	const TLearningArrayView<2, FVector> VelocityView = FeatureObject->InstanceData->View(FeatureObject->VelocityHandle);
-	const TLearningArrayView<1, FQuat> RelativeRotationView = FeatureObject->InstanceData->View(FeatureObject->RelativeRotationHandle);
-
-	if (Velocities.Num() != VelocityView.Num<1>())
+	OutKeys.SetNumUninitialized(OutElementNum);
+	OutValues.SetNumUninitialized(OutElementNum);
+	if (!GetMapObservationToArrayViews(OutKeys, OutValues, Element, Name))
 	{
-		UE_LOG(LogLearning, Error, TEXT("%s: Got wrong number of elements in array. Expected %i, got %i."), *GetName(), VelocityView.Num<1>(), Velocities.Num());
-		return;
+		OutKeys.Empty();
+		OutValues.Empty();
+		return false;
 	}
 
-	UE::Learning::Array::Copy<1, FVector>(VelocityView[AgentId], Velocities);
-	RelativeRotationView[AgentId] = RelativeRotation.Quaternion();
-	AgentIteration[AgentId]++;
+	return true;
 }
 
-#if UE_LEARNING_AGENTS_ENABLE_VISUAL_LOG
-void UVelocityArrayObservation::VisualLog(const UE::Learning::FIndexSet Instances) const
+bool ULearningAgentsObservationObject::GetMapObservationToArrayViews(TArrayView<FLearningAgentsObservationObjectElement> OutKeys, TArrayView<FLearningAgentsObservationObjectElement> OutValues, const FLearningAgentsObservationObjectElement Element, const FName Name) const
 {
-	UE_LEARNING_TRACE_CPUPROFILER_EVENT_SCOPE(UVelocityArrayObservation::VisualLog);
-
-	const TLearningArrayView<2, const FVector> VelocityView = FeatureObject->InstanceData->ConstView(FeatureObject->VelocityHandle);
-	const TLearningArrayView<1, const FQuat> RelativeRotationView = FeatureObject->InstanceData->ConstView(FeatureObject->RelativeRotationHandle);
-	const TLearningArrayView<2, const float> FeatureView = FeatureObject->InstanceData->ConstView(FeatureObject->FeatureHandle);
-
-	const int32 VelocityNum = VelocityView.Num<1>();
+	TArray<FLearningAgentsObservationObjectElement, TInlineAllocator<16>> Pairs;
+	Pairs.SetNumUninitialized(OutKeys.Num());
+	if (!GetSetObservationToArrayView(Pairs, Element, Name))
+	{
+		UE::Learning::Array::Set<1, FLearningAgentsObservationObjectElement>(OutKeys, FLearningAgentsObservationObjectElement());
+		UE::Learning::Array::Set<1, FLearningAgentsObservationObjectElement>(OutValues, FLearningAgentsObservationObjectElement());
+		return false;
+	}
 
-	for (const int32 Instance : Instances)
+	for (int32 PairIdx = 0; PairIdx < Pairs.Num(); PairIdx++)
 	{
-		if (const AActor* Actor = Cast<AActor>(Interactor->GetAgent(Instance)))
+		FLearningAgentsObservationObjectElement Key, Value;
+		if (!GetPairObservation(Pairs[PairIdx], Key, Value))
 		{
-			const FQuat RelativeRotation = RelativeRotationView[Instance];
-				
-			for (int32 VelocityIdx = 0; VelocityIdx < VelocityNum; VelocityIdx++)
-			{
-				const FVector Velocity = VelocityView[Instance][0];
-				const FVector LocalVelocity = RelativeRotation.UnrotateVector(Velocity);
-				const FVector Offset = UE::Learning::Agents::Debug::GridOffsetForIndex(VelocityIdx, VelocityNum);
-
-				UE_LEARNING_AGENTS_VLOG_ARROW(this, LogLearning, Display,
-					Actor->GetActorLocation() + Offset,
-					Actor->GetActorLocation() + Offset + Velocity,
-					VisualLogColor.ToFColor(true),
-					TEXT(""));
-
-				UE_LEARNING_AGENTS_VLOG_STRING(this, LogLearning, Display,
-					Actor->GetActorLocation() + Offset + Velocity,
-					VisualLogColor.ToFColor(true),
-					TEXT("Velocity %i: [% 6.3f % 6.3f % 6.3f]\nLocal Velocity: [% 6.3f % 6.3f % 6.3f]"),
-					VelocityIdx,
-					Velocity.X, Velocity.Y, Velocity.Z,
-					LocalVelocity.X, LocalVelocity.Y, LocalVelocity.Z);
-			}
-
-			UE_LEARNING_AGENTS_VLOG_TRANSFORM(this, LogLearning, Display,
-				Actor->GetActorLocation(),
-				RelativeRotation,
-				VisualLogColor.ToFColor(true),
-				TEXT("Agent %i\nScale: [% 6.2f]\nEncoded: %s"),
-				Instance,
-				FeatureObject->Scale,
-				*UE::Learning::Array::FormatFloat(FeatureView[Instance]));
+			UE::Learning::Array::Set<1, FLearningAgentsObservationObjectElement>(OutKeys, FLearningAgentsObservationObjectElement());
+			UE::Learning::Array::Set<1, FLearningAgentsObservationObjectElement>(OutValues, FLearningAgentsObservationObjectElement());
+			return false;
 		}
+
+		OutKeys[PairIdx] = Key;
+		OutValues[PairIdx] = Value;
 	}
-}
-#endif
 
-UScalarVelocityObservation* UScalarVelocityObservation::AddScalarVelocityObservation(ULearningAgentsInteractor* InInteractor, const FName Name, const float Scale)
-{
-	return UE::Learning::Agents::Observations::Private::AddObservation<UScalarVelocityObservation, UE::Learning::FScalarVelocityFeature>(InInteractor, Name, TEXT("AddScalarVelocityObservation"), 1, Scale);
+	return true;
 }
 
-void UScalarVelocityObservation::SetScalarVelocityObservation(const int32 AgentId, const float Velocity)
+bool ULearningAgentsObservationObject::GetEnumObservation(uint8& OutEnumValue, const UEnum* Enum, const FLearningAgentsObservationObjectElement Element, const FName Name) const
 {
-	if (!Interactor->HasAgent(AgentId))
+	if (!Enum)
 	{
-		UE_LOG(LogLearning, Error, TEXT("%s: AgentId %d not found in the agents set."), *GetName(), AgentId);
-		return;
+		UE_LOG(LogLearning, Error, TEXT("%s: Enum is nullptr."), *GetName());
+		OutEnumValue = 0;
+		return false;
 	}
 
-	FeatureObject->InstanceData->View(FeatureObject->VelocityHandle)[AgentId][0] = Velocity;
-	AgentIteration[AgentId]++;
-}
-
-void UScalarVelocityObservation::SetScalarVelocityObservationWithAxis(const int32 AgentId, const FVector Velocity, const FVector Axis)
-{
-	SetScalarVelocityObservation(AgentId, Velocity.Dot(Axis));
-}
-
-#if UE_LEARNING_AGENTS_ENABLE_VISUAL_LOG
-void UScalarVelocityObservation::VisualLog(const UE::Learning::FIndexSet Instances) const
-{
-	UE_LEARNING_TRACE_CPUPROFILER_EVENT_SCOPE(UScalarVelocityObservation::VisualLog);
+	int32 EnumValueNum;
+	if (!GetContinuousObservationNum(EnumValueNum, Element, Name))
+	{
+		OutEnumValue = 0;
+		return false;
+	}
+	
+	if (EnumValueNum != Enum->NumEnums() - 1)
+	{
+		UE_LOG(LogLearning, Error, TEXT("%s: Too many values for Enum '%s'. Expected %i, got %i."), *GetName(), *Enum->GetName(), Enum->NumEnums() - 1, EnumValueNum);
+		OutEnumValue = 0;
+		return false;
+	}
 
-	const TLearningArrayView<2, const float> VelocityView = FeatureObject->InstanceData->ConstView(FeatureObject->VelocityHandle);
-	const TLearningArrayView<2, const float> FeatureView = FeatureObject->InstanceData->ConstView(FeatureObject->FeatureHandle);
+	TArray<float, TInlineAllocator<32>> OneHot;
+	OneHot.SetNumUninitialized(EnumValueNum);
+	if (!GetContinuousObservationToArrayView(OneHot, Element, Name))
+	{
+		OutEnumValue = 0;
+		return false;
+	}
 
-	for (const int32 Instance : Instances)
+	int32 EnumValueIndex = INDEX_NONE;
+	for (int32 EnumIdx = 0; EnumIdx < EnumValueNum; EnumIdx++)
 	{
-		if (const AActor* Actor = Cast<AActor>(Interactor->GetAgent(Instance)))
+		if (OneHot[EnumIdx])
 		{
-			UE_LEARNING_AGENTS_VLOG_STRING(this, LogLearning, Display,
-				Actor->GetActorLocation(),
-				VisualLogColor.ToFColor(true),
-				TEXT("Agent %i\nScale: [% 6.2f]\nVelocity: [% 6.1f]\nEncoded: %s"),
-				Instance,
-				FeatureObject->Scale,
-				VelocityView[Instance][0],
-				*UE::Learning::Array::FormatFloat(FeatureView[Instance]));
+			EnumValueIndex = EnumIdx;
+			break;
 		}
 	}
-}
-#endif
 
+	if (EnumValueIndex == INDEX_NONE)
+	{
+		UE_LOG(LogLearning, Error, TEXT("%s: Index not found."), *GetName());
+		OutEnumValue = 0;
+		return false;
+	}
 
-UScalarVelocityArrayObservation* UScalarVelocityArrayObservation::AddScalarVelocityArrayObservation(ULearningAgentsInteractor* InInteractor, const FName Name, const int32 VelocityNum, const float Scale)
-{
-	if (VelocityNum < 1)
+	const int32 EnumValue = Enum->GetValueByIndex(EnumValueIndex);
+
+	if (EnumValue == INDEX_NONE)
 	{
-		UE_LOG(LogLearning, Error, TEXT("AddScalarVelocityArrayObservation: Number of elements in array must be at least 1, got %i."), VelocityNum);
-		return nullptr;
+		UE_LOG(LogLearning, Error, TEXT("%s: Enum Value not found for index %i."), *GetName(), EnumValueIndex);
+		OutEnumValue = 0;
+		return false;
 	}
 
-	return UE::Learning::Agents::Observations::Private::AddObservation<UScalarVelocityArrayObservation, UE::Learning::FScalarVelocityFeature>(InInteractor, Name, TEXT("AddScalarVelocityArrayObservation"), VelocityNum, Scale);
+	OutEnumValue = (uint8)EnumValue;
+	return true;
 }
 
-void UScalarVelocityArrayObservation::SetScalarVelocityArrayObservation(const int32 AgentId, const TArray<float>& Velocities)
+bool ULearningAgentsObservationObject::GetBitmaskObservation(int32& OutBitmaskValue, const UEnum* Enum, const FLearningAgentsObservationObjectElement Element, const FName Name) const
 {
-	if (!Interactor->HasAgent(AgentId))
+	if (!Enum)
 	{
-		UE_LOG(LogLearning, Error, TEXT("%s: AgentId %d not found in the agents set."), *GetName(), AgentId);
-		return;
+		UE_LOG(LogLearning, Error, TEXT("%s: Enum is nullptr."), *GetName());
+		OutBitmaskValue = 0;
+		return false;
 	}
 
-	const TLearningArrayView<2, float> VelocityView = FeatureObject->InstanceData->View(FeatureObject->VelocityHandle);
-
-	if (Velocities.Num() != VelocityView.Num<1>())
+	if (Enum->NumEnums() - 1 > 32)
 	{
-		UE_LOG(LogLearning, Error, TEXT("%s: Got wrong number of elements in array. Expected %i, got %i."), *GetName(), VelocityView.Num<1>(), Velocities.Num());
-		return;
+		UE_LOG(LogLearning, Error, TEXT("%s: Too many values in Enum to use as Bitmask (%i)."), *GetName(), Enum->NumEnums() - 1);
+		OutBitmaskValue = 0;
+		return false;
 	}
 
-	UE::Learning::Array::Copy<1, float>(VelocityView[AgentId], Velocities);
-	AgentIteration[AgentId]++;
-}
-
-void UScalarVelocityArrayObservation::SetScalarVelocityArrayObservationWithAxis(const int32 AgentId, const TArray<FVector>& Velocities, const FVector Axis)
-{
-	if (!Interactor->HasAgent(AgentId))
+	int32 EnumValueNum;
+	if (!GetContinuousObservationNum(EnumValueNum, Element, Name))
 	{
-		UE_LOG(LogLearning, Error, TEXT("%s: AgentId %d not found in the agents set."), *GetName(), AgentId);
-		return;
+		OutBitmaskValue = 0;
+		return false;
 	}
 
-	const TLearningArrayView<2, float> VelocityView = FeatureObject->InstanceData->View(FeatureObject->VelocityHandle);
-
-	if (Velocities.Num() != VelocityView.Num<1>())
+	if (EnumValueNum != Enum->NumEnums() - 1)
 	{
-		UE_LOG(LogLearning, Error, TEXT("%s: Got wrong number of elements in array. Expected %i, got %i."), *GetName(), VelocityView.Num<1>(), Velocities.Num());
-		return;
+		UE_LOG(LogLearning, Error, TEXT("%s: Too many values for Enum '%s'. Expected %i, got %i."), *GetName(), *Enum->GetName(), Enum->NumEnums() - 1, EnumValueNum);
+		OutBitmaskValue = 0;
+		return false;
 	}
 
-	for (int32 VelocityIdx = 0; VelocityIdx < VelocityView.Num<1>(); VelocityIdx++)
+	TArray<float, TInlineAllocator<32>> OneHot;
+	OneHot.Init(0.0f, EnumValueNum);
+	if (!GetContinuousObservationToArrayView(OneHot, Element, Name))
 	{
-		VelocityView[AgentId][VelocityIdx] = Velocities[VelocityIdx].Dot(Axis);
+		OutBitmaskValue = 0;
+		return false;
 	}
 
-	AgentIteration[AgentId]++;
-}
-
-#if UE_LEARNING_AGENTS_ENABLE_VISUAL_LOG
-void UScalarVelocityArrayObservation::VisualLog(const UE::Learning::FIndexSet Instances) const
-{
-	UE_LEARNING_TRACE_CPUPROFILER_EVENT_SCOPE(UScalarVelocityArrayObservation::VisualLog);
-
-	const TLearningArrayView<2, const float> VelocityView = FeatureObject->InstanceData->ConstView(FeatureObject->VelocityHandle);
-	const TLearningArrayView<2, const float> FeatureView = FeatureObject->InstanceData->ConstView(FeatureObject->FeatureHandle);
-
-	const int32 VelocityNum = VelocityView.Num<1>();
-
-	for (const int32 Instance : Instances)
+	OutBitmaskValue = 0;
+	for (int32 OneHotIdx = 0; OneHotIdx < EnumValueNum; OneHotIdx++)
 	{
-		if (const AActor* Actor = Cast<AActor>(Interactor->GetAgent(Instance)))
+		if (OneHot[OneHotIdx])
 		{
-			UE_LEARNING_AGENTS_VLOG_STRING(this, LogLearning, Display,
-				Actor->GetActorLocation(),
-				VisualLogColor.ToFColor(true),
-				TEXT("Agent %i\nScale: [% 6.2f]\nVelocities: %s\nEncoded: %s"),
-				Instance,
-				FeatureObject->Scale,
-				*UE::Learning::Array::FormatFloat(VelocityView[Instance]),
-				*UE::Learning::Array::FormatFloat(FeatureView[Instance]));
+			OutBitmaskValue |= (1 << OneHotIdx);
 		}
 	}
+	return true;
 }
-#endif
 
 
-UPlanarVelocityObservation* UPlanarVelocityObservation::AddPlanarVelocityObservation(ULearningAgentsInteractor* InInteractor, const FName Name, const float Scale, const FVector Axis0, const FVector Axis1)
+bool ULearningAgentsObservationObject::GetOptionalObservation(ELearningAgentsOptionalObservation& OutOption, FLearningAgentsObservationObjectElement& OutElement, const FLearningAgentsObservationObjectElement Element, const FName Name) const
 {
-	return UE::Learning::Agents::Observations::Private::AddObservation<UPlanarVelocityObservation, UE::Learning::FPlanarVelocityFeature>(
-		InInteractor, 
-		Name, 
-		TEXT("AddPlanarVelocityObservation"), 
-		1, 
-		Scale, 
-		Axis0.GetSafeNormal(UE_SMALL_NUMBER, FVector::ForwardVector),
-		Axis1.GetSafeNormal(UE_SMALL_NUMBER, FVector::RightVector));
+	FName OutName = NAME_None;
+	if (!GetExclusiveUnionObservation(OutName, OutElement, Element, Name))
+	{
+		OutOption = ELearningAgentsOptionalObservation::Null;
+		return false;
+	}
+
+	OutOption = OutName == TEXT("Null") ? ELearningAgentsOptionalObservation::Null : ELearningAgentsOptionalObservation::Valid;
+	return true;
 }
 
-void UPlanarVelocityObservation::SetPlanarVelocityObservation(const int32 AgentId, const FVector Velocity, const FRotator RelativeRotation)
+bool ULearningAgentsObservationObject::GetEitherObservation(ELearningAgentsEitherObservation& OutEither, FLearningAgentsObservationObjectElement& OutElement, const FLearningAgentsObservationObjectElement Element, const FName Name) const
 {
-	if (!Interactor->HasAgent(AgentId))
+	FName OutName = NAME_None;
+	if (!GetExclusiveUnionObservation(OutName, OutElement, Element, Name))
 	{
-		UE_LOG(LogLearning, Error, TEXT("%s: AgentId %d not found in the agents set."), *GetName(), AgentId);
-		return;
+		OutEither = ELearningAgentsEitherObservation::A;
+		return false;
 	}
 
-	FeatureObject->InstanceData->View(FeatureObject->VelocityHandle)[AgentId][0] = Velocity;
-	FeatureObject->InstanceData->View(FeatureObject->RelativeRotationHandle)[AgentId] = RelativeRotation.Quaternion();
-	AgentIteration[AgentId]++;
+	OutEither = OutName == TEXT("A") ? ELearningAgentsEitherObservation::A : ELearningAgentsEitherObservation::B;
+	return true;
 }
 
-#if UE_LEARNING_AGENTS_ENABLE_VISUAL_LOG
-void UPlanarVelocityObservation::VisualLog(const UE::Learning::FIndexSet Instances) const
+bool ULearningAgentsObservationObject::GetEncodingObservation(FLearningAgentsObservationObjectElement& OutElement, const FLearningAgentsObservationObjectElement Element, const FName Name) const
 {
-	UE_LEARNING_TRACE_CPUPROFILER_EVENT_SCOPE(UPlanarVelocityObservation::VisualLog);
+	if (!ObservationObject.IsValid(Element.ObjectElement))
+	{
+		UE_LOG(LogLearning, Error, TEXT("%s: Invalid Observation Object."), *GetName());
+		OutElement = FLearningAgentsObservationObjectElement();
+		return false;
+	}
 
-	const TLearningArrayView<2, const FVector> VelocityView = FeatureObject->InstanceData->ConstView(FeatureObject->VelocityHandle);
-	const TLearningArrayView<1, const FQuat> RelativeRotationView = FeatureObject->InstanceData->ConstView(FeatureObject->RelativeRotationHandle);
-	const TLearningArrayView<2, const float> FeatureView = FeatureObject->InstanceData->ConstView(FeatureObject->FeatureHandle);
+	if (ObservationObject.GetName(Element.ObjectElement) != Name)
+	{
+		UE_LOG(LogLearning, Warning, TEXT("%s: Observation name does not match. Observation is '%s' but asked for '%s'."), *GetName(), *ObservationObject.GetName(Element.ObjectElement).ToString(), *Name.ToString());
+	}
 
-	for (const int32 Instance : Instances)
+	if (ObservationObject.GetType(Element.ObjectElement) != UE::Learning::Observation::EType::Encoding)
 	{
-		if (const AActor* Actor = Cast<AActor>(Interactor->GetAgent(Instance)))
-		{
-			const FVector Velocity = VelocityView[Instance][0];
-			const FQuat RelativeRotation = RelativeRotationView[Instance];
-			const FVector LocalVelocity = RelativeRotation.UnrotateVector(Velocity);
-
-			UE_LEARNING_AGENTS_VLOG_ARROW(this, LogLearning, Display,
-				Actor->GetActorLocation(),
-				Actor->GetActorLocation() + Velocity,
-				VisualLogColor.ToFColor(true),
-				TEXT(""));
-
-			UE_LEARNING_AGENTS_VLOG_STRING(this, LogLearning, Display,
-				Actor->GetActorLocation() + Velocity,
-				VisualLogColor.ToFColor(true),
-				TEXT("Velocity: [% 6.3f % 6.3f % 6.3f]\nLocal Velocity: [% 6.3f % 6.3f % 6.3f]"),
-				Velocity.X, Velocity.Y, Velocity.Z,
-				LocalVelocity.X, LocalVelocity.Y, LocalVelocity.Z);
-
-			UE_LEARNING_AGENTS_VLOG_PLANE(this, LogLearning, Display,
-				Actor->GetActorLocation(),
-				RelativeRotation,
-				FeatureObject->Axis0,
-				FeatureObject->Axis1,
-				VisualLogColor.ToFColor(true),
-				TEXT(""));
-
-			UE_LEARNING_AGENTS_VLOG_TRANSFORM(this, LogLearning, Display,
-				Actor->GetActorLocation(),
-				RelativeRotation,
-				VisualLogColor.ToFColor(true),
-				TEXT("Agent %i\nScale: [% 6.2f]\nEncoded: %s"),
-				Instance,
-				FeatureObject->Scale,
-				*UE::Learning::Array::FormatFloat(FeatureView[Instance]));
-		}
+		UE_LOG(LogLearning, Error, TEXT("%s: Observation '%s' type does not match. Observation is '%s' but asked for '%s'."),
+			*GetName(),
+			*ObservationObject.GetName(Element.ObjectElement).ToString(),
+			UE::Learning::Agents::Observation::Private::GetObservationTypeString(ObservationObject.GetType(Element.ObjectElement)),
+			UE::Learning::Agents::Observation::Private::GetObservationTypeString(UE::Learning::Observation::EType::Encoding));
+		OutElement = FLearningAgentsObservationObjectElement();
+		return false;
 	}
-}
-#endif
 
+	OutElement = { ObservationObject.GetEncoding(Element.ObjectElement).Element };
+
+	return true;
+}
 
-UPlanarVelocityArrayObservation* UPlanarVelocityArrayObservation::AddPlanarVelocityArrayObservation(ULearningAgentsInteractor* InInteractor, const FName Name, const int32 VelocityNum, const float Scale, const FVector Axis0, const FVector Axis1)
+bool ULearningAgentsObservationObject::GetBoolObservation(bool& bOutValue, const FLearningAgentsObservationObjectElement Element, const FName Name) const
 {
-	if (VelocityNum < 1)
+	float OutValue = 0.0f;
+	if (!GetFloatObservation(OutValue, Element, 1.0f, Name))
 	{
-		UE_LOG(LogLearning, Error, TEXT("AddPlanarVelocityArrayObservation: Number of elements in array must be at least 1, got %i."), VelocityNum);
-		return nullptr;
+		bOutValue = false;
+		return false;
 	}
 
-	return UE::Learning::Agents::Observations::Private::AddObservation<UPlanarVelocityArrayObservation, UE::Learning::FPlanarVelocityFeature>(
-		InInteractor, 
-		Name, 
-		TEXT("AddPlanarVelocityArrayObservation"), 
-		VelocityNum, 
-		Scale, 
-		Axis0.GetSafeNormal(UE_SMALL_NUMBER, FVector::ForwardVector),
-		Axis1.GetSafeNormal(UE_SMALL_NUMBER, FVector::RightVector));
+	bOutValue = OutValue >= 0.0f;
+	return true;
 }
 
-void UPlanarVelocityArrayObservation::SetPlanarVelocityArrayObservation(const int32 AgentId, const TArray<FVector>& Velocities, const FRotator RelativeRotation)
+bool ULearningAgentsObservationObject::GetFloatObservation(float& OutValue, const FLearningAgentsObservationObjectElement Element, const float FloatScale, const FName Name) const
 {
-	if (!Interactor->HasAgent(AgentId))
+	float OutValuesData;
+	if (!GetContinuousObservationToArrayView(MakeArrayView(&OutValuesData, 1), Element, Name))
 	{
-		UE_LOG(LogLearning, Error, TEXT("%s: AgentId %d not found in the agents set."), *GetName(), AgentId);
-		return;
+		OutValue = 0.0f;
+		return false;
 	}
 
-	const TLearningArrayView<2, FVector> VelocityView = FeatureObject->InstanceData->View(FeatureObject->VelocityHandle);
-	const TLearningArrayView<1, FQuat> RelativeRotationView = FeatureObject->InstanceData->View(FeatureObject->RelativeRotationHandle);
+	OutValue = OutValuesData * FloatScale;
+	return true;
+}
 
-	if (Velocities.Num() != VelocityView.Num<1>())
+bool ULearningAgentsObservationObject::GetTranslationObservation(FVector& OutTranslation, const FLearningAgentsObservationObjectElement Element, const FTransform RelativeTransform, const float TranslationScale, const FName Name) const
+{
+	TStaticArray<float, 3> OutValues;
+	if (!GetContinuousObservationToArrayView(OutValues, Element, Name))
 	{
-		UE_LOG(LogLearning, Error, TEXT("%s: Got wrong number of elements in array. Expected %i, got %i."), *GetName(), VelocityView.Num<1>(), Velocities.Num());
-		return;
+		OutTranslation = FVector::ZeroVector;
+		return false;
 	}
 
-	UE::Learning::Array::Copy<1, FVector>(VelocityView[AgentId], Velocities);
-	RelativeRotationView[AgentId] = RelativeRotation.Quaternion();
-	AgentIteration[AgentId]++;
+	OutTranslation = RelativeTransform.TransformPosition(TranslationScale * FVector(OutValues[0], OutValues[1], OutValues[2]));
+	return true;
 }
 
-#if UE_LEARNING_AGENTS_ENABLE_VISUAL_LOG
-void UPlanarVelocityArrayObservation::VisualLog(const UE::Learning::FIndexSet Instances) const
+bool ULearningAgentsObservationObject::GetRotationObservation(FRotator& OutRotation, const FLearningAgentsObservationObjectElement Element, const FRotator RelativeRotation, const FName Name) const
 {
-	UE_LEARNING_TRACE_CPUPROFILER_EVENT_SCOPE(UPlanarVelocityArrayObservation::VisualLog);
-
-	const TLearningArrayView<2, const FVector> VelocityView = FeatureObject->InstanceData->ConstView(FeatureObject->VelocityHandle);
-	const TLearningArrayView<1, const FQuat> RelativeRotationView = FeatureObject->InstanceData->ConstView(FeatureObject->RelativeRotationHandle);
-	const TLearningArrayView<2, const float> FeatureView = FeatureObject->InstanceData->ConstView(FeatureObject->FeatureHandle);
-
-	const int32 VelocityNum = VelocityView.Num<1>();
-
-	for (const int32 Instance : Instances)
+	FQuat OutRotationQuat;
+	if (!GetRotationObservationAsQuat(OutRotationQuat, Element, FQuat::MakeFromRotator(RelativeRotation), Name))
 	{
-		if (const AActor* Actor = Cast<AActor>(Interactor->GetAgent(Instance)))
-		{
-			const FQuat RelativeRotation = RelativeRotationView[Instance];
+		OutRotation = FRotator::ZeroRotator;
+		return false;
+	}
 
-			for (int32 VelocityIdx = 0; VelocityIdx < VelocityNum; VelocityIdx++)
-			{
-				const FVector Velocity = VelocityView[Instance][0];
-				const FVector LocalVelocity = RelativeRotation.UnrotateVector(Velocity);
-				const FVector Offset = UE::Learning::Agents::Debug::GridOffsetForIndex(VelocityIdx, VelocityNum);
-
-				UE_LEARNING_AGENTS_VLOG_ARROW(this, LogLearning, Display,
-					Actor->GetActorLocation(),
-					Actor->GetActorLocation() + Velocity,
-					VisualLogColor.ToFColor(true),
-					TEXT(""));
-
-				UE_LEARNING_AGENTS_VLOG_STRING(this, LogLearning, Display,
-					Actor->GetActorLocation() + Velocity,
-					VisualLogColor.ToFColor(true),
-					TEXT("Velocity %i: [% 6.3f % 6.3f % 6.3f]\nLocal Velocity: [% 6.3f % 6.3f % 6.3f]"),
-					VelocityIdx,
-					Velocity.X, Velocity.Y, Velocity.Z,
-					LocalVelocity.X, LocalVelocity.Y, LocalVelocity.Z);
-			}
+	OutRotation = OutRotationQuat.Rotator();
+	return true;
+}
 
-			UE_LEARNING_AGENTS_VLOG_PLANE(this, LogLearning, Display,
-				Actor->GetActorLocation(),
-				RelativeRotation,
-				FeatureObject->Axis0,
-				FeatureObject->Axis1,
-				VisualLogColor.ToFColor(true),
-				TEXT(""));
-
-			UE_LEARNING_AGENTS_VLOG_TRANSFORM(this, LogLearning, Display,
-				Actor->GetActorLocation(),
-				RelativeRotation,
-				VisualLogColor.ToFColor(true),
-				TEXT("Agent %i\nScale: [% 6.2f]\nEncoded: %s"),
-				Instance,
-				FeatureObject->Scale,
-				*UE::Learning::Array::FormatFloat(FeatureView[Instance]));
-		}
+bool ULearningAgentsObservationObject::GetRotationObservationAsQuat(FQuat& OutRotation, const FLearningAgentsObservationObjectElement Element, const FQuat RelativeRotation, const FName Name) const
+{
+	TStaticArray<float, 6> OutValues;
+	if (!GetContinuousObservationToArrayView(OutValues, Element, Name))
+	{
+		OutRotation = FQuat::Identity;
+		return false;
 	}
-}
-#endif
 
-//------------------------------------------------------------------
+	const FVector LocalAxisForward = FVector(OutValues[0], OutValues[1], OutValues[2]);
+	const FVector LocalAxisRight = FVector(OutValues[3], OutValues[4], OutValues[5]);
+	const FVector AxisUp = LocalAxisForward.Cross(LocalAxisRight).GetSafeNormal(UE_SMALL_NUMBER, FVector::UpVector);
+	const FVector AxisRight = AxisUp.Cross(LocalAxisForward).GetSafeNormal(UE_SMALL_NUMBER, FVector::RightVector);
+	const FVector AxisForward = LocalAxisForward.GetSafeNormal(UE_SMALL_NUMBER, FVector::ForwardVector);
 
-UAngularVelocityObservation* UAngularVelocityObservation::AddAngularVelocityObservation(ULearningAgentsInteractor* InInteractor, const FName Name, const float Scale)
-{
-	return UE::Learning::Agents::Observations::Private::AddObservation<UAngularVelocityObservation, UE::Learning::FAngularVelocityFeature>(InInteractor, Name, TEXT("AddAngularVelocityObservation"), 1, Scale);
+	FMatrix RotationMatrix = FMatrix::Identity;
+	RotationMatrix.SetAxis(0, AxisForward);
+	RotationMatrix.SetAxis(1, AxisRight);
+	RotationMatrix.SetAxis(2, AxisUp);
+
+	OutRotation = RelativeRotation * RotationMatrix.ToQuat();
+	return true;
 }
 
-void UAngularVelocityObservation::SetAngularVelocityObservation(const int32 AgentId, const FVector AngularVelocity, const FRotator RelativeRotation)
+bool ULearningAgentsObservationObject::GetScaleObservation(FVector& OutScale, const FLearningAgentsObservationObjectElement Element, const FVector RelativeScale, const FName Name) const
 {
-	if (!Interactor->HasAgent(AgentId))
+	TStaticArray<float, 3> OutValues;
+	if (!GetContinuousObservationToArrayView(OutValues, Element, Name))
 	{
-		UE_LOG(LogLearning, Error, TEXT("%s: AgentId %d not found in the agents set."), *GetName(), AgentId);
-		return;
+		OutScale = FVector::OneVector;
+		return false;
 	}
 
-	FeatureObject->InstanceData->View(FeatureObject->AngularVelocityHandle)[AgentId][0] = AngularVelocity;
-	FeatureObject->InstanceData->View(FeatureObject->RelativeRotationHandle)[AgentId] = RelativeRotation.Quaternion();
-	AgentIteration[AgentId]++;
+	OutScale = RelativeScale * UE::Learning::Agents::Observation::Private::VectorExp(FVector(OutValues[0], OutValues[1], OutValues[2]));
+	return true;
 }
 
-#if UE_LEARNING_AGENTS_ENABLE_VISUAL_LOG
-void UAngularVelocityObservation::VisualLog(const UE::Learning::FIndexSet Instances) const
+bool ULearningAgentsObservationObject::GetTransformObservation(FTransform& OutTransform, const FLearningAgentsObservationObjectElement Element, const FTransform RelativeTransform, const float TranslationScale, const FName Name) const
 {
-	UE_LEARNING_TRACE_CPUPROFILER_EVENT_SCOPE(UAngularVelocityObservation::VisualLog);
+	TStaticArray<FName, 3> OutElementNames;
+	TStaticArray<FLearningAgentsObservationObjectElement, 3> OutElements;
+	if (!GetStructObservationToArrayViews(OutElementNames, OutElements, Element, Name))
+	{
+		OutTransform = FTransform::Identity;
+		return false;
+	}
 
-	const TLearningArrayView<2, const FVector> AngularVelocityView = FeatureObject->InstanceData->ConstView(FeatureObject->AngularVelocityHandle);
-	const TLearningArrayView<1, const FQuat> RelativeRotationView = FeatureObject->InstanceData->ConstView(FeatureObject->RelativeRotationHandle);
-	const TLearningArrayView<2, const float> FeatureView = FeatureObject->InstanceData->ConstView(FeatureObject->FeatureHandle);
+	const int32 TranslationElement = MakeArrayView(OutElementNames).Find(TEXT("Translation"));
+	FVector OutTranslation;
+	if (TranslationElement == INDEX_NONE || !GetTranslationObservation(OutTranslation, OutElements[TranslationElement], RelativeTransform, TranslationScale))
+	{
+		OutTransform = FTransform::Identity;
+		return false;
+	}
 
-	for (const int32 Instance : Instances)
+	const int32 RotationElement = MakeArrayView(OutElementNames).Find(TEXT("Rotation"));
+	FQuat OutRotation;
+	if (RotationElement == INDEX_NONE || !GetRotationObservationAsQuat(OutRotation, OutElements[RotationElement], RelativeTransform.GetRotation()))
 	{
-		if (const AActor* Actor = Cast<AActor>(Interactor->GetAgent(Instance)))
-		{
-			const FVector AngularVelocity = AngularVelocityView[Instance][0];
-			const FQuat RelativeRotation = RelativeRotationView[Instance];
-			const FVector LocalAngularVelocity = RelativeRotation.UnrotateVector(AngularVelocity);
-
-			UE_LEARNING_AGENTS_VLOG_ARROW(this, LogLearning, Display,
-				Actor->GetActorLocation(),
-				Actor->GetActorLocation() + AngularVelocity,
-				VisualLogColor.ToFColor(true),
-				TEXT(""));
-
-			UE_LEARNING_AGENTS_VLOG_STRING(this, LogLearning, Display,
-				Actor->GetActorLocation() + AngularVelocity,
-				VisualLogColor.ToFColor(true),
-				TEXT("Angular Velocity: [% 6.3f % 6.3f % 6.3f]\nLocal Angular Velocity: [% 6.3f % 6.3f % 6.3f]"),
-				AngularVelocity.X, AngularVelocity.Y, AngularVelocity.Z,
-				LocalAngularVelocity.X, LocalAngularVelocity.Y, LocalAngularVelocity.Z);
-
-			UE_LEARNING_AGENTS_VLOG_TRANSFORM(this, LogLearning, Display,
-				Actor->GetActorLocation(),
-				RelativeRotation,
-				VisualLogColor.ToFColor(true),
-				TEXT("Agent %i\nScale: [% 6.2f]\nEncoded: %s"),
-				Instance,
-				FeatureObject->Scale,
-				*UE::Learning::Array::FormatFloat(FeatureView[Instance]));
-		}
+		OutTransform = FTransform::Identity;
+		return false;
 	}
-}
-#endif
 
-UAngularVelocityArrayObservation* UAngularVelocityArrayObservation::AddAngularVelocityArrayObservation(ULearningAgentsInteractor* InInteractor, const FName Name, const int32 AngularVelocityNum, const float Scale)
-{
-	if (AngularVelocityNum < 1)
+	const int32 ScaleElement = MakeArrayView(OutElementNames).Find(TEXT("Scale"));
+	FVector OutScale;
+	if (ScaleElement == INDEX_NONE || !GetScaleObservation(OutScale, OutElements[ScaleElement], RelativeTransform.GetScale3D()))
 	{
-		UE_LOG(LogLearning, Error, TEXT("AddAngularVelocityArrayObservation: Number of elements in array must be at least 1, got %i."), AngularVelocityNum);
-		return nullptr;
+		OutTransform = FTransform::Identity;
+		return false;
 	}
 
-	return UE::Learning::Agents::Observations::Private::AddObservation<UAngularVelocityArrayObservation, UE::Learning::FAngularVelocityFeature>(InInteractor, Name, TEXT("AddAngularVelocityArrayObservation"), AngularVelocityNum, Scale);
+	OutTransform = FTransform(OutRotation, OutTranslation, OutScale);
+	return true;
 }
 
-void UAngularVelocityArrayObservation::SetAngularVelocityArrayObservation(const int32 AgentId, const TArray<FVector>& AngularVelocities, const FRotator RelativeRotation)
+bool ULearningAgentsObservationObject::GetAngleObservationRadians(float& OutAngle, const FLearningAgentsObservationObjectElement Element, const float RelativeAngle, const FName Name) const
 {
-	if (!Interactor->HasAgent(AgentId))
+	TStaticArray<float, 2> OutValues;
+	if (!GetContinuousObservationToArrayView(OutValues, Element, Name))
 	{
-		UE_LOG(LogLearning, Error, TEXT("%s: AgentId %d not found in the agents set."), *GetName(), AgentId);
-		return;
+		OutAngle = 0.0f;
+		return false;
 	}
 
-	const TLearningArrayView<2, FVector> AngularVelocityView = FeatureObject->InstanceData->View(FeatureObject->AngularVelocityHandle);
-	const TLearningArrayView<1, FQuat> RelativeRotationView = FeatureObject->InstanceData->View(FeatureObject->RelativeRotationHandle);
+	OutAngle = RelativeAngle + FMath::Atan2(OutValues[0], OutValues[1]);
+	return true;
+}
+
 
-	if (AngularVelocities.Num() != AngularVelocityView.Num<1>())
+bool ULearningAgentsObservationObject::GetAngleObservation(float& OutAngle, const FLearningAgentsObservationObjectElement Element, const float RelativeAngle, const FName Name) const
+{
+	if (!GetAngleObservationRadians(OutAngle, Element, FMath::DegreesToRadians(RelativeAngle), Name))
 	{
-		UE_LOG(LogLearning, Error, TEXT("%s: Got wrong number of elements in array. Expected %i, got %i."), *GetName(), AngularVelocityView.Num<1>(), AngularVelocities.Num());
-		return;
+		return false;
 	}
 
-	UE::Learning::Array::Copy<1, FVector>(AngularVelocityView[AgentId], AngularVelocities);
-	RelativeRotationView[AgentId] = RelativeRotation.Quaternion();
-	AgentIteration[AgentId]++;
+	OutAngle = FMath::RadiansToDegrees(OutAngle);
+	return true;
 }
 
-#if UE_LEARNING_AGENTS_ENABLE_VISUAL_LOG
-void UAngularVelocityArrayObservation::VisualLog(const UE::Learning::FIndexSet Instances) const
+bool ULearningAgentsObservationObject::GetVelocityObservation(FVector& OutVelocity, const FLearningAgentsObservationObjectElement Element, const FTransform RelativeTransform, const float VelocityScale, const FName Name) const
 {
-	UE_LEARNING_TRACE_CPUPROFILER_EVENT_SCOPE(UAngularVelocityArrayObservation::VisualLog);
-
-	const TLearningArrayView<2, const FVector> AngularVelocityView = FeatureObject->InstanceData->ConstView(FeatureObject->AngularVelocityHandle);
-	const TLearningArrayView<1, const FQuat> RelativeRotationView = FeatureObject->InstanceData->ConstView(FeatureObject->RelativeRotationHandle);
-	const TLearningArrayView<2, const float> FeatureView = FeatureObject->InstanceData->ConstView(FeatureObject->FeatureHandle);
-
-	const int32 AngularVelocityNum = AngularVelocityView.Num<1>();
-
-	for (const int32 Instance : Instances)
+	TStaticArray<float, 3> OutValues;
+	if (!GetContinuousObservationToArrayView(OutValues, Element, Name))
 	{
-		if (const AActor* Actor = Cast<AActor>(Interactor->GetAgent(Instance)))
-		{
-			const FQuat RelativeRotation = RelativeRotationView[Instance];
+		OutVelocity = FVector::ZeroVector;
+		return false;
+	}
 
-			for (int32 AngularVelocityIdx = 0; AngularVelocityIdx < AngularVelocityNum; AngularVelocityIdx++)
-			{
-				const FVector AngularVelocity = AngularVelocityView[Instance][0];
-				const FVector LocalAngularVelocity = RelativeRotation.UnrotateVector(AngularVelocity);
-				const FVector Offset = UE::Learning::Agents::Debug::GridOffsetForIndex(AngularVelocityIdx, AngularVelocityNum);
-
-				UE_LEARNING_AGENTS_VLOG_ARROW(this, LogLearning, Display,
-					Actor->GetActorLocation() + Offset,
-					Actor->GetActorLocation() + Offset + AngularVelocity,
-					VisualLogColor.ToFColor(true),
-					TEXT(""));
-
-				UE_LEARNING_AGENTS_VLOG_STRING(this, LogLearning, Display,
-					Actor->GetActorLocation() + Offset + AngularVelocity,
-					VisualLogColor.ToFColor(true),
-					TEXT("Angular Velocity %i: [% 6.3f % 6.3f % 6.3f]\nLocal Angular Velocity: [% 6.3f % 6.3f % 6.3f]"),
-					AngularVelocityIdx,
-					AngularVelocity.X, AngularVelocity.Y, AngularVelocity.Z,
-					LocalAngularVelocity.X, LocalAngularVelocity.Y, LocalAngularVelocity.Z);
-			}
+	OutVelocity = RelativeTransform.TransformVectorNoScale(VelocityScale * FVector(OutValues[0], OutValues[1], OutValues[2]));
+	return true;
+}
 
-			UE_LEARNING_AGENTS_VLOG_TRANSFORM(this, LogLearning, Display,
-				Actor->GetActorLocation(),
-				RelativeRotation,
-				VisualLogColor.ToFColor(true),
-				TEXT("Agent %i\nScale: [% 6.2f]\nEncoded: %s"),
-				Instance,
-				FeatureObject->Scale,
-				*UE::Learning::Array::FormatFloat(FeatureView[Instance]));
-		}
+bool ULearningAgentsObservationObject::GetDirectionObservation(FVector& OutDirection, const FLearningAgentsObservationObjectElement Element, const FTransform RelativeTransform, const FName Name) const
+{
+	TStaticArray<float, 3> OutValues;
+	if (!GetContinuousObservationToArrayView(OutValues, Element, Name))
+	{
+		OutDirection = FVector::ForwardVector;
+		return false;
 	}
+
+	OutDirection = RelativeTransform.TransformVectorNoScale(FVector(OutValues[0], OutValues[1], OutValues[2]).GetSafeNormal(UE_SMALL_NUMBER, FVector::ForwardVector));
+	return true;
 }
-#endif
 
-UScalarAngularVelocityObservation* UScalarAngularVelocityObservation::AddScalarAngularVelocityObservation(ULearningAgentsInteractor* InInteractor, const FName Name, const float Scale)
+bool ULearningAgentsObservationObject::GetLocationAlongSplineObservation(FVector& OutLocation, const FLearningAgentsObservationObjectElement Element, const FTransform RelativeTransform, const float TranslationScale, const FName Name)
 {
-	return UE::Learning::Agents::Observations::Private::AddObservation<UScalarAngularVelocityObservation, UE::Learning::FScalarAngularVelocityFeature>(InInteractor, Name, TEXT("AddScalarAngularVelocityObservation"), 1, Scale);
+	return GetTranslationObservation(OutLocation, Element, RelativeTransform, TranslationScale, Name);
 }
 
-void UScalarAngularVelocityObservation::SetScalarAngularVelocityObservation(const int32 AgentId, const float AngularVelocity)
+bool ULearningAgentsObservationObject::GetProportionAlongSplineObservation(bool& bOutIsClosedLoop, float& OutAngle, float& OutPropotion, const FLearningAgentsObservationObjectElement Element, const FName Name)
 {
-	if (!Interactor->HasAgent(AgentId))
+	FName SubName;
+	FLearningAgentsObservationObjectElement SubElement;
+	if (!GetExclusiveUnionObservation(SubName, SubElement, Element, Name))
 	{
-		UE_LOG(LogLearning, Error, TEXT("%s: AgentId %d not found in the agents set."), *GetName(), AgentId);
-		return;
+		bOutIsClosedLoop = false;
+		OutAngle = 0.0f;
+		OutPropotion = 0.0;
+		return false;
 	}
 
-	FeatureObject->InstanceData->View(FeatureObject->AngularVelocityHandle)[AgentId][0] = AngularVelocity;
-	AgentIteration[AgentId]++;
+	if (SubName == TEXT("Angle"))
+	{
+		bOutIsClosedLoop = true;
+		OutPropotion = 0.0f;
+		return GetAngleObservation(OutAngle, SubElement);
+	}
+	else
+	{
+		bOutIsClosedLoop = false;
+		OutAngle = 0.0f;
+		return GetFloatObservation(OutPropotion, SubElement);
+	}
 }
 
-void UScalarAngularVelocityObservation::SetScalarAngularVelocityObservationWithAxis(const int32 AgentId, const FVector AngularVelocity, const FVector Axis)
+bool ULearningAgentsObservationObject::GetDirectionAlongSplineObservation(FVector& OutDirection, const FLearningAgentsObservationObjectElement Element, const FTransform RelativeTransform, const FName Name)
 {
-	SetScalarAngularVelocityObservation(AgentId, AngularVelocity.Dot(Axis));
+	return GetDirectionObservation(OutDirection, Element, RelativeTransform, Name);
 }
 
-#if UE_LEARNING_AGENTS_ENABLE_VISUAL_LOG
-void UScalarAngularVelocityObservation::VisualLog(const UE::Learning::FIndexSet Instances) const
+bool ULearningAgentsObservationObject::GetPropertiesAlongSplineObservation(FVector& OutLocation, bool& bOutIsClosedLoop, float& OutAngle, float& OutPropotion, FVector& OutDirection, const FLearningAgentsObservationObjectElement Element, const FTransform RelativeTransform, const float TranslationScale, const FName Name)
 {
-	UE_LEARNING_TRACE_CPUPROFILER_EVENT_SCOPE(UScalarAngularVelocityObservation::VisualLog);
+	TStaticArray<FName, 3> OutElementNames;
+	TStaticArray<FLearningAgentsObservationObjectElement, 3> OutElements;
+	if (!GetStructObservationToArrayViews(OutElementNames, OutElements, Element, Name))
+	{
+		OutLocation = FVector::ZeroVector;
+		bOutIsClosedLoop = false;
+		OutAngle = 0.0f;
+		OutPropotion = 0.0f;
+		OutDirection = FVector::ForwardVector;
+		return false;
+	}
 
-	const TLearningArrayView<2, const float> AngularVelocityView = FeatureObject->InstanceData->ConstView(FeatureObject->AngularVelocityHandle);
-	const TLearningArrayView<2, const float> FeatureView = FeatureObject->InstanceData->ConstView(FeatureObject->FeatureHandle);
+	const int32 LocationElement = MakeArrayView(OutElementNames).Find(TEXT("Location"));
+	if (LocationElement == INDEX_NONE || !GetLocationAlongSplineObservation(OutLocation, OutElements[LocationElement], RelativeTransform, TranslationScale))
+	{
+		OutLocation = FVector::ZeroVector;
+		bOutIsClosedLoop = false;
+		OutAngle = 0.0f;
+		OutPropotion = 0.0f;
+		OutDirection = FVector::ForwardVector;
+		return false;
+	}
 
-	for (const int32 Instance : Instances)
+	const int32 ProportionElement = MakeArrayView(OutElementNames).Find(TEXT("Proportion"));
+	if (ProportionElement == INDEX_NONE || !GetProportionAlongSplineObservation(bOutIsClosedLoop, OutAngle, OutPropotion, OutElements[ProportionElement]))
 	{
-		if (const AActor* Actor = Cast<AActor>(Interactor->GetAgent(Instance)))
-		{
-			UE_LEARNING_AGENTS_VLOG_STRING(this, LogLearning, Display,
-				Actor->GetActorLocation(),
-				VisualLogColor.ToFColor(true),
-				TEXT("Agent %i\nScale: [% 6.2f]\nAngular Velocity: [% 6.1f]\nEncoded: %s"),
-				Instance,
-				FeatureObject->Scale,
-				AngularVelocityView[Instance][0],
-				*UE::Learning::Array::FormatFloat(FeatureView[Instance]));
-		}
+		OutLocation = FVector::ZeroVector;
+		bOutIsClosedLoop = false;
+		OutAngle = 0.0f;
+		OutPropotion = 0.0f;
+		OutDirection = FVector::ForwardVector;
+		return false;
+	}
+
+	const int32 DirectionElement = MakeArrayView(OutElementNames).Find(TEXT("Direction"));
+	if (DirectionElement == INDEX_NONE || !GetDirectionAlongSplineObservation(OutDirection, OutElements[ProportionElement], RelativeTransform))
+	{
+		OutLocation = FVector::ZeroVector;
+		bOutIsClosedLoop = false;
+		OutAngle = 0.0f;
+		OutPropotion = 0.0f;
+		OutDirection = FVector::ForwardVector;
+		return false;
 	}
-}
-#endif
 
+	return true;
+}
 
-UScalarAngularVelocityArrayObservation* UScalarAngularVelocityArrayObservation::AddScalarAngularVelocityArrayObservation(ULearningAgentsInteractor* InInteractor, const FName Name, const int32 AngularVelocityNum, const float Scale)
+bool ULearningAgentsObservationObject::GetProportionAlongRayObservation(float& OutProportion, const FLearningAgentsObservationObjectElement Element, const FName Name)
 {
-	if (AngularVelocityNum < 1)
+	if (!GetFloatObservation(OutProportion, Element, 1.0f, Name))
 	{
-		UE_LOG(LogLearning, Error, TEXT("AddScalarAngularVelocityArrayObservation: Number of elements in array must be at least 1, got %i."), AngularVelocityNum);
-		return nullptr;
+		OutProportion = 0.0f;
+		return false;
 	}
 
-	return UE::Learning::Agents::Observations::Private::AddObservation<UScalarAngularVelocityArrayObservation, UE::Learning::FScalarAngularVelocityFeature>(InInteractor, Name, TEXT("AddScalarAngularVelocityArrayObservation"), AngularVelocityNum, Scale);
+	OutProportion = 1.0f - OutProportion;
+	return true;
+}
+
+bool ULearningAgentsObservationObject::GetProportionAlongRaysObservationNum(int32& OutProportionNum, const FLearningAgentsObservationObjectElement Element, const FName Name)
+{
+	return GetStaticArrayObservationNum(OutProportionNum, Element, Name);
 }
 
-void UScalarAngularVelocityArrayObservation::SetScalarAngularVelocityArrayObservation(const int32 AgentId, const TArray<float>& AngularVelocities)
+bool ULearningAgentsObservationObject::GetProportionAlongRaysObservation(TArray<float>& OutProportions, const FLearningAgentsObservationObjectElement Element, const FName Name)
 {
-	if (!Interactor->HasAgent(AgentId))
+	int32 ProportionNum;
+	if (!GetProportionAlongRaysObservationNum(ProportionNum, Element, Name))
 	{
-		UE_LOG(LogLearning, Error, TEXT("%s: AgentId %d not found in the agents set."), *GetName(), AgentId);
-		return;
+		OutProportions.Empty();
+		return false;
 	}
 
-	const TLearningArrayView<2, float> AngularVelocityView = FeatureObject->InstanceData->View(FeatureObject->AngularVelocityHandle);
-
-	if (AngularVelocities.Num() != AngularVelocityView.Num<1>())
+	OutProportions.SetNumUninitialized(ProportionNum);
+	if (!GetProportionAlongRaysObservationToArrayView(OutProportions, Element, Name))
 	{
-		UE_LOG(LogLearning, Error, TEXT("%s: Got wrong number of elements in array. Expected %i, got %i."), *GetName(), AngularVelocityView.Num<1>(), AngularVelocities.Num());
-		return;
+		OutProportions.Empty();
+		return false;
 	}
 
-	UE::Learning::Array::Copy<1, float>(AngularVelocityView[AgentId], AngularVelocities);
-	AgentIteration[AgentId]++;
+	return true;
 }
 
-void UScalarAngularVelocityArrayObservation::SetScalarAngularVelocityArrayObservationWithAxis(const int32 AgentId, const TArray<FVector>& AngularVelocities, const FVector Axis)
+bool ULearningAgentsObservationObject::GetProportionAlongRaysObservationToArrayView(TArrayView<float> OutProportions, const FLearningAgentsObservationObjectElement Element, const FName Name)
 {
-	if (!Interactor->HasAgent(AgentId))
+	int32 ProportionNum;
+	if (!GetStaticArrayObservationNum(ProportionNum, Element, Name))
 	{
-		UE_LOG(LogLearning, Error, TEXT("%s: AgentId %d not found in the agents set."), *GetName(), AgentId);
-		return;
+		UE::Learning::Array::Zero<1, float>(OutProportions);
+		return false;
 	}
 
-	const TLearningArrayView<2, float> AngularVelocityView = FeatureObject->InstanceData->View(FeatureObject->AngularVelocityHandle);
-
-	if (AngularVelocities.Num() != AngularVelocityView.Num<1>())
+	if (ProportionNum != OutProportions.Num())
 	{
-		UE_LOG(LogLearning, Error, TEXT("%s: Got wrong number of elements in array. Expected %i, got %i."), *GetName(), AngularVelocityView.Num<1>(), AngularVelocities.Num());
-		return;
+		UE_LOG(LogLearning, Error, TEXT("%s: Observation '%s' size does not match. Observation is '%i' elements but asked for '%i'."),
+			*GetName(), *ObservationObject.GetName(Element.ObjectElement).ToString(),
+			ProportionNum, OutProportions.Num());
+		UE::Learning::Array::Zero<1, float>(OutProportions);
+		return false;
 	}
 
-	for (int32 AngularVelocityIdx = 0; AngularVelocityIdx < AngularVelocityView.Num<1>(); AngularVelocityIdx++)
+	TArray<FLearningAgentsObservationObjectElement, TInlineAllocator<32>> SubElements;
+	SubElements.SetNumUninitialized(ProportionNum);
+	if (!GetStaticArrayObservationToArrayView(SubElements, Element, Name))
 	{
-		AngularVelocityView[AgentId][AngularVelocityIdx] = AngularVelocities[AngularVelocityIdx].Dot(Axis);
+		UE::Learning::Array::Zero<1, float>(OutProportions);
+		return false;
 	}
-	AgentIteration[AgentId]++;
-}
-
-#if UE_LEARNING_AGENTS_ENABLE_VISUAL_LOG
-void UScalarAngularVelocityArrayObservation::VisualLog(const UE::Learning::FIndexSet Instances) const
-{
-	UE_LEARNING_TRACE_CPUPROFILER_EVENT_SCOPE(UScalarAngularVelocityArrayObservation::VisualLog);
-
-	const TLearningArrayView<2, const float> AngularVelocityView = FeatureObject->InstanceData->ConstView(FeatureObject->AngularVelocityHandle);
-	const TLearningArrayView<2, const float> FeatureView = FeatureObject->InstanceData->ConstView(FeatureObject->FeatureHandle);
 
-	for (const int32 Instance : Instances)
+	for (int32 SubElementIdx = 0; SubElementIdx < ProportionNum; SubElementIdx++)
 	{
-		if (const AActor* Actor = Cast<AActor>(Interactor->GetAgent(Instance)))
+		if (!GetProportionAlongRayObservation(OutProportions[SubElementIdx], SubElements[SubElementIdx]))
 		{
-			UE_LEARNING_AGENTS_VLOG_STRING(this, LogLearning, Display,
-				Actor->GetActorLocation(),
-				VisualLogColor.ToFColor(true),
-				TEXT("Agent %i\nScale: [% 6.2f]\nAngular Velocities: %s\nEncoded: %s"),
-				Instance,
-				FeatureObject->Scale,
-				*UE::Learning::Array::FormatFloat(AngularVelocityView[Instance]),
-				*UE::Learning::Array::FormatFloat(FeatureView[Instance]));
+			UE::Learning::Array::Zero<1, float>(OutProportions);
+			return false;
 		}
 	}
-}
-#endif
+
+	return true;
+}
\ No newline at end of file
