LearningAgents: Large restructuring to allow for structured actions and observations.

diff --git a/Engine/Plugins/Experimental/LearningAgents/Source/LearningAgents/Private/LearningAgentsPolicy.cpp b/Engine/Plugins/Experimental/LearningAgents/Source/LearningAgents/Private/LearningAgentsPolicy.cpp
--- a/Engine/Plugins/Experimental/LearningAgents/Source/LearningAgents/Private/LearningAgentsPolicy.cpp
+++ b/Engine/Plugins/Experimental/LearningAgents/Source/LearningAgents/Private/LearningAgentsPolicy.cpp
@@ -4,26 +4,95 @@
 
 #include "LearningAgentsManager.h"
 #include "LearningAgentsInteractor.h"
-#include "LearningAgentsHelpers.h"
-#include "LearningAgentsNeuralNetworkData.h"
-#include "LearningFeatureObject.h"
 #include "LearningNeuralNetwork.h"
-#include "LearningNeuralNetworkObject.h"
+#include "LearningPolicy.h"
 #include "LearningLog.h"
+#include "LearningRandom.h"
 
 #include "UObject/Package.h"
 #include "Misc/FileHelper.h"
 #include "Misc/Paths.h"
 #include "GameFramework/Actor.h"
 
+#include "NNERuntimeBasicCpuBuilder.h"
+
+namespace UE::Learning::Agents::Policy::Private
+{
+	static inline UE::NNE::RuntimeBasic::FModelBuilder::EActivationFunction GetBuilderActivationFunction(const ELearningAgentsActivationFunction ActivationFunction)
+	{
+		switch (ActivationFunction)
+		{
+		case ELearningAgentsActivationFunction::ReLU: return UE::NNE::RuntimeBasic::FModelBuilder::EActivationFunction::ReLU;
+		case ELearningAgentsActivationFunction::ELU: return UE::NNE::RuntimeBasic::FModelBuilder::EActivationFunction::ELU;
+		case ELearningAgentsActivationFunction::TanH: return UE::NNE::RuntimeBasic::FModelBuilder::EActivationFunction::TanH;
+		default:
+			UE_LOG(LogLearning, Error, TEXT("Unknown Activation Function"));
+			return UE::NNE::RuntimeBasic::FModelBuilder::EActivationFunction::ReLU;
+		}
+	}
+}
+
 ULearningAgentsPolicy::ULearningAgentsPolicy() : Super(FObjectInitializer::Get()) {}
 ULearningAgentsPolicy::ULearningAgentsPolicy(FVTableHelper& Helper) : Super(Helper) {}
 ULearningAgentsPolicy::~ULearningAgentsPolicy() = default;
 
+ULearningAgentsPolicy* ULearningAgentsPolicy::MakePolicy(
+	ULearningAgentsManager* InManager,
+	ULearningAgentsInteractor* InInteractor,
+	TSubclassOf<ULearningAgentsPolicy> Class,
+	const FName Name,
+	ULearningAgentsNeuralNetwork* EncoderNeuralNetworkAsset,
+	ULearningAgentsNeuralNetwork* PolicyNeuralNetworkAsset,
+	ULearningAgentsNeuralNetwork* DecoderNeuralNetworkAsset,
+	const bool bReinitializeEncoderNetwork,
+	const bool bReinitializePolicyNetwork,
+	const bool bReinitializeDecoderNetwork,
+	const FLearningAgentsPolicySettings& PolicySettings,
+	const int32 Seed)
+{
+	if (!InManager)
+	{
+		UE_LOG(LogLearning, Error, TEXT("MakePolicy: InManager is nullptr."));
+		return nullptr;
+	}
+
+	if (!Class)
+	{
+		UE_LOG(LogLearning, Error, TEXT("MakePolicy: Class is nullptr."));
+		return nullptr;
+	}
+
+	const FName UniqueName = MakeUniqueObjectName(InManager, Class, Name, EUniqueObjectNameOptions::GloballyUnique);
+
+	ULearningAgentsPolicy* Policy = NewObject<ULearningAgentsPolicy>(InManager, Class, UniqueName);
+	if (!Policy) { return nullptr; }
+
+	Policy->SetupPolicy(
+		InManager,
+		InInteractor,
+		EncoderNeuralNetworkAsset,
+		PolicyNeuralNetworkAsset,
+		DecoderNeuralNetworkAsset,
+		bReinitializeEncoderNetwork,
+		bReinitializePolicyNetwork,
+		bReinitializeDecoderNetwork,
+		PolicySettings,
+		Seed);
+
+	return Policy->IsSetup() ? Policy : nullptr;
+}
+
 void ULearningAgentsPolicy::SetupPolicy(
-	ULearningAgentsInteractor* InInteractor, 
+	ULearningAgentsManager* InManager,
+	ULearningAgentsInteractor* InInteractor,
+	ULearningAgentsNeuralNetwork* EncoderNeuralNetworkAsset,
+	ULearningAgentsNeuralNetwork* PolicyNeuralNetworkAsset,
+	ULearningAgentsNeuralNetwork* DecoderNeuralNetworkAsset,
+	const bool bReinitializeEncoderNetwork,
+	const bool bReinitializePolicyNetwork,
+	const bool bReinitializeDecoderNetwork,
 	const FLearningAgentsPolicySettings& PolicySettings,
-	ULearningAgentsNeuralNetwork* NeuralNetworkAsset)
+	const int32 Seed)
 {
 	if (IsSetup())
 	{
@@ -31,9 +100,9 @@ void ULearningAgentsPolicy::SetupPolicy(
 		return;
 	}
 
-	if (!Manager)
+	if (!InManager)
 	{
-		UE_LOG(LogLearning, Error, TEXT("%s: Must be attached to a LearningAgentsManager Actor."), *GetName());
+		UE_LOG(LogLearning, Error, TEXT("%s: InManager is nullptr."), *GetName());
 		return;
 	}
 
@@ -49,161 +118,273 @@ void ULearningAgentsPolicy::SetupPolicy(
 		return;
 	}
 
+	Manager = InManager;
 	Interactor = InInteractor;
 
-	// Setup Neural Network
+	// Compatibility Hashes and Sizes
+
+	const int32 MemoryStateSize = PolicySettings.bUseMemory ? PolicySettings.MemoryStateSize : 0;
+	const int32 ObservationVectorSize = Interactor->GetObservationVectorSize();
+	const int32 ObservationEncodedVectorSize = Interactor->GetObservationEncodedVectorSize();
+	const int32 ActionEncodedVectorSize = Interactor->GetActionEncodedVectorSize();
+	const int32 ActionDistributionVectorSize = Interactor->GetActionDistributionVectorSize();
 
-	const int32 NetworkInputNum = Interactor->GetObservationFeature().DimNum() + PolicySettings.MemoryStateSize;
-	const int32 NetworkOutputNum = 2 * Interactor->GetActionFeature().DimNum() + PolicySettings.MemoryStateSize;
+	const int32 ObservationCompatibilityHash = UE::Learning::Observation::GetSchemaObjectsCompatibilityHash(Interactor->GetObservationSchema(), Interactor->GetObservationSchemaElement());
+	const int32 ActionCompatibilityHash = UE::Learning::Action::GetSchemaObjectsCompatibilityHash(Interactor->GetActionSchema(), Interactor->GetActionSchemaElement());
+	
+	const int32 PolicyHashData[3] = { MemoryStateSize, ObservationEncodedVectorSize, ActionEncodedVectorSize };
+	const int32 PolicyCompatibilityHash = CityHash32((const char*)PolicyHashData, 3 * sizeof(int32));
 
-	if (NeuralNetworkAsset)
+	// Encoder
+
+	if (EncoderNeuralNetworkAsset)
 	{
-		// Use Existing Neural Network Asset
+		EncoderNetwork = EncoderNeuralNetworkAsset;
 
-		if (NeuralNetworkAsset->NeuralNetworkData)
+		if (EncoderNeuralNetworkAsset->NeuralNetworkData && !bReinitializeEncoderNetwork)
 		{
-			if (NeuralNetworkAsset->NeuralNetworkData->GetNetworkInterface()->GetInputNum() != NetworkInputNum ||
-				NeuralNetworkAsset->NeuralNetworkData->GetNetworkInterface()->GetOutputNum() != NetworkOutputNum)
+			if (EncoderNeuralNetworkAsset->NeuralNetworkData->GetCompatibilityHash() != ObservationCompatibilityHash)
 			{
-				UE_LOG(LogLearning, Error, TEXT("%s: Neural Network Asset provided during Setup is incorrect size: Inputs and outputs don't match what is required."), *GetName());
+				UE_LOG(LogLearning, Error, TEXT("%s: Encoder Network Asset provided during Setup is incompatible with Schema. Network hash is %i vs Schema hash %i."), *GetName(),
+					EncoderNeuralNetworkAsset->NeuralNetworkData->GetCompatibilityHash(),
+					ObservationCompatibilityHash);
 				return;
 			}
 
-			Network = NeuralNetworkAsset;
-		}
-		else
-		{
-			Network = NeuralNetworkAsset;
-			Network->NeuralNetworkData = NewObject<ULearningAgentsNeuralNetworkData>(Network);
-			
-			Network->NeuralNetworkData->CreateMemoryBackbone(
-				Interactor->GetObservationFeature().DimNum(),
-				2 * Interactor->GetActionFeature().DimNum(),
-				PolicySettings.MemoryStateSize,
-				PolicySettings.HiddenLayerSize,
-				FMath::Max(PolicySettings.LayerNum / 2, 1),
-				FMath::Max(PolicySettings.LayerNum / 2, 1));
-
+			if (EncoderNeuralNetworkAsset->NeuralNetworkData->GetInputSize() != ObservationVectorSize ||
+				EncoderNeuralNetworkAsset->NeuralNetworkData->GetOutputSize() != ObservationEncodedVectorSize)
+			{
+				UE_LOG(LogLearning, Error, TEXT("%s: Encoder Network Asset provided during Setup is incorrect size: Got inputs of size %i, expected %i. Got outputs of size %i, expected %i."), *GetName(),
+					DecoderNeuralNetworkAsset->NeuralNetworkData->GetInputSize(), ObservationVectorSize,
+					DecoderNeuralNetworkAsset->NeuralNetworkData->GetOutputSize(), ObservationEncodedVectorSize);
+				return;
+			}
 		}
 	}
-	else
-	{
-		// Create New Neural Network Asset
-
-		const FName UniqueName = MakeUniqueObjectName(this, ULearningAgentsNeuralNetwork::StaticClass(), TEXT("PolicyNetwork"), EUniqueObjectNameOptions::GloballyUnique);
-
-		Network = NewObject<ULearningAgentsNeuralNetwork>(this, UniqueName);
-		Network->NeuralNetworkData = NewObject<ULearningAgentsNeuralNetworkData>(Network);
-		
-		Network->NeuralNetworkData->CreateMemoryBackbone(
-			Interactor->GetObservationFeature().DimNum(),
-			2 * Interactor->GetActionFeature().DimNum(),
-			PolicySettings.MemoryStateSize,
-			PolicySettings.HiddenLayerSize,
-			FMath::Max(PolicySettings.LayerNum / 2, 1),
-			FMath::Max(PolicySettings.LayerNum / 2, 1));
 
+	if (!EncoderNetwork)
+	{
+		const FName UniqueName = MakeUniqueObjectName(this, ULearningAgentsNeuralNetwork::StaticClass(), TEXT("EncoderNetwork"), EUniqueObjectNameOptions::GloballyUnique);
+		EncoderNetwork = NewObject<ULearningAgentsNeuralNetwork>(this, UniqueName);
 	}
 
-	// Create Policy Object
-	UE::Learning::FNeuralNetworkPolicyFunctionSettings PolicyFunctionSettings;
-	PolicyFunctionSettings.ActionNoiseMin = PolicySettings.ActionNoiseMin;
-	PolicyFunctionSettings.ActionNoiseMax = PolicySettings.ActionNoiseMax;
-	PolicyFunctionSettings.ActionNoiseScale = PolicySettings.ActionNoiseScale;
-
-	PolicyObject = MakeShared<UE::Learning::FNeuralNetworkPolicyFunction>(
-		TEXT("PolicyObject"),
-		Manager->GetInstanceData().ToSharedRef(),
-		Manager->GetMaxAgentNum(),
-		Interactor->GetObservationFeature().DimNum(),
-		Interactor->GetActionFeature().DimNum(),
-		PolicySettings.MemoryStateSize,
-		Network->NeuralNetworkData->GetNetworkInterface(),
-		PolicySettings.ActionNoiseSeed,
-		UE::Learning::FNeuralNetworkInferenceSettings(),
-		PolicyFunctionSettings);
+	if (!EncoderNetwork->NeuralNetworkData || bReinitializeEncoderNetwork)
+	{
+		if (!EncoderNetwork->NeuralNetworkData)
+		{
+			EncoderNetwork->NeuralNetworkData = NewObject<ULearningNeuralNetworkData>(EncoderNetwork);
+		}
 
-	PolicyAgentIteration.SetNumUninitialized({ Manager->GetMaxAgentNum() });
-	UE::Learning::Array::Set<1, uint64>(PolicyAgentIteration, INDEX_NONE);
+		TArray<uint8> FileData;
+		uint32 EncoderInputSize, EncoderOutputSize;
+		UE::Learning::Observation::GenerateEncoderNetworkFileDataFromSchema(
+			FileData,
+			EncoderInputSize,
+			EncoderOutputSize,
+			Interactor->GetObservationSchema(),
+			Interactor->GetObservationSchemaElement(),
+			Seed);
+
+		UE_LEARNING_CHECK(EncoderInputSize == ObservationVectorSize);
+		UE_LEARNING_CHECK(EncoderOutputSize == ObservationEncodedVectorSize);
+
+		EncoderNetwork->NeuralNetworkData->Init(EncoderInputSize, EncoderOutputSize, ObservationCompatibilityHash, FileData);
+		EncoderNetwork->ForceMarkDirty();
+	}
 
-	PreEvaluationMemoryStateHandle = Manager->GetInstanceData()->Add<2, float>({ GetFName(), TEXT("PreEvaluationMemoryState") }, { Manager->GetMaxAgentNum(), PolicySettings.MemoryStateSize });
-	UE::Learning::Array::Set<2, float>(Manager->GetInstanceData()->View(PreEvaluationMemoryStateHandle), FLT_MAX);
+	// Policy
 
-	MemoryStateHandle = Manager->GetInstanceData()->Add<2, float>({ GetFName(), TEXT("MemoryState") }, { Manager->GetMaxAgentNum(), PolicySettings.MemoryStateSize });
-	UE::Learning::Array::Set<2, float>(Manager->GetInstanceData()->View(MemoryStateHandle), FLT_MAX);
+	if (PolicyNeuralNetworkAsset)
+	{
+		PolicyNetwork = PolicyNeuralNetworkAsset;
 
-	bIsSetup = true;
+		if (PolicyNeuralNetworkAsset->NeuralNetworkData && !bReinitializePolicyNetwork)
+		{
+			if (PolicyNeuralNetworkAsset->NeuralNetworkData->GetCompatibilityHash() != PolicyCompatibilityHash)
+			{
+				UE_LOG(LogLearning, Error, TEXT("%s: Policy Network Asset provided during Setup is incompatible with Schema. Network hash is %i vs Schema hash %i."), *GetName(),
+					PolicyNeuralNetworkAsset->NeuralNetworkData->GetCompatibilityHash(),
+					PolicyCompatibilityHash);
+				return;
+			}
 
-	OnAgentsAdded(Manager->GetAllAgentIds());
-}
+			if (PolicyNeuralNetworkAsset->NeuralNetworkData->GetInputSize() != ObservationEncodedVectorSize + MemoryStateSize ||
+				PolicyNeuralNetworkAsset->NeuralNetworkData->GetOutputSize() != ActionEncodedVectorSize + MemoryStateSize)
+			{
+				UE_LOG(LogLearning, Error, TEXT("%s: Policy Network Asset provided during Setup is incorrect size: Got inputs of size %i, expected %i. Got outputs of size %i, expected %i."), *GetName(),
+					PolicyNeuralNetworkAsset->NeuralNetworkData->GetInputSize(), ObservationEncodedVectorSize + MemoryStateSize,
+					PolicyNeuralNetworkAsset->NeuralNetworkData->GetOutputSize(), ActionEncodedVectorSize + MemoryStateSize);
+				return;
+			}
+		}
+	}
 
-void ULearningAgentsPolicy::OnAgentsAdded(const TArray<int32>& AgentIds)
-{
-	if (IsSetup())
+	if (!PolicyNetwork)
 	{
-		UE::Learning::Array::Set<1, uint64>(PolicyAgentIteration, 0, AgentIds);
-		UE::Learning::Array::Set<2, float>(Manager->GetInstanceData()->View(PreEvaluationMemoryStateHandle), 0.0f, AgentIds);
-		UE::Learning::Array::Set<2, float>(Manager->GetInstanceData()->View(MemoryStateHandle), 0.0f, AgentIds);
+		const FName UniqueName = MakeUniqueObjectName(this, ULearningAgentsNeuralNetwork::StaticClass(), TEXT("PolicyNetwork"), EUniqueObjectNameOptions::GloballyUnique);
+		PolicyNetwork = NewObject<ULearningAgentsNeuralNetwork>(this, UniqueName);
+	}
 
-		for (ULearningAgentsHelper* Helper : HelperObjects)
+	if (!PolicyNetwork->NeuralNetworkData || bReinitializePolicyNetwork)
+	{
+		if (!PolicyNetwork->NeuralNetworkData)
 		{
-			Helper->OnAgentsAdded(AgentIds);
+			PolicyNetwork->NeuralNetworkData = NewObject<ULearningNeuralNetworkData>(PolicyNetwork);
 		}
 
-		AgentsAdded(AgentIds);
+		UE::NNE::RuntimeBasic::FModelBuilder Builder;
+
+		const int32 PolicyHiddenLayerSize = PolicySettings.HiddenLayerSize;
+		const int32 PolicyLayerNum = PolicySettings.LayerNum;
+		const ELearningAgentsActivationFunction PolicyActivationFunction = PolicySettings.ActivationFunction;
+		const float PolicyInitialEncodedActionScale = PolicySettings.InitialEncodedActionScale;
+
+		TArray<uint8> FileData;
+		uint32 PolicyInputSize, PolicyOutputSize;
+		Builder.WriteFileDataAndReset(FileData, PolicyInputSize, PolicyOutputSize,
+			Builder.MakeMemoryBackbone(
+				MemoryStateSize,
+				Builder.MakeMLPWithRandomKaimingWeights(
+					ObservationEncodedVectorSize,
+					PolicyHiddenLayerSize,
+					PolicyHiddenLayerSize,
+					FMath::Max(PolicyLayerNum / 2 + 1, 2),
+					UE::Learning::Agents::Policy::Private::GetBuilderActivationFunction(PolicyActivationFunction),
+					true),
+				Builder.MakeMemoryCellWithLinearRandomKaimingWeights(
+					PolicyHiddenLayerSize,
+					PolicyHiddenLayerSize,
+					MemoryStateSize,
+					0.1f),
+				Builder.MakeSequence({
+					Builder.MakeMLPWithRandomKaimingWeights(
+						PolicyHiddenLayerSize,
+						ActionEncodedVectorSize,
+						PolicyHiddenLayerSize,
+						FMath::Max(PolicyLayerNum / 2 + 1, 2),
+						UE::Learning::Agents::Policy::Private::GetBuilderActivationFunction(PolicyActivationFunction)),
+					Builder.MakeDenormalize(
+						ActionEncodedVectorSize,
+						Builder.MakeWeightsZero(ActionEncodedVectorSize),
+						Builder.MakeWeightsConstant(ActionEncodedVectorSize, PolicyInitialEncodedActionScale))
+					})
+			));
+
+		UE_LEARNING_CHECK(PolicyInputSize == ObservationEncodedVectorSize + MemoryStateSize);
+		UE_LEARNING_CHECK(PolicyOutputSize == ActionEncodedVectorSize + MemoryStateSize);
+
+		PolicyNetwork->NeuralNetworkData->Init(PolicyInputSize, PolicyOutputSize, PolicyCompatibilityHash, FileData);
+		PolicyNetwork->ForceMarkDirty();
 	}
-}
 
-void ULearningAgentsPolicy::OnAgentsRemoved(const TArray<int32>& AgentIds)
-{
-	if (IsSetup())
+	// Decoder
+
+	if (DecoderNeuralNetworkAsset)
 	{
-		UE::Learning::Array::Set<1, uint64>(PolicyAgentIteration, INDEX_NONE, AgentIds);
-		UE::Learning::Array::Set<2, float>(Manager->GetInstanceData()->View(PreEvaluationMemoryStateHandle), FLT_MAX, AgentIds);
-		UE::Learning::Array::Set<2, float>(Manager->GetInstanceData()->View(MemoryStateHandle), FLT_MAX, AgentIds);
+		DecoderNetwork = DecoderNeuralNetworkAsset;
 
-		for (ULearningAgentsHelper* Helper : HelperObjects)
+		if (DecoderNeuralNetworkAsset->NeuralNetworkData && !bReinitializeDecoderNetwork)
 		{
-			Helper->OnAgentsRemoved(AgentIds);
-		}
+			if (DecoderNeuralNetworkAsset->NeuralNetworkData->GetCompatibilityHash() != ActionCompatibilityHash)
+			{
+				UE_LOG(LogLearning, Error, TEXT("%s: Decoder Network Asset provided during Setup is incompatible with Schema. Network hash is %i vs Schema hash %i."), *GetName(),
+					DecoderNeuralNetworkAsset->NeuralNetworkData->GetCompatibilityHash(),
+					ActionCompatibilityHash);
+				return;
+			}
 
-		AgentsRemoved(AgentIds);
+			if (DecoderNeuralNetworkAsset->NeuralNetworkData->GetInputSize() != ActionEncodedVectorSize ||
+				DecoderNeuralNetworkAsset->NeuralNetworkData->GetOutputSize() != ActionDistributionVectorSize)
+			{
+				UE_LOG(LogLearning, Error, TEXT("%s: Decoder Network Asset provided during Setup is incorrect size: Got inputs of size %i, expected %i. Got outputs of size %i, expected %i."), *GetName(),
+					DecoderNeuralNetworkAsset->NeuralNetworkData->GetInputSize(), ActionEncodedVectorSize,
+					DecoderNeuralNetworkAsset->NeuralNetworkData->GetOutputSize(), ActionDistributionVectorSize);
+				return;
+			}
+		}
 	}
-}
 
-void ULearningAgentsPolicy::OnAgentsReset(const TArray<int32>& AgentIds)
-{
-	if (IsSetup())
+	if (!DecoderNetwork)
 	{
-		UE::Learning::Array::Set<1, uint64>(PolicyAgentIteration, 0, AgentIds);
-		UE::Learning::Array::Set<2, float>(Manager->GetInstanceData()->View(PreEvaluationMemoryStateHandle), 0.0f, AgentIds);
-		UE::Learning::Array::Set<2, float>(Manager->GetInstanceData()->View(MemoryStateHandle), 0.0f, AgentIds);
+		const FName UniqueName = MakeUniqueObjectName(this, ULearningAgentsNeuralNetwork::StaticClass(), TEXT("DecoderNetwork"), EUniqueObjectNameOptions::GloballyUnique);
+		DecoderNetwork = NewObject<ULearningAgentsNeuralNetwork>(this, UniqueName);
+	}
 
-		for (ULearningAgentsHelper* Helper : HelperObjects)
+	if (!DecoderNetwork->NeuralNetworkData || bReinitializeDecoderNetwork)
+	{
+		if (!DecoderNetwork->NeuralNetworkData)
 		{
-			Helper->OnAgentsReset(AgentIds);
+			DecoderNetwork->NeuralNetworkData = NewObject<ULearningNeuralNetworkData>(DecoderNetwork);
 		}
 
-		AgentsReset(AgentIds);
+		TArray<uint8> FileData;
+		uint32 DecoderInputSize, DecoderOutputSize;
+		UE::Learning::Action::GenerateDecoderNetworkFileDataFromSchema(
+			FileData,
+			DecoderInputSize,
+			DecoderOutputSize,
+			Interactor->GetActionSchema(),
+			Interactor->GetActionSchemaElement(),
+			Seed);
+
+		UE_LEARNING_CHECK(DecoderInputSize == ActionEncodedVectorSize);
+		UE_LEARNING_CHECK(DecoderOutputSize == ActionDistributionVectorSize);
+
+		DecoderNetwork->NeuralNetworkData->Init(DecoderInputSize, DecoderOutputSize, ActionCompatibilityHash, FileData);
+		DecoderNetwork->ForceMarkDirty();
 	}
-}
 
-ULearningAgentsNeuralNetwork* ULearningAgentsPolicy::GetNetworkAsset()
-{
-	return Network;
-}
+	// Create Encoder / Policy / Decoder Objects
 
-UE::Learning::INeuralNetwork& ULearningAgentsPolicy::GetPolicyNetwork()
-{
-	return *Network->NeuralNetworkData->GetNetworkInterface();
-}
+	EncoderObject = MakeShared<UE::Learning::FNeuralNetworkFunction>(
+		Manager->GetMaxAgentNum(),
+		EncoderNetwork->NeuralNetworkData->GetNetwork(),
+		UE::Learning::FNeuralNetworkInferenceSettings());
 
-UE::Learning::FNeuralNetworkPolicyFunction& ULearningAgentsPolicy::GetPolicyObject()
-{
-	return *PolicyObject;
+	PolicyObject = MakeShared<UE::Learning::FNeuralNetworkPolicy>(
+		Manager->GetMaxAgentNum(),
+		ObservationEncodedVectorSize,
+		ActionEncodedVectorSize,
+		MemoryStateSize,
+		PolicyNetwork->NeuralNetworkData->GetNetwork(),
+		UE::Learning::FNeuralNetworkInferenceSettings());
+
+	DecoderObject = MakeShared<UE::Learning::FNeuralNetworkFunction>(
+		Manager->GetMaxAgentNum(),
+		DecoderNetwork->NeuralNetworkData->GetNetwork(),
+		UE::Learning::FNeuralNetworkInferenceSettings());
+
+	// State Variables
+
+	GlobalSeed = Seed;
+	Seeds.SetNumUninitialized({ Manager->GetMaxAgentNum() });
+	UE::Learning::Array::Set<1, uint32>(Seeds, INDEX_NONE);
+
+	ObservationVectorsEncoded.SetNumUninitialized({ Manager->GetMaxAgentNum(), ObservationEncodedVectorSize });
+	ActionVectorsEncoded.SetNumUninitialized({ Manager->GetMaxAgentNum(), ActionEncodedVectorSize });
+	ActionDistributionVectors.SetNumUninitialized({ Manager->GetMaxAgentNum(), ActionDistributionVectorSize });
+	UE::Learning::Array::Set(ObservationVectorsEncoded, FLT_MAX);
+	UE::Learning::Array::Set(ActionVectorsEncoded, FLT_MAX);
+	UE::Learning::Array::Set(ActionDistributionVectors, FLT_MAX);
+
+	PreEvaluationMemoryState.SetNumUninitialized({ Manager->GetMaxAgentNum(), MemoryStateSize });
+	MemoryState.SetNumUninitialized({ Manager->GetMaxAgentNum(), MemoryStateSize });
+	UE::Learning::Array::Set(PreEvaluationMemoryState, FLT_MAX);
+	UE::Learning::Array::Set(MemoryState, FLT_MAX);
+
+	ObservationVectorEncodedIteration.SetNumUninitialized({ Manager->GetMaxAgentNum() });
+	ActionVectorEncodedIteration.SetNumUninitialized({ Manager->GetMaxAgentNum() });
+	MemoryStateIteration.SetNumUninitialized({ Manager->GetMaxAgentNum() });
+	UE::Learning::Array::Set<1, uint64>(ObservationVectorEncodedIteration, INDEX_NONE);
+	UE::Learning::Array::Set<1, uint64>(ActionVectorEncodedIteration, INDEX_NONE);
+	UE::Learning::Array::Set<1, uint64>(MemoryStateIteration, INDEX_NONE);
+
+	bIsSetup = true;
+
+	Manager->AddListener(this);
 }
 
-void ULearningAgentsPolicy::LoadPolicyFromSnapshot(const FFilePath& File)
+void ULearningAgentsPolicy::OnAgentsAdded_Implementation(const TArray<int32>& AgentIds)
 {
 	if (!IsSetup())
 	{
@@ -211,10 +392,18 @@ void ULearningAgentsPolicy::LoadPolicyFromSnapshot(const FFilePath& File)
 		return;
 	}
 
-	Network->LoadNetworkFromSnapshot(File);
+	UE::Learning::Random::SampleIntArray(Seeds, GlobalSeed, AgentIds);
+	UE::Learning::Array::Set<2, float>(ObservationVectorsEncoded, FLT_MAX, AgentIds);
+	UE::Learning::Array::Set<2, float>(ActionVectorsEncoded, FLT_MAX, AgentIds);
+	UE::Learning::Array::Set<2, float>(ActionDistributionVectors, FLT_MAX, AgentIds);
+	UE::Learning::Array::Set<2, float>(PreEvaluationMemoryState, 0.0f, AgentIds);
+	UE::Learning::Array::Set<2, float>(MemoryState, 0.0f, AgentIds);
+	UE::Learning::Array::Set<1, uint64>(ObservationVectorEncodedIteration, 0, AgentIds);
+	UE::Learning::Array::Set<1, uint64>(ActionVectorEncodedIteration, 0, AgentIds);
+	UE::Learning::Array::Set<1, uint64>(MemoryStateIteration, 0, AgentIds);
 }
 
-void ULearningAgentsPolicy::SavePolicyToSnapshot(const FFilePath& File) const
+void ULearningAgentsPolicy::OnAgentsRemoved_Implementation(const TArray<int32>& AgentIds)
 {
 	if (!IsSetup())
 	{
@@ -222,10 +411,18 @@ void ULearningAgentsPolicy::SavePolicyToSnapshot(const FFilePath& File) const
 		return;
 	}
 
-	Network->SaveNetworkToSnapshot(File);
+	UE::Learning::Random::SampleIntArray(Seeds, GlobalSeed, AgentIds);
+	UE::Learning::Array::Set<2, float>(ObservationVectorsEncoded, FLT_MAX, AgentIds);
+	UE::Learning::Array::Set<2, float>(ActionVectorsEncoded, FLT_MAX, AgentIds);
+	UE::Learning::Array::Set<2, float>(ActionDistributionVectors, FLT_MAX, AgentIds);
+	UE::Learning::Array::Set<2, float>(PreEvaluationMemoryState, FLT_MAX, AgentIds);
+	UE::Learning::Array::Set<2, float>(MemoryState, FLT_MAX, AgentIds);
+	UE::Learning::Array::Set<1, uint64>(ObservationVectorEncodedIteration, INDEX_NONE, AgentIds);
+	UE::Learning::Array::Set<1, uint64>(ActionVectorEncodedIteration, INDEX_NONE, AgentIds);
+	UE::Learning::Array::Set<1, uint64>(MemoryStateIteration, INDEX_NONE, AgentIds);
 }
 
-void ULearningAgentsPolicy::UsePolicyFromAsset(ULearningAgentsNeuralNetwork* NeuralNetworkAsset)
+void ULearningAgentsPolicy::OnAgentsReset_Implementation(const TArray<int32>& AgentIds)
 {
 	if (!IsSetup())
 	{
@@ -233,54 +430,68 @@ void ULearningAgentsPolicy::UsePolicyFromAsset(ULearningAgentsNeuralNetwork* Neu
 		return;
 	}
 
-	if (!NeuralNetworkAsset || !NeuralNetworkAsset->NeuralNetworkData->GetNetworkInterface())
-	{
-		UE_LOG(LogLearning, Error, TEXT("%s: Asset is invalid."), *GetName());
-		return;
-	}
+	UE::Learning::Random::SampleIntArray(Seeds, GlobalSeed, AgentIds);
+	UE::Learning::Array::Set<2, float>(ObservationVectorsEncoded, FLT_MAX, AgentIds);
+	UE::Learning::Array::Set<2, float>(ActionVectorsEncoded, FLT_MAX, AgentIds);
+	UE::Learning::Array::Set<2, float>(ActionDistributionVectors, FLT_MAX, AgentIds);
+	UE::Learning::Array::Set<2, float>(PreEvaluationMemoryState, 0.0f, AgentIds);
+	UE::Learning::Array::Set<2, float>(MemoryState, 0.0f, AgentIds);
+	UE::Learning::Array::Set<1, uint64>(MemoryStateIteration, 0, AgentIds);
+	UE::Learning::Array::Set<1, uint64>(ObservationVectorEncodedIteration, 0, AgentIds);
+	UE::Learning::Array::Set<1, uint64>(ActionVectorEncodedIteration, 0, AgentIds);
+}
 
-	if (NeuralNetworkAsset == Network)
-	{
-		UE_LOG(LogLearning, Error, TEXT("%s: Asset is same as the current network."), *GetName());
-		return;
-	}
+UE::Learning::FNeuralNetworkFunction& ULearningAgentsPolicy::GetEncoderObject()
+{
+	return *EncoderObject;
+}
+
+UE::Learning::FNeuralNetworkPolicy& ULearningAgentsPolicy::GetPolicyObject()
+{
+	return *PolicyObject;
+}
+
+UE::Learning::FNeuralNetworkFunction& ULearningAgentsPolicy::GetDecoderObject()
+{
+	return *DecoderObject;
+}
 
-	if (NeuralNetworkAsset->NeuralNetworkData->GetNetworkInterface()->GetInputNum() != Network->NeuralNetworkData->GetNetworkInterface()->GetInputNum() ||
-		NeuralNetworkAsset->NeuralNetworkData->GetNetworkInterface()->GetOutputNum() != Network->NeuralNetworkData->GetNetworkInterface()->GetOutputNum())
+ULearningAgentsNeuralNetwork* ULearningAgentsPolicy::GetEncoderNetworkAsset()
+{
+	if (!IsSetup())
 	{
-		UE_LOG(LogLearning, Error, TEXT("%s: Failed to use asset as network settings don't match."), *GetName());
-		return;
+		UE_LOG(LogLearning, Error, TEXT("%s: Setup not complete."), *GetName());
+		return nullptr;
 	}
 
-	Network = NeuralNetworkAsset;
-	PolicyObject->UpdateNeuralNetwork(Network->NeuralNetworkData->GetNetworkInterface());
+	return EncoderNetwork;
 }
 
-void ULearningAgentsPolicy::LoadPolicyFromAsset(ULearningAgentsNeuralNetwork* NeuralNetworkAsset)
+ULearningAgentsNeuralNetwork* ULearningAgentsPolicy::GetPolicyNetworkAsset()
 {
 	if (!IsSetup())
 	{
 		UE_LOG(LogLearning, Error, TEXT("%s: Setup not complete."), *GetName());
-		return;
+		return nullptr;
 	}
 
-	Network->LoadNetworkFromAsset(NeuralNetworkAsset);
+	return PolicyNetwork;
 }
 
-void ULearningAgentsPolicy::SavePolicyToAsset(ULearningAgentsNeuralNetwork* NeuralNetworkAsset)
+ULearningAgentsNeuralNetwork* ULearningAgentsPolicy::GetDecoderNetworkAsset()
 {
 	if (!IsSetup())
 	{
 		UE_LOG(LogLearning, Error, TEXT("%s: Setup not complete."), *GetName());
-		return;
+		return nullptr;
 	}
 
-	Network->SaveNetworkToAsset(NeuralNetworkAsset);
+	return DecoderNetwork;
 }
 
-void ULearningAgentsPolicy::EvaluatePolicy()
+void ULearningAgentsPolicy::EncodeObservations()
 {
-	UE_LEARNING_TRACE_CPUPROFILER_EVENT_SCOPE(ULearningAgentsPolicy::EvaluatePolicy);
+	UE_LEARNING_TRACE_CPUPROFILER_EVENT_SCOPE(ULearningAgentsPolicy::EncodeObservations);
 
 	if (!IsSetup())
 	{
@@ -288,15 +499,19 @@ void ULearningAgentsPolicy::EvaluatePolicy()
 		return;
 	}
 
-	// Check Agents actually have encoded observations.
+	if (Manager->GetAgentNum() == 0)
+	{
+		UE_LOG(LogLearning, Warning, TEXT("%s: No agents added to Manager."), *GetName());
+	}
 
-	ValidAgentIds.Empty(Manager->GetAgentNum());
+	// Check which agents have had their observation vector set
 
+	ValidAgentIds.Empty(Manager->GetMaxAgentNum());
 	for (const int32 AgentId : Manager->GetAllAgentSet())
 	{
-		if (Interactor->GetObservationEncodingAgentIteration()[AgentId] == 0)
+		if (Interactor->ObservationVectorIteration[AgentId] == 0)
 		{
-			UE_LOG(LogLearning, Warning, TEXT("%s: Agent with id %i has not made observations so policy will not be evaluated for it."), *GetName(), AgentId);
+			UE_LOG(LogLearning, Warning, TEXT("%s: Agent with id %i does not have an observation vector ready be encoded. Was GatherObservations run without error?"), *GetName(), AgentId);
 			continue;
 		}
 
@@ -306,59 +521,74 @@ void ULearningAgentsPolicy::EvaluatePolicy()
 	ValidAgentSet = ValidAgentIds;
 	ValidAgentSet.TryMakeSlice();
 
-	// Record pre-evaluation state
-
-	TLearningArrayView<2, float> PreEvaluationMemoryStateView = Manager->GetInstanceData()->View(PreEvaluationMemoryStateHandle);
-	TLearningArrayView<2, float> MemoryStateView = Manager->GetInstanceData()->View(MemoryStateHandle);
+	// Encode Observations
 
-	UE::Learning::Array::Copy<2, float>(PreEvaluationMemoryStateView, MemoryStateView, ValidAgentSet);
+	EncoderObject->Evaluate(ObservationVectorsEncoded, Interactor->ObservationVectors, ValidAgentSet);
 
-	// Copy Observations and Memory State into input buffer
+	for (const int32 AgentId : ValidAgentSet)
+	{
+		ObservationVectorEncodedIteration[AgentId]++;
+	}
+}
 
-	TLearningArrayView<2, const float> ObservationsView = Manager->GetInstanceData()->ConstView(Interactor->GetObservationFeature().FeatureHandle);
-	TLearningArrayView<2, float> InputObservationView = Manager->GetInstanceData()->View(PolicyObject->InputObservationHandle);
-	TLearningArrayView<2, float> InputMemoryStateView = Manager->GetInstanceData()->View(PolicyObject->InputMemoryStateHandle);
+void ULearningAgentsPolicy::EvaluatePolicy()
+{
+	UE_LEARNING_TRACE_CPUPROFILER_EVENT_SCOPE(ULearningAgentsPolicy::EvaluatePolicy);
 
-	UE::Learning::Array::Copy<2, float>(InputObservationView, ObservationsView, ValidAgentSet);
-	UE::Learning::Array::Copy<2, float>(InputMemoryStateView, MemoryStateView, ValidAgentSet);
+	if (!IsSetup())
+	{
+		UE_LOG(LogLearning, Error, TEXT("%s: Setup not complete."), *GetName());
+		return;
+	}
 
-	// Evaluate Policy
+	if (Manager->GetAgentNum() == 0)
+	{
+		UE_LOG(LogLearning, Warning, TEXT("%s: No agents added to Manager."), *GetName());
+	}
 
-	PolicyObject->Evaluate(ValidAgentSet);
+	// Check agents actually have encoded observations
 
-	// Increment Policy Evaluation Iteration
+	ValidAgentIds.Empty(Manager->GetMaxAgentNum());
 
-	for (const int32 AgentId : ValidAgentSet)
+	for (const int32 AgentId : Manager->GetAllAgentSet())
 	{
-		PolicyAgentIteration[AgentId]++;
+		if (ObservationVectorEncodedIteration[AgentId] == 0)
+		{
+			UE_LOG(LogLearning, Warning, TEXT("%s: Agent with id %i has not encoded observations so policy will not be evaluated for it. Was EncodeObservations run without error?"), *GetName(), AgentId);
+			continue;
+		}
+
+		ValidAgentIds.Add(AgentId);
 	}
 
-	// Copy Actions and Memory State out of output buffer
+	ValidAgentSet = ValidAgentIds;
+	ValidAgentSet.TryMakeSlice();
 
-	TLearningArrayView<2, float> ActionsView = Manager->GetInstanceData()->View(Interactor->GetActionFeature().FeatureHandle);
-	TLearningArrayView<2, const float> OutputActionView = Manager->GetInstanceData()->ConstView(PolicyObject->OutputActionHandle);
-	TLearningArrayView<2, const float> OutputMemoryStateView = Manager->GetInstanceData()->ConstView(PolicyObject->OutputMemoryStateHandle);
+	// Copy pre-evaluation memory state
 
-	UE::Learning::Array::Copy<2, float>(ActionsView, OutputActionView, ValidAgentSet);
-	UE::Learning::Array::Copy<2, float>(MemoryStateView, OutputMemoryStateView, ValidAgentSet);
+	UE::Learning::Array::Copy<2, float>(PreEvaluationMemoryState, MemoryState, ValidAgentSet);
 
-	// Increment Action Encoding Iteration
+	// Evaluate policy
+
+	PolicyObject->Evaluate(
+		ActionVectorsEncoded,
+		MemoryState,
+		ObservationVectorsEncoded,
+		PreEvaluationMemoryState,
+		ValidAgentSet);
+
+	// Increment policy evaluation and action generation iteration
 
 	for (const int32 AgentId : ValidAgentSet)
 	{
-		Interactor->GetActionEncodingAgentIteration()[AgentId]++;
+		MemoryStateIteration[AgentId]++;
+		ActionVectorEncodedIteration[AgentId]++;
 	}
-
-	// Visual Logger
-
-#if UE_LEARNING_AGENTS_ENABLE_VISUAL_LOG
-	VisualLog(ValidAgentSet);
-#endif
 }
 
-void ULearningAgentsPolicy::RunInference()
+void ULearningAgentsPolicy::DecodeAndSampleActions(const float ActionNoiseScale)
 {
-	UE_LEARNING_TRACE_CPUPROFILER_EVENT_SCOPE(ULearningAgentsPolicy::RunInference);
+	UE_LEARNING_TRACE_CPUPROFILER_EVENT_SCOPE(ULearningAgentsPolicy::DecodeAndSampleActions);
 
 	if (!IsSetup())
 	{
@@ -366,30 +596,67 @@ void ULearningAgentsPolicy::RunInference()
 		return;
 	}
 
-	Interactor->EncodeObservations();
-	EvaluatePolicy();
-	Interactor->DecodeActions();
-}
-float ULearningAgentsPolicy::GetActionNoiseScale() const
-{
-	if (!IsSetup())
+	if (Manager->GetAgentNum() == 0)
 	{
-		UE_LOG(LogLearning, Error, TEXT("%s: Setup not complete."), *GetName());
-		return 0.0f;
+		UE_LOG(LogLearning, Warning, TEXT("%s: No agents added to Manager."), *GetName());
+	}
+
+	// Check which agents have had their encoded action vector set
+
+	ValidAgentIds.Empty(Manager->GetMaxAgentNum());
+
+	for (const int32 AgentId : Manager->GetAllAgentSet())
+	{
+		if (ActionVectorEncodedIteration[AgentId] == 0)
+		{
+			UE_LOG(LogLearning, Warning, TEXT("%s: Agent with id %i does not have an encoded action so actions will not be decoded for it. Was EvaluatePolicy run without error?"), *GetName(), AgentId);
+			continue;
+		}
+
+		ValidAgentIds.Add(AgentId);
 	}
 
-	return PolicyObject->InstanceData->ConstView(PolicyObject->ActionNoiseScaleHandle)[0];
+	ValidAgentSet = ValidAgentIds;
+	ValidAgentSet.TryMakeSlice();
+
+	// Decode to produce action distribution vectors
+
+	DecoderObject->Evaluate(ActionDistributionVectors, ActionVectorsEncoded, ValidAgentSet);
+
+	// Sample actions from action distribution vectors
+
+	for (const int32 AgentId : ValidAgentSet)
+	{
+		UE::Learning::Action::SampleVectorFromDistributionVector(
+			Seeds[AgentId],
+			Interactor->ActionVectors[AgentId],
+			ActionDistributionVectors[AgentId],
+			Interactor->ActionSchema->GetActionSchema(),
+			Interactor->ActionSchemaElement.SchemaElement,
+			ActionNoiseScale);
+	}
+
+	for (const int32 AgentId : ValidAgentSet)
+	{
+		Interactor->ActionVectorIteration[AgentId]++;
+	}
 }
 
-void ULearningAgentsPolicy::SetActionNoiseScale(const float ActionNoiseScale)
+void ULearningAgentsPolicy::RunInference(const float ActionNoiseScale)
 {
+	UE_LEARNING_TRACE_CPUPROFILER_EVENT_SCOPE(ULearningAgentsPolicy::RunInference);
+
 	if (!IsSetup())
 	{
 		UE_LOG(LogLearning, Error, TEXT("%s: Setup not complete."), *GetName());
 		return;
 	}
 
-	UE::Learning::Array::Set(PolicyObject->InstanceData->View(PolicyObject->ActionNoiseScaleHandle), ActionNoiseScale);
+	Interactor->GatherObservations();
+	EncodeObservations();
+	EvaluatePolicy();
+	DecodeAndSampleActions(ActionNoiseScale);
+	Interactor->ScatterActions();
 }
 
 void ULearningAgentsPolicy::GetMemoryState(TArray<float>& OutMemoryState, const int32 AgentId) const
@@ -408,10 +675,8 @@ void ULearningAgentsPolicy::GetMemoryState(TArray<float>& OutMemoryState, const
 		return;
 	}
 
-	TLearningArrayView<2, const float> MemoryStateView = Manager->GetInstanceData()->ConstView(MemoryStateHandle);
-
-	OutMemoryState.SetNumUninitialized(MemoryStateView.Num<1>());
-	UE::Learning::Array::Copy<1, float>(OutMemoryState, MemoryStateView[AgentId]);
+	OutMemoryState.SetNumUninitialized(MemoryState.Num<1>());
+	UE::Learning::Array::Copy<1, float>(OutMemoryState, MemoryState[AgentId]);
 }
 
 void ULearningAgentsPolicy::SetMemoryState(const int32 AgentId, const TArray<float>& InMemoryState)
@@ -428,74 +693,22 @@ void ULearningAgentsPolicy::SetMemoryState(const int32 AgentId, const TArray<flo
 		return;
 	}
 
-	TLearningArrayView<2, float> MemoryStateView = Manager->GetInstanceData()->View(MemoryStateHandle);
-
-	if (InMemoryState.Num() != MemoryStateView.Num<1>())
+	if (InMemoryState.Num() != MemoryState.Num<1>())
 	{
-		UE_LOG(LogLearning, Error, TEXT("%s: Memory State is incorrect size. Expected %i, got %i."), *GetName(), MemoryStateView.Num<1>(), InMemoryState.Num());
+		UE_LOG(LogLearning, Error, TEXT("%s: Memory State is incorrect size. Expected %i, got %i."), *GetName(), MemoryState.Num<1>(), InMemoryState.Num());
 		return;
 	}
 
-	UE::Learning::Array::Copy<1, float>(MemoryStateView[AgentId], InMemoryState);
-}
-
-TLearningArrayView<2, const float> ULearningAgentsPolicy::GetPreEvaluationMemoryStateView() const
-{
-	return Manager->GetInstanceData()->ConstView(PreEvaluationMemoryStateHandle);
-}
-
-TLearningArrayView<2, float> ULearningAgentsPolicy::GetPreEvaluationMemoryStateView()
-{
-	return Manager->GetInstanceData()->View(PreEvaluationMemoryStateHandle);
-}
-
-TLearningArrayView<2, const float> ULearningAgentsPolicy::GetMemoryStateView() const
-{
-	return Manager->GetInstanceData()->ConstView(MemoryStateHandle);
-}
-
-TLearningArrayView<2, float> ULearningAgentsPolicy::GetMemoryStateView()
-{
-	return Manager->GetInstanceData()->View(MemoryStateHandle);
+	UE::Learning::Array::Copy<1, float>(MemoryState[AgentId], InMemoryState);
 }
 
 int32 ULearningAgentsPolicy::GetMemoryStateSize() const
 {
-	return Manager->GetInstanceData()->ConstView(MemoryStateHandle).Num<1>();
-}
-
-#if UE_LEARNING_AGENTS_ENABLE_VISUAL_LOG
-void ULearningAgentsPolicy::VisualLog(const UE::Learning::FIndexSet AgentSet) const
-{
-	UE_LEARNING_TRACE_CPUPROFILER_EVENT_SCOPE(ULearningAgentsPolicy::VisualLog);
-
-	const TLearningArrayView<2, const float> InputObservationView = PolicyObject->InstanceData->ConstView(PolicyObject->InputObservationHandle);
-	const TLearningArrayView<2, const float> InputMemoryStateView = PolicyObject->InstanceData->ConstView(PolicyObject->InputMemoryStateHandle);
-	const TLearningArrayView<2, const float> OutputActionView = PolicyObject->InstanceData->ConstView(PolicyObject->OutputActionHandle);
-	const TLearningArrayView<2, const float> OutputActionMeanView = PolicyObject->InstanceData->ConstView(PolicyObject->OutputActionMeanHandle);
-	const TLearningArrayView<2, const float> OutputActionStdView = PolicyObject->InstanceData->ConstView(PolicyObject->OutputActionStdHandle);
-	const TLearningArrayView<1, const float> ActionNoiseScaleView = PolicyObject->InstanceData->ConstView(PolicyObject->ActionNoiseScaleHandle);
-
-	for (const int32 AgentId : AgentSet)
+	if (!IsSetup())
 	{
-		if (const AActor* Actor = Cast<AActor>(GetAgent(AgentId)))
-		{
-			UE_LEARNING_AGENTS_VLOG_STRING(this, LogLearning, Display,
-				Actor->GetActorLocation(),
-				VisualLogColor.ToFColor(true),
-				TEXT("Agent %i\nAction Noise Scale: [% 6.3f]\nObservation Input: %s\nObservation Input Stats (Min/Max/Mean/Std): %s\nMemory State Input: %s\nMemory State Input Stats (Min/Max/Mean/Std): %s\nOutput Mean: %s\nOutput Std: %s\nOutput Sample: %s\nOutput Stats (Min/Max/Mean/Std): %s"),
-				AgentId,
-				ActionNoiseScaleView[AgentId],
-				*UE::Learning::Array::FormatFloat(InputObservationView[AgentId]),
-				*UE::Learning::Agents::Debug::FloatArrayToStatsString(InputObservationView[AgentId]),
-				*UE::Learning::Array::FormatFloat(InputMemoryStateView[AgentId]),
-				*UE::Learning::Agents::Debug::FloatArrayToStatsString(InputMemoryStateView[AgentId]),
-				*UE::Learning::Array::FormatFloat(OutputActionMeanView[AgentId]),
-				*UE::Learning::Array::FormatFloat(OutputActionStdView[AgentId]),
-				*UE::Learning::Array::FormatFloat(OutputActionView[AgentId]),
-				*UE::Learning::Agents::Debug::FloatArrayToStatsString(OutputActionView[AgentId]));
-
-		}
+		UE_LOG(LogLearning, Error, TEXT("%s: Setup not complete."), *GetName());
+		return 0;
 	}
+
+	return MemoryState.Num<1>();
 }
-#endif
diff --git a/Engine/Plugins/Experimental/LearningAgents/Source/LearningAgents/Public/LearningAgentsActions.h b/Engine/Plugins/Experimental/LearningAgents/Source/LearningAgents/Public/LearningAgentsActions.h
--- a/Engine/Plugins/Experimental/LearningAgents/Source/LearningAgents/Public/LearningAgentsActions.h
+++ b/Engine/Plugins/Experimental/LearningAgents/Source/LearningAgents/Public/LearningAgentsActions.h
@@ -3,409 +3,405 @@
 #pragma once
 
 #include "LearningArray.h"
-#include "LearningAgentsDebug.h"
-
-#include "Templates/SharedPointer.h"
-#include "UObject/Object.h"
+#include "LearningAction.h"
 
 #include "LearningAgentsActions.generated.h"
 
-namespace UE::Learning
+class ULearningAgentsActionSchema;
+class ULearningAgentsActionObject;
+struct FLearningAgentsActionSchemaElement;
+struct FLearningAgentsActionObjectElement;
+
+/** An element of an Action Schema */
+USTRUCT(BlueprintType)
+struct LEARNINGAGENTS_API FLearningAgentsActionSchemaElement
+{
+	GENERATED_BODY()
+
+	UE::Learning::Action::FSchemaElement SchemaElement;
+};
+
+/** An element of an Action Object */
+USTRUCT(BlueprintType)
+struct LEARNINGAGENTS_API FLearningAgentsActionObjectElement
 {
-	struct FFeatureObject;
-	struct FFloatFeature;
-	struct FPlanarVelocityFeature;
-	struct FRotationVectorFeature;
-}
+	GENERATED_BODY()
 
-class ULearningAgentsInteractor;
+	UE::Learning::Action::FObjectElement ObjectElement;
+};
+
+/** Comparison operator for Action Object Elements */
+bool operator==(const FLearningAgentsActionObjectElement& Lhs, const FLearningAgentsActionObjectElement& Rhs);
 
-// For functions in this file, we are favoring having more verbose names such as "AddFloatAction" vs simply "Add" in 
-// order to keep it easy to find the correct function in blueprints.
+/** Hashing operator for Action Object Elements */
+uint32 GetTypeHash(const FLearningAgentsActionObjectElement& Element);
+
+/** Enum Type representing either action A or action B */
+UENUM(BlueprintType)
+enum class ELearningAgentsEitherAction : uint8
+{
+	A,
+	B,
+};
 
-//------------------------------------------------------------------
+/** Enum Type representing either a Null action or some Valid action */
+UENUM(BlueprintType)
+enum class ELearningAgentsOptionalAction : uint8
+{
+	Null,
+	Valid,
+};
 
 /**
- * The base class for all actions. Actions define the outputs from your agents. Action getters are marked non-pure by
- * convention as many of them do non-trivial amounts of work that can cause performance issues when marked pure in 
- * blueprints.
+ * Action Schema
+ * 
+ * This object is used to construct a schema describing some structure of actions.
  */
-UCLASS(Abstract, BlueprintType)
-class LEARNINGAGENTS_API ULearningAgentsAction : public UObject
+UCLASS(BlueprintType)
+class LEARNINGAGENTS_API ULearningAgentsActionSchema : public UObject
 {
 	GENERATED_BODY()
 
 public:
 
-	/** Reference to the Interactor this action is associated with. */
-	UPROPERTY(VisibleAnywhere, Transient, Category = "LearningAgents")
-	TObjectPtr<ULearningAgentsInteractor> Interactor;
+	/** Gets the internal action schema object */
+	const UE::Learning::Action::FSchema& GetActionSchema() const;
+
+	/**
+	 * Validates that the given object matches the schema. Will log errors on objects that don't match.
+	 * 
+	 * @param SchemaElement			Schema Element
+	 * @param Object				Action Object
+	 * @param ObjectElement			Action Object Element
+	 * @returns						true if the object matches the schema
+	 */
+	UFUNCTION(BlueprintPure, Category = "LearningAgents")
+	bool ValidateObjectMatchesSchema(
+		const FLearningAgentsActionSchemaElement SchemaElement,
+		const ULearningAgentsActionObject* Object,
+		const FLearningAgentsActionObjectElement ObjectElement) const;
 
 public:
 
-	/** Initialize the internal state for a given maximum number of agents */
-	void Init(const int32 MaxAgentNum);
+	UFUNCTION(BlueprintPure, Category = "LearningAgents")
+	FLearningAgentsActionSchemaElement SpecifyNullAction(const FName Name = TEXT("Null"));
 
-	/**
-	 * Called whenever agents are added to the associated ULearningAgentsInteractor object.
-	 * @param AgentIds Array of agent ids which have been added
-	 */
-	virtual void OnAgentsAdded(const TArray<int32>& AgentIds);
+	UFUNCTION(BlueprintPure, Category = "LearningAgents")
+	FLearningAgentsActionSchemaElement SpecifyContinuousAction(const int32 Size, const FName Name = TEXT("Continuous"));
 
-	/**
-	 * Called whenever agents are removed from the associated ULearningAgentsInteractor object.
-	 * @param AgentIds Array of agent ids which have been removed
-	 */
-	virtual void OnAgentsRemoved(const TArray<int32>& AgentIds);
+	UFUNCTION(BlueprintPure, Category = "LearningAgents", meta=(AutoCreateRefTerm="PriorProbabilities"))
+	FLearningAgentsActionSchemaElement SpecifyExclusiveDiscreteAction(const int32 Size, const TArray<float>& PriorProbabilities, const FName Name = TEXT("DiscreteExclusive"));
+	FLearningAgentsActionSchemaElement SpecifyExclusiveDiscreteActionFromArrayView(const int32 Size, const TArrayView<const float> PriorProbabilities, const FName Name = TEXT("DiscreteExclusive"));
 
-	/**
-	 * Called whenever agents are reset on the associated ULearningAgentsInteractor object.
-	 * @param AgentIds Array of agent ids which have been reset
-	 */
-	virtual void OnAgentsReset(const TArray<int32>& AgentIds);
+	UFUNCTION(BlueprintPure, Category = "LearningAgents", meta = (AutoCreateRefTerm = "PriorProbabilities"))
+	FLearningAgentsActionSchemaElement SpecifyInclusiveDiscreteAction(const int32 Size, const TArray<float>& PriorProbabilities, const FName Name = TEXT("DiscreteInclusive"));
+	FLearningAgentsActionSchemaElement SpecifyInclusiveDiscreteActionFromArrayView(const int32 Size, const TArrayView<const float> PriorProbabilities, const FName Name = TEXT("DiscreteInclusive"));
 
-	/** Get the number of times an action has been got for the given agent id. */
-	uint64 GetAgentGetIteration(const int32 AgentId) const;
+	UFUNCTION(BlueprintPure, Category = "LearningAgents")
+	FLearningAgentsActionSchemaElement SpecifyStructAction(const TMap<FName, FLearningAgentsActionSchemaElement>& Elements, const FName Name = TEXT("Struct"));
 
-	/** Get the number of times an action has been set for the given agent id. */
-	uint64 GetAgentSetIteration(const int32 AgentId) const;
+	UFUNCTION(BlueprintPure, Category = "LearningAgents")
+	FLearningAgentsActionSchemaElement SpecifyStructActionFromArrays(const TArray<FName>& ElementNames, const TArray<FLearningAgentsActionSchemaElement>& Elements, const FName Name = TEXT("Struct"));
+	FLearningAgentsActionSchemaElement SpecifyStructActionFromArrayViews(const TArrayView<const FName> ElementNames, const TArrayView<const FLearningAgentsActionSchemaElement> Elements, const FName Name = TEXT("Struct"));
 
-public:
-#if UE_LEARNING_AGENTS_ENABLE_VISUAL_LOG
-	/** Color used to draw this action in the visual log */
-	FLinearColor VisualLogColor = FColor::Blue;
+	UFUNCTION(BlueprintPure, Category = "LearningAgents", meta = (AutoCreateRefTerm = "PriorProbabilities"))
+	FLearningAgentsActionSchemaElement SpecifyExclusiveUnionAction(const TMap<FName, FLearningAgentsActionSchemaElement>& Elements, const TMap<FName, float>& PriorProbabilities, const FName Name = TEXT("ExclusiveUnion"));
+	
+	UFUNCTION(BlueprintPure, Category = "LearningAgents", meta = (AutoCreateRefTerm = "PriorProbabilities"))
+	FLearningAgentsActionSchemaElement SpecifyExclusiveUnionActionFromArrays(const TArray<FName>& ElementNames, const TArray<FLearningAgentsActionSchemaElement>& Elements, const TArray<float>& PriorProbabilities, const FName Name = TEXT("ExclusiveUnion"));
+	FLearningAgentsActionSchemaElement SpecifyExclusiveUnionActionFromArrayViews(const TArrayView<const FName> ElementNames, const TArrayView<const FLearningAgentsActionSchemaElement> Elements, const TArrayView<const float> PriorProbabilities, const FName Name = TEXT("ExclusiveUnion"));
 
-	/** Describes this action to the visual logger for debugging purposes. */
-	virtual void VisualLog(const UE::Learning::FIndexSet Instances) const {}
-#endif
+	UFUNCTION(BlueprintPure, Category = "LearningAgents", meta = (AutoCreateRefTerm = "PriorProbabilities"))
+	FLearningAgentsActionSchemaElement SpecifyInclusiveUnionAction(const TMap<FName, FLearningAgentsActionSchemaElement>& Elements, const TMap<FName, float>& PriorProbabilities, const FName Name = TEXT("InclusiveUnion"));
+	
+	UFUNCTION(BlueprintPure, Category = "LearningAgents", meta = (AutoCreateRefTerm = "PriorProbabilities"))
+	FLearningAgentsActionSchemaElement SpecifyInclusiveUnionActionFromArrays(const TArray<FName> ElementNames, const TArray<FLearningAgentsActionSchemaElement>& Elements, const TArray<float>& PriorProbabilities, const FName Name = TEXT("InclusiveUnion"));
+	FLearningAgentsActionSchemaElement SpecifyInclusiveUnionActionFromArrayViews(const TArrayView<const FName> ElementNames, const TArrayView<const FLearningAgentsActionSchemaElement> Elements, const TArrayView<const float> PriorProbabilities, const FName Name = TEXT("InclusiveUnion"));
 
-protected:
+	UFUNCTION(BlueprintPure, Category = "LearningAgents")
+	FLearningAgentsActionSchemaElement SpecifyStaticArrayAction(const FLearningAgentsActionSchemaElement Element, const int32 Num, const FName Name = TEXT("StaticArray"));
 
-	/** Number of times this action has been got for all agents */
-	TLearningArray<1, uint64, TInlineAllocator<32>> AgentGetIteration;
+	UFUNCTION(BlueprintPure, Category = "LearningAgents")
+	FLearningAgentsActionSchemaElement SpecifyPairAction(const FLearningAgentsActionSchemaElement Key, const FLearningAgentsActionSchemaElement Value, const FName Name = TEXT("Pair"));
 
-	/** Number of times this action has been set for all agents */
-	TLearningArray<1, uint64, TInlineAllocator<32>> AgentSetIteration;
-};
+	UFUNCTION(BlueprintPure, Category = "LearningAgents", meta = (AutoCreateRefTerm = "PriorProbabilities"))
+	FLearningAgentsActionSchemaElement SpecifyEnumAction(const UEnum* Enum, const TMap<uint8, float>& PriorProbabilities, const FName Name = TEXT("Enum"));
 
-//------------------------------------------------------------------
+	UFUNCTION(BlueprintPure, Category = "LearningAgents", meta = (AutoCreateRefTerm = "PriorProbabilities"))
+	FLearningAgentsActionSchemaElement SpecifyEnumActionFromArray(const UEnum* Enum, const TArray<float>& PriorProbabilities, const FName Name = TEXT("Enum"));
+	FLearningAgentsActionSchemaElement SpecifyEnumActionFromArrayView(const UEnum* Enum, const TArrayView<const float> PriorProbabilities, const FName Name = TEXT("Enum"));
 
-/** A simple float action. Used as a catch-all for situations where a more type-specific action does not exist yet. */
-UCLASS()
-class LEARNINGAGENTS_API UFloatAction : public ULearningAgentsAction
-{
-	GENERATED_BODY()
+	UFUNCTION(BlueprintPure, Category = "LearningAgents", meta = (AutoCreateRefTerm = "PriorProbabilities"))
+	FLearningAgentsActionSchemaElement SpecifyBitmaskAction(const UEnum* Enum, const TMap<uint8, float>& PriorProbabilities, const FName Name = TEXT("Bitmask"));
 
-public:
+	UFUNCTION(BlueprintPure, Category = "LearningAgents", meta = (AutoCreateRefTerm = "PriorProbabilities"))
+	FLearningAgentsActionSchemaElement SpecifyBitmaskActionFromArray(const UEnum* Enum, const TArray<float>& PriorProbabilities, const FName Name = TEXT("Bitmask"));
+	FLearningAgentsActionSchemaElement SpecifyBitmaskActionFromArrayView(const UEnum* Enum, const TArrayView<const float> PriorProbabilities, const FName Name = TEXT("Bitmask"));
 
-	/**
-	 * Adds a new float action to the given agent interactor. Call during ULearningAgentsInteractor::SetupActions event.
-	 * @param InInteractor The agent interactor to add this action to.
-	 * @param Name The name of this new action. Used for debugging.
-	 * @param Scale Used to normalize the data for the action.
-	 * @return The newly created action.
-	 */
-	UFUNCTION(BlueprintCallable, Category = "LearningAgents", meta = (DefaultToSelf = "InInteractor"))
-	static UFloatAction* AddFloatAction(ULearningAgentsInteractor* InInteractor, const FName Name = NAME_None, const float Scale = 1.0f);
+	UFUNCTION(BlueprintPure, Category = "LearningAgents")
+	FLearningAgentsActionSchemaElement SpecifyOptionalAction(const FLearningAgentsActionSchemaElement Element, const float PriorProbability = 0.5f, const FName Name = TEXT("Optional"));
 
-	/**
-	 * Gets the data for this action. Call during ULearningAgentsInteractor::GetActions event.
-	 * @param AgentId The agent id to get data for.
-	 * @return The current action value.
-	 */
-	UFUNCTION(BlueprintCallable, Category = "LearningAgents", meta = (AgentId = "-1"))
-	float GetFloatAction(const int32 AgentId);
+	UFUNCTION(BlueprintPure, Category = "LearningAgents")
+	FLearningAgentsActionSchemaElement SpecifyEitherAction(const FLearningAgentsActionSchemaElement A, const FLearningAgentsActionSchemaElement B, const float PriorProbabilityOfA = 0.5f, const FName Name = TEXT("Either"));
 
-	/**
-	 * Sets the data for this action. Call during ULearningAgentsController::SetActions event.
-	 * @param AgentId The agent id to set data for.
-	 * @param Value The current action value.
-	 */	
-	UFUNCTION(BlueprintCallable, Category = "LearningAgents", meta = (AgentId = "-1"))
-	void SetFloatAction(const int32 AgentId, const float Value);
-
-#if UE_LEARNING_AGENTS_ENABLE_VISUAL_LOG
-	/** Describes this action to the visual logger for debugging purposes. */
-	virtual void VisualLog(const UE::Learning::FIndexSet Instances) const override;
-#endif
-
-	TSharedPtr<UE::Learning::FFloatFeature> FeatureObject;
+	UFUNCTION(BlueprintPure, Category = "LearningAgents")
+	FLearningAgentsActionSchemaElement SpecifyEncodingAction(const FLearningAgentsActionSchemaElement Element, const int32 EncodingSize = 128, const FName Name = TEXT("Encoding"));
+
+	UFUNCTION(BlueprintPure, Category = "LearningAgents")
+	FLearningAgentsActionSchemaElement SpecifyBoolAction(const float PriorProbability = 0.5f, const FName Name = TEXT("Bool"));
+
+	UFUNCTION(BlueprintPure, Category = "LearningAgents")
+	FLearningAgentsActionSchemaElement SpecifyFloatAction(const FName Name = TEXT("Float"));
+
+	UFUNCTION(BlueprintPure, Category = "LearningAgents")
+	FLearningAgentsActionSchemaElement SpecifyTranslationAction(const FName Name = TEXT("Translation"));
+
+	UFUNCTION(BlueprintPure, Category = "LearningAgents")
+	FLearningAgentsActionSchemaElement SpecifyRotationAction(const FName Name = TEXT("Rotation"));
+
+	UFUNCTION(BlueprintPure, Category = "LearningAgents")
+	FLearningAgentsActionSchemaElement SpecifyScaleAction(const FName Name = TEXT("Scale"));
+
+	UFUNCTION(BlueprintPure, Category = "LearningAgents")
+	FLearningAgentsActionSchemaElement SpecifyTransformAction(const FName Name = TEXT("Transform"));
+
+	UFUNCTION(BlueprintPure, Category = "LearningAgents")
+	FLearningAgentsActionSchemaElement SpecifyAngleAction(const FName Name = TEXT("Angle"));
+
+	UFUNCTION(BlueprintPure, Category = "LearningAgents")
+	FLearningAgentsActionSchemaElement SpecifyVelocityAction(const FName Name = TEXT("Velocity"));
+
+	UFUNCTION(BlueprintPure, Category = "LearningAgents")
+	FLearningAgentsActionSchemaElement SpecifySpeedAction(const FName Name = TEXT("Speed"));
+
+private:
+
+	UE::Learning::Action::FSchema ActionSchema;
 };
 
-/** A simple float array action. Used as a catch-all for situations where a more type-specific action does not exist yet. */
-UCLASS()
-class LEARNINGAGENTS_API UFloatArrayAction : public ULearningAgentsAction
+/**
+ * Action Object
+ *
+ * This object is used to construct or get the values of actions.
+ */
+UCLASS(BlueprintType)
+class LEARNINGAGENTS_API ULearningAgentsActionObject : public UObject
 {
 	GENERATED_BODY()
 
 public:
 
-	/**
-	 * Adds a new float array action to the given agent interactor. Call during ULearningAgentsInteractor::SetupActions event.
-	 * @param InInteractor The agent interactor to add this action to.
-	 * @param Name The name of this new action. Used for debugging.
-	 * @param Num The number of floats in the array
-	 * @param Scale Used to normalize the data for the action.
-	 * @return The newly created action.
-	 */
-	UFUNCTION(BlueprintCallable, Category = "LearningAgents", meta = (DefaultToSelf = "InInteractor"))
-	static UFloatArrayAction* AddFloatArrayAction(ULearningAgentsInteractor* InInteractor, const FName Name = NAME_None, const int32 Num = 1, const float Scale = 1.0f);
+	/** Gets the internal action object */
+	const UE::Learning::Action::FObject& GetActionObject() const;
 
-	/**
-	 * Gets the data for this action. Call during ULearningAgentsInteractor::GetActions event.
-	 * @param AgentId The agent id to get data for.
-	 * @param OutValues The output array of floats
-	 */
-	UFUNCTION(BlueprintCallable, Category = "LearningAgents", meta = (AgentId = "-1"))
-	void GetFloatArrayAction(const int32 AgentId, TArray<float>& OutValues);
+	/** Gets the internal action object */
+	UE::Learning::Action::FObject& GetActionObject();
 
-	/**
-	 * Sets the data for this action. Call during ULearningAgentsController::SetActions event.
-	 * @param AgentId The agent id to set data for.
-	 * @param Values The current action values.
-	 */
-	UFUNCTION(BlueprintCallable, Category = "LearningAgents", meta = (AgentId = "-1"))
-	void SetFloatArrayAction(const int32 AgentId, const TArray<float>& Values);
+public:
 
-#if UE_LEARNING_AGENTS_ENABLE_VISUAL_LOG
-	/** Describes this action to the visual logger for debugging purposes. */
-	virtual void VisualLog(const UE::Learning::FIndexSet Instances) const override;
-#endif
+	UFUNCTION(BlueprintPure = false, Category = "LearningAgents")
+	void LogAction(const FLearningAgentsActionObjectElement Element);
 
-	TSharedPtr<UE::Learning::FFloatFeature> FeatureObject;
-};
+public:
 
-//------------------------------------------------------------------
+	UFUNCTION(BlueprintPure, Category = "LearningAgents")
+	FLearningAgentsActionObjectElement MakeNullAction(const FName Name = TEXT("Null"));
 
-/** A simple FVector action. */
-UCLASS()
-class LEARNINGAGENTS_API UVectorAction : public ULearningAgentsAction
-{
-	GENERATED_BODY()
+	UFUNCTION(BlueprintPure, Category = "LearningAgents")
+	FLearningAgentsActionObjectElement MakeContinuousAction(const TArray<float>& Values, const FName Name = TEXT("Continuous"));
+	FLearningAgentsActionObjectElement MakeContinuousActionFromArrayView(const TArrayView<const float> Values, const FName Name = TEXT("Continuous"));
 
-public:
+	UFUNCTION(BlueprintPure, Category = "LearningAgents")
+	FLearningAgentsActionObjectElement MakeExclusiveDiscreteAction(const int32 Index, const FName Name = TEXT("DiscreteExclusive"));
 
-	/**
-	 * Adds a new vector action to the given agent interactor. Call during ULearningAgentsInteractor::SetupActions event.
-	 * @param InInteractor The agent interactor to add this action to.
-	 * @param Name The name of this new action. Used for debugging.
-	 * @param Scale Used to normalize the data for the action.
-	 * @return The newly created action.
-	 */
-	UFUNCTION(BlueprintCallable, Category = "LearningAgents", meta = (DefaultToSelf = "InInteractor"))
-	static UVectorAction* AddVectorAction(ULearningAgentsInteractor* InInteractor, const FName Name = NAME_None, const float Scale = 1.0f);
+	UFUNCTION(BlueprintPure, Category = "LearningAgents")
+	FLearningAgentsActionObjectElement MakeInclusiveDiscreteAction(const TArray<int32>& Indices, const FName Name = TEXT("DiscreteInclusive"));
+	FLearningAgentsActionObjectElement MakeInclusiveDiscreteActionFromArrayView(const TArrayView<const int32> Indices, const FName Name = TEXT("DiscreteInclusive"));
 
-	/**
-	 * Gets the data for this action. Call during ULearningAgentsInteractor::GetActions event.
-	 * @param AgentId The agent id to get data for.
-	 * @return The current action values.
-	 */
-	UFUNCTION(BlueprintCallable, Category = "LearningAgents", meta = (AgentId = "-1"))
-	FVector GetVectorAction(const int32 AgentId);
+	UFUNCTION(BlueprintPure, Category = "LearningAgents")
+	FLearningAgentsActionObjectElement MakeStructAction(const TMap<FName, FLearningAgentsActionObjectElement>& Elements, const FName Name = TEXT("Struct"));
+	
+	UFUNCTION(BlueprintPure, Category = "LearningAgents")
+	FLearningAgentsActionObjectElement MakeStructActionFromArrays(const TArray<FName>& ElementNames, const TArray<FLearningAgentsActionObjectElement>& Elements, const FName Name = TEXT("Struct"));
+	FLearningAgentsActionObjectElement MakeStructActionFromArrayViews(const TArrayView<const FName> ElementNames, const TArrayView<const FLearningAgentsActionObjectElement> Elements, const FName Name = TEXT("Struct"));
 
-	/**
-	 * Sets the data for this action. Call during ULearningAgentsController::SetActions event.
-	 * @param AgentId The agent id to set data for.
-	 * @param Value The current action values.
-	 */
-	UFUNCTION(BlueprintCallable, Category = "LearningAgents", meta = (AgentId = "-1"))
-	void SetVectorAction(const int32 AgentId, const FVector Value);
+	UFUNCTION(BlueprintPure, Category = "LearningAgents")
+	FLearningAgentsActionObjectElement MakeExclusiveUnionAction(const FName ElementName, const FLearningAgentsActionObjectElement Element, const FName Name = TEXT("ExclusiveUnion"));
 
-#if UE_LEARNING_AGENTS_ENABLE_VISUAL_LOG
-	/** Describes this action to the visual logger for debugging purposes. */
-	virtual void VisualLog(const UE::Learning::FIndexSet Instances) const override;
-#endif
+	UFUNCTION(BlueprintPure, Category = "LearningAgents")
+	FLearningAgentsActionObjectElement MakeInclusiveUnionAction(const TMap<FName, FLearningAgentsActionObjectElement>& Elements, const FName Name = TEXT("InclusiveUnion"));
 
-	TSharedPtr<UE::Learning::FFloatFeature> FeatureObject;
-};
+	UFUNCTION(BlueprintPure, Category = "LearningAgents")
+	FLearningAgentsActionObjectElement MakeInclusiveUnionActionFromArrays(const TArray<FName>& ElementNames, const TArray<FLearningAgentsActionObjectElement>& Elements, const FName Name = TEXT("InclusiveUnion"));
+	FLearningAgentsActionObjectElement MakeInclusiveUnionActionFromArrayViews(const TArrayView<const FName> ElementNames, const TArrayView<const FLearningAgentsActionObjectElement> Elements, const FName Name = TEXT("InclusiveUnion"));
 
-/** A simple array of FVector action. */
-UCLASS()
-class LEARNINGAGENTS_API UVectorArrayAction : public ULearningAgentsAction
-{
-	GENERATED_BODY()
+	UFUNCTION(BlueprintPure, Category = "LearningAgents")
+	FLearningAgentsActionObjectElement MakeStaticArrayAction(const TArray<FLearningAgentsActionObjectElement>& Elements, const FName Name = TEXT("StaticArray"));
+	FLearningAgentsActionObjectElement MakeStaticArrayActionFromArrayView(const TArrayView<const FLearningAgentsActionObjectElement> Elements, const FName Name = TEXT("StaticArray"));
 
-public:
+	UFUNCTION(BlueprintPure, Category = "LearningAgents")
+	FLearningAgentsActionObjectElement MakePairAction(const FLearningAgentsActionObjectElement Key, const FLearningAgentsActionObjectElement Value, const FName Name = TEXT("Pair"));
 
-	/**
-	 * Adds a new vector action to the given agent interactor. Call during ULearningAgentsInteractor::SetupActions event.
-	 * @param InInteractor The agent interactor to add this action to.
-	 * @param Name The name of this new action. Used for debugging.
-	 * @param Num The number of vectors in the array
-	 * @param Scale Used to normalize the data for the action.
-	 * @return The newly created action.
-	 */
-	UFUNCTION(BlueprintCallable, Category = "LearningAgents", meta = (DefaultToSelf = "InInteractor"))
-	static UVectorArrayAction* AddVectorArrayAction(ULearningAgentsInteractor* InInteractor, const FName Name = NAME_None, const int32 Num = 1, const float Scale = 1.0f);
+	UFUNCTION(BlueprintPure, Category = "LearningAgents")
+	FLearningAgentsActionObjectElement MakeEnumAction(const UEnum* Enum, const uint8 EnumValue, const FName Name = TEXT("Enum"));
 
-	/**
-	 * Gets the data for this action. Call during ULearningAgentsInteractor::GetActions event.
-	 * @param AgentId The agent id to get data for.
-	 * @param OutVectors The current action values.
-	 */
-	UFUNCTION(BlueprintCallable, Category = "LearningAgents", meta = (AgentId = "-1"))
-	void GetVectorArrayAction(const int32 AgentId, TArray<FVector>& OutVectors);
+	UFUNCTION(BlueprintPure, Category = "LearningAgents")
+	FLearningAgentsActionObjectElement MakeBitmaskAction(const UEnum* Enum, const int32 BitmaskValue, const FName Name = TEXT("Bitmask"));
 
-	/**
-	 * Sets the data for this action. Call during ULearningAgentsController::SetActions event.
-	 * @param AgentId The agent id to set data for.
-	 * @param Vectors The current action values.
-	 */
-	UFUNCTION(BlueprintCallable, Category = "LearningAgents", meta = (AgentId = "-1"))
-	void SetVectorArrayAction(const int32 AgentId, const TArray<FVector>& Vectors);
+	UFUNCTION(BlueprintPure, Category = "LearningAgents")
+	FLearningAgentsActionObjectElement MakeOptionalAction(const FLearningAgentsActionObjectElement Element, const ELearningAgentsOptionalAction Option, const FName Name = TEXT("Optional"));
 
-#if UE_LEARNING_AGENTS_ENABLE_VISUAL_LOG
-	/** Describes this action to the visual logger for debugging purposes. */
-	virtual void VisualLog(const UE::Learning::FIndexSet Instances) const override;
-#endif
+	UFUNCTION(BlueprintPure, Category = "LearningAgents")
+	FLearningAgentsActionObjectElement MakeOptionalNullAction(const FName Name = TEXT("Optional"));
 
-	TSharedPtr<UE::Learning::FFloatFeature> FeatureObject;
-};
+	UFUNCTION(BlueprintPure, Category = "LearningAgents")
+	FLearningAgentsActionObjectElement MakeOptionalValidAction(const FLearningAgentsActionObjectElement Element, const FName Name = TEXT("Optional"));
 
-//------------------------------------------------------------------
+	UFUNCTION(BlueprintPure, Category = "LearningAgents")
+	FLearningAgentsActionObjectElement MakeEitherAction(const FLearningAgentsActionObjectElement Element, const ELearningAgentsEitherAction Either, const FName Name = TEXT("Either"));
 
-/** A planar velocity action. */
-UCLASS()
-class LEARNINGAGENTS_API UPlanarVelocityAction : public ULearningAgentsAction
-{
-	GENERATED_BODY()
+	UFUNCTION(BlueprintPure, Category = "LearningAgents", meta = (DisplayName = "Make Either A Action"))
+	FLearningAgentsActionObjectElement MakeEitherAAction(const FLearningAgentsActionObjectElement A, const FName Name = TEXT("Either"));
 
-public:
+	UFUNCTION(BlueprintPure, Category = "LearningAgents", meta = (DisplayName = "Make Either B Action"))
+	FLearningAgentsActionObjectElement MakeEitherBAction(const FLearningAgentsActionObjectElement B, const FName Name = TEXT("Either"));
 
-	/**
-	 * Adds a new planar velocity action to the given agent interactor. The axis parameters define the plane.
-	 * Call during ULearningAgentsInteractor::SetupActions event.
-	 * @param InInteractor The agent interactor to add this action to.
-	 * @param Name The name of this new action. Used for debugging.
-	 * @param Scale Used to normalize the data for the action.
-	 * @param Axis0 The forward axis of the plane.
-	 * @param Axis1 The right axis of the plane.
-	 * @return The newly created action.
-	 */
-	UFUNCTION(BlueprintCallable, Category = "LearningAgents", meta = (DefaultToSelf = "InInteractor"))
-	static UPlanarVelocityAction* AddPlanarVelocityAction(ULearningAgentsInteractor* InInteractor, const FName Name = NAME_None, const float Scale = 200.0f, const FVector Axis0 = FVector::ForwardVector, const FVector Axis1 = FVector::RightVector);
+	UFUNCTION(BlueprintPure, Category = "LearningAgents")
+	FLearningAgentsActionObjectElement MakeEncodingAction(const FLearningAgentsActionObjectElement Element, const FName Name = TEXT("Encoding"));
 
-	/**
-	 * Gets the data for this action. Call during ULearningAgentsInteractor::GetActions event.
-	 * @param AgentId The agent id to get data for.
-	 * @return The current action value.
-	 */
-	UFUNCTION(BlueprintCallable, Category = "LearningAgents", meta = (AgentId = "-1"))
-	FVector GetPlanarVelocityAction(const int32 AgentId);
+	UFUNCTION(BlueprintPure, Category = "LearningAgents")
+	FLearningAgentsActionObjectElement MakeBoolAction(const bool bValue, const FName Name = TEXT("Bool"));
 
-	/**
-	 * Sets the data for this action. Call during ULearningAgentsController::SetActions event.
-	 * @param AgentId The agent id this data corresponds to.
-	 * @param Velocity The velocity currently being observed.
-	 */
-	UFUNCTION(BlueprintCallable, Category = "LearningAgents", meta = (AgentId = "-1"))
-	void SetPlanarVelocityAction(const int32 AgentId, const FVector Velocity);
+	UFUNCTION(BlueprintPure, Category = "LearningAgents")
+	FLearningAgentsActionObjectElement MakeFloatAction(const float Value, const float FloatScale = 1.0f, const FName Name = TEXT("Float"));
 
-#if UE_LEARNING_AGENTS_ENABLE_VISUAL_LOG
-	/** Describes this action to the visual logger for debugging purposes. */
-	virtual void VisualLog(const UE::Learning::FIndexSet Instances) const override;
-#endif
+	UFUNCTION(BlueprintPure, Category = "LearningAgents")
+	FLearningAgentsActionObjectElement MakeTranslationAction(const FVector Translation, const FTransform RelativeTransform = FTransform(), const float TranslationScale = 100.0f, const FName Name = TEXT("Translation"));
 
-	TSharedPtr<UE::Learning::FPlanarVelocityFeature> FeatureObject;
-};
+	UFUNCTION(BlueprintPure, Category = "LearningAgents")
+	FLearningAgentsActionObjectElement MakeRotationAction(const FRotator Rotation, const FRotator RelativeRotation = FRotator::ZeroRotator, const float RotationScale = 90.0f, const FName Name = TEXT("Rotation"));
+	FLearningAgentsActionObjectElement MakeRotationActionFromQuat(const FQuat Rotation, const FQuat RelativeRotation = FQuat::Identity, const float RotationScale = 90.0f, const FName Name = TEXT("Rotation"));
 
-/** A rotation action. */
-UCLASS()
-class LEARNINGAGENTS_API URotationAction : public ULearningAgentsAction
-{
-	GENERATED_BODY()
+	UFUNCTION(BlueprintPure, Category = "LearningAgents")
+	FLearningAgentsActionObjectElement MakeScaleAction(const FVector Scale, const FVector RelativeScale = FVector(1,1,1), const FName Name = TEXT("Scale"));
+
+	UFUNCTION(BlueprintPure, Category = "LearningAgents")
+	FLearningAgentsActionObjectElement MakeTransformAction(const FTransform Transform, const FTransform RelativeTransform = FTransform(), const float TranslationScale = 100.0f, const FName Name = TEXT("Transform"));
+
+	UFUNCTION(BlueprintPure, Category = "LearningAgents")
+	FLearningAgentsActionObjectElement MakeAngleAction(const float Angle, const float RelativeAngle = 0.0f, const float AngleScale = 90.0f, const FName Name = TEXT("Angle"));
+
+	UFUNCTION(BlueprintPure, Category = "LearningAgents")
+	FLearningAgentsActionObjectElement MakeAngleActionRadians(const float Angle, const float RelativeAngle = 0.0f, const float AngleScale = 1.57079632679f, const FName Name = TEXT("Angle"));
+
+	UFUNCTION(BlueprintPure, Category = "LearningAgents")
+	FLearningAgentsActionObjectElement MakeVelocityAction(const FVector Velocity, const FTransform RelativeTransform = FTransform(), const float VelocityScale = 200.0f, const FName Name = TEXT("Velocity"));
+
+	UFUNCTION(BlueprintPure, Category = "LearningAgents")
+	FLearningAgentsActionObjectElement MakeSpeedAction(const float Speed, const float SpeedScale = 200.0f, const FName Name = TEXT("Speed"));
 
 public:
 
-	/**
-	 * Adds a new rotation action to the given agent interactor. Call during ULearningAgentsInteractor::SetupActions event.
-	 * @param InInteractor The agent interactor to add this action to.
-	 * @param Name The name of this new action. Used for debugging.
-	 * @param Scale Used to normalize the data for the action.
-	 * @return The newly created action.
-	 */
-	UFUNCTION(BlueprintCallable, Category = "LearningAgents", meta = (DefaultToSelf = "InInteractor"))
-	static URotationAction* AddRotationAction(ULearningAgentsInteractor* InInteractor, const FName Name = NAME_None, const float Scale = 180.0f);
+	UFUNCTION(BlueprintPure = false, Category = "LearningAgents")
+	UPARAM(DisplayName = "Success") bool GetNullAction(const FLearningAgentsActionObjectElement Element, const FName Name = TEXT("Null")) const;
 
-	/**
-	 * Gets the data for this action as a rotator. Call during ULearningAgentsInteractor::GetActions event.
-	 * @param AgentId The agent id to get data for.
-	 * @return The current action value.
-	 */
-	UFUNCTION(BlueprintCallable, Category = "LearningAgents", meta = (AgentId = "-1"))
-	FRotator GetRotationAction(const int32 AgentId);
+	UFUNCTION(BlueprintPure = false, Category = "LearningAgents")
+	UPARAM(DisplayName = "Success") bool GetContinuousActionNum(int32& OutNum, const FLearningAgentsActionObjectElement Element, const FName Name = TEXT("Continuous")) const;
 
-	/**
-	 * Gets the data for this action as a rotation vector. Call during ULearningAgentsInteractor::GetActions event.
-	 * @param AgentId The agent id to get data for.
-	 * @return The current action value.
-	 */
-	UFUNCTION(BlueprintCallable, Category = "LearningAgents", meta = (AgentId = "-1"))
-	FVector GetRotationActionAsRotationVector(const int32 AgentId);
+	UFUNCTION(BlueprintPure = false, Category = "LearningAgents")
+	UPARAM(DisplayName = "Success") bool GetContinuousAction(TArray<float>& OutValues, const FLearningAgentsActionObjectElement Element, const FName Name = TEXT("Continuous")) const;
+	UPARAM(DisplayName = "Success") bool GetContinuousActionToArrayView(TArrayView<float> OutValues, const FLearningAgentsActionObjectElement Element, const FName Name = TEXT("Continuous")) const;
+	
+	UFUNCTION(BlueprintPure = false, Category = "LearningAgents")
+	UPARAM(DisplayName = "Success") bool GetExclusiveDiscreteAction(int32& OutIndex, const FLearningAgentsActionObjectElement Element, const FName Name = TEXT("DiscreteExclusive")) const;
 
-	/**
-	 * Gets the data for this action as a quaternion. Call during ULearningAgentsInteractor::GetActions event.
-	 * @param AgentId The agent id to get data for.
-	 * @return The current action value.
-	 */
-	UFUNCTION(BlueprintCallable, Category = "LearningAgents", meta = (AgentId = "-1"))
-	FQuat GetRotationActionAsQuat(const int32 AgentId);
+	UFUNCTION(BlueprintPure = false, Category = "LearningAgents")
+	UPARAM(DisplayName = "Success") bool GetInclusiveDiscreteActionNum(int32& OutNum, const FLearningAgentsActionObjectElement Element, const FName Name = TEXT("DiscreteInclusive")) const;
 
-#if UE_LEARNING_AGENTS_ENABLE_VISUAL_LOG
-	/** Describes this action to the visual logger for debugging purposes. */
-	virtual void VisualLog(const UE::Learning::FIndexSet Instances) const override;
-#endif
+	UFUNCTION(BlueprintPure = false, Category = "LearningAgents")
+	UPARAM(DisplayName = "Success") bool GetInclusiveDiscreteAction(TArray<int32>& OutIndices, const FLearningAgentsActionObjectElement Element, const FName Name = TEXT("DiscreteInclusive")) const;
+	UPARAM(DisplayName = "Success") bool GetInclusiveDiscreteActionToArrayView(TArrayView<int32> OutIndices, const FLearningAgentsActionObjectElement Element, const FName Name = TEXT("DiscreteInclusive")) const;
 
-	TSharedPtr<UE::Learning::FRotationVectorFeature> FeatureObject;
-};
+	UFUNCTION(BlueprintPure = false, Category = "LearningAgents")
+	UPARAM(DisplayName = "Success") bool GetStructActionNum(int32& OutNum, const FLearningAgentsActionObjectElement Element, const FName Name = TEXT("Struct")) const;
 
-/** An array of rotation actions. */
-UCLASS()
-class LEARNINGAGENTS_API URotationArrayAction : public ULearningAgentsAction
-{
-	GENERATED_BODY()
+	UFUNCTION(BlueprintPure = false, Category = "LearningAgents")
+	UPARAM(DisplayName = "Success") bool GetStructAction(TMap<FName, FLearningAgentsActionObjectElement>& OutElements, const FLearningAgentsActionObjectElement Element, const FName Name = TEXT("Struct")) const;
+	
+	UFUNCTION(BlueprintPure = false, Category = "LearningAgents")
+	UPARAM(DisplayName = "Success") bool GetStructActionToArrays(TArray<FName>& OutElementNames, TArray<FLearningAgentsActionObjectElement>& OutElements, const FLearningAgentsActionObjectElement Element, const FName Name = TEXT("Struct")) const;
+	UPARAM(DisplayName = "Success") bool GetStructActionToArrayViews(TArrayView<FName> OutElementNames, TArrayView<FLearningAgentsActionObjectElement> OutElements, const FLearningAgentsActionObjectElement Element, const FName Name = TEXT("Struct")) const;
 
-public:
+	UFUNCTION(BlueprintPure = false, Category = "LearningAgents")
+	UPARAM(DisplayName = "Success") bool GetExclusiveUnionAction(FName& OutElementName, FLearningAgentsActionObjectElement& OutElement, const FLearningAgentsActionObjectElement Element, const FName Name = TEXT("ExclusiveUnion")) const;
 
-	/**
-	 * Adds a new rotation array action to the given agent interactor. Call during ULearningAgentsInteractor::SetupActions event.
-	 * @param InInteractor The agent interactor to add this action to.
-	 * @param Name The name of this new action. Used for debugging.
-	 * @param RotationNum The number of rotations in the array.
-	 * @param Scale Used to normalize the data for the action.
-	 * @return The newly created action.
-	 */
-	UFUNCTION(BlueprintCallable, Category = "LearningAgents", meta = (DefaultToSelf = "InInteractor"))
-	static URotationArrayAction* AddRotationArrayAction(ULearningAgentsInteractor* InInteractor, const FName Name = NAME_None, const int32 RotationNum = 1, const float Scale = 180.0f);
+	UFUNCTION(BlueprintPure = false, Category = "LearningAgents")
+	UPARAM(DisplayName = "Success") bool GetInclusiveUnionActionNum(int32& OutNum, const FLearningAgentsActionObjectElement Element, const FName Name = TEXT("InclusiveUnion")) const;
 
-	/**
-	 * Gets the data for this action as rotators. Call during ULearningAgentsInteractor::GetActions event.
-	 * @param AgentId The agent id to get data for.
-	 * @param OutRotations The current action values.
-	 */
-	UFUNCTION(BlueprintCallable, Category = "LearningAgents", meta = (AgentId = "-1"))
-	void GetRotationArrayAction(const int32 AgentId, TArray<FRotator>& OutRotations);
+	UFUNCTION(BlueprintPure = false, Category = "LearningAgents")
+	UPARAM(DisplayName = "Success") bool GetInclusiveUnionAction(TMap<FName, FLearningAgentsActionObjectElement>& OutElements, const FLearningAgentsActionObjectElement Element, const FName Name = TEXT("InclusiveUnion")) const;
+	
+	UFUNCTION(BlueprintPure = false, Category = "LearningAgents")
+	UPARAM(DisplayName = "Success") bool GetInclusiveUnionActionToArrays(TArray<FName>& OutElementNames, TArray<FLearningAgentsActionObjectElement>& OutElements, const FLearningAgentsActionObjectElement Element, const FName Name = TEXT("InclusiveUnion")) const;
+	UPARAM(DisplayName = "Success") bool GetInclusiveUnionActionToArrayViews(TArrayView<FName> OutElementNames, TArrayView<FLearningAgentsActionObjectElement> OutElements, const FLearningAgentsActionObjectElement Element, const FName Name = TEXT("InclusiveUnion")) const;
 
-	/**
-	 * Gets the data for this action as rotation vectors. Call during ULearningAgentsInteractor::GetActions event.
-	 * @param AgentId The agent id to get data for.
-	 * @param OutRotationVectors The current action values.
-	 */
-	UFUNCTION(BlueprintCallable, Category = "LearningAgents", meta = (AgentId = "-1"))
-	void GetRotationArrayActionAsRotationVectors(const int32 AgentId, TArray<FVector>& OutRotationVectors);
+	UFUNCTION(BlueprintPure = false, Category = "LearningAgents")
+	UPARAM(DisplayName = "Success") bool GetStaticArrayActionNum(int32& OutNum, const FLearningAgentsActionObjectElement Element, const FName Name = TEXT("StaticArray")) const;
 
-	/**
-	 * Gets the data for this action as quaternions. Call during ULearningAgentsInteractor::GetActions event.
-	 * @param AgentId The agent id to get data for.
-	 * @param OutRotations The current action values.
-	 */
-	UFUNCTION(BlueprintCallable, Category = "LearningAgents", meta = (AgentId = "-1"))
-	void GetRotationArrayActionAsQuats(const int32 AgentId, TArray<FQuat>& OutRotations);
+	UFUNCTION(BlueprintPure = false, Category = "LearningAgents")
+	UPARAM(DisplayName = "Success") bool GetStaticArrayAction(TArray<FLearningAgentsActionObjectElement>& OutElements, const FLearningAgentsActionObjectElement Element, const FName Name = TEXT("StaticArray")) const;
+	UPARAM(DisplayName = "Success") bool GetStaticArrayActionToArrayView(TArrayView<FLearningAgentsActionObjectElement> OutElements, const FLearningAgentsActionObjectElement Element, const FName Name = TEXT("StaticArray")) const;
+
+	UFUNCTION(BlueprintPure = false, Category = "LearningAgents")
+	UPARAM(DisplayName = "Success") bool GetPairAction(FLearningAgentsActionObjectElement& OutKey, FLearningAgentsActionObjectElement& OutValue, const FLearningAgentsActionObjectElement Element, const FName Name = TEXT("Pair")) const;
+
+	UFUNCTION(BlueprintPure = false, Category = "LearningAgents")
+	UPARAM(DisplayName = "Success") bool GetEnumAction(uint8& OutEnumValue, const UEnum* Enum, const FLearningAgentsActionObjectElement Element, const FName Name = TEXT("Enum")) const;
+
+	UFUNCTION(BlueprintPure = false, Category = "LearningAgents")
+	UPARAM(DisplayName = "Success") bool GetBitmaskAction(int32& OutBitmaskValue, const UEnum* Enum, const FLearningAgentsActionObjectElement Element, const FName Name = TEXT("Bitmask")) const;
+
+	UFUNCTION(BlueprintPure = false, Category = "LearningAgents", meta = (ExpandEnumAsExecs = "OutOption"))
+	UPARAM(DisplayName = "Success") bool GetOptionalAction(ELearningAgentsOptionalAction& OutOption, FLearningAgentsActionObjectElement& OutElement, const FLearningAgentsActionObjectElement Element, const FName Name = TEXT("Optional")) const;
+	
+	UFUNCTION(BlueprintPure = false, Category = "LearningAgents", meta = (ExpandEnumAsExecs = "OutEither"))
+	UPARAM(DisplayName = "Success") bool GetEitherAction(ELearningAgentsEitherAction& OutEither, FLearningAgentsActionObjectElement& OutElement, const FLearningAgentsActionObjectElement Element, const FName Name = TEXT("Either")) const;
+
+	UFUNCTION(BlueprintPure = false, Category = "LearningAgents")
+	UPARAM(DisplayName = "Success") bool GetEncodingAction(FLearningAgentsActionObjectElement& OutElement, const FLearningAgentsActionObjectElement Element, const FName Name = TEXT("Encoding")) const;
+
+	UFUNCTION(BlueprintPure = false, Category = "LearningAgents")
+	UPARAM(DisplayName = "Success") bool GetBoolAction(bool& bOutValue, const FLearningAgentsActionObjectElement Element, const FName Name = TEXT("Bool")) const;
+
+	UFUNCTION(BlueprintPure = false, Category = "LearningAgents")
+	UPARAM(DisplayName = "Success") bool GetFloatAction(float& OutValue, const FLearningAgentsActionObjectElement Element, const float FloatScale = 1.0f, const FName Name = TEXT("Float")) const;
+
+	UFUNCTION(BlueprintPure = false, Category = "LearningAgents")
+	UPARAM(DisplayName = "Success") bool GetTranslationAction(FVector& OutTranslation, const FLearningAgentsActionObjectElement Element, const FTransform RelativeTransform = FTransform(), const float TranslationScale = 100.0f, const FName Name = TEXT("Translation")) const;
+
+	UFUNCTION(BlueprintPure = false, Category = "LearningAgents")
+	UPARAM(DisplayName = "Success") bool GetRotationAction(FRotator& OutRotation, const FLearningAgentsActionObjectElement Element, const FRotator RelativeRotation = FRotator::ZeroRotator, const float RotationScale = 90.0f, const FName Name = TEXT("Rotation")) const;
+	UPARAM(DisplayName = "Success") bool GetRotationActionAsQuat(FQuat& OutRotation, const FLearningAgentsActionObjectElement Element, const FQuat RelativeRotation = FQuat::Identity, const float RotationScale = 90.0f, const FName Name = TEXT("Rotation")) const;
+
+	UFUNCTION(BlueprintPure = false, Category = "LearningAgents")
+	UPARAM(DisplayName = "Success") bool GetScaleAction(FVector& OutScale, const FLearningAgentsActionObjectElement Element, const FVector RelativeScale = FVector(1,1,1), const float Scale = 1.0f, const FName Name = TEXT("Scale")) const;
+
+	UFUNCTION(BlueprintPure = false, Category = "LearningAgents")
+	UPARAM(DisplayName = "Success") bool GetTransformAction(FTransform& OutTransform, const FLearningAgentsActionObjectElement Element, const FTransform RelativeTransform = FTransform(), const float TranslationScale = 100.0f, const float RotationScale = 1.0f, const float ScaleScale = 1.0f, const FName Name = TEXT("Transform")) const;
+
+	UFUNCTION(BlueprintPure = false, Category = "LearningAgents")
+	UPARAM(DisplayName = "Success") bool GetAngleAction(float& OutAngle, const FLearningAgentsActionObjectElement Element, const float RelativeAngle = 0.0f, const float AngleScale = 90.0f, const FName Name = TEXT("Angle")) const;
+
+	UFUNCTION(BlueprintPure = false, Category = "LearningAgents")
+	UPARAM(DisplayName = "Success") bool GetAngleActionRadians(float& OutAngle, const FLearningAgentsActionObjectElement Element, const float RelativeAngle = 0.0f, const float AngleScale = 1.57079632679f, const FName Name = TEXT("Angle")) const;
+
+	UFUNCTION(BlueprintPure = false, Category = "LearningAgents")
+	UPARAM(DisplayName = "Success") bool GetVelocityAction(FVector& OutVelocity, const FLearningAgentsActionObjectElement Element, const FTransform RelativeTransform = FTransform(), const float VelocityScale = 200.0f, const FName Name = TEXT("Velocity")) const;
+
+	UFUNCTION(BlueprintPure = false, Category = "LearningAgents")
+	UPARAM(DisplayName = "Success") bool GetSpeedAction(float& OutSpeed, const FLearningAgentsActionObjectElement Element, const float SpeedScale = 200.0f, const FName Name = TEXT("Speed")) const;
 
-#if UE_LEARNING_AGENTS_ENABLE_VISUAL_LOG
-	/** Describes this action to the visual logger for debugging purposes. */
-	virtual void VisualLog(const UE::Learning::FIndexSet Instances) const override;
-#endif
+private:
 
-	TSharedPtr<UE::Learning::FRotationVectorFeature> FeatureObject;
+	UE::Learning::Action::FObject ActionObject;
 };
diff --git a/Engine/Plugins/Experimental/LearningAgents/Source/LearningAgents/Public/LearningAgentsController.h b/Engine/Plugins/Experimental/LearningAgents/Source/LearningAgents/Public/LearningAgentsController.h
--- a/Engine/Plugins/Experimental/LearningAgents/Source/LearningAgents/Public/LearningAgentsController.h
+++ b/Engine/Plugins/Experimental/LearningAgents/Source/LearningAgents/Public/LearningAgentsController.h
@@ -2,17 +2,24 @@
 
 #pragma once
 
-#include "LearningAgentsManagerComponent.h"
+#include "LearningAgentsManagerListener.h"
+#include "LearningAgentsObservations.h"
+#include "LearningAgentsActions.h"
+
 #include "LearningArray.h"
+
 #include "LearningAgentsController.generated.h"
 
+
+DECLARE_DYNAMIC_DELEGATE_RetVal_FourParams(FLearningAgentsActionObjectElement, FEvaluateAgentControllerDelegate, ULearningAgentsActionObject*, InActionObject, const FLearningAgentsObservationObjectElement, InObservationObjectElement, const ULearningAgentsObservationObject*, InObservationObject, const int32, AgentId);
+
 /**
- * A controller provides a method for injecting actions into the learning agents system from some other existing
- * behavior, e.g. we may want to gather demonstrations from a human or AI behavior tree controlling our agent(s)
- * for imitation learning purposes.
+ * A controller is an object that can be used to construct actions from observations - essentially a hand-made Policy. This can be useful for making 
+ * a learning agents system that uses some other existing behavior, e.g. we may want to gather demonstrations from a human or AI behavior tree 
+ * controlling our agent(s) for imitation learning purposes.
  */
-UCLASS(Abstract, BlueprintType, Blueprintable)
-class LEARNINGAGENTS_API ULearningAgentsController : public ULearningAgentsManagerComponent
+UCLASS(Abstract, HideDropdown, BlueprintType, Blueprintable)
+class LEARNINGAGENTS_API ULearningAgentsController : public ULearningAgentsManagerListener
 {
 	GENERATED_BODY()
 
@@ -23,28 +30,74 @@ public:
 	ULearningAgentsController(FVTableHelper& Helper);
 	virtual ~ULearningAgentsController();
 
+	/**
+	 * Constructs a new controller for the given agent interactor.
+	 * 
+	 * @param InManager			The input Manager
+	 * @param InInteractor		The input Interactor component
+	 * @param Class				The controller class
+	 * @param Name				The controller name
+	 */
+	UFUNCTION(BlueprintCallable, Category = "LearningAgents", meta = (DeterminesOutputType = "Class"))
+	static ULearningAgentsController* MakeController(
+		ULearningAgentsManager* InManager, 
+		ULearningAgentsInteractor* InInteractor, 
+		TSubclassOf<ULearningAgentsController> Class,
+		const FName Name = TEXT("Controller"));
+
 	/** Initializes this object to be used with the given agent interactor. */
 	UFUNCTION(BlueprintCallable, Category = "LearningAgents")
-	void SetupController(ULearningAgentsInteractor* InInteractor);
+	void SetupController(ULearningAgentsManager* InManager, ULearningAgentsInteractor* InInteractor);
+
+	/**
+	 * This callback should be implemented by the Controller and should produce an array of action object elements, from an array of observation 
+	 * object elements.
+	 * 
+	 * @param OutActionObjectElements		Output Action Object Elements. This should be the same size as the input AgentIds and 
+	 *                                      InObservationObjectElements arrays.
+	 * @param InActionObject				Action object used to construct output elements.
+	 * @param InObservationObjectElements	Input observation object elements
+	 * @param InObservationObject			Input observation object
+	 * @param AgentIds						Agent ids associated with each observation
+	 */
+	UFUNCTION(BlueprintNativeEvent, Category = "LearningAgents", Meta = (ForceAsFunction))
+	void EvaluateAgentController(
+		TArray<FLearningAgentsActionObjectElement>& OutActionObjectElements, 
+		ULearningAgentsActionObject* InActionObject, 
+		const TArray<FLearningAgentsObservationObjectElement>& InObservationObjectElements,
+		const ULearningAgentsObservationObject* InObservationObject,
+		const TArray<int32>& AgentIds);
 
 	/**
-	 * During this event, you should set the actions of your agents.
-	 * @param AgentIds The list of agent ids to set actions for.
-	 * @see LearningAgentsActions.h for the list of available actions.
-	 * @see ULearningAgentsInteractor::GetAgent to get the agent corresponding to each id.
+	 * This is a convenience function that can be used to implement `EvaluateAgentController` in terms of a callback which iterates over every agent
+	 * individually and computes the actions from the observations.
+	 *
+	 * @param OutActionObjectElements		Output Action Object Elements. This should be the same size as the input AgentIds and
+	 *                                      InObservationObjectElements arrays.
+	 * @param InActionObject				Action object used to construct output elements.
+	 * @param InObservationObjectElements	Input observation object elements
+	 * @param InObservationObject			Input observation object
+	 * @param AgentIds						Agent ids associated with each observation
+	 * @param Delegate						The Delegate used to compute actions from observations.
 	 */
-	UFUNCTION(BlueprintNativeEvent, Category = "LearningAgents")
-	void SetActions(const TArray<int32>& AgentIds);
+	UFUNCTION(BlueprintPure = false, Category = "LearningAgents")
+	void EvaluateAgentControllerUsingDelegate(
+		TArray<FLearningAgentsActionObjectElement>& OutActionObjectElements,
+		ULearningAgentsActionObject* InActionObject,
+		const TArray<FLearningAgentsObservationObjectElement>& InObservationObjectElements,
+		const ULearningAgentsObservationObject* InObservationObject,
+		const TArray<int32>& AgentIds,
+		const FEvaluateAgentControllerDelegate& Delegate);
 
 	/**
-	 * Call this function when it is time to gather all the actions for your agents. This should be called roughly 
-	 * whenever you are calling ULearningAgentsInteractor::EncodeObservations. This will call this controller's SetActions event.
+	 * Call this function when it is time to evaluate the controller and produce the actions for the agents. This should be called after 
+	 * GatherObservations but before ScatterActions. This will call this controller's EvaluateAgentController event.
 	 */
 	UFUNCTION(BlueprintCallable, Category = "LearningAgents")
-	void EncodeActions();
+	void EvaluateController();
 
 	/**
-	 * Calls EncodeObservations, followed by EncodeActions, followed by DecodeActions
+	 * Calls GatherObservations, followed by EvaluateController, followed by ScatterActions
 	 */
 	UFUNCTION(BlueprintCallable, Category = "LearningAgents")
 	void RunController();
@@ -53,19 +106,15 @@ protected:
 
 	/**
 	 * Gets the agent interactor associated with this component.
+	 * 
 	 * @param AgentClass The class to cast the agent interactor to (in blueprint).
 	 */
 	UFUNCTION(BlueprintPure, Category = "LearningAgents", meta = (DeterminesOutputType = "InteractorClass"))
 	ULearningAgentsInteractor* GetInteractor(const TSubclassOf<ULearningAgentsInteractor> InteractorClass) const;
 
-private:
+protected:
 
 	/** The agent interactor this controller is associated with. */
 	UPROPERTY(VisibleAnywhere, Transient, Category = "LearningAgents")
 	TObjectPtr<ULearningAgentsInteractor> Interactor;
-
-	/** Temp buffers used to record the set of agents that are valid for encoding */
-	TArray<int32> ValidAgentIds;
-	UE::Learning::FIndexSet ValidAgentSet;
-	TBitArray<> ValidAgentStatus;
 };
