LearningAgents: Large restructuring to allow for structured actions and observations.

diff --git a/Engine/Plugins/Experimental/LearningAgents/Source/LearningAgents/Public/LearningAgentsCritic.h b/Engine/Plugins/Experimental/LearningAgents/Source/LearningAgents/Public/LearningAgentsCritic.h
--- a/Engine/Plugins/Experimental/LearningAgents/Source/LearningAgents/Public/LearningAgentsCritic.h
+++ b/Engine/Plugins/Experimental/LearningAgents/Source/LearningAgents/Public/LearningAgentsCritic.h
@@ -2,10 +2,9 @@
 
 #pragma once
 
-#include "LearningAgentsManagerComponent.h"
+#include "LearningAgentsManagerListener.h"
 
 #include "LearningAgentsNeuralNetwork.h" // Included for ELearningAgentsActivationFunction
-#include "LearningAgentsDebug.h"
 #include "LearningArray.h"
 #include "UObject/ObjectPtr.h"
 
@@ -13,11 +12,9 @@
 
 namespace UE::Learning
 {
-	struct INeuralNetwork;
-	struct FNeuralNetworkCriticFunction;
+	struct FNeuralNetworkCritic;
 }
 
-struct FFilePath;
 class ULearningAgentsNeuralNetwork;
 class ULearningAgentsInteractor;
 class ULearningAgentsPolicy;
@@ -43,9 +40,9 @@ public:
 	ELearningAgentsActivationFunction ActivationFunction = ELearningAgentsActivationFunction::ELU;
 };
 
-/** A critic used by some algorithms for training the managed agents. */
+/** A critic used for training the policy. Can be used at inference time to estimate the discounted returns.  */
 UCLASS(BlueprintType, Blueprintable, meta = (BlueprintSpawnableComponent))
-class LEARNINGAGENTS_API ULearningAgentsCritic : public ULearningAgentsManagerComponent
+class LEARNINGAGENTS_API ULearningAgentsCritic : public ULearningAgentsManagerListener
 {
 	GENERATED_BODY()
 
@@ -57,69 +54,64 @@ public:
 	virtual ~ULearningAgentsCritic();
 
 	/**
-	 * Initializes this object to be used with the given agent interactor and critic settings.
-	 * @param InInteractor The input Interactor component
-	 * @param InPolicy The input Policy component
-	 * @param CriticSettings The critic settings to use
-	 * @param NeuralNetworkAsset Optional Network Asset to use. If provided must match the given CriticSettings. If not
-	 * provided or asset is empty then a new neural network object will be created according to the given 
-	 * CriticSettings and used.
+	 * Constructs a Critic to be used with the given agent interactor and critic settings.
+	 * 
+	 * @param InManager						The input Manager
+	 * @param InInteractor					The input Interactor object
+	 * @param InPolicy						The input Policy object
+	 * @param Class							The critic class
+	 * @param Name							The critic name
+	 * @param CriticNeuralNetworkAsset		Optional Network Asset to use. If not provided, asset is empty, or
+	 *										bReinitializeNetwork is set then a new neural network object will be created
+	 *										according to the given CriticSettings.
+	 * @param bReinitializeCriticNetwork	If to reinitialize the critic network
+	 * @param CriticSettings				The critic settings to use on creation of a new critic network
+	 * @param Seed							Random seed to use for initializing the critic weights
 	 */
-	UFUNCTION(BlueprintCallable, Category = "LearningAgents")
-	void SetupCritic(
+	UFUNCTION(BlueprintCallable, Category = "LearningAgents", meta = (Class = "/Script/LearningAgents.LearningAgentsCritic", DeterminesOutputType = "Class"))
+	static ULearningAgentsCritic* MakeCritic(
+		ULearningAgentsManager* InManager,
 		ULearningAgentsInteractor* InInteractor,
 		ULearningAgentsPolicy* InPolicy,
+		TSubclassOf<ULearningAgentsCritic> Class,
+		const FName Name = TEXT("Critic"),
+		ULearningAgentsNeuralNetwork* CriticNeuralNetworkAsset = nullptr,
+		const bool bReinitializeCriticNetwork = true,
 		const FLearningAgentsCriticSettings& CriticSettings = FLearningAgentsCriticSettings(),
-		ULearningAgentsNeuralNetwork* NeuralNetworkAsset = nullptr);
-
-public:
-
-	//~ Begin ULearningAgentsManagerComponent Interface
-	virtual void OnAgentsAdded(const TArray<int32>& AgentIds) override;
-	virtual void OnAgentsRemoved(const TArray<int32>& AgentIds) override;
-	virtual void OnAgentsReset(const TArray<int32>& AgentIds) override;
-	//~ End ULearningAgentsManagerComponent Interface
-
-// ----- Load / Save -----
-public:
+		const int32 Seed = 1234);
 
 	/**
-	 * Load a snapshot's weights into this critic.
-	 * @param File The snapshot file.
-	 */
-	UFUNCTION(BlueprintCallable, Category = "LearningAgents", meta = (RelativePath))
-	void LoadCriticFromSnapshot(const FFilePath& File);
-
-	/**
-	 * Save this critic's weights into a snapshot.
-	 * @param File The snapshot file.
-	 */
-	UFUNCTION(BlueprintCallable, BlueprintPure = false, Category = "LearningAgents", meta = (RelativePath))
-	void SaveCriticToSnapshot(const FFilePath& File) const;
-
-	/**
-	 * Use a ULearningAgentsNeuralNetwork asset directly for this critic rather than making a copy. If used 
-	 * during training then this asset's weights will be updated as training progresses.
-	 * @param NeuralNetworkAsset The asset to use.
+	 * Initializes a critic to be used with the given agent interactor and critic settings.
+	 * 
+	 * @param InManager						The input Manager
+	 * @param InInteractor					The input Interactor object
+	 * @param InPolicy						The input Policy object
+	 * @param CriticNeuralNetworkAsset		Optional Network Asset to use. If not provided, asset is empty, or
+	 *										bReinitializeNetwork is set then a new neural network object will be created
+	 *										according to the given CriticSettings.
+	 * @param bReinitializeCriticNetwork	If to reinitialize the critic network
+	 * @param CriticSettings				The critic settings to use on creation of a new critic network
+	 * @param Seed							Random seed to use for initializing the critic weights
 	 */
 	UFUNCTION(BlueprintCallable, Category = "LearningAgents")
-	void UseCriticFromAsset(ULearningAgentsNeuralNetwork* NeuralNetworkAsset);
+	void SetupCritic(
+		ULearningAgentsManager* InManager,
+		ULearningAgentsInteractor* InInteractor,
+		ULearningAgentsPolicy* InPolicy,
+		ULearningAgentsNeuralNetwork* CriticNeuralNetworkAsset = nullptr,
+		const bool bReinitializeCriticNetwork = true,
+		const FLearningAgentsCriticSettings& CriticSettings = FLearningAgentsCriticSettings(),
+		const int32 Seed = 1234);
 
-	/**
-	 * Load a ULearningAgentsNeuralNetwork asset's weights into this critic.
-	 * @param NeuralNetworkAsset The asset to load from.
-	 */
-	UFUNCTION(BlueprintCallable, Category = "LearningAgents")
-	void LoadCriticFromAsset(ULearningAgentsNeuralNetwork* NeuralNetworkAsset);
+public:
 
-	/**
-	 * Save this critic's weights to a ULearningAgentsNeuralNetwork asset.
-	 * @param NeuralNetworkAsset The asset to save to.
-	 */
-	UFUNCTION(BlueprintCallable, Category = "LearningAgents", Meta=(DevelopmentOnly))
-	void SaveCriticToAsset(ULearningAgentsNeuralNetwork* NeuralNetworkAsset);
+	//~ Begin ULearningAgentsManagerListener Interface
+	virtual void OnAgentsAdded_Implementation(const TArray<int32>& AgentIds) override;
+	virtual void OnAgentsRemoved_Implementation(const TArray<int32>& AgentIds) override;
+	virtual void OnAgentsReset_Implementation(const TArray<int32>& AgentIds) override;
+	//~ End ULearningAgentsManagerListener Interface
 
-// ----- Evaluation -----
+// ----- Blueprint public interface -----
 public:
 
 	/**
@@ -141,17 +133,15 @@ public:
 	UFUNCTION(BlueprintPure, Category = "LearningAgents", Meta=(AgentId = -1))
 	float GetEstimatedDiscountedReturn(const int32 AgentId) const;
 
+	/** Gets the current Network Asset being used */
+	UFUNCTION(BlueprintPure, Category = "LearningAgents")
+	ULearningAgentsNeuralNetwork* GetCriticNetworkAsset();
+
 // ----- Non-blueprint public interface -----
 public:
 
-	/** Gets the current Network Asset being used */
-	ULearningAgentsNeuralNetwork* GetNetworkAsset();
-
-	/** Get a reference to this critic's neural network. */
-	UE::Learning::INeuralNetwork& GetCriticNetwork();
-	
 	/** Get a reference to this critic's critic function object. */
-	UE::Learning::FNeuralNetworkCriticFunction& GetCriticObject();
+	UE::Learning::FNeuralNetworkCritic& GetCriticObject();
 
 // ----- Private Data -----
 private:
@@ -166,27 +156,18 @@ private:
 
 	/** The underlying neural network. */
 	UPROPERTY(VisibleAnywhere, Transient, Category = "LearningAgents")
-	TObjectPtr<ULearningAgentsNeuralNetwork> Network;
+	TObjectPtr<ULearningAgentsNeuralNetwork> CriticNetwork;
 
 	/** Internal Critic Function Object */
-	TSharedPtr<UE::Learning::FNeuralNetworkCriticFunction> CriticObject;
+	TSharedPtr<UE::Learning::FNeuralNetworkCritic> CriticObject;
 
-// ----- Private Iteration Checks ----- 
-private:
+	/** Buffer to store the computed returns for each agent */
+	TLearningArray<1, float, TInlineAllocator<32>> Returns;
 
-	/** Number of times critic has been evaluated for all agents */
-	TLearningArray<1, uint64, TInlineAllocator<32>> CriticAgentIteration;
+	/** Number of times the returns have been computed for all agents */
+	TLearningArray<1, uint64, TInlineAllocator<32>> ReturnsIteration;
 
 	/** Temp buffers used to record the set of agents that are valid for evaluation */
 	TArray<int32> ValidAgentIds;
 	UE::Learning::FIndexSet ValidAgentSet;
-
-private:
-#if UE_LEARNING_AGENTS_ENABLE_VISUAL_LOG
-	/** Color used to draw this action in the visual log */
-	FLinearColor VisualLogColor = FColor::Orange;
-
-	/** Describes this critic to the visual logger for debugging purposes. */
-	void VisualLog(const UE::Learning::FIndexSet AgentSet) const;
-#endif
 };
diff --git a/Engine/Plugins/Experimental/LearningAgents/Source/LearningAgents/Public/LearningAgentsInteractor.h b/Engine/Plugins/Experimental/LearningAgents/Source/LearningAgents/Public/LearningAgentsInteractor.h
--- a/Engine/Plugins/Experimental/LearningAgents/Source/LearningAgents/Public/LearningAgentsInteractor.h
+++ b/Engine/Plugins/Experimental/LearningAgents/Source/LearningAgents/Public/LearningAgentsInteractor.h
@@ -2,7 +2,9 @@
 
 #pragma once
 
-#include "LearningAgentsManagerComponent.h"
+#include "LearningAgentsManagerListener.h"
+#include "LearningAgentsObservations.h"
+#include "LearningAgentsActions.h"
 
 #include "LearningArray.h"
 #include "Containers/Array.h"
@@ -14,195 +16,306 @@
 
 #include "LearningAgentsInteractor.generated.h"
 
-namespace UE::Learning
-{
-	struct FArrayMap;
-	struct FConcatenateFeature;
-	struct FFeatureObject;
-}
+class ULearningAgentsNeuralNetwork;
 
-class ALearningAgentsManager;
-class ULearningAgentsAction;
-class ULearningAgentsObservation;
+DECLARE_DYNAMIC_DELEGATE_RetVal_TwoParams(FLearningAgentsObservationObjectElement, FGatherAgentObservationDelegate, ULearningAgentsObservationObject*, InObservationObject, const int32, AgentId);
+DECLARE_DYNAMIC_DELEGATE_ThreeParams(FScatterAgentActionDelegate, const FLearningAgentsActionObjectElement, InActionObjectElement, const ULearningAgentsActionObject*, InActionObject, const int32, AgentId);
 
 /**
  * ULearningAgentsInteractor defines how agents interact with the environment through their observations and actions.
- * Additionally, it provides methods to be called during the inference process of those agents.
  *
- * To use this class, you need to implement the SetupObservations and SetupActions functions (as well as their
- * corresponding SetObservations and SetActions functions), which will define the size of inputs and outputs to your
- * policy. Before you can do inference, you need to call SetupInteractor, which will initialize the underlying data
- * structure, and you need to call AddAgent for each object you want controlled by this agent interactor.
+ * To use this class, you need to implement `SpecifyAgentObservations` and `SpecifyAgentActions`, which will define 
+ * the structure of inputs and outputs to your policy. You also need to implement `GatherAgentObservations` and 
+ * `ScatterAgentActions` which will dictate how those observations are gathered, and actions actuated in your
+ * environment.
  */
-UCLASS(Abstract, BlueprintType, Blueprintable)
-class LEARNINGAGENTS_API ULearningAgentsInteractor : public ULearningAgentsManagerComponent
+UCLASS(Abstract, HideDropdown, BlueprintType, Blueprintable)
+class LEARNINGAGENTS_API ULearningAgentsInteractor : public ULearningAgentsManagerListener
 {
 	GENERATED_BODY()
 
 public:
 
+	friend class ULearningAgentsController;
+	friend class ULearningAgentsPolicy;
+	friend class ULearningAgentsRecorder;
+	friend class ULearningAgentsTrainer;
+
 	// These constructors/destructors are needed to make forward declarations happy
 	ULearningAgentsInteractor();
 	ULearningAgentsInteractor(FVTableHelper& Helper);
 	virtual ~ULearningAgentsInteractor();
 
-	/** Initializes this object and runs the setup events for observations and actions. */
+	/**
+	 * Constructs an Interactor.
+	 *
+	 * @param InManager						The input Manager
+	 * @param Class							The interactor class
+	 * @param Name							The interactor name
+	 */
+	UFUNCTION(BlueprintCallable, Category = "LearningAgents", meta = (DeterminesOutputType = "Class"))
+	static ULearningAgentsInteractor* MakeInteractor(
+		ULearningAgentsManager* InManager, 
+		TSubclassOf<ULearningAgentsInteractor> Class,
+		const FName Name = TEXT("Interactor"));
+
+	/**
+	 * Initializes an Interactor.
+	 *
+	 * @param InManager						The input Manager
+	 */
 	UFUNCTION(BlueprintCallable, Category = "LearningAgents")
-	void SetupInteractor();
+	void SetupInteractor(ULearningAgentsManager* InManager);
 
 public:
 
-	//~ Begin ULearningAgentsManagerComponent Interface
-	virtual void OnAgentsAdded(const TArray<int32>& AgentIds) override;
-	virtual void OnAgentsRemoved(const TArray<int32>& AgentIds) override;
-	virtual void OnAgentsReset(const TArray<int32>& AgentIds) override;
-	//~ End ULearningAgentsManagerComponent Interface
+	//~ Begin ULearningAgentsManagerListener Interface
+	virtual void OnAgentsAdded_Implementation(const TArray<int32>& AgentIds) override;
+	virtual void OnAgentsRemoved_Implementation(const TArray<int32>& AgentIds) override;
+	virtual void OnAgentsReset_Implementation(const TArray<int32>& AgentIds) override;
+	//~ End ULearningAgentsManagerListener Interface
 
 // ----- Observations -----
 public:
 
 	/**
-	 * During this event, all observations should be added to this agent interactor.
-	 * @see LearningAgentsObservations.h for the list of available observations.
+	 * Specifies the structure of the observations using the observation schema
+	 * 
+	 * @param OutObservationSchemaElement		Output Schema Element
+	 * @param InObservationSchema				Observation Schema
 	 */
-	UFUNCTION(BlueprintNativeEvent, Category = "LearningAgents")
-	void SetupObservations();
+	UFUNCTION(BlueprintNativeEvent, Category = "LearningAgents", Meta = (ForceAsFunction))
+	void SpecifyAgentObservations(FLearningAgentsObservationSchemaElement& OutObservationSchemaElement, ULearningAgentsObservationSchema* InObservationSchema);
 
 	/**
-	 * During this event, all observations should be set for each agent.
-	 * @param AgentIds The list of agent ids to set observations for.
-	 * @see LearningAgentsObservations.h for the list of available observations.
-	 * @see GetAgent to get the agent corresponding to each id.
+	 * Gathers all the observations for the given agents. The structure of the Observation Elements output by this function should match that defined
+	 * by the schema.
+	 *
+	 * @param OutObservationObjectElements		Output Observation Element. This should be the same size as AgentIds.
+	 * @param InObservationObject				Observation Object.
+	 * @param AgentIds							Set of Agent Ids to gather observations for.
 	 */
-	UFUNCTION(BlueprintNativeEvent, Category = "LearningAgents")
-	void SetObservations(const TArray<int32>& AgentIds);
+	UFUNCTION(BlueprintNativeEvent, Category = "LearningAgents", Meta = (ForceAsFunction))
+	void GatherAgentObservations(TArray<FLearningAgentsObservationObjectElement>& OutObservationObjectElements, ULearningAgentsObservationObject* InObservationObject, const TArray<int32>& AgentIds);
 
 	/**
-	 * Used by objects derived from ULearningAgentsObservation to add themselves to this agent interactor during their
-	 * creation. You shouldn't need to call this directly.
+	 * This is a convenience function that can be used to implement `GatherAgentObservations` in terms of a callback which iterates over every agent
+	 * individually and gathers the observations.
+	 * 
+	 * @param OutObservationObjectElements		Output Observation Element. This should be the same size as AgentIds.
+	 * @param InObservationObject				Observation Object.
+	 * @param AgentIds							Set of Agent Ids to gather observations for.
+	 * @param Delegate							The Delegate used to gather each agent's observations.
 	 */
-	void AddObservation(TObjectPtr<ULearningAgentsObservation> Object, const TSharedRef<UE::Learning::FFeatureObject>& Feature);
+	UFUNCTION(BlueprintPure = false, Category = "LearningAgents")
+	void GatherAgentObservationsUsingDelegate(
+			TArray<FLearningAgentsObservationObjectElement>& OutObservationElements,
+			ULearningAgentsObservationObject* InObservationObject, 
+			const TArray<int32>& AgentIds,
+			const FGatherAgentObservationDelegate& Delegate);
 
 // ----- Actions -----
 public:
 
 	/**
-	 * During this event, all actions should be added to this agent interactor.
-	 * @see LearningAgentsActions.h for the list of available actions.
+	 * Specifies the structure of the actions using the action schema
+	 *
+	 * @param OutActionSchemaElement			Output Schema Element
+	 * @param InActionSchema					Action Schema
 	 */
-	UFUNCTION(BlueprintNativeEvent, Category = "LearningAgents")
-	void SetupActions();
+	UFUNCTION(BlueprintNativeEvent, Category = "LearningAgents", Meta = (ForceAsFunction))
+	void SpecifyAgentActions(FLearningAgentsActionSchemaElement& OutActionSchemaElement, ULearningAgentsActionSchema* InActionSchema);
 
 	/**
-	 * During this event, you should retrieve the actions and apply them to your agents.
-	 * @param AgentIds The list of agent ids to get actions for.
-	 * @see LearningAgentsActions.h for the list of available actions.
-	 * @see GetAgent to get the agent corresponding to each id.
+	 * Scatters all the actions for the given agents and applies their affect in the world. The structure of the Action Elements given as input to 
+	 * this function will match that defined by the schema.
+	 *
+	 * @param InActionObjectElements			Input Actions Element. This will be the same size as AgentIds.
+	 * @param InActionObject					Action Object.
+	 * @param AgentIds							Set of Agent Ids to scatter observations for.
 	 */
-	UFUNCTION(BlueprintNativeEvent, Category = "LearningAgents")
-	void GetActions(const TArray<int32>& AgentIds);
+	UFUNCTION(BlueprintNativeEvent, Category = "LearningAgents", Meta = (ForceAsFunction))
+	void ScatterAgentActions(const TArray<FLearningAgentsActionObjectElement>& InActionObjectElements, const ULearningAgentsActionObject* InActionObject, const TArray<int32>& AgentIds);
 
 	/**
-	 * Used by objects derived from ULearningAgentsAction to add themselves to this agent interactor during their creation.
-	 * You shouldn't need to call this directly.
+	 * This is a convenience function that can be used to implement `ScatterAgentActions` in terms of a callback which iterates over every agent
+	 * individually and scatters the actions.
+	 *
+	 * @param InActionObjectElements			Input Actions Element. This will be the same size as AgentIds.
+	 * @param InActionObject					Action Object.
+	 * @param AgentIds							Set of Agent Ids to scatter observations for.
+	 * @param Delegate							The Delegate used to scatter each agent's actions.
 	 */
-	void AddAction(TObjectPtr<ULearningAgentsAction> Object, const TSharedRef<UE::Learning::FFeatureObject>& Feature);
-
-// ----- Encoding / Decoding -----
+	UFUNCTION(BlueprintPure = false, Category = "LearningAgents")
+	void ScatterAgentObservationsUsingDelegate(
+		const TArray<FLearningAgentsActionObjectElement>& InActionObjectElements, 
+		const ULearningAgentsActionObject* InActionObject,
+		const TArray<int32>& AgentIds,
+		const FScatterAgentActionDelegate& Delegate);
+
+// ----- Blueprint public interface -----
 public:
 
+	/** Gathers all the observations for all agents. This will call GatherAgentObservations. */
+	UFUNCTION(BlueprintCallable, Category = "LearningAgents")
+	void GatherObservations();
+
+	/** Scatters all the actions for all agents. This will call ScatterAgentActions. */
+	UFUNCTION(BlueprintCallable, Category = "LearningAgents")
+	void ScatterActions();
+
 	/**
-	 * Call this function when it is time to gather all the observations for your agents. This can be done each frame or
-	 * you can consider wiring it up to some kind of meaningful event, e.g. a hypothetical "OnAITurnStarted" if you have
-	 * a turn-based game. This will call this agent interactor's SetObservations event.
+	 * Get the current buffered observation vector for the given agent.
+	 * 
+	 * @param OutObservationVector				Output Observation Vector
+	 * @param OutObservationCompatibilityHash	Output Compatibility Hash used to identify which schema this vector is compatible with.
+	 * @param AgentId							Agent Id to look up the observation vector for.
 	 */
-	UFUNCTION(BlueprintCallable, Category = "LearningAgents")
-	void EncodeObservations();
+	UFUNCTION(BlueprintCallable, Category = "LearningAgents", Meta = (AgentId = "-1"))
+	void GetObservationVector(TArray<float>& OutObservationVector, int32& OutObservationCompatibilityHash, const int32 AgentId);
 
 	/**
-	 * Call this function when it is time for your agents to take their actions. You most likely want to call this after
-	 * your policy's EvaluatePolicy function to ensure that you are receiving the latest actions. This will call this
-	 * agent interactor's GetActions event.
+	 * Get the current buffered action vector for the given agent.
+	 *
+	 * @param OutActionVector					Output Action Vector
+	 * @param OutActionCompatibilityHash		Output Compatibility Hash used to identify which schema this vector is compatible with.
+	 * @param AgentId							Agent Id to look up the action vector for.
 	 */
-	UFUNCTION(BlueprintCallable, Category = "LearningAgents")
-	void DecodeActions();
+	UFUNCTION(BlueprintCallable, Category = "LearningAgents", Meta = (AgentId = "-1"))
+	void GetActionVector(TArray<float>& OutActionVector, int32& OutActionCompatibilityHash, const int32 AgentId);
 
 	/**
-	 * Gets the observation vector used by a given agent. Should be called only after EncodeObservations.
+	 * Sets the current buffered observation vector for the given agent.
 	 *
-	 * @param AgentId				The AgentId to look-up the observation vector for
-	 * @param OutObservationVector	The observation vector for the given agent
+	 * @param ObservationVector					Observation Vector
+	 * @param InObservationCompatibilityHash	Compatibility Hash used to identify which schema this vector is compatible with.
+	 * @param AgentId							Agent Id to set the observation vector for.
+	 * @param bIncrementIteration				If to increment the iteration number used to keep track of associated actions and observations.
 	 */
-	UFUNCTION(BlueprintPure, Category = "LearningAgents", meta = (AgentId = "-1"))
-	void GetObservationVector(const int32 AgentId, TArray<float>& OutObservationVector) const;
+	UFUNCTION(BlueprintCallable, Category = "LearningAgents", Meta = (AgentId = "-1"))
+	void SetObservationVector(const TArray<float>& ObservationVector, const int32 InObservationCompatibilityHash, const int32 AgentId, bool bIncrementIteration = true);
 
 	/**
-	 * Gets the action vector used by a given agent. Should be called only after EncodeActions or EvaluatePolicy.
+	 * Sets the current buffered action vector for the given agent.
 	 *
-	 * @param AgentId				The AgentId to look-up the action vector for
-	 * @param OutActionVector		The action vector for the given agent
+	 * @param ActionVector						Action Vector
+	 * @param InActionCompatibilityHash			Compatibility Hash used to identify which schema this vector is compatible with.
+	 * @param AgentId							Agent Id to set the observation vector for.
+	 * @param bIncrementIteration				If to increment the iteration number used to keep track of associated actions and observations.
+	 */
+	UFUNCTION(BlueprintCallable, Category = "LearningAgents", Meta = (AgentId = "-1"))
+	void SetActionVector(const TArray<float>& ActionVector, const int32 InActionCompatibilityHash, const int32 AgentId, bool bIncrementIteration = true);
+
+	/**
+	 * Returns true if GatherObservations or SetObservationVector has been called and the observation vector already set for the given agent.
 	 */
 	UFUNCTION(BlueprintPure, Category = "LearningAgents", meta = (AgentId = "-1"))
-	void GetActionVector(const int32 AgentId, TArray<float>& OutActionVector) const;
+	bool HasObservationVector(const int32 AgentId) const;
 
-// ----- Non-blueprint public interface -----
-public:
+	/**
+	 * Returns true if DecodeAndSampleActions on the policy or SetActionVector has been called and the action vector already set for the given agent.
+	 */
+	UFUNCTION(BlueprintPure, Category = "LearningAgents", meta = (AgentId = "-1"))
+	bool HasActionVector(const int32 AgentId) const;
 
-	/** Get a reference to this agent interactor's observation feature. */
-	UE::Learning::FFeatureObject& GetObservationFeature() const;
+	/** Gets the size of the observation vector used by this interactor. */
+	UFUNCTION(BlueprintPure, Category = "LearningAgents")
+	int32 GetObservationVectorSize() const;
 
-	/** Get a reference to this agent interactor's action feature. */
-	UE::Learning::FFeatureObject& GetActionFeature() const;
+	/** Gets the size of the encoded observation vector used by this interactor. */
+	UFUNCTION(BlueprintPure, Category = "LearningAgents")
+	int32 GetObservationEncodedVectorSize() const;
 
-	/** Get a const array view of this agent interactor's observation objects. */
-	TConstArrayView<ULearningAgentsObservation*> GetObservationObjects() const;
+	/** Gets the size of the action vector used by this interactor. */
+	UFUNCTION(BlueprintPure, Category = "LearningAgents")
+	int32 GetActionVectorSize() const;
 
-	/** Get a const array view of this agent interactor's action objects. */
-	TConstArrayView<ULearningAgentsAction*> GetActionObjects() const;
+	/** Gets the size of the action distribution vector used by this interactor. */
+	UFUNCTION(BlueprintPure, Category = "LearningAgents")
+	int32 GetActionDistributionVectorSize() const;
 
-	/** Get the number of times observations have been encoded for all agents */
-	TLearningArrayView<1, uint64> GetObservationEncodingAgentIteration();
+	/** Gets the size of the encoded action vector used by this interactor. */
+	UFUNCTION(BlueprintPure, Category = "LearningAgents")
+	int32 GetActionEncodedVectorSize() const;
 
-	/** Get the number of times actions have been encoded for all agents */
-	TLearningArrayView<1, uint64> GetActionEncodingAgentIteration();
+// ----- Non-blueprint public interface -----
+public:
 
 	/** Encode Observations for a specific set of agents */
-	void EncodeObservations(const UE::Learning::FIndexSet AgentSet);
+	void GatherObservations(const UE::Learning::FIndexSet AgentSet, bool bIncrementIteration = true);
 
-	/** Decode Actions for a specific set of agents */
-	void DecodeActions(const UE::Learning::FIndexSet AgentSet);
+	/** Scatter Actions for a specific set of agents */
+	void ScatterActions(const UE::Learning::FIndexSet AgentSet);
+
+	/** Gets the internal observation schema object */
+	const UE::Learning::Observation::FSchema& GetObservationSchema() const;
+
+	/** Gets the internal observation schema element */
+	UE::Learning::Observation::FSchemaElement GetObservationSchemaElement() const;
+
+	/** Gets the internal action schema object */
+	const UE::Learning::Action::FSchema& GetActionSchema() const;
+
+	/** Gets the internal action schema element */
+	UE::Learning::Action::FSchemaElement GetActionSchemaElement() const;
 
 private:
 
-	/** The list of current observation objects. */
+	/** Observation Schema used by this interactor */
+	UPROPERTY(VisibleAnywhere, Transient, Category = "LearningAgents")
+	TObjectPtr<ULearningAgentsObservationSchema> ObservationSchema;
+
+	/** Observation Schema Element used by this interactor */
 	UPROPERTY(VisibleAnywhere, Transient, Category = "LearningAgents")
-	TArray<TObjectPtr<ULearningAgentsObservation>> ObservationObjects;
+	FLearningAgentsObservationSchemaElement ObservationSchemaElement;
 
-	/** The list of current action objects. */
+	/** Action Schema used by this interactor */
 	UPROPERTY(VisibleAnywhere, Transient, Category = "LearningAgents")
-	TArray<TObjectPtr<ULearningAgentsAction>> ActionObjects;
+	TObjectPtr<ULearningAgentsActionSchema> ActionSchema;
+
+	/** Action Schema Element used by this interactor */
+	UPROPERTY(VisibleAnywhere, Transient, Category = "LearningAgents")
+	FLearningAgentsActionSchemaElement ActionSchemaElement;
+
+	/** Observation Object used by this interactor */
+	UPROPERTY(VisibleAnywhere, Transient, Category = "LearningAgents")
+	TObjectPtr<ULearningAgentsObservationObject> ObservationObject;
+
+	/** Observation Object Elements used by this interactor */
+	UPROPERTY(VisibleAnywhere, Transient, Category = "LearningAgents")
+	TArray<FLearningAgentsObservationObjectElement> ObservationObjectElements;
+
+	/** Action Object used by this interactor */
+	UPROPERTY(VisibleAnywhere, Transient, Category = "LearningAgents")
+	TObjectPtr<ULearningAgentsActionObject> ActionObject;
+
+	/** Action Object Elements used by this interactor */
+	UPROPERTY(VisibleAnywhere, Transient, Category = "LearningAgents")
+	TArray<FLearningAgentsActionObjectElement> ActionObjectElements;
 
 // ----- Private Data -----
 private:
 
-	TArray<TSharedRef<UE::Learning::FFeatureObject>, TInlineAllocator<16>> ObservationFeatures;
-	TArray<TSharedRef<UE::Learning::FFeatureObject>, TInlineAllocator<16>> ActionFeatures;
+	/** Buffer of Observation Vectors for each agent */
+	TLearningArray<2, float> ObservationVectors;
 
-	TSharedPtr<UE::Learning::FConcatenateFeature> Observations;
-	TSharedPtr<UE::Learning::FConcatenateFeature> Actions;
+	/** Buffer of Action Vectors for each agent */
+	TLearningArray<2, float> ActionVectors;
 
-// ----- Private Iteration Checks ----- 
-private:
+	/** Compatibility Hash for Observation Schema */
+	int32 ObservationCompatibilityHash = 0;
+
+	/** Compatibility Hash for Actiuon Schema */
+	int32 ActionCompatibilityHash = 0;
 
-	/** Number of times observations have been encoded for all agents */
-	TLearningArray<1, uint64, TInlineAllocator<32>> ObservationEncodingAgentIteration;
+	/** Number of times observation vector has been set for all agents */
+	TLearningArray<1, uint64, TInlineAllocator<32>> ObservationVectorIteration;
 
-	/** Number of times actions have been encoded for all agents */
-	TLearningArray<1, uint64, TInlineAllocator<32>> ActionEncodingAgentIteration;
+	/** Number of times action vector has been set for all agents */
+	TLearningArray<1, uint64, TInlineAllocator<32>> ActionVectorIteration;
 
 	/** Temp buffers used to record the set of agents that are valid for encoding/decoding */
 	TArray<int32> ValidAgentIds;
 	UE::Learning::FIndexSet ValidAgentSet;
-	TBitArray<TInlineAllocator<32>> ValidAgentStatus;
 
 };
diff --git a/Engine/Plugins/Experimental/LearningAgents/Source/LearningAgents/Public/LearningAgentsManager.h b/Engine/Plugins/Experimental/LearningAgents/Source/LearningAgents/Public/LearningAgentsManager.h
--- a/Engine/Plugins/Experimental/LearningAgents/Source/LearningAgents/Public/LearningAgentsManager.h
+++ b/Engine/Plugins/Experimental/LearningAgents/Source/LearningAgents/Public/LearningAgentsManager.h
@@ -4,22 +4,18 @@
 
 #include "LearningArray.h"
 
-#include "GameFramework/Actor.h"
 #include "Containers/Array.h"
 #include "Containers/ArrayView.h"
 #include "Containers/ContainerAllocationPolicies.h"
 #include "Templates/SharedPointer.h"
+
 #include "UObject/Object.h"
 #include "UObject/ObjectPtr.h"
+#include "Components/ActorComponent.h"
 
 #include "LearningAgentsManager.generated.h"
 
-namespace UE::Learning
-{
-	struct FArrayMap;
-}
-
-class ULearningAgentsManagerComponent;
+class ULearningAgentsManagerListener;
 
 /**
  * The agent manager is responsible for tracking which game objects are agents. It's the central class around which
@@ -28,20 +24,22 @@ class ULearningAgentsManagerComponent;
  * If you have multiple different types of objects you want controlled by Learning Agents, you should consider creating
  * one agent manager per object type, rather than trying to share an agent manager.
  */
-UCLASS(BlueprintType, Blueprintable, meta=(ChildCanTick))
-class LEARNINGAGENTS_API ALearningAgentsManager : public AActor
+UCLASS(BlueprintType, Blueprintable, meta = (BlueprintSpawnableComponent))
+class LEARNINGAGENTS_API ULearningAgentsManager : public UActorComponent
 {
 	GENERATED_BODY()
 
 public:
 
 	// These constructors/destructors are needed to make forward declarations happy
-	ALearningAgentsManager();
-	ALearningAgentsManager(FVTableHelper& Helper);
-	virtual ~ALearningAgentsManager();
+	ULearningAgentsManager();
+	ULearningAgentsManager(FVTableHelper& Helper);
+	virtual ~ULearningAgentsManager();
 
 	virtual void PostInitProperties() override;
 
+	virtual void TickComponent(float DeltaTime, enum ELevelTick TickType, FActorComponentTickFunction* ThisTickFunction) override;
+
 // ----- Agent Management -----
 public:
 
@@ -112,6 +110,7 @@ public:
 	/**
 	 * Gets the agent with the given id. Calling this from blueprint with the appropriate AgentClass will automatically
 	 * cast the object to the given type. If not in a blueprint, you should use one of the other GetAgent overloads.
+	 * 
 	 * @param AgentId The id of the agent to get.
 	 * @param AgentClass The class to cast the agent object to (in blueprint).
 	 * @return The agent object.
@@ -122,6 +121,7 @@ public:
 	/**
 	 * Gets the agents associated with a set of ids. Calling this from blueprint with the appropriate AgentClass will 
 	 * automatically cast the object to the given type.
+	 * 
 	 * @param AgentIds The ids of the agents to get.
 	 * @param AgentClass The class to cast the agent objects to (in blueprint).
 	 * @param OutAgents The output array of agent objects.
@@ -132,6 +132,7 @@ public:
 	/**
 	 * Gets all added agents. Calling this from blueprint with the appropriate AgentClass will automatically
 	 * cast the object to the given type.
+	 * 
 	 * @param AgentClass The class to cast the agent objects to (in blueprint).
 	 * @param OutAgents The output array of agent objects.
 	 * @param OutAgentIds The output array of agent ids.
@@ -141,6 +142,7 @@ public:
 
 	/**
 	 * Gets the agent id associated with a given agent.
+	 * 
 	 * @param Agent The agent object.
 	 * @return The agent id.
 	 */
@@ -149,6 +151,7 @@ public:
 
 	/**
 	 * Gets the agent ids associated with a set of agents.
+	 * 
 	 * @param OutAgentIds The ids of the agents.
 	 * @param InAgents The agent objects.
 	 */
@@ -157,6 +160,7 @@ public:
 
 	/**
 	 * Gets the number of agents added
+	 * 
 	 * @return The number of agents added.
 	 */
 	UFUNCTION(BlueprintPure, Category = "LearningAgents")
@@ -172,22 +176,15 @@ public:
 	UFUNCTION(BlueprintPure, Category = "LearningAgents", meta = (AgentId = "-1"))
 	bool HasAgent(const int32 AgentId) const;
 
-// ----- Tick Prerequisites -----
 public:
 
-	/**
-	 * Adds this manager as a tick prerequisite of the given actor objects.
-	 * @param InAgents The actor object which this manager will be added as a tick prerequisite to.
-	 */
+	/** Adds a listener to be tracked by this manager. */
 	UFUNCTION(BlueprintCallable, Category = "LearningAgents")
-	void AddManagerAsTickPrerequisiteOfAgents(const TArray<AActor*>& InAgents);
+	void AddListener(ULearningAgentsManagerListener* Listener);
 
-	/**
-	 * Adds the given actor objects as tick prerequisite of this manager.
-	 * @param InAgents The actor object which will be added as tick prerequisites to this manager.
-	 */
+	/** Removes a listener from being tracked by this manager. */
 	UFUNCTION(BlueprintCallable, Category = "LearningAgents")
-	void AddAgentsAsTickPrerequisiteOfManager(const TArray<AActor*>& InAgents);
+	void RemoveListener(ULearningAgentsManagerListener* Listener);
 
 
 // ----- Non-blueprint public interface -----
@@ -205,9 +202,6 @@ public:
 	/** Gets the set of agent ids currently added as an FIndexSet */
 	UE::Learning::FIndexSet GetAllAgentSet() const;
 
-	/** Get a const reference to this manager's underlying instance data. */
-	const TSharedPtr<UE::Learning::FArrayMap>& GetInstanceData() const;
-
 	/** Get a const array view of this manager's agent objects. */
 	TConstArrayView<TObjectPtr<UObject>> GetAgents() const;
 
@@ -221,15 +215,16 @@ private:
 	UPROPERTY(VisibleAnywhere, Transient, Category = "LearningAgents")
 	TArray<TObjectPtr<UObject>> Agents;
 
-private:
+	/** The list of current listeners. */
+	UPROPERTY(VisibleAnywhere, Transient, Category = "LearningAgents")
+	TArray<TObjectPtr<ULearningAgentsManagerListener>> Listeners;
 
-	/** Object containing all of the instance data required by all agents */
-	TSharedPtr<UE::Learning::FArrayMap> InstanceData;
+private:
 
 	/** Update the agent sets to keep them in sync with the id lists. */
 	void UpdateAgentSets();
 
-	/** Array of agent ids to be passed to events such as ULearningAgentsManagerComponent::OnAgentAdded. */
+	/** Array of agent ids to be passed to events such as ULearningAgentsManagerListener::OnAgentAdded. */
 	TArray<int32> OnEventAgentIds;
 
 	/** Array of agent ids currently in use and associated with each agent object. */
diff --git a/Engine/Plugins/Experimental/LearningAgents/Source/LearningAgents/Public/LearningAgentsNeuralNetwork.h b/Engine/Plugins/Experimental/LearningAgents/Source/LearningAgents/Public/LearningAgentsNeuralNetwork.h
--- a/Engine/Plugins/Experimental/LearningAgents/Source/LearningAgents/Public/LearningAgentsNeuralNetwork.h
+++ b/Engine/Plugins/Experimental/LearningAgents/Source/LearningAgents/Public/LearningAgentsNeuralNetwork.h
@@ -20,7 +20,7 @@ enum class ELearningAgentsActivationFunction : uint8
 	TanH	UMETA(DisplayName = "TanH"),
 };
 
-class ULearningAgentsNeuralNetworkData;
+class ULearningNeuralNetworkData;
 
 /** A neural network data asset. */
 UCLASS(BlueprintType)
@@ -34,14 +34,13 @@ public:
 	ULearningAgentsNeuralNetwork(FVTableHelper& Helper);
 	virtual ~ULearningAgentsNeuralNetwork();
 
-	/**
-	 * Resets this network asset to be empty.
-	 */
+	/** Resets this network asset to be empty. */
 	UFUNCTION(CallInEditor, Category = "LearningAgents")
 	void ResetNetwork();
 
 	/**
 	 * Load this network from a snapshot.
+	 * 
 	 * @param File The snapshot file.
 	 */
 	UFUNCTION(BlueprintCallable, Category = "LearningAgents", meta = (RelativePath))
@@ -49,23 +48,26 @@ public:
 
 	/**
 	 * Save this network into a snapshot.
+	 * 
 	 * @param File The snapshot file.
 	 */
 	UFUNCTION(BlueprintCallable, Category = "LearningAgents", meta = (RelativePath))
 	void SaveNetworkToSnapshot(const FFilePath& File);
 
 	/**
-	 * Copy another asset's network weights into this network. Network sizes must match.
+	 * Copy another asset into this network.
+	 * 
 	 * @param NeuralNetworkAsset The asset to load from.
 	 */
 	UFUNCTION(BlueprintCallable, Category = "LearningAgents")
 	void LoadNetworkFromAsset(ULearningAgentsNeuralNetwork* NeuralNetworkAsset);
 
 	/**
-	 * Copy this network's weights to another asset. Network sizes must match.
+	 * Copy this network into another asset.
+	 * 
 	 * @param NeuralNetworkAsset The asset to save to.
 	 */
-	UFUNCTION(BlueprintCallable, Category = "LearningAgents", Meta = (DevelopmentOnly))
+	UFUNCTION(BlueprintCallable, Category = "LearningAgents")
 	void SaveNetworkToAsset(ULearningAgentsNeuralNetwork* NeuralNetworkAsset);
 
 public:
@@ -76,6 +78,6 @@ public:
 public:
 
 	/** The internal Neural Network Data */
-	UPROPERTY();
-	TObjectPtr<ULearningAgentsNeuralNetworkData> NeuralNetworkData;
+	UPROPERTY(BlueprintReadOnly, VisibleAnywhere, Category = "LearningAgents");
+	TObjectPtr<ULearningNeuralNetworkData> NeuralNetworkData;
 };
