LearningAgents: Large restructuring to allow for structured actions and observations.

diff --git a/Engine/Plugins/Experimental/LearningAgents/Source/LearningAgents/Public/LearningAgentsObservations.h b/Engine/Plugins/Experimental/LearningAgents/Source/LearningAgents/Public/LearningAgentsObservations.h
--- a/Engine/Plugins/Experimental/LearningAgents/Source/LearningAgents/Public/LearningAgentsObservations.h
+++ b/Engine/Plugins/Experimental/LearningAgents/Source/LearningAgents/Public/LearningAgentsObservations.h
@@ -3,1363 +3,555 @@
 #pragma once
 
 #include "LearningArray.h"
-#include "LearningAgentsDebug.h"
+#include "LearningObservation.h"
 
-#include "Templates/SharedPointer.h"
-#include "UObject/Object.h"
+#include "Engine/EngineTypes.h"
+#include "Kismet/BlueprintFunctionLibrary.h"
 
 #include "LearningAgentsObservations.generated.h"
 
-namespace UE::Learning
-{
-	struct FFeatureObject;
-	struct FFloatFeature;
-	struct FTimeFeature;
-	struct FAngleFeature;
-	struct FRotationFeature;
-	struct FDirectionFeature;
-	struct FPlanarDirectionFeature;
-	struct FPositionFeature;
-	struct FScalarPositionFeature;
-	struct FPlanarPositionFeature;
-	struct FVelocityFeature;
-	struct FScalarVelocityFeature;
-	struct FPlanarVelocityFeature;
-	struct FAngularVelocityFeature;
-	struct FScalarAngularVelocityFeature;
-}
-
-class ULearningAgentsInteractor;
-
-// For functions in this file, we are favoring having more verbose names such as "AddFloatObservation" vs simply "Add" in 
-// order to keep it easy to find the correct function in blueprints.
-
-//------------------------------------------------------------------
-
-/** The base class for all observations. Observations define the inputs to your agents. */
-UCLASS(Abstract, BlueprintType)
-class LEARNINGAGENTS_API ULearningAgentsObservation : public UObject
-{
-	GENERATED_BODY()
-
-public:
-
-	/** Reference to the Interactor this observation is associated with. */
-	UPROPERTY(VisibleAnywhere, Transient, Category = "LearningAgents")
-	TObjectPtr<ULearningAgentsInteractor> Interactor;
-
-public:
-
-	/** Initialize the internal state for a given maximum number of agents */
-	void Init(const int32 MaxAgentNum);
-
-	/**
-	 * Called whenever agents are added to the associated ULearningAgentsInteractor object.
-	 * @param AgentIds Array of agent ids which have been added
-	 */
-	virtual void OnAgentsAdded(const TArray<int32>& AgentIds);
-
-	/**
-	 * Called whenever agents are removed from the associated ULearningAgentsInteractor object.
-	 * @param AgentIds Array of agent ids which have been removed
-	 */
-	virtual void OnAgentsRemoved(const TArray<int32>& AgentIds);
-
-	/**
-	 * Called whenever agents are reset on the associated ULearningAgentsInteractor object.
-	 * @param AgentIds Array of agent ids which have been reset
-	 */
-	virtual void OnAgentsReset(const TArray<int32>& AgentIds);
-
-	/** Get the number of times an observation has been set for the given agent id. */
-	uint64 GetAgentIteration(const int32 AgentId) const;
-
-public:
-#if UE_LEARNING_AGENTS_ENABLE_VISUAL_LOG
-
-	/** Color used to draw this observation in the visual log */
-	FLinearColor VisualLogColor = FColor::Red;
-
-	/** Describes this observation to the visual logger for debugging purposes. */
-	virtual void VisualLog(const UE::Learning::FIndexSet Instances) const {}
-#endif
+class ULearningAgentsObservationSchema;
+class ULearningAgentsObservationObject;
+struct FLearningAgentsObservationSchemaElement;
+struct FLearningAgentsObservationObjectElement;
 
-protected:
+class USplineComponent;
 
-	/** Number of times this observation has been set for all agents */
-	TLearningArray<1, uint64, TInlineAllocator<32>> AgentIteration;
-};
-
-//------------------------------------------------------------------
-
-/** A simple float observation. Used as a catch-all for situations where a more type-specific observation does not exist yet. */
-UCLASS()
-class LEARNINGAGENTS_API UFloatObservation : public ULearningAgentsObservation
+/** An element of an Observation Schema */
+USTRUCT(BlueprintType)
+struct LEARNINGAGENTS_API FLearningAgentsObservationSchemaElement
 {
 	GENERATED_BODY()
 
-public:
-
-	/**
-	 * Adds a new float observation to the given agent interactor. Call during ULearningAgentsInteractor::SetupObservations event.
-	 * @param InInteractor The agent interactor to add this observation to.
-	 * @param Name The name of this new observation. Used for debugging.
-	 * @param Scale Used to normalize the data for the observation.
-	 * @return The newly created observation.
-	 */
-	UFUNCTION(BlueprintCallable, Category = "LearningAgents", meta=(DefaultToSelf = "InInteractor"))
-	static UFloatObservation* AddFloatObservation(ULearningAgentsInteractor* InInteractor, const FName Name = NAME_None, const float Scale = 1.0f);
-
-	/**
-	 * Sets the data for this observation. Call during ULearningAgentsInteractor::SetObservations event.
-	 * @param AgentId The agent id this data corresponds to.
-	 * @param Observation The value currently being observed.
-	 */
-	UFUNCTION(BlueprintCallable, Category = "LearningAgents", meta = (AgentId = "-1"))
-	void SetFloatObservation(const int32 AgentId, const float Observation);
-
-#if UE_LEARNING_AGENTS_ENABLE_VISUAL_LOG
-	/** Describes this observation to the visual logger for debugging purposes. */
-	virtual void VisualLog(const UE::Learning::FIndexSet Instances) const override;
-#endif
-
-	TSharedPtr<UE::Learning::FFloatFeature> FeatureObject;
+	UE::Learning::Observation::FSchemaElement SchemaElement;
 };
 
-/** A simple array of floats observation. Used as a catch-all for situations where a more type-specific observation does not exist yet. */
-UCLASS()
-class LEARNINGAGENTS_API UFloatArrayObservation : public ULearningAgentsObservation
+/** An element of an Observation Object */
+USTRUCT(BlueprintType)
+struct LEARNINGAGENTS_API FLearningAgentsObservationObjectElement
 {
 	GENERATED_BODY()
 
-public:
-
-	/**
-	 * Adds a new float array observation to the given agent interactor. Call during ULearningAgentsInteractor::SetupObservations event.
-	 * @param InInteractor The agent interactor to add this observation to.
-	 * @param Name The name of this new observation. Used for debugging.
-	 * @param Num The number of floats in the array
-	 * @param Scale Used to normalize the data for the observation.
-	 * @return The newly created observation.
-	 */
-	UFUNCTION(BlueprintCallable, Category = "LearningAgents", meta = (DefaultToSelf = "InInteractor"))
-	static UFloatArrayObservation* AddFloatArrayObservation(ULearningAgentsInteractor* InInteractor, const FName Name = NAME_None, const int32 Num = 1, const float Scale = 1.0f);
-
-	/**
-	 * Sets the data for this observation. Call during ULearningAgentsInteractor::SetObservations event.
-	 * @param AgentId The agent id this data corresponds to.
-	 * @param Observation The value currently being observed.
-	 */
-	UFUNCTION(BlueprintCallable, Category = "LearningAgents", meta = (AgentId = "-1"))
-	void SetFloatArrayObservation(const int32 AgentId, const TArray<float>& Observation);
-
-#if UE_LEARNING_AGENTS_ENABLE_VISUAL_LOG
-	/** Describes this observation to the visual logger for debugging purposes. */
-	virtual void VisualLog(const UE::Learning::FIndexSet Instances) const override;
-#endif
-
-	TSharedPtr<UE::Learning::FFloatFeature> FeatureObject;
+	UE::Learning::Observation::FObjectElement ObjectElement;
 };
 
-//------------------------------------------------------------------
-
-/** A simple observation for an FVector. */
-UCLASS()
-class LEARNINGAGENTS_API UVectorObservation : public ULearningAgentsObservation
-{
-	GENERATED_BODY()
-
-public:
-
-	/**
-	 * Adds a new vector observation to the given agent interactor. Call during ULearningAgentsInteractor::SetupObservations event.
-	 * @param InInteractor The agent interactor to add this observation to.
-	 * @param Name The name of this new observation. Used for debugging.
-	 * @param Scale Used to normalize the data for the observation.
-	 * @return The newly created observation.
-	 */
-	UFUNCTION(BlueprintCallable, Category = "LearningAgents", meta = (DefaultToSelf = "InInteractor"))
-	static UVectorObservation* AddVectorObservation(ULearningAgentsInteractor* InInteractor, const FName Name = NAME_None, const float Scale = 1.0f);
+/** Comparison operator for Observation Object Elements */
+bool operator==(const FLearningAgentsObservationObjectElement& Lhs, const FLearningAgentsObservationObjectElement& Rhs);
 
-	/**
-	 * Sets the data for this observation. Call during ULearningAgentsInteractor::SetObservations event.
-	 * @param AgentId The agent id this data corresponds to.
-	 * @param Observation The values currently being observed.
-	 */
-	UFUNCTION(BlueprintCallable, Category = "LearningAgents", meta = (AgentId = "-1"))
-	void SetVectorObservation(const int32 AgentId, const FVector Observation);
+/** Hashing operator for Observation Object Elements */
+uint32 GetTypeHash(const FLearningAgentsObservationObjectElement& Element);
 
-#if UE_LEARNING_AGENTS_ENABLE_VISUAL_LOG
-	/** Describes this observation to the visual logger for debugging purposes. */
-	virtual void VisualLog(const UE::Learning::FIndexSet Instances) const override;
-#endif
-
-	TSharedPtr<UE::Learning::FFloatFeature> FeatureObject;
-};
-
-/** A simple observation for an array of FVector. */
-UCLASS()
-class LEARNINGAGENTS_API UVectorArrayObservation : public ULearningAgentsObservation
+/** Enum Type representing either observation A or observation B */
+UENUM(BlueprintType)
+enum class ELearningAgentsEitherObservation : uint8
 {
-	GENERATED_BODY()
-
-public:
-
-	/**
-	 * Adds a new vector array observation to the given agent interactor. Call during ULearningAgentsInteractor::SetupObservations event.
-	 * @param InInteractor The agent interactor to add this observation to.
-	 * @param Name The name of this new observation. Used for debugging.
-	 * @param Num The number of vectors in the array
-	 * @param Scale Used to normalize the data for the observation.
-	 * @return The newly created observation.
-	 */
-	UFUNCTION(BlueprintCallable, Category = "LearningAgents", meta = (DefaultToSelf = "InInteractor"))
-	static UVectorArrayObservation* AddVectorArrayObservation(ULearningAgentsInteractor* InInteractor, const FName Name = NAME_None, const int32 Num = 1, const float Scale = 1.0f);
-
-	/**
-	 * Sets the data for this observation. Call during ULearningAgentsInteractor::SetObservations event.
-	 * @param AgentId The agent id this data corresponds to.
-	 * @param Observation The values currently being observed.
-	 */
-	UFUNCTION(BlueprintCallable, Category = "LearningAgents", meta = (AgentId = "-1"))
-	void SetVectorArrayObservation(const int32 AgentId, const TArray<FVector>& Observation);
-
-#if UE_LEARNING_AGENTS_ENABLE_VISUAL_LOG
-	/** Describes this observation to the visual logger for debugging purposes. */
-	virtual void VisualLog(const UE::Learning::FIndexSet Instances) const override;
-#endif
-
-	TSharedPtr<UE::Learning::FFloatFeature> FeatureObject;
+	A,
+	B,
 };
 
-//------------------------------------------------------------------
-
-/** An observation of an enumeration. */
-UCLASS()
-class LEARNINGAGENTS_API UEnumObservation : public ULearningAgentsObservation
+/** Enum Type representing either a Null observation or some Valid observation */
+UENUM(BlueprintType)
+enum class ELearningAgentsOptionalObservation : uint8
 {
-	GENERATED_BODY()
-
-public:
-
-	/**
-	 * Adds a new enum observation to the given agent interactor. Call during ULearningAgentsInteractor::SetupObservations event.
-	 * @param InInteractor The agent interactor to add this observation to.
-	 * @param EnumType The type of enum to use
-	 * @param Name The name of this new observation. Used for debugging.
-	 * @return The newly created observation.
-	 */
-	UFUNCTION(BlueprintCallable, Category = "LearningAgents", meta = (DefaultToSelf = "InInteractor"))
-	static UEnumObservation* AddEnumObservation(ULearningAgentsInteractor* InInteractor, const UEnum* EnumType, const FName Name = NAME_None);
-
-	/**
-	 * Sets the data for this observation. Call during ULearningAgentsInteractor::SetObservations event.
-	 * @param AgentId The agent id this data corresponds to.
-	 * @param Value The enum value currently being observed.
-	 */
-	UFUNCTION(BlueprintCallable, Category = "LearningAgents", meta = (AgentId = "-1"))
-	void SetEnumObservation(const int32 AgentId, const uint8 Value);
-
-#if UE_LEARNING_AGENTS_ENABLE_VISUAL_LOG
-	/** Describes this observation to the visual logger for debugging purposes. */
-	virtual void VisualLog(const UE::Learning::FIndexSet Instances) const override;
-#endif
-
-	const UEnum* Enum = nullptr;
-	TSharedPtr<UE::Learning::FFloatFeature> FeatureObject;
+	Null,
+	Valid,
 };
 
-/** An observation of an array of enumerations. */
+/**
+ * Observation Functions
+ * 
+ * Convenience functions useful for constructing observations.
+ */
 UCLASS()
-class LEARNINGAGENTS_API UEnumArrayObservation : public ULearningAgentsObservation
+class ULearningAgentsObservationFunctions : public UBlueprintFunctionLibrary
 {
 	GENERATED_BODY()
 
-public:
-
-	/**
-	 * Adds a new enum array observation to the given agent interactor. Call during ULearningAgentsInteractor::SetupObservations event.
-	 * @param InInteractor The agent interactor to add this observation to.
-	 * @param EnumType The type of enum to use
-	 * @param Name The name of this new observation. Used for debugging.
-	 * @param EnumNum The number of enum observations in the array
-	 * @return The newly created observation.
-	 */
-	UFUNCTION(BlueprintCallable, Category = "LearningAgents", meta = (DefaultToSelf = "InInteractor"))
-	static UEnumArrayObservation* AddEnumArrayObservation(ULearningAgentsInteractor* InInteractor, const UEnum* EnumType, const FName Name = NAME_None, const int32 EnumNum = 1);
-
-	/**
-	 * Sets the data for this observation. Call during ULearningAgentsInteractor::SetObservations event.
-	 * @param AgentId The agent id this data corresponds to.
-	 * @param Values The enum values currently being observed.
-	 */
-	UFUNCTION(BlueprintCallable, Category = "LearningAgents", meta = (AgentId = "-1"))
-	void SetEnumArrayObservation(const int32 AgentId, const TArray<uint8>& Values);
-
-#if UE_LEARNING_AGENTS_ENABLE_VISUAL_LOG
-	/** Describes this observation to the visual logger for debugging purposes. */
-	virtual void VisualLog(const UE::Learning::FIndexSet Instances) const override;
-#endif
-
-	const UEnum* Enum = nullptr;
-	TSharedPtr<UE::Learning::FFloatFeature> FeatureObject;
+	/** Project a transform onto the ground plane, leaving just rotation around the vertical axis */
+	UFUNCTION(BlueprintPure, Category = "Learning Agents")
+	static FTransform ProjectTransformOntoGroundPlane(const FTransform Transform, const FVector LocalForwardVector = FVector::ForwardVector, const float GroundPlaneHeight = 0.0f);
 };
 
-//------------------------------------------------------------------
-
-/** An observation of a time relative to another time. */
-UCLASS()
-class LEARNINGAGENTS_API UTimeObservation : public ULearningAgentsObservation
+/**
+ * Observation Schema
+ *
+ * This object is used to construct a schema describing some structure of observations.
+ */
+UCLASS(BlueprintType)
+class LEARNINGAGENTS_API ULearningAgentsObservationSchema : public UObject
 {
 	GENERATED_BODY()
 
 public:
 
-	/**
-	 * Adds a new time observation to the given agent interactor. Call during ULearningAgentsInteractor::SetupObservations event.
-	 * @param InInteractor The agent interactor to add this observation to.
-	 * @param Name The name of this new observation. Used for debugging.
-	 * @param Scale Used to normalize the data for the observation.
-	 * @return The newly created observation.
-	 */
-	UFUNCTION(BlueprintCallable, Category = "LearningAgents", meta = (DefaultToSelf = "InInteractor"))
-	static UTimeObservation* AddTimeObservation(ULearningAgentsInteractor* InInteractor, const FName Name = NAME_None, const float Scale = 1.0f);
+	/** Gets the internal observation schema object */
+	const UE::Learning::Observation::FSchema& GetObservationSchema() const;
 
 	/**
-	 * Sets the data for this observation. Call during ULearningAgentsInteractor::SetObservations event.
-	 * @param AgentId The agent id this data corresponds to.
-	 * @param Time The time currently being observed.
-	 * @param RelativeTime The time the provided time should be encoded relative to.
+	 * Validates that the given object matches the schema. Will log errors on objects that don't match.
+	 *
+	 * @param SchemaElement			Schema Element
+	 * @param Object				Observation Object
+	 * @param ObjectElement			Observation Object Element
+	 * @returns						true if the object matches the schema
 	 */
-	UFUNCTION(BlueprintCallable, Category = "LearningAgents", meta = (AgentId = "-1"))
-	void SetTimeObservation(const int32 AgentId, const float Time, const float RelativeTime = 0.0f);
-
-#if UE_LEARNING_AGENTS_ENABLE_VISUAL_LOG
-	/** Describes this observation to the visual logger for debugging purposes. */
-	virtual void VisualLog(const UE::Learning::FIndexSet Instances) const override;
-#endif
-
-	TSharedPtr<UE::Learning::FTimeFeature> FeatureObject;
-};
-
-/** An observation of an array of times. */
-UCLASS()
-class LEARNINGAGENTS_API UTimeArrayObservation : public ULearningAgentsObservation
-{
-	GENERATED_BODY()
+	UFUNCTION(BlueprintPure, Category = "LearningAgents")
+	bool ValidateObjectMatchesSchema(
+		const FLearningAgentsObservationSchemaElement SchemaElement,
+		const ULearningAgentsObservationObject* Object,
+		const FLearningAgentsObservationObjectElement ObjectElement) const;
 
 public:
 
-	/**
-	 * Adds a new angle observation to the given agent interactor. Call during ULearningAgentsInteractor::SetupObservations event.
-	 * @param InInteractor The agent interactor to add this observation to.
-	 * @param Name The name of this new observation. Used for debugging.
-	 * @param TimeNum The number of times in the array
-	 * @param Scale Used to normalize the data for the observation.
-	 * @return The newly created observation.
-	 */
-	UFUNCTION(BlueprintCallable, Category = "LearningAgents", meta = (DefaultToSelf = "InInteractor"))
-	static UTimeArrayObservation* AddTimeArrayObservation(ULearningAgentsInteractor* InInteractor, const FName Name = NAME_None, const int32 TimeNum = 1, const float Scale = 1.0f);
+	// Basic Observations
 
-	/**
-	 * Sets the data for this observation. Call during ULearningAgentsInteractor::SetObservations event.
-	 * @param AgentId The agent id this data corresponds to.
-	 * @param Time The times currently being observed.
-	 * @param RelativeTime The time the provided time should be encoded relative to.
-	 */
-	UFUNCTION(BlueprintCallable, Category = "LearningAgents", meta = (AgentId = "-1"))
-	void SetTimeArrayObservation(const int32 AgentId, const TArray<float>& Times, const float RelativeTime = 0.0f);
+	UFUNCTION(BlueprintPure, Category = "LearningAgents")
+	FLearningAgentsObservationSchemaElement SpecifyNullObservation(const FName Name = TEXT("Null"));
 
-#if UE_LEARNING_AGENTS_ENABLE_VISUAL_LOG
-	/** Describes this observation to the visual logger for debugging purposes. */
-	virtual void VisualLog(const UE::Learning::FIndexSet Instances) const override;
-#endif
+	UFUNCTION(BlueprintPure, Category = "LearningAgents")
+	FLearningAgentsObservationSchemaElement SpecifyContinuousObservation(const int32 Size, const FName Name = TEXT("Continuous"));
 
-	TSharedPtr<UE::Learning::FTimeFeature> FeatureObject;
-};
+	UFUNCTION(BlueprintPure, Category = "LearningAgents")
+	FLearningAgentsObservationSchemaElement SpecifyExclusiveDiscreteObservation(const int32 Size, const FName Name = TEXT("ExclusiveDiscrete"));
 
-//------------------------------------------------------------------
+	UFUNCTION(BlueprintPure, Category = "LearningAgents")
+	FLearningAgentsObservationSchemaElement SpecifyInclusiveDiscreteObservation(const int32 Size, const FName Name = TEXT("InclusiveDiscrete"));
 
-/** An observation of an angle relative to another angle. */
-UCLASS()
-class LEARNINGAGENTS_API UAngleObservation : public ULearningAgentsObservation
-{
-	GENERATED_BODY()
+	UFUNCTION(BlueprintPure, Category = "LearningAgents")
+	FLearningAgentsObservationSchemaElement SpecifyIndexObservation(const int32 Size, const FName Name = TEXT("Index"));
 
-public:
+	UFUNCTION(BlueprintPure, Category = "LearningAgents")
+	FLearningAgentsObservationSchemaElement SpecifyStructObservation(const TMap<FName, FLearningAgentsObservationSchemaElement>& Elements, const FName Name = TEXT("Struct"));
+	
+	UFUNCTION(BlueprintPure, Category = "LearningAgents")
+	FLearningAgentsObservationSchemaElement SpecifyStructObservationFromArrays(const TArray<FName>& ElementNames, const TArray<FLearningAgentsObservationSchemaElement>& Elements, const FName Name = TEXT("Struct"));
+	FLearningAgentsObservationSchemaElement SpecifyStructObservationFromArrayViews(const TArrayView<const FName> ElementNames, const TArrayView<const FLearningAgentsObservationSchemaElement> Elements, const FName Name = TEXT("Struct"));
 
-	/**
-	 * Adds a new angle observation to the given agent interactor. Call during ULearningAgentsInteractor::SetupObservations event.
-	 * @param InInteractor The agent interactor to add this observation to.
-	 * @param Name The name of this new observation. Used for debugging.
-	 * @param Scale Used to normalize the data for the observation. Angle observations are encoded as directions. 
-	 * @return The newly created observation.
-	 */
-	UFUNCTION(BlueprintCallable, Category = "LearningAgents", meta = (DefaultToSelf = "InInteractor"))
-	static UAngleObservation* AddAngleObservation(ULearningAgentsInteractor* InInteractor, const FName Name = NAME_None, const float Scale = 1.0f);
+	UFUNCTION(BlueprintPure, Category = "LearningAgents")
+	FLearningAgentsObservationSchemaElement SpecifyExclusiveUnionObservation(const TMap<FName, FLearningAgentsObservationSchemaElement>& Elements, const int32 EncodingSize = 128, const FName Name = TEXT("ExclusiveUnion"));
+	
+	UFUNCTION(BlueprintPure, Category = "LearningAgents")
+	FLearningAgentsObservationSchemaElement SpecifyExclusiveUnionObservationFromArrays(const TArray<FName>& ElementNames, const TArray<FLearningAgentsObservationSchemaElement>& Elements, const int32 EncodingSize = 128, const FName Name = TEXT("ExclusiveUnion"));
+	FLearningAgentsObservationSchemaElement SpecifyExclusiveUnionObservationFromArrayViews(const TArrayView<const FName> ElementNames, const TArrayView<const FLearningAgentsObservationSchemaElement> Elements, const int32 EncodingSize = 128, const FName Name = TEXT("ExclusiveUnion"));
 
-	/**
-	 * Sets the data for this observation. Call during ULearningAgentsInteractor::SetObservations event.
-	 * @param AgentId The agent id this data corresponds to.
-	 * @param Angle The angle currently being observed.
-	 * @param RelativeAngle The frame of reference angle.
-	 */
-	UFUNCTION(BlueprintCallable, Category = "LearningAgents", meta = (AgentId = "-1"))
-	void SetAngleObservation(const int32 AgentId, const float Angle, const float RelativeAngle = 0.0f);
+	UFUNCTION(BlueprintPure, Category = "LearningAgents")
+	FLearningAgentsObservationSchemaElement SpecifyInclusiveUnionObservation(const TMap<FName, FLearningAgentsObservationSchemaElement>& Elements, const int32 AttentionEncodingSize = 32, const int32 AttentionHeadNum = 4, const int32 ValueEncodingSize = 32, const FName Name = TEXT("InclusiveUnion"));
+	
+	UFUNCTION(BlueprintPure, Category = "LearningAgents")
+	FLearningAgentsObservationSchemaElement SpecifyInclusiveUnionObservationFromArrays(const TArray<FName>& ElementNames, const TArray<FLearningAgentsObservationSchemaElement>& Elements, const int32 AttentionEncodingSize = 32, const int32 AttentionHeadNum = 4, const int32 ValueEncodingSize = 32, const FName Name = TEXT("InclusiveUnion"));
+	FLearningAgentsObservationSchemaElement SpecifyInclusiveUnionObservationFromArrayViews(const TArrayView<const FName> ElementNames, const TArrayView<const FLearningAgentsObservationSchemaElement> Elements, const int32 AttentionEncodingSize = 32, const int32 AttentionHeadNum = 4, const int32 ValueEncodingSize = 32, const FName Name = TEXT("InclusiveUnion"));
 
-#if UE_LEARNING_AGENTS_ENABLE_VISUAL_LOG
-	/** Describes this observation to the visual logger for debugging purposes. */
-	virtual void VisualLog(const UE::Learning::FIndexSet Instances) const override;
-#endif
+	UFUNCTION(BlueprintPure, Category = "LearningAgents")
+	FLearningAgentsObservationSchemaElement SpecifyStaticArrayObservation(const FLearningAgentsObservationSchemaElement Element, const int32 Num, const FName Name = TEXT("StaticArray"));
 
-	TSharedPtr<UE::Learning::FAngleFeature> FeatureObject;
-};
+	UFUNCTION(BlueprintPure, Category = "LearningAgents")
+	FLearningAgentsObservationSchemaElement SpecifySetObservation(const FLearningAgentsObservationSchemaElement Element, const int32 MaxNum, const int32 AttentionEncodingSize = 32, const int32 AttentionHeadNum = 4, const int32 ValueEncodingSize = 32, const FName Name = TEXT("Set"));
 
-/** An observation of an array of angles. */
-UCLASS()
-class LEARNINGAGENTS_API UAngleArrayObservation : public ULearningAgentsObservation
-{
-	GENERATED_BODY()
+	UFUNCTION(BlueprintPure, Category = "LearningAgents")
+	FLearningAgentsObservationSchemaElement SpecifyPairObservation(const FLearningAgentsObservationSchemaElement Key, const FLearningAgentsObservationSchemaElement Value, const FName Name = TEXT("Pair"));
 
-public:
+	UFUNCTION(BlueprintPure, Category = "LearningAgents")
+	FLearningAgentsObservationSchemaElement SpecifyArrayObservation(const FLearningAgentsObservationSchemaElement Element, const int32 MaxNum, const int32 AttentionEncodingSize = 32, const int32 AttentionHeadNum = 4, const int32 ValueEncodingSize = 32, const FName Name = TEXT("Array"));
 
-	/**
-	 * Adds a new angle observation to the given agent interactor. Call during ULearningAgentsInteractor::SetupObservations event.
-	 * @param InInteractor The agent interactor to add this observation to.
-	 * @param Name The name of this new observation. Used for debugging.
-	 * @param AngleNum The number of angles in the array
-	 * @param Scale Used to normalize the data for the observation. Angle observations are encoded as directions. 
-	 * @return The newly created observation.
-	 */
-	UFUNCTION(BlueprintCallable, Category = "LearningAgents", meta = (DefaultToSelf = "InInteractor"))
-	static UAngleArrayObservation* AddAngleArrayObservation(ULearningAgentsInteractor* InInteractor, const FName Name = NAME_None, const int32 AngleNum = 1, const float Scale = 1.0f);
+	UFUNCTION(BlueprintPure, Category = "LearningAgents")
+	FLearningAgentsObservationSchemaElement SpecifyMapObservation(const FLearningAgentsObservationSchemaElement KeyElement, const FLearningAgentsObservationSchemaElement ValueElement, const int32 MaxNum, const int32 AttentionEncodingSize = 32, const int32 AttentionHeadNum = 4, const int32 ValueEncodingSize = 32, const FName Name = TEXT("Map"));
 
-	/**
-	 * Sets the data for this observation. Call during ULearningAgentsInteractor::SetObservations event.
-	 * @param AgentId The agent id this data corresponds to.
-	 * @param Angles The angles currently being observed.
-	 * @param RelativeAngle The frame of reference angle.
-	 */
-	UFUNCTION(BlueprintCallable, Category = "LearningAgents", meta = (AgentId = "-1"))
-	void SetAngleArrayObservation(const int32 AgentId, const TArray<float>& Angles, const float RelativeAngle = 0.0f);
+	UFUNCTION(BlueprintPure, Category = "LearningAgents")
+	FLearningAgentsObservationSchemaElement SpecifyEnumObservation(const UEnum* Enum, const FName Name = TEXT("Enum"));
 
-#if UE_LEARNING_AGENTS_ENABLE_VISUAL_LOG
-	/** Describes this observation to the visual logger for debugging purposes. */
-	virtual void VisualLog(const UE::Learning::FIndexSet Instances) const override;
-#endif
+	UFUNCTION(BlueprintPure, Category = "LearningAgents")
+	FLearningAgentsObservationSchemaElement SpecifyBitmaskObservation(const UEnum* Enum, const FName Name = TEXT("Bitmask"));
 
-	TSharedPtr<UE::Learning::FAngleFeature> FeatureObject;
-};
+	UFUNCTION(BlueprintPure, Category = "LearningAgents")
+	FLearningAgentsObservationSchemaElement SpecifyOptionalObservation(const FLearningAgentsObservationSchemaElement Element, const int32 EncodingSize = 128, const FName Name = TEXT("Optional"));
 
-/** An observation of a rotation relative to another rotation. */
-UCLASS()
-class LEARNINGAGENTS_API URotationObservation : public ULearningAgentsObservation
-{
-	GENERATED_BODY()
+	UFUNCTION(BlueprintPure, Category = "LearningAgents")
+	FLearningAgentsObservationSchemaElement SpecifyEitherObservation(const FLearningAgentsObservationSchemaElement A, const FLearningAgentsObservationSchemaElement B, const int32 EncodingSize = 128, const FName Name = TEXT("Either"));
 
-public:
+	UFUNCTION(BlueprintPure, Category = "LearningAgents")
+	FLearningAgentsObservationSchemaElement SpecifyEncodingObservation(const FLearningAgentsObservationSchemaElement Element, const int32 EncodingSize = 128, const FName Name = TEXT("Encoding"));
 
-	/**
-	 * Adds a new rotation observation to the given agent interactor. Call during ULearningAgentsInteractor::SetupObservations event.
-	 * @param InInteractor The agent interactor to add this observation to.
-	 * @param Name The name of this new observation. Used for debugging.
-	 * @param Scale Used to normalize the data for the observation. Rotation observations are encoded as directions. 
-	 * @return The newly created observation.
-	 */
-	UFUNCTION(BlueprintCallable, Category = "LearningAgents", meta = (DefaultToSelf = "InInteractor"))
-	static URotationObservation* AddRotationObservation(ULearningAgentsInteractor* InInteractor, const FName Name = NAME_None, const float Scale = 1.0f);
+	UFUNCTION(BlueprintPure, Category = "LearningAgents")
+	FLearningAgentsObservationSchemaElement SpecifyBoolObservation(const FName Name = TEXT("Bool"));
 
-	/**
-	 * Sets the data for this observation. Call during ULearningAgentsInteractor::SetObservations event.
-	 * @param AgentId The agent id this data corresponds to.
-	 * @param Rotation The rotation currently being observed.
-	 * @param RelativeRotation The frame of reference rotation.
-	 */
-	UFUNCTION(BlueprintCallable, Category = "LearningAgents", meta = (AgentId = "-1"))
-	void SetRotationObservation(const int32 AgentId, const FRotator Rotation, const FRotator RelativeRotation = FRotator::ZeroRotator);
+	UFUNCTION(BlueprintPure, Category = "LearningAgents")
+	FLearningAgentsObservationSchemaElement SpecifyFloatObservation(const FName Name = TEXT("Float"));
 
-	/**
-	 * Sets the data for this observation. Call during ULearningAgentsInteractor::SetObservations event.
-	 * @param AgentId The agent id this data corresponds to.
-	 * @param Rotation The rotation currently being observed.
-	 * @param RelativeRotation The frame of reference rotation.
-	 */
-	UFUNCTION(BlueprintCallable, Category = "LearningAgents", meta = (AgentId = "-1"))
-	void SetRotationObservationFromQuat(const int32 AgentId, const FQuat Rotation, const FQuat RelativeRotation);
+	UFUNCTION(BlueprintPure, Category = "LearningAgents")
+	FLearningAgentsObservationSchemaElement SpecifyTranslationObservation(const FName Name = TEXT("Translation"));
 
-#if UE_LEARNING_AGENTS_ENABLE_VISUAL_LOG
-	/** Describes this observation to the visual logger for debugging purposes. */
-	virtual void VisualLog(const UE::Learning::FIndexSet Instances) const override;
-#endif
+	UFUNCTION(BlueprintPure, Category = "LearningAgents")
+	FLearningAgentsObservationSchemaElement SpecifyRotationObservation(const FName Name = TEXT("Rotation"));
 
-	TSharedPtr<UE::Learning::FRotationFeature> FeatureObject;
-};
+	UFUNCTION(BlueprintPure, Category = "LearningAgents")
+	FLearningAgentsObservationSchemaElement SpecifyScaleObservation(const FName Name = TEXT("Scale"));
 
-/** An observation of an array of rotations. */
-UCLASS()
-class LEARNINGAGENTS_API URotationArrayObservation : public ULearningAgentsObservation
-{
-	GENERATED_BODY()
+	UFUNCTION(BlueprintPure, Category = "LearningAgents")
+	FLearningAgentsObservationSchemaElement SpecifyTransformObservation(const FName Name = TEXT("Transform"));
 
-public:
+	UFUNCTION(BlueprintPure, Category = "LearningAgents")
+	FLearningAgentsObservationSchemaElement SpecifyAngleObservation(const FName Name = TEXT("Angle"));
 
-	/**
-	 * Adds a new rotation observation to the given agent interactor. Call during ULearningAgentsInteractor::SetupObservations event.
-	 * @param InInteractor The agent interactor to add this observation to.
-	 * @param Name The name of this new observation. Used for debugging.
-	 * @param RotationNum The number of rotations in the array
-	 * @param Scale Used to normalize the data for the observation. Rotation observations are encoded as directions. 
-	 * @return The newly created observation.
-	 */
-	UFUNCTION(BlueprintCallable, Category = "LearningAgents", meta = (DefaultToSelf = "InInteractor"))
-	static URotationArrayObservation* AddRotationArrayObservation(ULearningAgentsInteractor* InInteractor, const FName Name = NAME_None, const int32 RotationNum = 1, const float Scale = 1.0f);
+	UFUNCTION(BlueprintPure, Category = "LearningAgents")
+	FLearningAgentsObservationSchemaElement SpecifyVelocityObservation(const FName Name = TEXT("Velocity"));
 
-	/**
-	 * Sets the data for this observation. Call during ULearningAgentsInteractor::SetObservations event.
-	 * @param AgentId The agent id this data corresponds to.
-	 * @param Rotations The rotations currently being observed.
-	 * @param RelativeRotation The frame of reference rotation.
-	 */
-	UFUNCTION(BlueprintCallable, Category = "LearningAgents", meta = (AgentId = "-1"))
-	void SetRotationArrayObservation(const int32 AgentId, const TArray<FRotator>& Rotations, const FRotator RelativeRotation = FRotator::ZeroRotator);
+	UFUNCTION(BlueprintPure, Category = "LearningAgents")
+	FLearningAgentsObservationSchemaElement SpecifyDirectionObservation(const FName Name = TEXT("Direction"));
 
-	/**
-	 * Sets the data for this observation. Call during ULearningAgentsInteractor::SetObservations event.
-	 * @param AgentId The agent id this data corresponds to.
-	 * @param Rotations The rotations currently being observed.
-	 * @param RelativeRotation The frame of reference rotation.
-	 */
-	UFUNCTION(BlueprintCallable, Category = "LearningAgents", meta = (AgentId = "-1"))
-	void SetRotationArrayObservationFromQuats(const int32 AgentId, const TArray<FQuat>& Rotations, const FQuat RelativeRotation);
+	// Spline Observations
 
-#if UE_LEARNING_AGENTS_ENABLE_VISUAL_LOG
-	/** Describes this observation to the visual logger for debugging purposes. */
-	virtual void VisualLog(const UE::Learning::FIndexSet Instances) const override;
-#endif
+	UFUNCTION(BlueprintPure, Category = "LearningAgents")
+	FLearningAgentsObservationSchemaElement SpecifyLocationAlongSplineObservation(const FName Name = TEXT("LocationAlongSpline"));
 
-	TSharedPtr<UE::Learning::FRotationFeature> FeatureObject;
-};
+	UFUNCTION(BlueprintPure, Category = "LearningAgents")
+	FLearningAgentsObservationSchemaElement SpecifyProportionAlongSplineObservation(const FName Name = TEXT("ProportionAlongSpline"));
 
-//------------------------------------------------------------------
+	UFUNCTION(BlueprintPure, Category = "LearningAgents")
+	FLearningAgentsObservationSchemaElement SpecifyDirectionAlongSplineObservation(const FName Name = TEXT("DirectionAlongSpline"));
 
-/** An observation of a direction vector. */
-UCLASS()
-class LEARNINGAGENTS_API UDirectionObservation : public ULearningAgentsObservation
-{
-	GENERATED_BODY()
+	UFUNCTION(BlueprintPure, Category = "LearningAgents")
+	FLearningAgentsObservationSchemaElement SpecifyPropertiesAlongSplineObservation(const FName Name = TEXT("PropertiesAlongSpline"));
 
-public:
+	// Ray Cast Observations
 
-	/**
-	 * Adds a new direction observation to the given agent interactor. Call during ULearningAgentsInteractor::SetupObservations event.
-	 * @param InInteractor The agent interactor to add this observation to.
-	 * @param Name The name of this new observation. Used for debugging.
-	 * @param Scale Used to normalize the data for the observation.
-	 * @return The newly created observation.
-	 */
-	UFUNCTION(BlueprintCallable, Category = "LearningAgents", meta = (DefaultToSelf = "InInteractor"))
-	static UDirectionObservation* AddDirectionObservation(ULearningAgentsInteractor* InInteractor, const FName Name = NAME_None, const float Scale = 1.0f);
+	UFUNCTION(BlueprintPure, Category = "LearningAgents")
+	FLearningAgentsObservationSchemaElement SpecifyProportionAlongRayObservation(const FName Name = TEXT("ProportionAlongRay"));
 
-	/**
-	 * Sets the data for this observation. The relative rotation can be used to make this observation relative to the
-	 * agent's perspective, e.g. by passing the agent's forward rotation. Call during ULearningAgentsInteractor::SetObservations event.
-	 * @param AgentId The agent id this data corresponds to.
-	 * @param Direction The direction currently being observed.
-	 * @param RelativeRotation The frame of reference rotation.
-	 */
-	UFUNCTION(BlueprintCallable, Category = "LearningAgents", meta = (AgentId = "-1"))
-	void SetDirectionObservation(const int32 AgentId, const FVector Direction, const FRotator RelativeRotation = FRotator::ZeroRotator);
+	UFUNCTION(BlueprintPure, Category = "LearningAgents")
+	FLearningAgentsObservationSchemaElement SpecifyProportionAlongRaysObservation(const int32 Num, const FName Name = TEXT("ProportionAlongRays"));
 
-#if UE_LEARNING_AGENTS_ENABLE_VISUAL_LOG
-	/** Describes this observation to the visual logger for debugging purposes. */
-	virtual void VisualLog(const UE::Learning::FIndexSet Instances) const override;
-#endif
+private:
 
-	TSharedPtr<UE::Learning::FDirectionFeature> FeatureObject;
+	UE::Learning::Observation::FSchema ObservationSchema;
 };
 
-/** An observation of an array of direction vectors. */
-UCLASS()
-class LEARNINGAGENTS_API UDirectionArrayObservation : public ULearningAgentsObservation
+/**
+ * Observation Object
+ *
+ * This object is used to construct or get the values of observations.
+ */
+UCLASS(BlueprintType)
+class LEARNINGAGENTS_API ULearningAgentsObservationObject : public UObject
 {
 	GENERATED_BODY()
 
 public:
 
-	/**
-	 * Adds a new direction array observation to the given agent interactor. Call during ULearningAgentsInteractor::SetupObservations event.
-	 * @param InInteractor The agent interactor to add this observation to.
-	 * @param Name The name of this new observation. Used for debugging.
-	 * @param DirectionNum The number of directions in the array
-	 * @param Scale Used to normalize the data for the observation.
-	 * @return The newly created observation.
-	 */
-	UFUNCTION(BlueprintCallable, Category = "LearningAgents", meta = (DefaultToSelf = "InInteractor"))
-	static UDirectionArrayObservation* AddDirectionArrayObservation(ULearningAgentsInteractor* InInteractor, const FName Name = NAME_None, const int32 DirectionNum = 1, const float Scale = 1.0f);
-
-	/**
-	 * Sets the data for this observation. The relative rotation can be used to make this observation relative to the
-	 * agent's perspective, e.g. by passing the agent's forward rotation. Call during ULearningAgentsInteractor::SetObservations event.
-	 * @param AgentId The agent id this data corresponds to.
-	 * @param Directions The directions currently being observed.
-	 * @param RelativeRotation The frame of reference rotation.
-	 */
-	UFUNCTION(BlueprintCallable, Category = "LearningAgents", meta = (AgentId = "-1"))
-	void SetDirectionArrayObservation(const int32 AgentId, const TArray<FVector>& Directions, const FRotator RelativeRotation = FRotator::ZeroRotator);
+	/** Gets the internal observation object */
+	const UE::Learning::Observation::FObject& GetObservationObject() const;
 
-#if UE_LEARNING_AGENTS_ENABLE_VISUAL_LOG
-	/** Describes this observation to the visual logger for debugging purposes. */
-	virtual void VisualLog(const UE::Learning::FIndexSet Instances) const override;
-#endif
-
-	TSharedPtr<UE::Learning::FDirectionFeature> FeatureObject;
-};
-
-/** An observation of a direction vector projected onto a plane. */
-UCLASS()
-class LEARNINGAGENTS_API UPlanarDirectionObservation : public ULearningAgentsObservation
-{
-	GENERATED_BODY()
+	/** Gets the internal observation object */
+	UE::Learning::Observation::FObject& GetObservationObject();
 
 public:
 
-	/**
-	 * Adds a new planar direction observation to the given agent interactor. The axis parameters define the plane. Call
-	 * during ULearningAgentsInteractor::SetupObservations event.
-	 * @param InInteractor The agent interactor to add this observation to.
-	 * @param Name The name of this new observation. Used for debugging.
-	 * @param Scale Used to normalize the data for the observation.
-	 * @param Axis0 The forward axis of the plane.
-	 * @param Axis1 The right axis of the plane.
-	 * @return The newly created observation.
-	 */
-	UFUNCTION(BlueprintCallable, Category = "LearningAgents", meta = (DefaultToSelf = "InInteractor"))
-	static UPlanarDirectionObservation* AddPlanarDirectionObservation(ULearningAgentsInteractor* InInteractor, const FName Name = NAME_None, const float Scale = 1.0f, const FVector Axis0 = FVector::ForwardVector, const FVector Axis1 = FVector::RightVector);
-
-	/**
-	 * Sets the data for this observation. The relative rotation can be used to make this observation relative to the
-	 * agent's perspective, e.g. by passing the agent's forward rotation. Call during ULearningAgentsInteractor::SetObservations event.
-	 * @param AgentId The agent id this data corresponds to.
-	 * @param Direction The direction currently being observed.
-	 * @param RelativeRotation The frame of reference rotation.
-	 */
-	UFUNCTION(BlueprintCallable, Category = "LearningAgents", meta = (AgentId = "-1"))
-	void SetPlanarDirectionObservation(const int32 AgentId, const FVector Direction, const FRotator RelativeRotation = FRotator::ZeroRotator);
-
-#if UE_LEARNING_AGENTS_ENABLE_VISUAL_LOG
-	/** Describes this observation to the visual logger for debugging purposes. */
-	virtual void VisualLog(const UE::Learning::FIndexSet Instances) const override;
-#endif
-
-	TSharedPtr<UE::Learning::FPlanarDirectionFeature> FeatureObject;
-};
-
-/** An observation of an array of direction vectors projected onto a plane. */
-UCLASS()
-class LEARNINGAGENTS_API UPlanarDirectionArrayObservation : public ULearningAgentsObservation
-{
-	GENERATED_BODY()
+	UFUNCTION(BlueprintPure = false, Category = "LearningAgents")
+	void LogObservation(const FLearningAgentsObservationObjectElement Element);
 
 public:
 
-	/**
-	 * Adds a new planar direction observation to the given agent interactor. The axis parameters define the plane. Call
-	 * during ULearningAgentsInteractor::SetupObservations event.
-	 * @param InInteractor The agent interactor to add this observation to.
-	 * @param Name The name of this new observation. Used for debugging.
-	 * @param DirectionNum The number of directions in the array
-	 * @param Scale Used to normalize the data for the observation.
-	 * @param Axis0 The forward axis of the plane.
-	 * @param Axis1 The right axis of the plane.
-	 * @return The newly created observation.
-	 */
-	UFUNCTION(BlueprintCallable, Category = "LearningAgents", meta = (DefaultToSelf = "InInteractor"))
-	static UPlanarDirectionArrayObservation* AddPlanarDirectionArrayObservation(ULearningAgentsInteractor* InInteractor, const FName Name = NAME_None, const int32 DirectionNum = 1, const float Scale = 1.0f, const FVector Axis0 = FVector::ForwardVector, const FVector Axis1 = FVector::RightVector);
+	// Basic Observations
 
-	/**
-	 * Sets the data for this observation. The relative rotation can be used to make this observation relative to the
-	 * agent's perspective, e.g. by passing the agent's forward rotation. Call during ULearningAgentsInteractor::SetObservations event.
-	 * @param AgentId The agent id this data corresponds to.
-	 * @param Directions The directions currently being observed.
-	 * @param RelativeRotation The frame of reference rotation.
-	 */
-	UFUNCTION(BlueprintCallable, Category = "LearningAgents", meta = (AgentId = "-1"))
-	void SetPlanarDirectionArrayObservation(const int32 AgentId, const TArray<FVector>& Directions, const FRotator RelativeRotation = FRotator::ZeroRotator);
+	UFUNCTION(BlueprintPure, Category = "LearningAgents")
+	FLearningAgentsObservationObjectElement MakeNullObservation(const FName Name = TEXT("Null"));
 
-#if UE_LEARNING_AGENTS_ENABLE_VISUAL_LOG
-	/** Describes this observation to the visual logger for debugging purposes. */
-	virtual void VisualLog(const UE::Learning::FIndexSet Instances) const override;
-#endif
+	UFUNCTION(BlueprintPure, Category = "LearningAgents")
+	FLearningAgentsObservationObjectElement MakeContinuousObservation(const TArray<float>& Values, const FName Name = TEXT("Continuous"));
+	FLearningAgentsObservationObjectElement MakeContinuousObservationFromArrayView(const TArrayView<const float> Values, const FName Name = TEXT("Continuous"));
 
-	TSharedPtr<UE::Learning::FPlanarDirectionFeature> FeatureObject;
-};
+	UFUNCTION(BlueprintPure, Category = "LearningAgents")
+	FLearningAgentsObservationObjectElement MakeExclusiveDiscreteObservation(const int32 DiscreteIndex, const int32 Size, const FName Name = TEXT("ExclusiveDiscrete"));
 
-//------------------------------------------------------------------
+	UFUNCTION(BlueprintPure, Category = "LearningAgents")
+	FLearningAgentsObservationObjectElement MakeInclusiveDiscreteObservation(const TArray<int32>& DiscreteIndices, const int32 Size, const FName Name = TEXT("InclusiveDiscrete"));
+	FLearningAgentsObservationObjectElement MakeInclusiveDiscreteObservationFromArrayView(const TArrayView<const int32> DiscreteIndices, const int32 Size, const FName Name = TEXT("InclusiveDiscrete"));
 
-/** An observation of a position vector. */
-UCLASS()
-class LEARNINGAGENTS_API UPositionObservation : public ULearningAgentsObservation
-{
-	GENERATED_BODY()
+	UFUNCTION(BlueprintPure, Category = "LearningAgents")
+	FLearningAgentsObservationObjectElement MakeIndexObservation(const int32 Index, const int32 Size, const FName Name = TEXT("Index"));
 
-public:
+	UFUNCTION(BlueprintPure, Category = "LearningAgents")
+	FLearningAgentsObservationObjectElement MakeStructObservation(const TMap<FName, FLearningAgentsObservationObjectElement>& Elements, const FName Name = TEXT("Struct"));
+	
+	UFUNCTION(BlueprintPure, Category = "LearningAgents")
+	FLearningAgentsObservationObjectElement MakeStructObservationFromArrays(const TArray<FName>& ElementNames, const TArray<FLearningAgentsObservationObjectElement>& Elements, const FName Name = TEXT("Struct"));
+	FLearningAgentsObservationObjectElement MakeStructObservationFromArrayViews(const TArrayView<const FName> ElementNames, const TArrayView<const FLearningAgentsObservationObjectElement> Elements, const FName Name = TEXT("Struct"));
 
-	/**
-	 * Adds a new position observation to the given agent interactor. Call during ULearningAgentsInteractor::SetupObservations event.
-	 * @param InInteractor The agent interactor to add this observation to.
-	 * @param Name The name of this new observation. Used for debugging.
-	 * @param Scale Used to normalize the data for the observation.
-	 * @return The newly created observation.
-	 */
-	UFUNCTION(BlueprintCallable, Category = "LearningAgents", meta = (DefaultToSelf = "InInteractor"))
-	static UPositionObservation* AddPositionObservation(ULearningAgentsInteractor* InInteractor, const FName Name = NAME_None, const float Scale = 100.0f);
+	UFUNCTION(BlueprintPure, Category = "LearningAgents")
+	FLearningAgentsObservationObjectElement MakeExclusiveUnionObservation(const FName ElementName, const FLearningAgentsObservationObjectElement Element, const FName Name = TEXT("ExclusiveUnion"));
 
-	/**
-	 * Sets the data for this observation. The relative position & rotation can be used to make this observation
-	 * relative to the agent's perspective, e.g. by passing the agent's position & forward rotation. Call during
-	 * ULearningAgentsInteractor::SetObservations event.
-	 * @param AgentId The agent id this data corresponds to.
-	 * @param Position The position currently being observed.
-	 * @param RelativePosition The vector Position will be offset from.
-	 * @param RelativeRotation The frame of reference rotation.
-	 */
-	UFUNCTION(BlueprintCallable, Category = "LearningAgents", meta = (AgentId = "-1"))
-	void SetPositionObservation(const int32 AgentId, const FVector Position, const FVector RelativePosition = FVector::ZeroVector, const FRotator RelativeRotation = FRotator::ZeroRotator);
+	UFUNCTION(BlueprintPure, Category = "LearningAgents")
+	FLearningAgentsObservationObjectElement MakeInclusiveUnionObservation(const TMap<FName, FLearningAgentsObservationObjectElement>& Elements, const FName Name = TEXT("InclusiveUnion"));
+	
+	UFUNCTION(BlueprintPure, Category = "LearningAgents")
+	FLearningAgentsObservationObjectElement MakeInclusiveUnionObservationFromArrays(const TArray<FName>& ElementNames, const TArray<FLearningAgentsObservationObjectElement>& Elements, const FName Name = TEXT("InclusiveUnion"));
+	FLearningAgentsObservationObjectElement MakeInclusiveUnionObservationFromArrayViews(const TArrayView<const FName> ElementNames, const TArrayView<const FLearningAgentsObservationObjectElement> Elements, const FName Name = TEXT("InclusiveUnion"));
 
-#if UE_LEARNING_AGENTS_ENABLE_VISUAL_LOG
-	/** Describes this observation to the visual logger for debugging purposes. */
-	virtual void VisualLog(const UE::Learning::FIndexSet Instances) const override;
-#endif
+	UFUNCTION(BlueprintPure, Category = "LearningAgents")
+	FLearningAgentsObservationObjectElement MakeStaticArrayObservation(const TArray<FLearningAgentsObservationObjectElement>& Elements, const FName Name = TEXT("StaticArray"));
+	FLearningAgentsObservationObjectElement MakeStaticArrayObservationFromArrayView(const TArrayView<const FLearningAgentsObservationObjectElement> Elements, const FName Name = TEXT("StaticArray"));
 
-	TSharedPtr<UE::Learning::FPositionFeature> FeatureObject;
-};
+	UFUNCTION(BlueprintPure, Category = "LearningAgents")
+	FLearningAgentsObservationObjectElement MakeSetObservation(const TSet<FLearningAgentsObservationObjectElement>& Elements, const FName Name = TEXT("Set"));
+	
+	UFUNCTION(BlueprintPure, Category = "LearningAgents")
+	FLearningAgentsObservationObjectElement MakeSetObservationFromArray(const TArray<FLearningAgentsObservationObjectElement>& Elements, const FName Name = TEXT("Set"));
+	FLearningAgentsObservationObjectElement MakeSetObservationFromArrayView(const TArrayView<const FLearningAgentsObservationObjectElement> Elements, const FName Name = TEXT("Set"));
 
-/** An observation of an array of positions. */
-UCLASS()
-class LEARNINGAGENTS_API UPositionArrayObservation : public ULearningAgentsObservation
-{
-	GENERATED_BODY()
+	UFUNCTION(BlueprintPure, Category = "LearningAgents")
+	FLearningAgentsObservationObjectElement MakePairObservation(const FLearningAgentsObservationObjectElement Key, const FLearningAgentsObservationObjectElement Value, const FName Name = TEXT("Pair"));
 
-public:
+	UFUNCTION(BlueprintPure, Category = "LearningAgents")
+	FLearningAgentsObservationObjectElement MakeArrayObservation(const TArray<FLearningAgentsObservationObjectElement>& Elements, const FName Name = TEXT("Array"));
+	FLearningAgentsObservationObjectElement MakeArrayObservationFromArrayView(const TArrayView<const FLearningAgentsObservationObjectElement> Elements, const FName Name = TEXT("Array"));
 
-	/**
-	 * Adds a new position array observation to the given agent interactor. Call during ULearningAgentsInteractor::SetupObservations event.
-	 * @param InInteractor The agent interactor to add this observation to.
-	 * @param Name The name of this new observation. Used for debugging.
-	 * @param PositionNum The number of positions in the array.
-	 * @param Scale Used to normalize the data for the observation.
-	 * @return The newly created observation.
-	 */
-	UFUNCTION(BlueprintCallable, Category = "LearningAgents", meta = (DefaultToSelf = "InInteractor"))
-	static UPositionArrayObservation* AddPositionArrayObservation(ULearningAgentsInteractor* InInteractor, const FName Name = NAME_None, const int32 PositionNum = 1, const float Scale = 100.0f);
+	UFUNCTION(BlueprintPure, Category = "LearningAgents")
+	FLearningAgentsObservationObjectElement MakeMapObservation(const TMap<FLearningAgentsObservationObjectElement, FLearningAgentsObservationObjectElement>& Map, const FName Name = TEXT("Map"));
+	
+	UFUNCTION(BlueprintPure, Category = "LearningAgents")
+	FLearningAgentsObservationObjectElement MakeMapObservationFromArrays(const TArray<FLearningAgentsObservationObjectElement>& Keys, const TArray<FLearningAgentsObservationObjectElement>& Values, const FName Name = TEXT("Map"));
+	FLearningAgentsObservationObjectElement MakeMapObservationFromArrayViews(const TArrayView<const FLearningAgentsObservationObjectElement> Keys, const TArrayView<const FLearningAgentsObservationObjectElement> Values, const FName Name = TEXT("Map"));
 
-	/**
-	 * Sets the data for this observation. The relative position & rotation can be used to make this observation
-	 * relative to the agent's perspective, e.g. by passing the agent's position & forward rotation. Call during
-	 * ULearningAgentsInteractor::SetObservations event.
-	 * @param AgentId The agent id this data corresponds to.
-	 * @param Positions The positions currently being observed.
-	 * @param RelativePosition The vector Positions will be offset from.
-	 * @param RelativeRotation The frame of reference rotation.
-	 */
-	UFUNCTION(BlueprintCallable, Category = "LearningAgents", meta = (AgentId = "-1"))
-	void SetPositionArrayObservation(int32 AgentId, const TArray<FVector>& Positions, const FVector RelativePosition = FVector::ZeroVector, const FRotator RelativeRotation = FRotator::ZeroRotator);
-
-#if UE_LEARNING_AGENTS_ENABLE_VISUAL_LOG
-	/** Describes this observation to the visual logger for debugging purposes. */
-	virtual void VisualLog(const UE::Learning::FIndexSet Instances) const override;
-#endif
-
-	TSharedPtr<UE::Learning::FPositionFeature> FeatureObject;
-};
+	UFUNCTION(BlueprintPure, Category = "LearningAgents")
+	FLearningAgentsObservationObjectElement MakeEnumObservation(const UEnum* Enum, const uint8 EnumValue, const FName Name = TEXT("Enum"));
 
-/** An observation of a position along a single axis. Can be useful for providing information like object heights. */
-UCLASS()
-class LEARNINGAGENTS_API UScalarPositionObservation : public ULearningAgentsObservation
-{
-	GENERATED_BODY()
-
-public:
-
-	/**
-	 * Adds a new scalar position observation to the given agent interactor. Call during 
-	 * ULearningAgentsInteractor::SetupObservations event.
-	 * @param InInteractor The agent interactor to add this observation to.
-	 * @param Name The name of this new observation. Used for debugging.
-	 * @param Scale Used to normalize the data for the observation.
-	 * @return The newly created observation.
-	 */
-	UFUNCTION(BlueprintCallable, Category = "LearningAgents", meta = (DefaultToSelf = "InInteractor"))
-	static UScalarPositionObservation* AddScalarPositionObservation(ULearningAgentsInteractor* InInteractor, const FName Name = NAME_None, const float Scale = 100.0f);
+	UFUNCTION(BlueprintPure, Category = "LearningAgents")
+	FLearningAgentsObservationObjectElement MakeBitmaskObservation(const UEnum* Enum, const int32 BitmaskValue, const FName Name = TEXT("Bitmask"));
 
-	/**
-	 * Sets the data for this observation. The relative position can be used to make this observation
-	 * relative to the agent's perspective, e.g. by passing the agent's position. Call during 
-	 * ULearningAgentsInteractor::SetObservations event.
-	 * @param AgentId The agent id this data corresponds to.
-	 * @param Position The position currently being observed.
-	 * @param RelativePosition The vector Position will be offset from.
-	 */
-	UFUNCTION(BlueprintCallable, Category = "LearningAgents", meta = (AgentId = "-1"))
-	void SetScalarPositionObservation(const int32 AgentId, const float Position, const float RelativePosition = 0.0f);
+	UFUNCTION(BlueprintPure, Category = "LearningAgents")
+	FLearningAgentsObservationObjectElement MakeOptionalObservation(const FLearningAgentsObservationObjectElement Element, const ELearningAgentsOptionalObservation Option, const FName Name = TEXT("Optional"));
 
-	/**
-	 * Sets the data for this observation. The relative position can be used to make this observation
-	 * relative to the agent's perspective, e.g. by passing the agent's position. Call during
-	 * ULearningAgentsInteractor::SetObservations event.
-	 * @param AgentId The agent id this data corresponds to.
-	 * @param Position The position currently being observed.
-	 * @param RelativePosition The vector Position will be offset from.
-	 * @param Axis The axis along which to encode the positions
-	 */
-	UFUNCTION(BlueprintCallable, Category = "LearningAgents", meta = (AgentId = "-1"))
-	void SetScalarPositionObservationWithAxis(const int32 AgentId, const FVector Position, const FVector RelativePosition = FVector::ZeroVector, const FVector Axis = FVector::UpVector);
+	UFUNCTION(BlueprintPure, Category = "LearningAgents")
+	FLearningAgentsObservationObjectElement MakeOptionalNullObservation(const FName Name = TEXT("Optional"));
 
-#if UE_LEARNING_AGENTS_ENABLE_VISUAL_LOG
-	/** Describes this observation to the visual logger for debugging purposes. */
-	virtual void VisualLog(const UE::Learning::FIndexSet Instances) const override;
-#endif
+	UFUNCTION(BlueprintPure, Category = "LearningAgents")
+	FLearningAgentsObservationObjectElement MakeOptionalValidObservation(const FLearningAgentsObservationObjectElement Element, const FName Name = TEXT("Optional"));
 
-	TSharedPtr<UE::Learning::FScalarPositionFeature> FeatureObject;
-};
+	UFUNCTION(BlueprintPure, Category = "LearningAgents")
+	FLearningAgentsObservationObjectElement MakeEitherObservation(const FLearningAgentsObservationObjectElement Element, const ELearningAgentsEitherObservation Either, const FName Name = TEXT("Either"));
 
-/** An observation of an array of positions along a single axis. */
-UCLASS()
-class LEARNINGAGENTS_API UScalarPositionArrayObservation : public ULearningAgentsObservation
-{
-	GENERATED_BODY()
+	UFUNCTION(BlueprintPure, Category = "LearningAgents", meta=(DisplayName="Make Either A Observation"))
+	FLearningAgentsObservationObjectElement MakeEitherAObservation(const FLearningAgentsObservationObjectElement A, const FName Name = TEXT("Either"));
 
-public:
+	UFUNCTION(BlueprintPure, Category = "LearningAgents", meta = (DisplayName = "Make Either B Observation"))
+	FLearningAgentsObservationObjectElement MakeEitherBObservation(const FLearningAgentsObservationObjectElement B, const FName Name = TEXT("Either"));
 
-	/**
-	 * Adds a new scalar position array observation to the given agent interactor. Call during 
-	 * ULearningAgentsInteractor::SetupObservations event.
-	 * @param InInteractor The agent interactor to add this observation to.
-	 * @param Name The name of this new observation. Used for debugging.
-	 * @param PositionNum The number of positions in the array.
-	 * @param Scale Used to normalize the data for the observation.
-	 * @return The newly created observation.
-	 */
-	UFUNCTION(BlueprintCallable, Category = "LearningAgents", meta = (DefaultToSelf = "InInteractor"))
-	static UScalarPositionArrayObservation* AddScalarPositionArrayObservation(ULearningAgentsInteractor* InInteractor, const FName Name = NAME_None, const int32 PositionNum = 1, const float Scale = 100.0f);
+	UFUNCTION(BlueprintPure, Category = "LearningAgents")
+	FLearningAgentsObservationObjectElement MakeEncodingObservation(const FLearningAgentsObservationObjectElement Element, const FName Name = TEXT("Encoding"));
 
-	/**
-	 * Sets the data for this observation. The relative position can be used to make this observation
-	 * relative to the agent's perspective, e.g. by passing the agent's position. Call during
-	 * ULearningAgentsInteractor::SetObservations event.
-	 * @param AgentId The agent id this data corresponds to.
-	 * @param Positions The positions currently being observed.
-	 * @param RelativePosition The vector Positions will be offset from.
-	 */
-	UFUNCTION(BlueprintCallable, Category = "LearningAgents", meta = (AgentId = "-1"))
-	void SetScalarPositionArrayObservation(const int32 AgentId, const TArray<float>& Positions, const float RelativePosition = 0.0f);
+	UFUNCTION(BlueprintPure, Category = "LearningAgents")
+	FLearningAgentsObservationObjectElement MakeBoolObservation(const bool bValue, const FName Name = TEXT("Bool"));
 
-	/**
-	 * Sets the data for this observation. The relative position can be used to make this observation
-	 * relative to the agent's perspective, e.g. by passing the agent's position. Call during
-	 * ULearningAgentsInteractor::SetObservations event.
-	 * @param AgentId The agent id this data corresponds to.
-	 * @param Positions The positions currently being observed.
-	 * @param RelativePosition The vector Positions will be offset from.
-	 * @param Axis The axis along which to encode the positions
-	 */
-	UFUNCTION(BlueprintCallable, Category = "LearningAgents", meta = (AgentId = "-1"))
-	void SetScalarPositionArrayObservationWithAxis(const int32 AgentId, const TArray<FVector>& Positions, const FVector RelativePosition = FVector::ZeroVector, const FVector Axis = FVector::UpVector);
+	UFUNCTION(BlueprintPure, Category = "LearningAgents")
+	FLearningAgentsObservationObjectElement MakeFloatObservation(const float Value, const float FloatScale = 1.0f, const FName Name = TEXT("Float"));
 
-#if UE_LEARNING_AGENTS_ENABLE_VISUAL_LOG
-	/** Describes this observation to the visual logger for debugging purposes. */
-	virtual void VisualLog(const UE::Learning::FIndexSet Instances) const override;
-#endif
+	UFUNCTION(BlueprintPure, Category = "LearningAgents")
+	FLearningAgentsObservationObjectElement MakeTranslationObservation(const FVector Translation, const FTransform RelativeTransform = FTransform(), const float TranslationScale = 100.0f, const FName Name = TEXT("Translation"));
 
-	TSharedPtr<UE::Learning::FScalarPositionFeature> FeatureObject;
-};
+	UFUNCTION(BlueprintPure, Category = "LearningAgents")
+	FLearningAgentsObservationObjectElement MakeRotationObservation(const FRotator Rotation, const FRotator RelativeRotation = FRotator::ZeroRotator, const FName Name = TEXT("Rotation"));
+	FLearningAgentsObservationObjectElement MakeRotationObservationFromQuat(const FQuat Rotation, const FQuat RelativeRotation = FQuat::Identity, const FName Name = TEXT("Rotation"));
 
-/** An observation of a position projected onto a plane. */
-UCLASS()
-class LEARNINGAGENTS_API UPlanarPositionObservation : public ULearningAgentsObservation
-{
-	GENERATED_BODY()
+	UFUNCTION(BlueprintPure, Category = "LearningAgents")
+	FLearningAgentsObservationObjectElement MakeScaleObservation(const FVector Scale, const FVector RelativeScale = FVector(1,1,1), const FName Name = TEXT("Scale"));
 
-public:
+	UFUNCTION(BlueprintPure, Category = "LearningAgents")
+	FLearningAgentsObservationObjectElement MakeTransformObservation(const FTransform Transform, const FTransform RelativeTransform = FTransform(), const float TranslationScale = 100.0f, const FName Name = TEXT("Transform"));
 
-	/**
-	 * Adds a new planar position observation to the given agent interactor. The axis parameters define the plane. Call
-	 * during ULearningAgentsInteractor::SetupObservations event.
-	 * @param InInteractor The agent interactor to add this observation to.
-	 * @param Name The name of this new observation. Used for debugging.
-	 * @param Scale Used to normalize the data for the observation.
-	 * @param Axis0 The forward axis of the plane.
-	 * @param Axis1 The right axis of the plane.
-	 * @return The newly created observation.
-	 */
-	UFUNCTION(BlueprintCallable, Category = "LearningAgents", meta = (DefaultToSelf = "InInteractor"))
-	static UPlanarPositionObservation* AddPlanarPositionObservation(ULearningAgentsInteractor* InInteractor, const FName Name = NAME_None, const float Scale = 100.0f, const FVector Axis0 = FVector::ForwardVector, const FVector Axis1 = FVector::RightVector);
+	UFUNCTION(BlueprintPure, Category = "LearningAgents")
+	FLearningAgentsObservationObjectElement MakeAngleObservation(const float Angle, const float RelativeAngle = 0.0f, const FName Name = TEXT("Angle"));
 
-	/**
-	 * Sets the data for this observation. The relative position & rotation can be used to make this observation
-	 * relative to the agent's perspective, e.g. by passing the agent's position & forward rotation. Call during
-	 * ULearningAgentsInteractor::SetObservations event.
-	 * @param AgentId The agent id this data corresponds to.
-	 * @param Position The position currently being observed.
-	 * @param RelativePosition The vector Position will be offset from.
-	 * @param RelativeRotation The frame of reference rotation.
-	 */
-	UFUNCTION(BlueprintCallable, Category = "LearningAgents", meta = (AgentId = "-1"))
-	void SetPlanarPositionObservation(const int32 AgentId, const FVector Position, const FVector RelativePosition = FVector::ZeroVector, const FRotator RelativeRotation = FRotator::ZeroRotator);
+	UFUNCTION(BlueprintPure, Category = "LearningAgents")
+	FLearningAgentsObservationObjectElement MakeAngleObservationRadians(const float Angle, const float RelativeAngle = 0.0f, const FName Name = TEXT("Angle"));
 
-#if UE_LEARNING_AGENTS_ENABLE_VISUAL_LOG
-	/** Describes this observation to the visual logger for debugging purposes. */
-	virtual void VisualLog(const UE::Learning::FIndexSet Instances) const override;
-#endif
+	UFUNCTION(BlueprintPure, Category = "LearningAgents")
+	FLearningAgentsObservationObjectElement MakeVelocityObservation(const FVector Velocity, const FTransform RelativeTransform = FTransform(), const float VelocityScale = 200.0f, const FName Name = TEXT("Velocity"));
 
-	TSharedPtr<UE::Learning::FPlanarPositionFeature> FeatureObject;
-};
+	UFUNCTION(BlueprintPure, Category = "LearningAgents")
+	FLearningAgentsObservationObjectElement MakeDirectionObservation(const FVector Direction, const FTransform RelativeTransform = FTransform(), const FName Name = TEXT("Direction"));
 
-/** An observation of an array of positions projected onto a plane. */
-UCLASS()
-class LEARNINGAGENTS_API UPlanarPositionArrayObservation : public ULearningAgentsObservation
-{
-	GENERATED_BODY()
+	// Spline Observations
 
-public:
+	UFUNCTION(BlueprintPure, Category = "LearningAgents")
+	FLearningAgentsObservationObjectElement MakeLocationAlongSplineObservation(const USplineComponent* SplineComponent, const float DistanceAlongSpline, const FTransform RelativeTransform = FTransform(), const float TranslationScale = 100.0f, const FName Name = TEXT("LocationAlongSpline"));
 
-	/**
-	 * Adds a new planar position array observation to the given agent interactor. The axis parameters define the plane.
-	 * Call during ULearningAgentsInteractor::SetupObservations event.
-	 * @param InInteractor The agent interactor to add this observation to.
-	 * @param Name The name of this new observation. Used for debugging.
-	 * @param PositionNum The number of positions in the array.
-	 * @param Scale Used to normalize the data for the observation.
-	 * @param Axis0 The forward axis of the plane.
-	 * @param Axis1 The right axis of the plane.
-	 * @return The newly created observation.
-	 */
-	UFUNCTION(BlueprintCallable, Category = "LearningAgents", meta = (DefaultToSelf = "InInteractor"))
-	static UPlanarPositionArrayObservation* AddPlanarPositionArrayObservation(ULearningAgentsInteractor* InInteractor, const FName Name = NAME_None, const int32 PositionNum = 1, const float Scale = 100.0f, const FVector Axis0 = FVector::ForwardVector, const FVector Axis1 = FVector::RightVector);
+	UFUNCTION(BlueprintPure, Category = "LearningAgents")
+	FLearningAgentsObservationObjectElement MakeProportionAlongSplineObservation(const USplineComponent* SplineComponent, const float DistanceAlongSpline, const FName Name = TEXT("ProportionAlongSpline"));
 
-	/**
-	 * Sets the data for this observation. The relative position & rotation can be used to make this observation
-	 * relative to the agent's perspective, e.g. by passing the agent's position & forward rotation. Call during
-	 * ULearningAgentsInteractor::SetObservations event.
-	 * @param AgentId The agent id this data corresponds to.
-	 * @param Positions The positions currently being observed.
-	 * @param RelativePosition The vector Positions will be offset from.
-	 * @param RelativeRotation The frame of reference rotation.
-	 */
-	UFUNCTION(BlueprintCallable, Category = "LearningAgents", meta = (AgentId = "-1"))
-	void SetPlanarPositionArrayObservation(const int32 AgentId, const TArray<FVector>& Positions, const FVector RelativePosition = FVector::ZeroVector, const FRotator RelativeRotation = FRotator::ZeroRotator);
+	UFUNCTION(BlueprintPure, Category = "LearningAgents")
+	FLearningAgentsObservationObjectElement MakeDirectionAlongSplineObservation(const USplineComponent* SplineComponent, const float DistanceAlongSpline, const FTransform RelativeTransform = FTransform(), const FName Name = TEXT("DirectionAlongSpline"));
 
-#if UE_LEARNING_AGENTS_ENABLE_VISUAL_LOG
-	/** Describes this observation to the visual logger for debugging purposes. */
-	virtual void VisualLog(const UE::Learning::FIndexSet Instances) const override;
-#endif
+	UFUNCTION(BlueprintPure, Category = "LearningAgents")
+	FLearningAgentsObservationObjectElement MakePropertiesAlongSplineObservation(const USplineComponent* SplineComponent, const float DistanceAlongSpline, const FTransform RelativeTransform = FTransform(), const float TranslationScale = 100.0f, const FName Name = TEXT("PropertiesAlongSpline"));
 
-	TSharedPtr<UE::Learning::FPlanarPositionFeature> FeatureObject;
-};
+	// Ray Cast Observations
 
-//------------------------------------------------------------------
+	UFUNCTION(BlueprintPure, Category = "LearningAgents")
+	FLearningAgentsObservationObjectElement MakeProportionAlongRayObservation(const FVector RayStart, const FVector RayEnd, const FTransform RayTransform = FTransform(), const ECollisionChannel CollisionChannel = ECollisionChannel::ECC_WorldStatic, const FName Name = TEXT("ProportionAlongRay"));
 
-/** An observation of a velocity. */
-UCLASS()
-class LEARNINGAGENTS_API UVelocityObservation : public ULearningAgentsObservation
-{
-	GENERATED_BODY()
+	UFUNCTION(BlueprintPure, Category = "LearningAgents")
+	FLearningAgentsObservationObjectElement MakeProportionAlongRaysObservation(const TArray<FVector>& RayStarts, const TArray<FVector>& RayEnds, const FTransform RayTransform, const ECollisionChannel CollisionChannel = ECollisionChannel::ECC_WorldStatic, const FName Name = TEXT("ProportionAlongRays"));
 
 public:
 
-	/**
-	 * Adds a new velocity observation to the given agent interactor. Call during ULearningAgentsInteractor::SetupObservations event.
-	 * @param InInteractor The agent interactor to add this observation to.
-	 * @param Name The name of this new observation. Used for debugging.
-	 * @param Scale Used to normalize the data for the observation.
-	 * @return The newly created observation.
-	 */
-	UFUNCTION(BlueprintCallable, Category = "LearningAgents", meta = (DefaultToSelf = "InInteractor"))
-	static UVelocityObservation* AddVelocityObservation(ULearningAgentsInteractor* InInteractor, const FName Name = NAME_None, const float Scale = 200.0f);
+	// Basic Observations
 
-	/**
-	 * Sets the data for this observation. The relative rotation can be used to make this observation relative to the
-	 * agent's perspective, e.g. by passing the agent's forward rotation. Call during ULearningAgentsInteractor::SetObservations event.
-	 * @param AgentId The agent id this data corresponds to.
-	 * @param Velocity The velocity currently being observed.
-	 * @param RelativeRotation The frame of reference rotation.
-	 */
-	UFUNCTION(BlueprintCallable, Category = "LearningAgents", meta = (AgentId = "-1"))
-	void SetVelocityObservation(const int32 AgentId, const FVector Velocity, const FRotator RelativeRotation = FRotator::ZeroRotator);
+	UFUNCTION(BlueprintPure = false, Category = "LearningAgents")
+	UPARAM(DisplayName = "Success") bool GetNullObservation(const FLearningAgentsObservationObjectElement Element, const FName Name = TEXT("Null")) const;
 
-#if UE_LEARNING_AGENTS_ENABLE_VISUAL_LOG
-	/** Describes this observation to the visual logger for debugging purposes. */
-	virtual void VisualLog(const UE::Learning::FIndexSet Instances) const override;
-#endif
+	UFUNCTION(BlueprintPure = false, Category = "LearningAgents")
+	UPARAM(DisplayName = "Success") bool GetContinuousObservationNum(int32& OutNum, const FLearningAgentsObservationObjectElement Element, const FName Name = TEXT("Continuous")) const;
 
-	TSharedPtr<UE::Learning::FVelocityFeature> FeatureObject;
-};
+	UFUNCTION(BlueprintPure = false, Category = "LearningAgents")
+	UPARAM(DisplayName = "Success") bool GetContinuousObservation(TArray<float>& OutValues, const FLearningAgentsObservationObjectElement Element, const FName Name = TEXT("Continuous")) const;
+	UPARAM(DisplayName = "Success") bool GetContinuousObservationToArrayView(TArrayView<float> OutValues, const FLearningAgentsObservationObjectElement Element, const FName Name = TEXT("Continuous")) const;
 
-/** An observation of an array of velocities. */
-UCLASS()
-class LEARNINGAGENTS_API UVelocityArrayObservation : public ULearningAgentsObservation
-{
-	GENERATED_BODY()
+	UFUNCTION(BlueprintPure = false, Category = "LearningAgents")
+	UPARAM(DisplayName = "Success") bool GetExclusiveDiscreteObservation(int32& OutIndex, const FLearningAgentsObservationObjectElement Element, const FName Name = TEXT("ExclusiveDiscrete")) const;
 
-public:
+	UFUNCTION(BlueprintPure = false, Category = "LearningAgents")
+	UPARAM(DisplayName = "Success") bool GetInclusiveDiscreteObservationNum(int32& OutNum, const FLearningAgentsObservationObjectElement Element, const FName Name = TEXT("InclusiveDiscrete")) const;
 
-	/**
-	 * Adds a new velocity observation to the given agent interactor. Call during ULearningAgentsInteractor::SetupObservations event.
-	 * @param InInteractor The agent interactor to add this observation to.
-	 * @param Name The name of this new observation. Used for debugging.
-	 * @param VelocityNum The number of velocities in the array.
-	 * @param Scale Used to normalize the data for the observation.
-	 * @return The newly created observation.
-	 */
-	UFUNCTION(BlueprintCallable, Category = "LearningAgents", meta = (DefaultToSelf = "InInteractor"))
-	static UVelocityArrayObservation* AddVelocityArrayObservation(ULearningAgentsInteractor* InInteractor, const FName Name = NAME_None, const int32 VelocityNum = 1, const float Scale = 200.0f);
+	UFUNCTION(BlueprintPure = false, Category = "LearningAgents")
+	UPARAM(DisplayName = "Success") bool GetInclusiveDiscreteObservation(TArray<int32>& OutIndices, const FLearningAgentsObservationObjectElement Element, const FName Name = TEXT("InclusiveDiscrete")) const;
+	UPARAM(DisplayName = "Success") bool GetInclusiveDiscreteObservationToArrayView(TArrayView<int32> OutIndices, const FLearningAgentsObservationObjectElement Element, const FName Name = TEXT("InclusiveDiscrete")) const;
 
-	/**
-	 * Sets the data for this observation. The relative rotation can be used to make this observation relative to the
-	 * agent's perspective, e.g. by passing the agent's forward rotation. Call during ULearningAgentsInteractor::SetObservations event.
-	 * @param AgentId The agent id this data corresponds to.
-	 * @param Velocities The velocities currently being observed.
-	 * @param RelativeRotation The frame of reference rotation.
-	 */
-	UFUNCTION(BlueprintCallable, Category = "LearningAgents", meta = (AgentId = "-1"))
-	void SetVelocityArrayObservation(const int32 AgentId, const TArray<FVector>& Velocities, const FRotator RelativeRotation = FRotator::ZeroRotator);
+	UFUNCTION(BlueprintPure = false, Category = "LearningAgents")
+	UPARAM(DisplayName = "Success") bool GetIndexObservation(int32& OutIndex, const FLearningAgentsObservationObjectElement Element, const FName Name = TEXT("Index")) const;
 
-#if UE_LEARNING_AGENTS_ENABLE_VISUAL_LOG
-	/** Describes this observation to the visual logger for debugging purposes. */
-	virtual void VisualLog(const UE::Learning::FIndexSet Instances) const override;
-#endif
+	UFUNCTION(BlueprintPure = false, Category = "LearningAgents")
+	UPARAM(DisplayName = "Success") bool GetStructObservationNum(int32& OutNum, const FLearningAgentsObservationObjectElement Element, const FName Name = TEXT("Struct")) const;
 
-	TSharedPtr<UE::Learning::FVelocityFeature> FeatureObject;
-};
+	UFUNCTION(BlueprintPure = false, Category = "LearningAgents")
+	UPARAM(DisplayName = "Success") bool GetStructObservation(TMap<FName, FLearningAgentsObservationObjectElement>& OutElements, const FLearningAgentsObservationObjectElement Element, const FName Name = TEXT("Struct")) const;
+	
+	UFUNCTION(BlueprintPure = false, Category = "LearningAgents")
+	UPARAM(DisplayName = "Success") bool GetStructObservationToArrays(TArray<FName>& OutElementNames, TArray<FLearningAgentsObservationObjectElement>& OutElements, const FLearningAgentsObservationObjectElement Element, const FName Name = TEXT("Struct")) const;
+	UPARAM(DisplayName = "Success") bool GetStructObservationToArrayViews(TArrayView<FName> OutElementNames, TArrayView<FLearningAgentsObservationObjectElement> OutElements, const FLearningAgentsObservationObjectElement Element, const FName Name = TEXT("Struct")) const;
 
-/** An observation of a velocity along a single axis. */
-UCLASS()
-class LEARNINGAGENTS_API UScalarVelocityObservation : public ULearningAgentsObservation
-{
-	GENERATED_BODY()
+	UFUNCTION(BlueprintPure = false, Category = "LearningAgents")
+	UPARAM(DisplayName = "Success") bool GetExclusiveUnionObservation(FName& OutElementName, FLearningAgentsObservationObjectElement& OutElement, const FLearningAgentsObservationObjectElement Element, const FName Name = TEXT("ExclusiveUnion")) const;
 
-public:
+	UFUNCTION(BlueprintPure = false, Category = "LearningAgents")
+	UPARAM(DisplayName = "Success") bool GetInclusiveUnionObservationNum(int32& OutNum, const FLearningAgentsObservationObjectElement Element, const FName Name = TEXT("InclusiveUnion")) const;
 
-	/**
-	 * Adds a new scalar velocity observation to the given agent interactor. 
-	 * Call during ULearningAgentsInteractor::SetupObservations event.
-	 * @param InInteractor The agent interactor to add this observation to.
-	 * @param Name The name of this new observation. Used for debugging.
-	 * @param Scale Used to normalize the data for the observation.
-	 * @return The newly created observation.
-	 */
-	UFUNCTION(BlueprintCallable, Category = "LearningAgents", meta = (DefaultToSelf = "InInteractor"))
-	static UScalarVelocityObservation* AddScalarVelocityObservation(ULearningAgentsInteractor* InInteractor, const FName Name = NAME_None, const float Scale = 200.0f);
+	UFUNCTION(BlueprintPure = false, Category = "LearningAgents")
+	UPARAM(DisplayName = "Success") bool GetInclusiveUnionObservation(TMap<FName, FLearningAgentsObservationObjectElement>& OutElements, const FLearningAgentsObservationObjectElement Element, const FName Name = TEXT("InclusiveUnion")) const;
+	
+	UFUNCTION(BlueprintPure = false, Category = "LearningAgents")
+	UPARAM(DisplayName = "Success") bool GetInclusiveUnionObservationToArrays(TArray<FName>& OutElementNames, TArray<FLearningAgentsObservationObjectElement>& OutElements, const FLearningAgentsObservationObjectElement Element, const FName Name = TEXT("InclusiveUnion")) const;
+	UPARAM(DisplayName = "Success") bool GetInclusiveUnionObservationToArrayViews(TArrayView<FName> OutElementNames, TArrayView<FLearningAgentsObservationObjectElement> OutElements, const FLearningAgentsObservationObjectElement Element, const FName Name = TEXT("InclusiveUnion")) const;
 
-	/**
-	 * Sets the data for this observation. Call during ULearningAgentsInteractor::SetObservations event.
-	 * @param AgentId The agent id this data corresponds to.
-	 * @param Velocity The velocity currently being observed.
-	 */
-	UFUNCTION(BlueprintCallable, Category = "LearningAgents", meta = (AgentId = "-1"))
-	void SetScalarVelocityObservation(const int32 AgentId, const float Velocity);
+	UFUNCTION(BlueprintPure = false, Category = "LearningAgents")
+	UPARAM(DisplayName = "Success") bool GetStaticArrayObservationNum(int32& OutNum, const FLearningAgentsObservationObjectElement Element, const FName Name = TEXT("StaticArray")) const;
 
-	/**
-	 * Sets the data for this observation. Call during ULearningAgentsInteractor::SetObservations event.
-	 * @param AgentId The agent id this data corresponds to.
-	 * @param Velocity The velocity currently being observed.
-	 * @param Axis The axis to encode the velocity along
-	 */
-	UFUNCTION(BlueprintCallable, Category = "LearningAgents", meta = (AgentId = "-1"))
-	void SetScalarVelocityObservationWithAxis(const int32 AgentId, const FVector Velocity, const FVector Axis = FVector::UpVector);
+	UFUNCTION(BlueprintPure = false, Category = "LearningAgents")
+	UPARAM(DisplayName = "Success") bool GetStaticArrayObservation(TArray<FLearningAgentsObservationObjectElement>& OutElements, const FLearningAgentsObservationObjectElement Element, const FName Name = TEXT("StaticArray")) const;
+	UPARAM(DisplayName = "Success") bool GetStaticArrayObservationToArrayView(TArrayView<FLearningAgentsObservationObjectElement> OutElements, const FLearningAgentsObservationObjectElement Element, const FName Name = TEXT("StaticArray")) const;
 
-#if UE_LEARNING_AGENTS_ENABLE_VISUAL_LOG
-	/** Describes this observation to the visual logger for debugging purposes. */
-	virtual void VisualLog(const UE::Learning::FIndexSet Instances) const override;
-#endif
+	UFUNCTION(BlueprintPure = false, Category = "LearningAgents")
+	UPARAM(DisplayName = "Success") bool GetSetObservationNum(int32& OutNum, const FLearningAgentsObservationObjectElement Element, const FName Name = TEXT("Set")) const;
 
-	TSharedPtr<UE::Learning::FScalarVelocityFeature> FeatureObject;
-};
+	UFUNCTION(BlueprintPure = false, Category = "LearningAgents")
+	UPARAM(DisplayName = "Success") bool GetSetObservation(TSet<FLearningAgentsObservationObjectElement>& OutElements, const FLearningAgentsObservationObjectElement Element, const FName Name = TEXT("Set")) const;
+	
+	UFUNCTION(BlueprintPure = false, Category = "LearningAgents")
+	UPARAM(DisplayName = "Success") bool GetSetObservationToArray(TArray<FLearningAgentsObservationObjectElement>& OutElements, const FLearningAgentsObservationObjectElement Element, const FName Name = TEXT("Set")) const;
+	UPARAM(DisplayName = "Success") bool GetSetObservationToArrayView(TArrayView<FLearningAgentsObservationObjectElement> OutElements, const FLearningAgentsObservationObjectElement Element, const FName Name = TEXT("Set")) const;
 
-/** An observation of an array of velocities along a single axis. */
-UCLASS()
-class LEARNINGAGENTS_API UScalarVelocityArrayObservation : public ULearningAgentsObservation
-{
-	GENERATED_BODY()
+	UFUNCTION(BlueprintPure = false, Category = "LearningAgents")
+	UPARAM(DisplayName = "Success") bool GetPairObservation(FLearningAgentsObservationObjectElement& OutKey, FLearningAgentsObservationObjectElement& OutValue, const FLearningAgentsObservationObjectElement Element, const FName Name = TEXT("Pair")) const;
 
-public:
+	UFUNCTION(BlueprintPure = false, Category = "LearningAgents")
+	UPARAM(DisplayName = "Success") bool GetArrayObservationNum(int32& OutNum, const FLearningAgentsObservationObjectElement Element, const FName Name = TEXT("Array")) const;
 
-	/**
-	 * Adds a new scalar velocity observation to the given agent interactor.
-	 * Call during ULearningAgentsInteractor::SetupObservations event.
-	 * @param InInteractor The agent interactor to add this observation to.
-	 * @param Name The name of this new observation. Used for debugging.
-	 * @param VelocityNum The number of velocities in the array.
-	 * @param Scale Used to normalize the data for the observation.
-	 * @return The newly created observation.
-	 */
-	UFUNCTION(BlueprintCallable, Category = "LearningAgents", meta = (DefaultToSelf = "InInteractor"))
-	static UScalarVelocityArrayObservation* AddScalarVelocityArrayObservation(ULearningAgentsInteractor* InInteractor, const FName Name = NAME_None, const int32 VelocityNum = 1, const float Scale = 200.0f);
+	UFUNCTION(BlueprintPure = false, Category = "LearningAgents")
+	UPARAM(DisplayName = "Success") bool GetArrayObservation(TArray<FLearningAgentsObservationObjectElement>& OutElements, const FLearningAgentsObservationObjectElement Element, const FName Name = TEXT("Array")) const;
+	UPARAM(DisplayName = "Success") bool GetArrayObservationToArrayView(TArrayView<FLearningAgentsObservationObjectElement> OutElements, const FLearningAgentsObservationObjectElement Element, const FName Name = TEXT("Array")) const;
 
-	/**
-	 * Sets the data for this observation. Call during ULearningAgentsInteractor::SetObservations event.
-	 * @param AgentId The agent id this data corresponds to.
-	 * @param Velocities The velocities currently being observed.
-	 */
-	UFUNCTION(BlueprintCallable, Category = "LearningAgents", meta = (AgentId = "-1"))
-	void SetScalarVelocityArrayObservation(const int32 AgentId, const TArray<float>& Velocities);
+	UFUNCTION(BlueprintPure = false, Category = "LearningAgents")
+	UPARAM(DisplayName = "Success") bool GetMapObservationNum(int32& OutNum, const FLearningAgentsObservationObjectElement Element, const FName Name = TEXT("Map")) const;
 
-	/**
-	 * Sets the data for this observation. Call during ULearningAgentsInteractor::SetObservations event.
-	 * @param AgentId The agent id this data corresponds to.
-	 * @param Velocities The velocities currently being observed.
-	 * @param Axis The axis to encode the velocity along
-	 */
-	UFUNCTION(BlueprintCallable, Category = "LearningAgents", meta = (AgentId = "-1"))
-	void SetScalarVelocityArrayObservationWithAxis(const int32 AgentId, const TArray<FVector>& Velocities, const FVector Axis = FVector::UpVector);
+	UFUNCTION(BlueprintPure = false, Category = "LearningAgents")
+	UPARAM(DisplayName = "Success") bool GetMapObservation(TMap<FLearningAgentsObservationObjectElement,FLearningAgentsObservationObjectElement>& OutElements, const FLearningAgentsObservationObjectElement Element, const FName Name = TEXT("Map")) const;
+	
+	UFUNCTION(BlueprintPure = false, Category = "LearningAgents")
+	UPARAM(DisplayName = "Success") bool GetMapObservationToArrays(TArray<FLearningAgentsObservationObjectElement>& OutKeys, TArray<FLearningAgentsObservationObjectElement>& OutValues, const FLearningAgentsObservationObjectElement Element, const FName Name = TEXT("Map")) const;
+	UPARAM(DisplayName = "Success") bool GetMapObservationToArrayViews(TArrayView<FLearningAgentsObservationObjectElement> OutKeys, TArrayView<FLearningAgentsObservationObjectElement> OutValues, const FLearningAgentsObservationObjectElement Element, const FName Name = TEXT("Map")) const;
 
-#if UE_LEARNING_AGENTS_ENABLE_VISUAL_LOG
-	/** Describes this observation to the visual logger for debugging purposes. */
-	virtual void VisualLog(const UE::Learning::FIndexSet Instances) const override;
-#endif
+	UFUNCTION(BlueprintPure = false, Category = "LearningAgents")
+	UPARAM(DisplayName = "Success") bool GetEnumObservation(uint8& OutEnumValue, const UEnum* Enum, const FLearningAgentsObservationObjectElement Element, const FName Name = TEXT("Enum")) const;
 
-	TSharedPtr<UE::Learning::FScalarVelocityFeature> FeatureObject;
-};
+	UFUNCTION(BlueprintPure = false, Category = "LearningAgents")
+	UPARAM(DisplayName = "Success") bool GetBitmaskObservation(int32& OutBitmaskValue, const UEnum* Enum, const FLearningAgentsObservationObjectElement Element, const FName Name = TEXT("Bitmask")) const;
 
-/** An observation of a velocity projected onto a plane. */
-UCLASS()
-class LEARNINGAGENTS_API UPlanarVelocityObservation : public ULearningAgentsObservation
-{
-	GENERATED_BODY()
+	UFUNCTION(BlueprintPure = false, Category = "LearningAgents", meta = (ExpandEnumAsExecs = "OutOption"))
+	UPARAM(DisplayName = "Success") bool GetOptionalObservation(ELearningAgentsOptionalObservation& OutOption, FLearningAgentsObservationObjectElement& OutElement, const FLearningAgentsObservationObjectElement Element, const FName Name = TEXT("Optional")) const;
 
-public:
+	UFUNCTION(BlueprintPure = false, Category = "LearningAgents", meta = (ExpandEnumAsExecs = "OutEither"))
+	UPARAM(DisplayName = "Success") bool GetEitherObservation(ELearningAgentsEitherObservation& OutEither, FLearningAgentsObservationObjectElement& OutElement, const FLearningAgentsObservationObjectElement Element, const FName Name = TEXT("Either")) const;
 
-	/**
-	 * Adds a new planar velocity observation to the given agent interactor. The axis parameters define the plane.
-	 * Call during ULearningAgentsInteractor::SetupObservations event.
-	 * @param InInteractor The agent interactor to add this observation to.
-	 * @param Name The name of this new observation. Used for debugging.
-	 * @param Scale Used to normalize the data for the observation.
-	 * @param Axis0 The forward axis of the plane.
-	 * @param Axis1 The right axis of the plane.
-	 * @return The newly created observation.
-	 */
-	UFUNCTION(BlueprintCallable, Category = "LearningAgents", meta = (DefaultToSelf = "InInteractor"))
-	static UPlanarVelocityObservation* AddPlanarVelocityObservation(ULearningAgentsInteractor* InInteractor, const FName Name = NAME_None, const float Scale = 200.0f, const FVector Axis0 = FVector::ForwardVector, const FVector Axis1 = FVector::RightVector);
+	UFUNCTION(BlueprintPure = false, Category = "LearningAgents")
+	UPARAM(DisplayName = "Success") bool GetEncodingObservation(FLearningAgentsObservationObjectElement& OutElement, const FLearningAgentsObservationObjectElement Element, const FName Name = TEXT("Encoding")) const;
 
-	/**
-	 * Sets the data for this observation. The relative rotation can be used to make this observation relative to the
-	 * agent's perspective, e.g. by passing the agent's forward rotation. Call during ULearningAgentsInteractor::SetObservations event.
-	 * @param AgentId The agent id this data corresponds to.
-	 * @param Velocity The velocity currently being observed.
-	 * @param RelativeRotation The frame of reference rotation.
-	 */
-	UFUNCTION(BlueprintCallable, Category = "LearningAgents", meta = (AgentId = "-1"))
-	void SetPlanarVelocityObservation(const int32 AgentId, const FVector Velocity, const FRotator RelativeRotation = FRotator::ZeroRotator);
+	UFUNCTION(BlueprintPure = false, Category = "LearningAgents")
+	UPARAM(DisplayName = "Success") bool GetBoolObservation(bool& bOutValue, const FLearningAgentsObservationObjectElement Element, const FName Name = TEXT("Bool")) const;
 
-#if UE_LEARNING_AGENTS_ENABLE_VISUAL_LOG
-	/** Describes this observation to the visual logger for debugging purposes. */
-	virtual void VisualLog(const UE::Learning::FIndexSet Instances) const override;
-#endif
+	UFUNCTION(BlueprintPure = false, Category = "LearningAgents")
+	UPARAM(DisplayName = "Success") bool GetFloatObservation(float& OutValue, const FLearningAgentsObservationObjectElement Element, const float FloatScale = 1.0f, const FName Name = TEXT("Float")) const;
 
-	TSharedPtr<UE::Learning::FPlanarVelocityFeature> FeatureObject;
-};
+	UFUNCTION(BlueprintPure = false, Category = "LearningAgents")
+	UPARAM(DisplayName = "Success") bool GetTranslationObservation(FVector& OutTranslation, const FLearningAgentsObservationObjectElement Element, const FTransform RelativeTransform = FTransform(), const float TranslationScale = 100.0f, const FName Name = TEXT("Translation")) const;
 
-/** An observation of an array of velocities projected onto a plane. */
-UCLASS()
-class LEARNINGAGENTS_API UPlanarVelocityArrayObservation : public ULearningAgentsObservation
-{
-	GENERATED_BODY()
+	UFUNCTION(BlueprintPure = false, Category = "LearningAgents")
+	UPARAM(DisplayName = "Success") bool GetRotationObservation(FRotator& OutRotation, const FLearningAgentsObservationObjectElement Element, const FRotator RelativeRotation = FRotator::ZeroRotator, const FName Name = TEXT("Rotation")) const;
+	UPARAM(DisplayName = "Success") bool GetRotationObservationAsQuat(FQuat& OutRotation, const FLearningAgentsObservationObjectElement Element, const FQuat RelativeRotation = FQuat::Identity, const FName Name = TEXT("Rotation")) const;
 
-public:
+	UFUNCTION(BlueprintPure = false, Category = "LearningAgents")
+	UPARAM(DisplayName = "Success") bool GetScaleObservation(FVector& OutScale, const FLearningAgentsObservationObjectElement Element, const FVector RelativeScale = FVector(1, 1, 1), const FName Name = TEXT("Scale")) const;
 
-	/**
-	 * Adds a new planar velocity observation to the given agent interactor. The axis parameters define the plane.
-	 * Call during ULearningAgentsInteractor::SetupObservations event.
-	 * @param InInteractor The agent interactor to add this observation to.
-	 * @param Name The name of this new observation. Used for debugging.
-	 * @param VelocityNum The number of velocities in the array.
-	 * @param Scale Used to normalize the data for the observation.
-	 * @param Axis0 The forward axis of the plane.
-	 * @param Axis1 The right axis of the plane.
-	 * @return The newly created observation.
-	 */
-	UFUNCTION(BlueprintCallable, Category = "LearningAgents", meta = (DefaultToSelf = "InInteractor"))
-	static UPlanarVelocityArrayObservation* AddPlanarVelocityArrayObservation(ULearningAgentsInteractor* InInteractor, const FName Name = NAME_None, const int32 VelocityNum = 1, const float Scale = 200.0f, const FVector Axis0 = FVector::ForwardVector, const FVector Axis1 = FVector::RightVector);
+	UFUNCTION(BlueprintPure = false, Category = "LearningAgents")
+	UPARAM(DisplayName = "Success") bool GetTransformObservation(FTransform& OutTransform, const FLearningAgentsObservationObjectElement Element, const FTransform RelativeTransform = FTransform(), const float TranslationScale = 100.0f, const FName Name = TEXT("Transform")) const;
 
-	/**
-	 * Sets the data for this observation. The relative rotation can be used to make this observation relative to the
-	 * agent's perspective, e.g. by passing the agent's forward rotation. Call during ULearningAgentsInteractor::SetObservations event.
-	 * @param AgentId The agent id this data corresponds to.
-	 * @param Velocities The velocities currently being observed.
-	 * @param RelativeRotation The frame of reference rotation.
-	 */
-	UFUNCTION(BlueprintCallable, Category = "LearningAgents", meta = (AgentId = "-1"))
-	void SetPlanarVelocityArrayObservation(const int32 AgentId, const TArray<FVector>& Velocities, const FRotator RelativeRotation = FRotator::ZeroRotator);
+	UFUNCTION(BlueprintPure = false, Category = "LearningAgents")
+	UPARAM(DisplayName = "Success") bool GetAngleObservation(float& OutAngle, const FLearningAgentsObservationObjectElement Element, const float RelativeAngle = 0.0f, const FName Name = TEXT("Angle")) const;
 
-#if UE_LEARNING_AGENTS_ENABLE_VISUAL_LOG
-	/** Describes this observation to the visual logger for debugging purposes. */
-	virtual void VisualLog(const UE::Learning::FIndexSet Instances) const override;
-#endif
+	UFUNCTION(BlueprintPure = false, Category = "LearningAgents")
+	UPARAM(DisplayName = "Success") bool GetAngleObservationRadians(float& OutAngle, const FLearningAgentsObservationObjectElement Element, const float RelativeAngle = 0.0f, const FName Name = TEXT("Angle")) const;
 
-	TSharedPtr<UE::Learning::FPlanarVelocityFeature> FeatureObject;
-};
+	UFUNCTION(BlueprintPure = false, Category = "LearningAgents")
+	UPARAM(DisplayName = "Success") bool GetVelocityObservation(FVector& OutVelocity, const FLearningAgentsObservationObjectElement Element, const FTransform RelativeTransform = FTransform(), const float VelocityScale = 200.0f, const FName Name = TEXT("Velocity")) const;
 
-//------------------------------------------------------------------
+	UFUNCTION(BlueprintPure = false, Category = "LearningAgents")
+	UPARAM(DisplayName = "Success") bool GetDirectionObservation(FVector& OutDirection, const FLearningAgentsObservationObjectElement Element, const FTransform RelativeTransform = FTransform(), const FName Name = TEXT("Direction")) const;
 
-/** An observation of an angular velocity. */
-UCLASS()
-class LEARNINGAGENTS_API UAngularVelocityObservation : public ULearningAgentsObservation
-{
-	GENERATED_BODY()
+	// Spline Observations
 
-public:
+	UFUNCTION(BlueprintPure = false, Category = "LearningAgents")
+	UPARAM(DisplayName = "Success") bool GetLocationAlongSplineObservation(FVector& OutLocation, const FLearningAgentsObservationObjectElement Element, const FTransform RelativeTransform = FTransform(), const float TranslationScale = 100.0f, const FName Name = TEXT("LocationAlongSpline"));
 
-	/**
-	 * Adds a new angular velocity observation to the given agent interactor. Call during ULearningAgentsInteractor::SetupObservations event.
-	 * @param InInteractor The agent interactor to add this observation to.
-	 * @param Name The name of this new observation. Used for debugging.
-	 * @param Scale Used to normalize the data for the observation.
-	 * @return The newly created observation.
-	 */
-	UFUNCTION(BlueprintCallable, Category = "LearningAgents", meta = (DefaultToSelf = "InInteractor"))
-	static UAngularVelocityObservation* AddAngularVelocityObservation(ULearningAgentsInteractor* InInteractor, const FName Name = NAME_None, const float Scale = 180.0f);
+	UFUNCTION(BlueprintPure = false, Category = "LearningAgents")
+	UPARAM(DisplayName = "Success") bool GetProportionAlongSplineObservation(bool& bOutIsClosedLoop, float& OutAngle, float& OutPropotion, const FLearningAgentsObservationObjectElement Element, const FName Name = TEXT("ProportionAlongSpline"));
 
-	/**
-	 * Sets the data for this observation. The relative rotation can be used to make this observation relative to the
-	 * agent's perspective, e.g. by passing the agent's forward rotation. Call during ULearningAgentsInteractor::SetObservations event.
-	 * @param AgentId The agent id this data corresponds to.
-	 * @param AngularVelocity The angular velocity currently being observed.
-	 * @param RelativeRotation The frame of reference rotation.
-	 */
-	UFUNCTION(BlueprintCallable, Category = "LearningAgents", meta = (AgentId = "-1"))
-	void SetAngularVelocityObservation(const int32 AgentId, const FVector AngularVelocity, const FRotator RelativeRotation = FRotator::ZeroRotator);
+	UFUNCTION(BlueprintPure = false, Category = "LearningAgents")
+	UPARAM(DisplayName = "Success") bool GetDirectionAlongSplineObservation(FVector& OutDirection, const FLearningAgentsObservationObjectElement Element, const FTransform RelativeTransform = FTransform(), const FName Name = TEXT("DirectionAlongSpline"));
 
-#if UE_LEARNING_AGENTS_ENABLE_VISUAL_LOG
-	/** Describes this observation to the visual logger for debugging purposes. */
-	virtual void VisualLog(const UE::Learning::FIndexSet Instances) const override;
-#endif
+	UFUNCTION(BlueprintPure = false, Category = "LearningAgents")
+	UPARAM(DisplayName = "Success") bool GetPropertiesAlongSplineObservation(FVector& OutLocation, bool& bOutIsClosedLoop, float& OutAngle, float& OutPropotion, FVector& OutDirection, const FLearningAgentsObservationObjectElement Element, const FTransform RelativeTransform = FTransform(), const float TranslationScale = 100.0f, const FName Name = TEXT("PropertiesAlongSpline"));
 
-	TSharedPtr<UE::Learning::FAngularVelocityFeature> FeatureObject;
-};
+	// Ray Cast Observations
 
-/** An observation of an array of angular velocities. */
-UCLASS()
-class LEARNINGAGENTS_API UAngularVelocityArrayObservation : public ULearningAgentsObservation
-{
-	GENERATED_BODY()
-
-public:
-
-	/**
-	 * Adds a new angular velocity array observation to the given agent interactor. Call during ULearningAgentsInteractor::SetupObservations event.
-	 * @param InInteractor The agent interactor to add this observation to.
-	 * @param Name The name of this new observation. Used for debugging.
-	 * @param AngularVelocityNum The number of angular velocities in the array.
-	 * @param Scale Used to normalize the data for the observation.
-	 * @return The newly created observation.
-	 */
-	UFUNCTION(BlueprintCallable, Category = "LearningAgents", meta = (DefaultToSelf = "InInteractor"))
-	static UAngularVelocityArrayObservation* AddAngularVelocityArrayObservation(ULearningAgentsInteractor* InInteractor, const FName Name = NAME_None, const int32 AngularVelocityNum = 1, const float Scale = 180.0f);
-
-	/**
-	 * Sets the data for this observation. The relative rotation can be used to make this observation relative to the
-	 * agent's perspective, e.g. by passing the agent's forward rotation. Call during ULearningAgentsInteractor::SetObservations event.
-	 * @param AgentId The agent id this data corresponds to.
-	 * @param Velocities The angular velocities currently being observed.
-	 * @param RelativeRotation The frame of reference rotation.
-	 */
-	UFUNCTION(BlueprintCallable, Category = "LearningAgents", meta = (AgentId = "-1"))
-	void SetAngularVelocityArrayObservation(const int32 AgentId, const TArray<FVector>& AngularVelocities, const FRotator RelativeRotation = FRotator::ZeroRotator);
-
-#if UE_LEARNING_AGENTS_ENABLE_VISUAL_LOG
-	/** Describes this observation to the visual logger for debugging purposes. */
-	virtual void VisualLog(const UE::Learning::FIndexSet Instances) const override;
-#endif
-
-	TSharedPtr<UE::Learning::FAngularVelocityFeature> FeatureObject;
-};
-
-/** An observation of a scalar angular velocity. */
-UCLASS()
-class LEARNINGAGENTS_API UScalarAngularVelocityObservation : public ULearningAgentsObservation
-{
-	GENERATED_BODY()
-
-public:
-
-	/**
-	 * Adds a new scalar angular velocity observation to the given agent interactor.
-	 * Call during ULearningAgentsInteractor::SetupObservations event.
-	 * @param InInteractor The agent interactor to add this observation to.
-	 * @param Name The name of this new observation. Used for debugging.
-	 * @param Scale Used to normalize the data for the observation.
-	 * @return The newly created observation.
-	 */
-	UFUNCTION(BlueprintCallable, Category = "LearningAgents", meta = (DefaultToSelf = "InInteractor"))
-	static UScalarAngularVelocityObservation* AddScalarAngularVelocityObservation(ULearningAgentsInteractor* InInteractor, const FName Name = NAME_None, const float Scale = 180.0f);
-
-	/**
-	 * Sets the data for this observation. Call during ULearningAgentsInteractor::SetObservations event.
-	 * @param AgentId The agent id this data corresponds to.
-	 * @param AngularVelocity The angular velocity currently being observed.
-	 */
-	UFUNCTION(BlueprintCallable, Category = "LearningAgents", meta = (AgentId = "-1"))
-	void SetScalarAngularVelocityObservation(const int32 AgentId, const float AngularVelocity);
-
-	/**
-	 * Sets the data for this observation. Call during ULearningAgentsInteractor::SetObservations event.
-	 * @param AgentId The agent id this data corresponds to.
-	 * @param AngularVelocity The angular velocity currently being observed.
-	 * @param Axis The axis to encode the angular velocity around
-	 */
-	UFUNCTION(BlueprintCallable, Category = "LearningAgents", meta = (AgentId = "-1"))
-	void SetScalarAngularVelocityObservationWithAxis(const int32 AgentId, const FVector AngularVelocity, const FVector Axis = FVector::UpVector);
-
-#if UE_LEARNING_AGENTS_ENABLE_VISUAL_LOG
-	/** Describes this observation to the visual logger for debugging purposes. */
-	virtual void VisualLog(const UE::Learning::FIndexSet Instances) const override;
-#endif
-
-	TSharedPtr<UE::Learning::FScalarAngularVelocityFeature> FeatureObject;
-};
-
-/** An observation of an array of scalar angular velocities. */
-UCLASS()
-class LEARNINGAGENTS_API UScalarAngularVelocityArrayObservation : public ULearningAgentsObservation
-{
-	GENERATED_BODY()
+	UFUNCTION(BlueprintPure, Category = "LearningAgents")
+	UPARAM(DisplayName = "Success") bool GetProportionAlongRayObservation(float& OutProportion, const FLearningAgentsObservationObjectElement Element, const FName Name = TEXT("ProportionAlongRay"));
 
-public:
+	UFUNCTION(BlueprintPure, Category = "LearningAgents")
+	UPARAM(DisplayName = "Success") bool GetProportionAlongRaysObservationNum(int32& OutProportionNum, const FLearningAgentsObservationObjectElement Element, const FName Name = TEXT("ProportionAlongRays"));
 
-	/**
-	 * Adds a new scalar angular velocity array observation to the given agent interactor.
-	 * Call during ULearningAgentsInteractor::SetupObservations event.
-	 * @param InInteractor The agent interactor to add this observation to.
-	 * @param Name The name of this new observation. Used for debugging.
-	 * @param AngularVelocityNum The number of angular velocities in the array.
-	 * @param Scale Used to normalize the data for the observation.
-	 * @return The newly created observation.
-	 */
-	UFUNCTION(BlueprintCallable, Category = "LearningAgents", meta = (DefaultToSelf = "InInteractor"))
-	static UScalarAngularVelocityArrayObservation* AddScalarAngularVelocityArrayObservation(ULearningAgentsInteractor* InInteractor, const FName Name = NAME_None, const int32 AngularVelocityNum = 1, const float Scale = 180.0f);
+	UFUNCTION(BlueprintPure, Category = "LearningAgents")
+	UPARAM(DisplayName = "Success") bool GetProportionAlongRaysObservation(TArray<float>& OutProportions, const FLearningAgentsObservationObjectElement Element, const FName Name = TEXT("ProportionAlongRays"));
+	UPARAM(DisplayName = "Success") bool GetProportionAlongRaysObservationToArrayView(TArrayView<float> OutProportions, const FLearningAgentsObservationObjectElement Element, const FName Name = TEXT("ProportionAlongRays"));
 
-	/**
-	 * Sets the data for this observation. Call during ULearningAgentsInteractor::SetObservations event.
-	 * @param AgentId The agent id this data corresponds to.
-	 * @param AngularVelocities The angular velocities currently being observed.
-	 */
-	UFUNCTION(BlueprintCallable, Category = "LearningAgents", meta = (AgentId = "-1"))
-	void SetScalarAngularVelocityArrayObservation(const int32 AgentId, const TArray<float>& AngularVelocities);
-
-	/**
-	 * Sets the data for this observation. Call during ULearningAgentsInteractor::SetObservations event.
-	 * @param AgentId The agent id this data corresponds to.
-	 * @param AngularVelocities The angular velocities currently being observed.
-	 * @param Axis The axis to encode the angular velocity around
-	 */
-	UFUNCTION(BlueprintCallable, Category = "LearningAgents", meta = (AgentId = "-1"))
-	void SetScalarAngularVelocityArrayObservationWithAxis(const int32 AgentId, const TArray<FVector>& AngularVelocities, const FVector Axis = FVector::UpVector);
 
-#if UE_LEARNING_AGENTS_ENABLE_VISUAL_LOG
-	/** Describes this observation to the visual logger for debugging purposes. */
-	virtual void VisualLog(const UE::Learning::FIndexSet Instances) const override;
-#endif
+private:
 
-	TSharedPtr<UE::Learning::FScalarAngularVelocityFeature> FeatureObject;
+	UE::Learning::Observation::FObject ObservationObject;
 };
