LearningAgents: Large restructuring to allow for structured actions and observations.

diff --git a/Engine/Plugins/Experimental/LearningAgents/Source/LearningAgents/Public/LearningAgentsPolicy.h b/Engine/Plugins/Experimental/LearningAgents/Source/LearningAgents/Public/LearningAgentsPolicy.h
--- a/Engine/Plugins/Experimental/LearningAgents/Source/LearningAgents/Public/LearningAgentsPolicy.h
+++ b/Engine/Plugins/Experimental/LearningAgents/Source/LearningAgents/Public/LearningAgentsPolicy.h
@@ -2,23 +2,20 @@
 
 #pragma once
 
-#include "LearningAgentsManagerComponent.h"
+#include "LearningAgentsManagerListener.h"
 
 #include "LearningAgentsNeuralNetwork.h" // Included for ELearningAgentsActivationFunction
-#include "LearningAgentsDebug.h"
 #include "LearningArray.h"
-#include "LearningArrayMap.h"
 #include "UObject/ObjectPtr.h"
 
 #include "LearningAgentsPolicy.generated.h"
 
 namespace UE::Learning
 {
-	struct INeuralNetwork;
-	struct FNeuralNetworkPolicyFunction;
+	struct FNeuralNetworkPolicy;
+	struct FNeuralNetworkFunction;
 }
 
-struct FFilePath;
 class ULearningAgentsNeuralNetwork;
 
 /** The configurable settings for a ULearningAgentsPolicy. */
@@ -29,22 +26,6 @@ struct LEARNINGAGENTS_API FLearningAgentsPolicySettings
 
 public:
 
-	/** Seed for the action noise used by the policy  */
-	UPROPERTY(EditAnywhere, Category = "LearningAgents", meta = (ClampMin = "0", UIMin = "0"))
-	int32 ActionNoiseSeed = 1234;
-
-	/** Minimum action noise used by the policy */
-	UPROPERTY(EditAnywhere, Category = "LearningAgents", meta = (ClampMin = "0.0", UIMin = "0.0"))
-	float ActionNoiseMin = 0.25f;
-
-	/** Maximum action noise used by the policy */
-	UPROPERTY(EditAnywhere, Category = "LearningAgents", meta = (ClampMin = "0.0", UIMin = "0.0"))
-	float ActionNoiseMax = 0.25f;
-
-	/** Scale of the action noise used by the policy should be 1.0 during training. */
-	UPROPERTY(EditAnywhere, Category = "LearningAgents", meta = (ClampMin = "0.0", UIMin = "0.0"))
-	float ActionNoiseScale = 1.0f;
-
 	/** Total layers for policy network including input, hidden, and output layers */
 	UPROPERTY(EditAnywhere, Category = "LearningAgents", meta = (ClampMin = "2", UIMin = "2"))
 	int32 LayerNum = 3;
@@ -53,112 +34,147 @@ public:
 	UPROPERTY(EditAnywhere, Category = "LearningAgents", meta = (ClampMin = "1", UIMin = "1"))
 	int32 HiddenLayerSize = 128;
 
+	/** If this network should use memory or not. Setting this to false is equivalent to setting MemoryStateSize to 0. */
+	UPROPERTY(EditAnywhere, Category = "LearningAgents")
+	bool bUseMemory = true;
+
 	/** Number of neurons in the memory state of the policy network */
-	UPROPERTY(EditAnywhere, Category = "LearningAgents", meta = (ClampMin = "1", UIMin = "1"))
-	int32 MemoryStateSize = 128;
+	UPROPERTY(EditAnywhere, Category = "LearningAgents", meta = (ClampMin = "0", UIMin = "0"))
+	int32 MemoryStateSize = 64;
+
+	/** Initial scale for encoded actions. Setting this <1 reduces the initial bias from the network random weights and tends to help training. */
+	UPROPERTY(EditAnywhere, Category = "LearningAgents", meta = (ClampMin = "0", UIMin = "1"))
+	float InitialEncodedActionScale = 0.1f;
 
 	/** Activation function to use on hidden layers of the policy network */
 	UPROPERTY(EditAnywhere, Category = "LearningAgents")
 	ELearningAgentsActivationFunction ActivationFunction = ELearningAgentsActivationFunction::ELU;
 };
 
-/** A policy that maps from observations to actions for the managed agents. */
+/** A policy that maps from observations to actions. */
 UCLASS(BlueprintType, Blueprintable, meta = (BlueprintSpawnableComponent))
-class LEARNINGAGENTS_API ULearningAgentsPolicy : public ULearningAgentsManagerComponent
+class LEARNINGAGENTS_API ULearningAgentsPolicy : public ULearningAgentsManagerListener
 {
 	GENERATED_BODY()
 
 public:
 
+	friend class ULearningAgentsCritic;
+	friend class ULearningAgentsTrainer;
+
 	// These constructors/destructors are needed to make forward declarations happy
 	ULearningAgentsPolicy();
 	ULearningAgentsPolicy(FVTableHelper& Helper);
 	virtual ~ULearningAgentsPolicy();
 
+	/**
+	 * Constructs this object to be used with the given agent interactor and policy settings.
+	 * 
+	 * @param InManager						The input Manager
+	 * @param InInteractor					The input Interactor component
+	 * @param Class							The policy class
+	 * @param Name							The policy name
+	 * @param EncoderNeuralNetworkAsset		Optional Encoder Network Asset to use. If not provided, asset is empty, or
+	 *										bReinitializeEncoderNetwork is set then a new neural network object will be
+	 *										created.
+	 * @param PolicyNeuralNetworkAsset		Optional Policy Network Asset to use. If not provided, asset is empty, or
+	 *										bReinitializePolicyNetwork is set then a new neural network object will be 
+	 *										created according to the given PolicySettings.
+	 * @param DecoderNeuralNetworkAsset		Optional Decoder Network Asset to use. If not provided, asset is empty, or
+	 *										bReinitializeDecoderNetwork is set then a new neural network object will be
+	 *										created.
+	 * @param bReinitializeEncoderNetwork	If to reinitialize the encoder network
+	 * @param bReinitializePolicyNetwork	If to reinitialize the policy network
+	 * @param bReinitializeDecoderNetwork	If to reinitialize the decoder network
+	 * @param PolicySettings				The policy settings to use on creation of a new policy network
+	 * @param Seed							Random seed to use for initializing network weights and policy sampling
+	 */
+	UFUNCTION(BlueprintCallable, Category = "LearningAgents", meta=(Class = "/Script/LearningAgents.LearningAgentsPolicy", DeterminesOutputType = "Class"))
+	static ULearningAgentsPolicy* MakePolicy(
+		ULearningAgentsManager* InManager,
+		ULearningAgentsInteractor* InInteractor,
+		TSubclassOf<ULearningAgentsPolicy> Class,
+		const FName Name = TEXT("Policy"),
+		ULearningAgentsNeuralNetwork* EncoderNeuralNetworkAsset = nullptr,
+		ULearningAgentsNeuralNetwork* PolicyNeuralNetworkAsset = nullptr,
+		ULearningAgentsNeuralNetwork* DecoderNeuralNetworkAsset = nullptr,
+		const bool bReinitializeEncoderNetwork = true,
+		const bool bReinitializePolicyNetwork = true,
+		const bool bReinitializeDecoderNetwork = true,
+		const FLearningAgentsPolicySettings& PolicySettings = FLearningAgentsPolicySettings(),
+		const int32 Seed = 1234);
+
 	/**
 	 * Initializes this object to be used with the given agent interactor and policy settings.
-	 * @param InInteractor The input Interactor component
-	 * @param PolicySettings The policy settings to use
-	 * @param NeuralNetworkAsset Optional Network Asset to use. If provided must match the given PolicySettings. If not
-	 * provided or asset is empty then a new neural network object will be created according to the given 
-	 * PolicySettings and used.
+	 * 
+	 * @param InManager						The input Manager
+	 * @param InInteractor					The input Interactor component
+	 * @param EncoderNeuralNetworkAsset		Optional Encoder Network Asset to use. If not provided, asset is empty, or
+	 *										bReinitializeEncoderNetwork is set then a new neural network object will be
+	 *										created.
+	 * @param PolicyNeuralNetworkAsset		Optional Policy Network Asset to use. If not provided, asset is empty, or
+	 *										bReinitializePolicyNetwork is set then a new neural network object will be
+	 *										created according to the given PolicySettings.
+	 * @param DecoderNeuralNetworkAsset		Optional Decoder Network Asset to use. If not provided, asset is empty, or
+	 *										bReinitializeDecoderNetwork is set then a new neural network object will be
+	 *										created.
+	 * @param bReinitializeEncoderNetwork	If to reinitialize the encoder network
+	 * @param bReinitializePolicyNetwork	If to reinitialize the policy network
+	 * @param bReinitializeDecoderNetwork	If to reinitialize the decoder network
+	 * @param PolicySettings				The policy settings to use on creation of a new policy network
+	 * @param Seed							Random seed to use for initializing network weights and policy sampling
 	 */
 	UFUNCTION(BlueprintCallable, Category = "LearningAgents")
-	void SetupPolicy(ULearningAgentsInteractor* InInteractor, 
+	void SetupPolicy(
+		ULearningAgentsManager* InManager,
+		ULearningAgentsInteractor* InInteractor,
+		ULearningAgentsNeuralNetwork* EncoderNeuralNetworkAsset = nullptr,
+		ULearningAgentsNeuralNetwork* PolicyNeuralNetworkAsset = nullptr,
+		ULearningAgentsNeuralNetwork* DecoderNeuralNetworkAsset = nullptr,
+		const bool bReinitializeEncoderNetwork = true,
+		const bool bReinitializePolicyNetwork = true,
+		const bool bReinitializeDecoderNetwork = true,
 		const FLearningAgentsPolicySettings& PolicySettings = FLearningAgentsPolicySettings(),
-		ULearningAgentsNeuralNetwork* NeuralNetworkAsset = nullptr);
+		const int32 Seed = 1234);
 
 public:
 
-	//~ Begin ULearningAgentsManagerComponent Interface
-	virtual void OnAgentsAdded(const TArray<int32>& AgentIds) override;
-	virtual void OnAgentsRemoved(const TArray<int32>& AgentIds) override;
-	virtual void OnAgentsReset(const TArray<int32>& AgentIds) override;
-	//~ End ULearningAgentsManagerComponent Interface
+	//~ Begin ULearningAgentsManagerListener Interface
+	virtual void OnAgentsAdded_Implementation(const TArray<int32>& AgentIds) override;
+	virtual void OnAgentsRemoved_Implementation(const TArray<int32>& AgentIds) override;
+	virtual void OnAgentsReset_Implementation(const TArray<int32>& AgentIds) override;
+	//~ End ULearningAgentsManagerListener Interface
 
-// ----- Load / Save -----
-public:
 
-	/**
-	 * Load a snapshot's weights into this policy.
-	 * @param File The snapshot file.
-	 */
-	UFUNCTION(BlueprintCallable, Category = "LearningAgents", meta = (RelativePath))
-	void LoadPolicyFromSnapshot(const FFilePath& File);
-
-	/**
-	 * Save this policy's weights into a snapshot.
-	 * @param File The snapshot file.
-	 */
-	UFUNCTION(BlueprintCallable, BlueprintPure = false, Category = "LearningAgents", meta = (RelativePath))
-	void SavePolicyToSnapshot(const FFilePath& File) const;
+// ----- Blueprint public interface -----
+public:
 
 	/**
-	 * Use a ULearningAgentsNeuralNetwork asset directly for this critic rather than making a copy. If used
-	 * during training then this asset's weights will be updated as training progresses.
-	 * @param NeuralNetworkAsset The asset to use.
+	 * Encodes the buffered observation vectors using the Encoder network. This should be called after GatherObservations and before EvaluatePolicy.
 	 */
 	UFUNCTION(BlueprintCallable, Category = "LearningAgents")
-	void UsePolicyFromAsset(ULearningAgentsNeuralNetwork* NeuralNetworkAsset);
+	void EncodeObservations();
 
 	/**
-	 * Load a ULearningAgentsNeuralNetwork asset's weights into this policy.
-	 * @param NeuralNetworkAsset The asset to load from.
+	 * Calling this function will run the underlying neural network on the previously encoded observations to populate
+	 * the encoded actions. This should be called after EncodeObservations and before DecodeAndSampleActions.
 	 */
 	UFUNCTION(BlueprintCallable, Category = "LearningAgents")
-	void LoadPolicyFromAsset(ULearningAgentsNeuralNetwork* NeuralNetworkAsset);
-
-	/**
-	 * Save this policy's weights to a ULearningAgentsNeuralNetwork asset.
-	 * @param NeuralNetworkAsset The asset to save to.
-	 */
-	UFUNCTION(BlueprintCallable, Category = "LearningAgents", Meta = (DevelopmentOnly))
-	void SavePolicyToAsset(ULearningAgentsNeuralNetwork* NeuralNetworkAsset);
-
-// ----- Evaluation -----
-public:
+	void EvaluatePolicy();
 
 	/**
-	 * Calling this function will run the underlying neural network on the previously buffered observations to populate
-	 * the output action buffer. This should be called after the associated agent interactor's EncodeObservations and
-	 * before its DecodeActions.
+	 * Decodes and samples action vectors using the Decoder network. This should be called after EvaluatePolicy and before Scatter Actions.
+	 * 
+	 * @param ActionNoiseScale		Scale of the action noise to use during sampling. Set this to zero to always sample the mean (expected) action.
 	 */
 	UFUNCTION(BlueprintCallable, Category = "LearningAgents")
-	void EvaluatePolicy();
+	void DecodeAndSampleActions(const float ActionNoiseScale = 1.0f);
 
 	/**
-	 * Calls EncodeObservations, followed by EvaluatePolicy, followed by DecodeActions
+	 * Calls GatherObservations, EncodeObservations, EvaluatePolicy, DecodeAndSampleActions, ScatterActions
 	 */
 	UFUNCTION(BlueprintCallable, Category = "LearningAgents")
-	void RunInference();
-
-	/** Get the action noise scale used. */
-	UFUNCTION(BlueprintPure, Category = "LearningAgents")
-	float GetActionNoiseScale() const;
-
-	/** Set the action noise scale used. */
-	UFUNCTION(BlueprintCallable, Category = "LearningAgents")
-	void SetActionNoiseScale(const float ActionNoiseScale);
+	void RunInference(const float ActionNoiseScale = 1.0f);
 
 	/**
 	 * Gets the current memory state for a given agent as represented by an abstract vector learned by the policy.
@@ -178,32 +194,34 @@ public:
 	UFUNCTION(BlueprintCallable, Category = "LearningAgents", Meta = (AgentId = -1))
 	void SetMemoryState(const int32 AgentId, const TArray<float>& InMemoryState);
 
-// ----- Non-blueprint public interface -----
-public:
+	/** Gets the size of the memory state */
+	UFUNCTION(BlueprintPure, Category = "LearningAgents")
+	int32 GetMemoryStateSize() const;
 
-	/** Gets the current Network Asset being used */
-	ULearningAgentsNeuralNetwork* GetNetworkAsset();
+	/** Gets the current Encoder Network Asset being used */
+	UFUNCTION(BlueprintPure, Category = "LearningAgents")
+	ULearningAgentsNeuralNetwork* GetEncoderNetworkAsset();
 
-	/** Get a reference to this policy's neural network. */
-	UE::Learning::INeuralNetwork& GetPolicyNetwork();
+	/** Gets the current Policy Network Asset being used */
+	UFUNCTION(BlueprintPure, Category = "LearningAgents")
+	ULearningAgentsNeuralNetwork* GetPolicyNetworkAsset();
 
-	/** Get a reference to this policy's policy function object. */
-	UE::Learning::FNeuralNetworkPolicyFunction& GetPolicyObject();
+	/** Gets the current Encoder Network Asset being used */
+	UFUNCTION(BlueprintPure, Category = "LearningAgents")
+	ULearningAgentsNeuralNetwork* GetDecoderNetworkAsset();
 
-	/** Gets a view of the memory state for the given agent before evaluation */
-	TLearningArrayView<2, const float> GetPreEvaluationMemoryStateView() const;
 
-	/** Gets a view of the memory state for the given agent before evaluation */
-	TLearningArrayView<2, float> GetPreEvaluationMemoryStateView();
+// ----- Non-blueprint public interface -----
+public:
 
-	/** Gets a view of the current memory state for the given agent */
-	TLearningArrayView<2, const float> GetMemoryStateView() const;
+	/** Get a reference to this policy's encoder function object. */
+	UE::Learning::FNeuralNetworkFunction& GetEncoderObject();
 
-	/** Gets a view of the current memory state for the given agent */
-	TLearningArrayView<2, float> GetMemoryStateView();
+	/** Get a reference to this policy's policy function object. */
+	UE::Learning::FNeuralNetworkPolicy& GetPolicyObject();
 
-	/** Gets the size of the memory state */
-	int32 GetMemoryStateSize() const;
+	/** Get a reference to this policy's decoder function object. */
+	UE::Learning::FNeuralNetworkFunction& GetDecoderObject();
 
 // ----- Private Data -----
 private:
@@ -212,35 +230,58 @@ private:
 	UPROPERTY(VisibleAnywhere, Transient, Category = "LearningAgents")
 	TObjectPtr<ULearningAgentsInteractor> Interactor;
 
-	/** The underlying neural network. */
+	/** The underlying encoder neural network. */
 	UPROPERTY(VisibleAnywhere, Transient, Category = "LearningAgents")
-	TObjectPtr<ULearningAgentsNeuralNetwork> Network;
+	TObjectPtr<ULearningAgentsNeuralNetwork> EncoderNetwork;
 
-	/** Internal Policy Function Object */
-	TSharedPtr<UE::Learning::FNeuralNetworkPolicyFunction> PolicyObject;
+	/** The underlying policy neural network. */
+	UPROPERTY(VisibleAnywhere, Transient, Category = "LearningAgents")
+	TObjectPtr<ULearningAgentsNeuralNetwork> PolicyNetwork;
+
+	/** The underlying decoder neural network. */
+	UPROPERTY(VisibleAnywhere, Transient, Category = "LearningAgents")
+	TObjectPtr<ULearningAgentsNeuralNetwork> DecoderNetwork;
 
 	/** The internal memory state of each agent before evaluation. */
-	UE::Learning::TArrayMapHandle<2, float> PreEvaluationMemoryStateHandle;
+	TLearningArray<2, float> PreEvaluationMemoryState;
 
 	/** The internal memory state of each agent. */
-	UE::Learning::TArrayMapHandle<2, float> MemoryStateHandle;
+	TLearningArray<2, float> MemoryState;
 
-// ----- Private Iteration Checks ----- 
-private:
+	/** The global random seed state used to initialize network weights and during sampling */
+	uint32 GlobalSeed = 0;
+
+	/** The individual per-agent seeds used during sampling of actions. */
+	TLearningArray<1, uint32> Seeds;
+
+	/** Encoded observation vector buffers */
+	TLearningArray<2, float> ObservationVectorsEncoded;
+
+	/** Encoded action vector buffers */
+	TLearningArray<2, float> ActionVectorsEncoded;
+
+	/** Action distribution vector buffers */
+	TLearningArray<2, float> ActionDistributionVectors;
+
+	/** Internal Encoder Function Object */
+	TSharedPtr<UE::Learning::FNeuralNetworkFunction> EncoderObject;
 
-	/** Number of times policy has been evaluated for all agents. */
-	TLearningArray<1, uint64, TInlineAllocator<32>> PolicyAgentIteration;
+	/** Internal Policy Function Object */
+	TSharedPtr<UE::Learning::FNeuralNetworkPolicy> PolicyObject;
+
+	/** Internal Decoder Function Object */
+	TSharedPtr<UE::Learning::FNeuralNetworkFunction> DecoderObject;
+
+	/** Number of times encoded observation vector has been set for all agents */
+	TLearningArray<1, uint64, TInlineAllocator<32>> ObservationVectorEncodedIteration;
+
+	/** Number of times encoded action vector has been set for all agents */
+	TLearningArray<1, uint64, TInlineAllocator<32>> ActionVectorEncodedIteration;
+
+	/** Number of times the memory state has been updated for all agents. */
+	TLearningArray<1, uint64, TInlineAllocator<32>> MemoryStateIteration;
 
 	/** Temp buffers used to record the set of agents that are valid for evaluation */
 	TArray<int32> ValidAgentIds;
 	UE::Learning::FIndexSet ValidAgentSet;
-
-private:
-#if UE_LEARNING_AGENTS_ENABLE_VISUAL_LOG
-	/** Color used to draw this action in the visual log */
-	FLinearColor VisualLogColor = FColor::Purple;
-
-	/** Describes this policy to the visual logger for debugging purposes. */
-	void VisualLog(const UE::Learning::FIndexSet AgentSet) const;
-#endif
 };
diff --git a/Engine/Plugins/Experimental/LearningAgents/Source/LearningAgentsTraining/Private/LearningAgentsCompletions.cpp b/Engine/Plugins/Experimental/LearningAgents/Source/LearningAgentsTraining/Private/LearningAgentsCompletions.cpp
--- a/Engine/Plugins/Experimental/LearningAgents/Source/LearningAgentsTraining/Private/LearningAgentsCompletions.cpp
+++ b/Engine/Plugins/Experimental/LearningAgents/Source/LearningAgentsTraining/Private/LearningAgentsCompletions.cpp
@@ -2,363 +2,123 @@
 
 #include "LearningAgentsCompletions.h"
 
-#include "LearningAgentsTrainer.h"
-#include "LearningAgentsManager.h"
-#include "LearningArray.h"
-#include "LearningArrayMap.h"
 #include "LearningCompletion.h"
-#include "LearningCompletionObject.h"
 #include "LearningLog.h"
 
-#include "GameFramework/Actor.h"
-
-namespace UE::Learning::Agents::Completions::Private
+namespace UE::Learning::Agents
 {
-	template<typename CompletionUObject, typename CompletionFObject, typename... InArgTypes>
-	CompletionUObject* AddCompletion(ULearningAgentsTrainer* InAgentTrainer, const FName Name, const TCHAR* FunctionName, InArgTypes&& ...Args)
+	ELearningAgentsCompletion GetLearningAgentsCompletion(const ECompletionMode CompletionMode)
 	{
-		if (!InAgentTrainer)
+		switch (CompletionMode)
 		{
-			UE_LOG(LogLearning, Error, TEXT("%s: InAgentTrainer is nullptr."), FunctionName);
-			return nullptr;
+		case ECompletionMode::Running: return ELearningAgentsCompletion::Running;
+		case ECompletionMode::Terminated: return ELearningAgentsCompletion::Termination;
+		case ECompletionMode::Truncated: return ELearningAgentsCompletion::Truncation;
+		default:UE_LOG(LogLearning, Error, TEXT("Unknown Completion Mode.")); return ELearningAgentsCompletion::Running;
 		}
+	}
 
-		if (!InAgentTrainer->HasAgentManager())
+	ECompletionMode GetCompletionMode(const ELearningAgentsCompletion Completion)
+	{
+		switch (Completion)
 		{
-			UE_LOG(LogLearning, Error, TEXT("%s: Must be attached to a LearningAgentsManager Actor."), *InAgentTrainer->GetName());
-			return nullptr;
+		case ELearningAgentsCompletion::Running: return ECompletionMode::Running;
+		case ELearningAgentsCompletion::Termination: return ECompletionMode::Terminated;
+		case ELearningAgentsCompletion::Truncation: return ECompletionMode::Truncated;
+		default:UE_LOG(LogLearning, Error, TEXT("Unknown Completion.")); return ECompletionMode::Running;
 		}
-
-		const FName UniqueName = MakeUniqueObjectName(InAgentTrainer, CompletionUObject::StaticClass(), Name, EUniqueObjectNameOptions::GloballyUnique);
-
-		CompletionUObject* Completion = NewObject<CompletionUObject>(InAgentTrainer, UniqueName);
-		Completion->Init(InAgentTrainer->GetAgentManager()->GetMaxAgentNum());
-		Completion->AgentTrainer = InAgentTrainer;
-		Completion->CompletionObject = MakeShared<CompletionFObject>(
-			Completion->GetFName(),
-			InAgentTrainer->GetAgentManager()->GetInstanceData().ToSharedRef(),
-			InAgentTrainer->GetAgentManager()->GetMaxAgentNum(),
-			Forward<InArgTypes>(Args)...);
-
-		InAgentTrainer->AddCompletion(Completion, Completion->CompletionObject.ToSharedRef());
-
-		return Completion;
 	}
 }
 
-//------------------------------------------------------------------
-
-void ULearningAgentsCompletion::Init(const int32 MaxAgentNum)
+bool ULearningAgentsCompletions::CompletionIsRunning(const ELearningAgentsCompletion Completion)
 {
-	AgentIteration.SetNumUninitialized({ MaxAgentNum });
-	UE::Learning::Array::Set<1, uint64>(AgentIteration, INDEX_NONE);
+	return Completion == ELearningAgentsCompletion::Running;
 }
 
-void ULearningAgentsCompletion::OnAgentsAdded(const TArray<int32>& AgentIds)
+bool ULearningAgentsCompletions::CompletionIsCompleted(const ELearningAgentsCompletion Completion)
 {
-	UE::Learning::Array::Set<1, uint64>(AgentIteration, 0, AgentIds);
+	return Completion == ELearningAgentsCompletion::Truncation || Completion == ELearningAgentsCompletion::Termination;
 }
 
-void ULearningAgentsCompletion::OnAgentsRemoved(const TArray<int32>& AgentIds)
+bool ULearningAgentsCompletions::CompletionIsTruncation(const ELearningAgentsCompletion Completion)
 {
-	UE::Learning::Array::Set<1, uint64>(AgentIteration, INDEX_NONE, AgentIds);
+	return Completion == ELearningAgentsCompletion::Truncation;
 }
 
-void ULearningAgentsCompletion::OnAgentsReset(const TArray<int32>& AgentIds)
+bool ULearningAgentsCompletions::CompletionIsTermination(const ELearningAgentsCompletion Completion)
 {
-	UE::Learning::Array::Set<1, uint64>(AgentIteration, 0, AgentIds);
+	return Completion == ELearningAgentsCompletion::Termination;
 }
 
-uint64 ULearningAgentsCompletion::GetAgentIteration(const int32 AgentId) const
+ELearningAgentsCompletion ULearningAgentsCompletions::CompletionOr(ELearningAgentsCompletion A, ELearningAgentsCompletion B)
 {
-	return AgentIteration[AgentId];
-}
-
-//------------------------------------------------------------------
-
-UConditionalCompletion* UConditionalCompletion::AddConditionalCompletion(ULearningAgentsTrainer* InAgentTrainer, const FName Name, const ELearningAgentsCompletion InCompletionMode)
-{
-	return UE::Learning::Agents::Completions::Private::AddCompletion<UConditionalCompletion, UE::Learning::FConditionalCompletion>(InAgentTrainer, Name, TEXT("AddConditionalCompletion"), UE::Learning::Agents::GetCompletionMode(InCompletionMode));
-}
-
-void UConditionalCompletion::SetConditionalCompletion(const int32 AgentId, const bool bIsComplete)
-{
-	if (!AgentTrainer->HasAgent(AgentId))
+	if ((A == ELearningAgentsCompletion::Running && B != ELearningAgentsCompletion::Running) ||
+		(A == ELearningAgentsCompletion::Truncation && B == ELearningAgentsCompletion::Termination))
 	{
-		UE_LOG(LogLearning, Error, TEXT("%s: AgentId %d not found in the agents set."), *GetName(), AgentId);
-		return;
+		return B;
 	}
-
-	CompletionObject->InstanceData->View(CompletionObject->ConditionHandle)[AgentId] = bIsComplete;
-	AgentIteration[AgentId]++;
-}
-
-#if UE_LEARNING_AGENTS_ENABLE_VISUAL_LOG
-void UConditionalCompletion::VisualLog(const UE::Learning::FIndexSet Instances) const
-{
-	UE_LEARNING_TRACE_CPUPROFILER_EVENT_SCOPE(UConditionalCompletion::VisualLog);
-
-	const TLearningArrayView<1, const bool> ConditionView = CompletionObject->InstanceData->ConstView(CompletionObject->ConditionHandle);
-	const TLearningArrayView<1, const UE::Learning::ECompletionMode> CompletionView = CompletionObject->InstanceData->ConstView(CompletionObject->CompletionHandle);
-
-	for (const int32 Instance : Instances)
+	else
 	{
-		if (const AActor* Actor = Cast<AActor>(AgentTrainer->GetAgent(Instance)))
-		{
-			UE_LEARNING_AGENTS_VLOG_STRING(this, LogLearning, Display,
-				Actor->GetActorLocation(),
-				VisualLogColor.ToFColor(true),
-				TEXT("Agent %i\nCondition: %s\nCompletion: %s"),
-				Instance,
-				ConditionView[Instance] ? TEXT("true") : TEXT("false"),
-				UE::Learning::Completion::CompletionModeString(CompletionView[Instance]));
-		}
+		return A;
 	}
 }
-#endif
-
-//------------------------------------------------------------------
 
-UTimeElapsedCompletion* UTimeElapsedCompletion::AddTimeElapsedCompletion(ULearningAgentsTrainer* InAgentTrainer, const FName Name, const float Threshold, const ELearningAgentsCompletion InCompletionMode)
+ELearningAgentsCompletion ULearningAgentsCompletions::CompletionAnd(ELearningAgentsCompletion A, ELearningAgentsCompletion B)
 {
-	return UE::Learning::Agents::Completions::Private::AddCompletion<UTimeElapsedCompletion, UE::Learning::FTimeElapsedCompletion>(InAgentTrainer, Name, TEXT("AddTimeElapsedCompletion"), Threshold, UE::Learning::Agents::GetCompletionMode(InCompletionMode));
-}
-
-void UTimeElapsedCompletion::SetTimeElapsedCompletion(const int32 AgentId, const float Time)
-{
-	if (!AgentTrainer->HasAgent(AgentId))
+	if (A == ELearningAgentsCompletion::Running ||
+		(A == ELearningAgentsCompletion::Truncation && B != ELearningAgentsCompletion::Running))
 	{
-		UE_LOG(LogLearning, Error, TEXT("%s: AgentId %d not found in the agents set."), *GetName(), AgentId);
-		return;
+		return A;
 	}
-
-	CompletionObject->InstanceData->View(CompletionObject->TimeHandle)[AgentId] = Time;
-	AgentIteration[AgentId]++;
-}
-
-#if UE_LEARNING_AGENTS_ENABLE_VISUAL_LOG
-void UTimeElapsedCompletion::VisualLog(const UE::Learning::FIndexSet Instances) const
-{
-	UE_LEARNING_TRACE_CPUPROFILER_EVENT_SCOPE(UTimeElapsedCompletion::VisualLog);
-
-	const TLearningArrayView<1, const float> TimeView = CompletionObject->InstanceData->ConstView(CompletionObject->TimeHandle);
-	const TLearningArrayView<1, const float> ThresholdView = CompletionObject->InstanceData->ConstView(CompletionObject->ThresholdHandle);
-	const TLearningArrayView<1, const UE::Learning::ECompletionMode> CompletionView = CompletionObject->InstanceData->ConstView(CompletionObject->CompletionHandle);
-
-	for (const int32 Instance : Instances)
+	else
 	{
-		if (const AActor* Actor = Cast<AActor>(AgentTrainer->GetAgent(Instance)))
-		{
-			UE_LEARNING_AGENTS_VLOG_STRING(this, LogLearning, Display,
-				Actor->GetActorLocation(),
-				VisualLogColor.ToFColor(true),
-				TEXT("Agent %i\nThreshold: [%6.3f]\nTime: [%6.3f]\nCompletion: %s"),
-				Instance,
-				ThresholdView[Instance],
-				TimeView[Instance],
-				UE::Learning::Completion::CompletionModeString(CompletionView[Instance]));
-		}
+		return B;
 	}
 }
-#endif
-
-//------------------------------------------------------------------
 
-UPlanarPositionDifferenceCompletion* UPlanarPositionDifferenceCompletion::AddPlanarPositionDifferenceCompletion(
-	ULearningAgentsTrainer* InAgentTrainer,
-	const FName Name,
-	const float Threshold,
-	const ELearningAgentsCompletion InCompletionMode,
-	const FVector Axis0,
-	const FVector Axis1)
+ELearningAgentsCompletion ULearningAgentsCompletions::CompletionNot(const ELearningAgentsCompletion A, const ELearningAgentsCompletion NotRunningType)
 {
-	return UE::Learning::Agents::Completions::Private::AddCompletion<UPlanarPositionDifferenceCompletion, UE::Learning::FPlanarPositionDifferenceCompletion>(
-		InAgentTrainer,
-		Name,
-		TEXT("AddPlanarPositionDifferenceCompletion"),
-		1,
-		Threshold,
-		UE::Learning::Agents::GetCompletionMode(InCompletionMode),
-		Axis0.GetSafeNormal(UE_SMALL_NUMBER, FVector::ForwardVector),
-		Axis1.GetSafeNormal(UE_SMALL_NUMBER, FVector::RightVector));
+	return A == ELearningAgentsCompletion::Running ? NotRunningType : ELearningAgentsCompletion::Running;
 }
 
-void UPlanarPositionDifferenceCompletion::SetPlanarPositionDifferenceCompletion(const int32 AgentId, const FVector Position0, const FVector Position1)
+ELearningAgentsCompletion ULearningAgentsCompletions::CompletionOnCondition(const bool bCondition, const ELearningAgentsCompletion CompletionType)
 {
-	if (!AgentTrainer->HasAgent(AgentId))
-	{
-		UE_LOG(LogLearning, Error, TEXT("%s: AgentId %d not found in the agents set."), *GetName(), AgentId);
-		return;
-	}
-
-	CompletionObject->InstanceData->View(CompletionObject->Position0Handle)[AgentId][0] = Position0;
-	CompletionObject->InstanceData->View(CompletionObject->Position1Handle)[AgentId][0] = Position1;
-	AgentIteration[AgentId]++;
+	return bCondition ? CompletionType : ELearningAgentsCompletion::Running;
 }
 
-#if UE_LEARNING_AGENTS_ENABLE_VISUAL_LOG
-void UPlanarPositionDifferenceCompletion::VisualLog(const UE::Learning::FIndexSet Instances) const
+ELearningAgentsCompletion ULearningAgentsCompletions::CompletionOnTimeElapsed(const float Time, const float TimeThreshold, const ELearningAgentsCompletion CompletionType)
 {
-	UE_LEARNING_TRACE_CPUPROFILER_EVENT_SCOPE(UPlanarPositionDifferenceCompletion::VisualLog);
-
-	const TLearningArrayView<2, const FVector> Position0View = CompletionObject->InstanceData->ConstView(CompletionObject->Position0Handle);
-	const TLearningArrayView<2, const FVector> Position1View = CompletionObject->InstanceData->ConstView(CompletionObject->Position1Handle);
-	const TLearningArrayView<1, const float> ThresholdView = CompletionObject->InstanceData->ConstView(CompletionObject->ThresholdHandle);
-	const TLearningArrayView<1, const UE::Learning::ECompletionMode> CompletionView = CompletionObject->InstanceData->ConstView(CompletionObject->CompletionHandle);
-
-	for (const int32 Instance : Instances)
-	{
-		if (const AActor* Actor = Cast<AActor>(AgentTrainer->GetAgent(Instance)))
-		{
-			const FVector Position0 = Position0View[Instance][0];
-			const FVector Position1 = Position1View[Instance][0];
-
-			const FVector PlanarPosition0 = FVector(CompletionObject->Axis0.Dot(Position0), CompletionObject->Axis1.Dot(Position0), 0.0f);
-			const FVector PlanarPosition1 = FVector(CompletionObject->Axis0.Dot(Position1), CompletionObject->Axis1.Dot(Position1), 0.0f);
-
-			UE_LEARNING_AGENTS_VLOG_LOCATION(this, LogLearning, Display,
-				Position0,
-				10,
-				VisualLogColor.ToFColor(true),
-				TEXT("Position0: [% 6.1f % 6.1f % 6.1f]\nPlanar Position0: [% 6.1f % 6.1f]"),
-				Position0.X, Position0.Y, Position0.Z,
-				PlanarPosition0.X, PlanarPosition0.Y);
-
-			UE_LEARNING_AGENTS_VLOG_PLANE(this, LogLearning, Display,
-				Position0,
-				FQuat::Identity,
-				CompletionObject->Axis0,
-				CompletionObject->Axis1,
-				VisualLogColor.ToFColor(true),
-				TEXT(""));
-
-			UE_LEARNING_AGENTS_VLOG_LOCATION(this, LogLearning, Display,
-				Position1,
-				10,
-				VisualLogColor.ToFColor(true),
-				TEXT("Position1: [% 6.1f % 6.1f % 6.1f]\nPlanar Position1: [% 6.1f % 6.1f]"),
-				Position1.X, Position1.Y, Position1.Z,
-				PlanarPosition1.X, PlanarPosition1.Y);
-
-			UE_LEARNING_AGENTS_VLOG_PLANE(this, LogLearning, Display,
-				Position1,
-				FQuat::Identity,
-				CompletionObject->Axis0,
-				CompletionObject->Axis1,
-				VisualLogColor.ToFColor(true),
-				TEXT(""));
-
-			UE_LEARNING_AGENTS_VLOG_SEGMENT(this, LogLearning, Display,
-				Position0,
-				Position1,
-				VisualLogColor.ToFColor(true),
-				TEXT("Agent %i\nDistance: [% 6.3f]\nPlanar Distance: [% 6.3f]\nThreshold: [% 6.2f]\nCompletion: %s"),
-				Instance,
-				FVector::Distance(Position0, Position1),
-				FVector::Distance(PlanarPosition0, PlanarPosition1),
-				ThresholdView[Instance],
-				UE::Learning::Completion::CompletionModeString(CompletionView[Instance]));
-		}
-	}
+	return CompletionOnCondition(Time > TimeThreshold, CompletionType);
 }
-#endif
 
-
-UPlanarPositionSimilarityCompletion* UPlanarPositionSimilarityCompletion::AddPlanarPositionSimilarityCompletion(
-	ULearningAgentsTrainer* InAgentTrainer,
-	const FName Name,
-	const float Threshold,
-	const ELearningAgentsCompletion InCompletionMode,
-	const FVector Axis0,
-	const FVector Axis1)
+ELearningAgentsCompletion ULearningAgentsCompletions::CompletionOnEpisodeStepsRecorded(const int32 EpisodeSteps, const int32 MaxEpisodeSteps, const ELearningAgentsCompletion CompletionType)
 {
-	return UE::Learning::Agents::Completions::Private::AddCompletion<UPlanarPositionSimilarityCompletion, UE::Learning::FPlanarPositionSimilarityCompletion>(
-		InAgentTrainer,
-		Name,
-		TEXT("AddPlanarPositionSimilarityCompletion"),
-		1,
-		Threshold,
-		UE::Learning::Agents::GetCompletionMode(InCompletionMode),
-		Axis0.GetSafeNormal(UE_SMALL_NUMBER, FVector::ForwardVector),
-		Axis1.GetSafeNormal(UE_SMALL_NUMBER, FVector::RightVector));
+	return CompletionOnCondition(EpisodeSteps >= MaxEpisodeSteps, CompletionType);
 }
 
-void UPlanarPositionSimilarityCompletion::SetPlanarPositionSimilarityCompletion(const int32 AgentId, const FVector Position0, const FVector Position1)
+ELearningAgentsCompletion ULearningAgentsCompletions::CompletionOnTranslationDifferenceBelowThreshold(const FVector TranslationA, const FVector TranslationB, const float DistanceThreshold, const ELearningAgentsCompletion CompletionType)
 {
-	if (!AgentTrainer->HasAgent(AgentId))
-	{
-		UE_LOG(LogLearning, Error, TEXT("%s: AgentId %d not found in the agents set."), *GetName(), AgentId);
-		return;
-	}
-
-	CompletionObject->InstanceData->View(CompletionObject->Position0Handle)[AgentId][0] = Position0;
-	CompletionObject->InstanceData->View(CompletionObject->Position1Handle)[AgentId][0] = Position1;
-	AgentIteration[AgentId]++;
+	return CompletionOnCondition(FVector::Distance(TranslationA, TranslationB) < DistanceThreshold, CompletionType);
 }
 
-#if UE_LEARNING_AGENTS_ENABLE_VISUAL_LOG
-void UPlanarPositionSimilarityCompletion::VisualLog(const UE::Learning::FIndexSet Instances) const
+ELearningAgentsCompletion ULearningAgentsCompletions::CompletionOnTranslationDifferenceAboveThreshold(const FVector TranslationA, const FVector TranslationB, const float DistanceThreshold, const ELearningAgentsCompletion CompletionType)
 {
-	UE_LEARNING_TRACE_CPUPROFILER_EVENT_SCOPE(UPlanarPositionSimilarityCompletion::VisualLog);
-
-	const TLearningArrayView<2, const FVector> Position0View = CompletionObject->InstanceData->ConstView(CompletionObject->Position0Handle);
-	const TLearningArrayView<2, const FVector> Position1View = CompletionObject->InstanceData->ConstView(CompletionObject->Position1Handle);
-	const TLearningArrayView<1, const float> ThresholdView = CompletionObject->InstanceData->ConstView(CompletionObject->ThresholdHandle);
-	const TLearningArrayView<1, const UE::Learning::ECompletionMode> CompletionView = CompletionObject->InstanceData->ConstView(CompletionObject->CompletionHandle);
-
-	for (const int32 Instance : Instances)
-	{
-		if (const AActor* Actor = Cast<AActor>(AgentTrainer->GetAgent(Instance)))
-		{
-			const FVector Position0 = Position0View[Instance][0];
-			const FVector Position1 = Position1View[Instance][0];
-
-			const FVector PlanarPosition0 = FVector(CompletionObject->Axis0.Dot(Position0), CompletionObject->Axis1.Dot(Position0), 0.0f);
-			const FVector PlanarPosition1 = FVector(CompletionObject->Axis0.Dot(Position1), CompletionObject->Axis1.Dot(Position1), 0.0f);
-
-			UE_LEARNING_AGENTS_VLOG_LOCATION(this, LogLearning, Display,
-				Position0,
-				10,
-				VisualLogColor.ToFColor(true),
-				TEXT("Position0: [% 6.1f % 6.1f % 6.1f]\nPlanar Position0: [% 6.1f % 6.1f]"),
-				Position0.X, Position0.Y, Position0.Z,
-				PlanarPosition0.X, PlanarPosition0.Y);
-
-			UE_LEARNING_AGENTS_VLOG_PLANE(this, LogLearning, Display,
-				Position0,
-				FQuat::Identity,
-				CompletionObject->Axis0,
-				CompletionObject->Axis1,
-				VisualLogColor.ToFColor(true),
-				TEXT(""));
-
-			UE_LEARNING_AGENTS_VLOG_LOCATION(this, LogLearning, Display,
-				Position1,
-				10,
-				VisualLogColor.ToFColor(true),
-				TEXT("Position1: [% 6.1f % 6.1f % 6.1f]\nPlanar Position1: [% 6.1f % 6.1f]"),
-				Position1.X, Position1.Y, Position1.Z,
-				PlanarPosition1.X, PlanarPosition1.Y);
+	return CompletionOnCondition(FVector::Distance(TranslationA, TranslationB) > DistanceThreshold, CompletionType);
+}
 
-			UE_LEARNING_AGENTS_VLOG_PLANE(this, LogLearning, Display,
-				Position1,
-				FQuat::Identity,
-				CompletionObject->Axis0,
-				CompletionObject->Axis1,
-				VisualLogColor.ToFColor(true),
-				TEXT(""));
+ELearningAgentsCompletion ULearningAgentsCompletions::CompletionOnTranslationOutsideBounds(
+	const FVector Translation,
+	const FTransform BoundsTransform,
+	const FVector BoundsMins,
+	const FVector BoundsMaxs,
+	const ELearningAgentsCompletion CompletionType)
+{
+	const FVector LocalTranslation = BoundsTransform.InverseTransformPosition(Translation);
 
-			UE_LEARNING_AGENTS_VLOG_SEGMENT(this, LogLearning, Display,
-				Position0,
-				Position1,
-				VisualLogColor.ToFColor(true),
-				TEXT("Agent %i\nDistance: [% 6.3f]\nPlanar Distance: [% 6.3f]\nThreshold: [% 6.2f]\nCompletion: %s"),
-				Instance,
-				FVector::Distance(Position0, Position1),
-				FVector::Distance(PlanarPosition0, PlanarPosition1),
-				ThresholdView[Instance],
-				UE::Learning::Completion::CompletionModeString(CompletionView[Instance]));
-		}
-	}
+	return CompletionOnCondition(
+		LocalTranslation.X < BoundsMins.X || LocalTranslation.X > BoundsMaxs.X ||
+		LocalTranslation.Y < BoundsMins.Y || LocalTranslation.Y > BoundsMaxs.Y ||
+		LocalTranslation.Z < BoundsMins.Z || LocalTranslation.Z > BoundsMaxs.Z,
+		CompletionType);
 }
-#endif
+
diff --git a/Engine/Plugins/Experimental/LearningAgents/Source/LearningAgentsTraining/Private/LearningAgentsImitationTrainer.cpp b/Engine/Plugins/Experimental/LearningAgents/Source/LearningAgentsTraining/Private/LearningAgentsImitationTrainer.cpp
--- a/Engine/Plugins/Experimental/LearningAgents/Source/LearningAgentsTraining/Private/LearningAgentsImitationTrainer.cpp
+++ b/Engine/Plugins/Experimental/LearningAgents/Source/LearningAgentsTraining/Private/LearningAgentsImitationTrainer.cpp
@@ -2,15 +2,13 @@
 
 #include "LearningAgentsImitationTrainer.h"
 
+#include "LearningAgentsManager.h"
 #include "LearningAgentsInteractor.h"
-#include "LearningArrayMap.h"
 #include "LearningExperience.h"
-#include "LearningFeatureObject.h"
 #include "LearningLog.h"
 #include "LearningImitationTrainer.h"
 #include "LearningAgentsRecording.h"
 #include "LearningAgentsPolicy.h"
-#include "LearningNeuralNetworkObject.h"
 
 #include "GenericPlatform/GenericPlatformMisc.h"
 #include "HAL/FileManager.h"
@@ -20,27 +18,73 @@ ULearningAgentsImitationTrainer::ULearningAgentsImitationTrainer() : Super(FObje
 ULearningAgentsImitationTrainer::ULearningAgentsImitationTrainer(FVTableHelper& Helper) : Super(Helper) {}
 ULearningAgentsImitationTrainer::~ULearningAgentsImitationTrainer() = default;
 
-void ULearningAgentsImitationTrainer::EndPlay(const EEndPlayReason::Type EndPlayReason)
+void ULearningAgentsImitationTrainer::BeginDestroy()
 {
 	if (IsTraining())
 	{
 		EndTraining();
 	}
 
-	Super::EndPlay(EndPlayReason);
+	Super::BeginDestroy();
 }
 
-void ULearningAgentsImitationTrainer::BeginTraining(
-	ULearningAgentsPolicy* InPolicy, 
-	const ULearningAgentsRecording* Recording,
-	const FLearningAgentsImitationTrainerSettings& ImitationTrainerSettings,
-	const FLearningAgentsImitationTrainerTrainingSettings& ImitationTrainerTrainingSettings,
-	const FLearningAgentsTrainerPathSettings& ImitationTrainerPathSettings,
-	const bool bReinitializePolicyNetwork)
+ULearningAgentsImitationTrainer* ULearningAgentsImitationTrainer::MakeImitationTrainer(
+	ULearningAgentsManager* InManager,
+	ULearningAgentsInteractor* InInteractor,
+	ULearningAgentsPolicy* InPolicy,
+	TSubclassOf<ULearningAgentsImitationTrainer> Class)
 {
-	if (IsTraining())
+	if (!InManager)
 	{
-		UE_LOG(LogLearning, Error, TEXT("%s: Cannot begin training as we are already training!"), *GetName());
+		UE_LOG(LogLearning, Error, TEXT("MakeImitationTrainer: InManager is nullptr."));
+		return nullptr;
+	}
+
+	if (!Class)
+	{
+		UE_LOG(LogLearning, Error, TEXT("MakeImitationTrainer: Class is nullptr."));
+		return nullptr;
+	}
+
+	const FName UniqueName = MakeUniqueObjectName(InManager, Class, TEXT("ImitationTrainer"), EUniqueObjectNameOptions::GloballyUnique);
+
+	ULearningAgentsImitationTrainer* ImitationTrainer = NewObject<ULearningAgentsImitationTrainer>(InManager, Class, UniqueName);
+	if (!ImitationTrainer) { return nullptr; }
+
+	ImitationTrainer->SetupImitationTrainer(
+		InManager,
+		InInteractor,
+		InPolicy);
+
+	return ImitationTrainer->IsSetup() ? ImitationTrainer : nullptr;
+}
+
+void ULearningAgentsImitationTrainer::SetupImitationTrainer(
+	ULearningAgentsManager* InManager,
+	ULearningAgentsInteractor* InInteractor,
+	ULearningAgentsPolicy* InPolicy)
+{
+	if (IsSetup())
+	{
+		UE_LOG(LogLearning, Error, TEXT("%s: Setup already performed!"), *GetName());
+		return;
+	}
+
+	if (!InManager)
+	{
+		UE_LOG(LogLearning, Error, TEXT("%s: InManager is nullptr."), *GetName());
+		return;
+	}
+
+	if (!InInteractor)
+	{
+		UE_LOG(LogLearning, Error, TEXT("%s: InInteractor is nullptr."), *GetName());
+		return;
+	}
+
+	if (!InInteractor->IsSetup())
+	{
+		UE_LOG(LogLearning, Error, TEXT("%s: %s's Setup must be run before it can be used."), *GetName(), *InInteractor->GetName());
 		return;
 	}
 
@@ -56,6 +100,33 @@ void ULearningAgentsImitationTrainer::BeginTraining(
 		return;
 	}
 
+	Manager = InManager;
+	Interactor = InInteractor;
+	Policy = InPolicy;
+
+	bIsSetup = true;
+
+	Manager->AddListener(this);
+}
+
+void ULearningAgentsImitationTrainer::BeginTraining(
+	const ULearningAgentsRecording* Recording,
+	const FLearningAgentsImitationTrainerSettings& ImitationTrainerSettings,
+	const FLearningAgentsImitationTrainerTrainingSettings& ImitationTrainerTrainingSettings,
+	const FLearningAgentsTrainerPathSettings& ImitationTrainerPathSettings)
+{
+	if (!IsSetup())
+	{
+		UE_LOG(LogLearning, Error, TEXT("%s: Setup not complete."), *GetName());
+		return;
+	}
+
+	if (IsTraining())
+	{
+		UE_LOG(LogLearning, Error, TEXT("%s: Cannot begin training as we are already training!"), *GetName());
+		return;
+	}
+
 	if (!Recording)
 	{
 		UE_LOG(LogLearning, Error, TEXT("%s: Recording is nullptr."), *GetName());
@@ -68,8 +139,6 @@ void ULearningAgentsImitationTrainer::BeginTraining(
 		return;
 	}
 
-	Policy = InPolicy;
-
 	// Record Timeout Setting
 
 	TrainerTimeout = ImitationTrainerSettings.TrainerCommunicationTimeout;
@@ -103,9 +172,9 @@ void ULearningAgentsImitationTrainer::BeginTraining(
 
 	// Sizes
 
-	const int32 ObservationNum = Policy->GetPolicyObject().ObservationNum;
-	const int32 ActionNum = Policy->GetPolicyObject().ActionNum;
-	const int32 MemoryStateNum = Policy->GetPolicyObject().MemoryStateNum;
+	const int32 ObservationNum = Interactor->GetObservationVectorSize();
+	const int32 ActionNum = Interactor->GetActionVectorSize();
+	const int32 MemoryStateNum = Policy->GetMemoryStateSize();
 
 	// Get Number of Steps
 
@@ -137,6 +206,11 @@ void ULearningAgentsImitationTrainer::BeginTraining(
 
 	// Copy into Flat Arrays
 
+	TLearningArray<1, int32> RecordedEpisodeStarts;
+	TLearningArray<1, int32> RecordedEpisodeLengths;
+	TLearningArray<2, float> RecordedObservations;
+	TLearningArray<2, float> RecordedActions;
+
 	RecordedEpisodeStarts.SetNumUninitialized({ TotalEpisodeNum });
 	RecordedEpisodeLengths.SetNumUninitialized({ TotalEpisodeNum });
 	RecordedObservations.SetNumUninitialized({ TotalStepNum, ObservationNum });
@@ -149,10 +223,13 @@ void ULearningAgentsImitationTrainer::BeginTraining(
 		if (Record.ObservationDimNum != ObservationNum) { continue; }
 		if (Record.ActionDimNum != ActionNum) { continue; }
 
+		TLearningArrayView<2, const float> ObservationsView = TLearningArrayView<2, const float>(Record.ObservationData.GetData(), { Record.StepNum, Record.ObservationDimNum });
+		TLearningArrayView<2, const float> ActionsView = TLearningArrayView<2, const float>(Record.ActionData.GetData(), { Record.StepNum, Record.ActionDimNum });
+
 		RecordedEpisodeStarts[EpisodeIdx] = StepIdx;
 		RecordedEpisodeLengths[EpisodeIdx] = Record.StepNum;
-		UE::Learning::Array::Copy(RecordedObservations.Slice(StepIdx, Record.StepNum), Record.Observations);
-		UE::Learning::Array::Copy(RecordedActions.Slice(StepIdx, Record.StepNum), Record.Actions);
+		UE::Learning::Array::Copy(RecordedObservations.Slice(StepIdx, Record.StepNum), ObservationsView);
+		UE::Learning::Array::Copy(RecordedActions.Slice(StepIdx, Record.StepNum), ActionsView);
 		EpisodeIdx++;
 		StepIdx += Record.StepNum;
 	}
@@ -167,18 +244,17 @@ void ULearningAgentsImitationTrainer::BeginTraining(
 
 	UE::Learning::FImitationTrainerTrainingSettings ImitationTrainingSettings;
 	ImitationTrainingSettings.IterationNum = ImitationTrainerTrainingSettings.NumberOfIterations;
-	ImitationTrainingSettings.LearningRatePolicy = ImitationTrainerTrainingSettings.LearningRate;
+	ImitationTrainingSettings.LearningRate = ImitationTrainerTrainingSettings.LearningRate;
 	ImitationTrainingSettings.LearningRateDecay = ImitationTrainerTrainingSettings.LearningRateDecay;
 	ImitationTrainingSettings.WeightDecay = ImitationTrainerTrainingSettings.WeightDecay;
-	ImitationTrainingSettings.PolicyBatchSize = ImitationTrainerTrainingSettings.BatchSize;
+	ImitationTrainingSettings.BatchSize = ImitationTrainerTrainingSettings.BatchSize;
+	ImitationTrainingSettings.Window = ImitationTrainerTrainingSettings.Window;
+	ImitationTrainingSettings.ActionRegularizationWeight = ImitationTrainerTrainingSettings.ActionRegularizationWeight;
+	ImitationTrainingSettings.ActionEntropyWeight = ImitationTrainerTrainingSettings.ActionEntropyWeight;
 	ImitationTrainingSettings.Seed = ImitationTrainerTrainingSettings.RandomSeed;
 	ImitationTrainingSettings.Device = UE::Learning::Agents::GetTrainerDevice(ImitationTrainerTrainingSettings.Device);
 	ImitationTrainingSettings.bUseTensorboard = ImitationTrainerTrainingSettings.bUseTensorboard;
-
-	const UE::Learning::EImitationTrainerFlags TrainerFlags = 
-		bReinitializePolicyNetwork ? 
-		UE::Learning::EImitationTrainerFlags::None : 
-		UE::Learning::EImitationTrainerFlags::UseInitialPolicyNetwork;
+	ImitationTrainingSettings.bSaveSnapshots = ImitationTrainerTrainingSettings.bSaveSnapshots;
 
 	ImitationTrainer = MakeUnique<UE::Learning::FSharedMemoryImitationTrainer>(
 		GetName(),
@@ -191,36 +267,44 @@ void ULearningAgentsImitationTrainer::BeginTraining(
 		ObservationNum,
 		ActionNum,
 		MemoryStateNum,
-		Policy->GetPolicyNetwork(),
+		*Policy->GetPolicyNetworkAsset()->NeuralNetworkData,
+		*Policy->GetEncoderNetworkAsset()->NeuralNetworkData,
+		*Policy->GetDecoderNetworkAsset()->NeuralNetworkData,
+		Interactor->GetObservationSchema(),
+		Interactor->GetObservationSchemaElement(),
+		Interactor->GetActionSchema(),
+		Interactor->GetActionSchemaElement(),
 		ImitationTrainingSettings);
 
 	UE_LOG(LogLearning, Display, TEXT("%s: Sending / Receiving initial policy..."), *GetName());
 
 	UE::Learning::ETrainerResponse Response = UE::Learning::ETrainerResponse::Success;
 
-	Response = ImitationTrainer->SendPolicy(Policy->GetPolicyNetwork(), TrainerTimeout);
-
+	Response = ImitationTrainer->SendPolicy(*Policy->GetPolicyNetworkAsset()->NeuralNetworkData, TrainerTimeout);
 	if (Response != UE::Learning::ETrainerResponse::Success)
 	{
-		UE_LOG(LogLearning, Error, TEXT("%s: Error sending policy to trainer: %s. Check log for errors."), *GetName(), UE::Learning::Trainer::GetResponseString(Response));
+		UE_LOG(LogLearning, Error, TEXT("%s: Error sending policy to trainer: %s. Check log for additional errors."), *GetName(), UE::Learning::Trainer::GetResponseString(Response));
 		ImitationTrainer->Terminate();
 		bHasTrainingFailed = true;
 		return;
 	}
 
-	if (!bReinitializePolicyNetwork)
+	Response = ImitationTrainer->SendEncoder(*Policy->GetEncoderNetworkAsset()->NeuralNetworkData, TrainerTimeout);
+	if (Response != UE::Learning::ETrainerResponse::Success)
 	{
-		Response = ImitationTrainer->RecvPolicy(Policy->GetPolicyNetwork(), TrainerTimeout);
-
-		if (Response != UE::Learning::ETrainerResponse::Success)
-		{
-			UE_LOG(LogLearning, Error, TEXT("%s: Error receiving policy from trainer: %s. Check log for errors."), *GetName(), UE::Learning::Trainer::GetResponseString(Response));
-			bHasTrainingFailed = true;
-			ImitationTrainer->Terminate();
-			return;
-		}
+		UE_LOG(LogLearning, Error, TEXT("%s: Error sending encoder to trainer: %s. Check log for additional errors."), *GetName(), UE::Learning::Trainer::GetResponseString(Response));
+		ImitationTrainer->Terminate();
+		bHasTrainingFailed = true;
+		return;
+	}
 
-		Policy->GetNetworkAsset()->ForceMarkDirty();
+	Response = ImitationTrainer->SendDecoder(*Policy->GetDecoderNetworkAsset()->NeuralNetworkData, TrainerTimeout);
+	if (Response != UE::Learning::ETrainerResponse::Success)
+	{
+		UE_LOG(LogLearning, Error, TEXT("%s: Error sending decoder to trainer: %s. Check log for additional errors."), *GetName(), UE::Learning::Trainer::GetResponseString(Response));
+		ImitationTrainer->Terminate();
+		bHasTrainingFailed = true;
+		return;
 	}
 
 	UE_LOG(LogLearning, Display, TEXT("%s: Sending Experience..."), *GetName());
@@ -236,7 +320,7 @@ void ULearningAgentsImitationTrainer::BeginTraining(
 
 	if (Response != UE::Learning::ETrainerResponse::Success)
 	{
-		UE_LOG(LogLearning, Error, TEXT("%s: Error sending experience to trainer: %s. Check log for errors."), *GetName(), UE::Learning::Trainer::GetResponseString(Response));
+		UE_LOG(LogLearning, Error, TEXT("%s: Error sending experience to trainer: %s. Check log for additional errors."), *GetName(), UE::Learning::Trainer::GetResponseString(Response));
 		bHasTrainingFailed = true;
 		ImitationTrainer->Terminate();
 		return;
@@ -247,6 +331,12 @@ void ULearningAgentsImitationTrainer::BeginTraining(
 
 void ULearningAgentsImitationTrainer::DoneTraining()
 {
+	if (!IsSetup())
+	{
+		UE_LOG(LogLearning, Error, TEXT("%s: Setup not complete."), *GetName());
+		return;
+	}
+
 	if (IsTraining())
 	{
 		// Wait for Trainer to finish
@@ -261,6 +351,12 @@ void ULearningAgentsImitationTrainer::DoneTraining()
 
 void ULearningAgentsImitationTrainer::EndTraining()
 {
+	if (!IsSetup())
+	{
+		UE_LOG(LogLearning, Error, TEXT("%s: Setup not complete."), *GetName());
+		return;
+	}
+
 	if (IsTraining())
 	{
 		UE_LOG(LogLearning, Display, TEXT("%s: Stopping training..."), *GetName());
@@ -273,6 +369,12 @@ void ULearningAgentsImitationTrainer::IterateTraining()
 {
 	UE_LEARNING_TRACE_CPUPROFILER_EVENT_SCOPE(ULearningAgentsImitationTrainer::IterateTraining);
 
+	if (!IsSetup())
+	{
+		UE_LOG(LogLearning, Error, TEXT("%s: Setup not complete."), *GetName());
+		return;
+	}
+
 	if (!IsTraining())
 	{
 		UE_LOG(LogLearning, Error, TEXT("%s: Training not running."), *GetName());
@@ -281,9 +383,7 @@ void ULearningAgentsImitationTrainer::IterateTraining()
 
 	if (ImitationTrainer->HasPolicyOrCompleted())
 	{
-		UE::Learning::ETrainerResponse Response = ImitationTrainer->RecvPolicy(Policy->GetPolicyNetwork(), TrainerTimeout);
-		Policy->GetNetworkAsset()->ForceMarkDirty();
-
+		UE::Learning::ETrainerResponse Response = ImitationTrainer->RecvPolicy(*Policy->GetPolicyNetworkAsset()->NeuralNetworkData, TrainerTimeout);
 		if (Response == UE::Learning::ETrainerResponse::Completed)
 		{
 			UE_LOG(LogLearning, Display, TEXT("%s: Trainer completed training."), *GetName());
@@ -292,22 +392,47 @@ void ULearningAgentsImitationTrainer::IterateTraining()
 		}
 		else if (Response != UE::Learning::ETrainerResponse::Success)
 		{
-			UE_LOG(LogLearning, Error, TEXT("%s: Error receiving policy from trainer. Check log for errors."), *GetName());
+			UE_LOG(LogLearning, Error, TEXT("%s: Error receiving policy from trainer: %s. Check log for additional errors."), *GetName(), UE::Learning::Trainer::GetResponseString(Response));
+			bHasTrainingFailed = true;
+			EndTraining();
+			return;
+		}
+		Policy->GetPolicyNetworkAsset()->ForceMarkDirty();
+
+		Response = ImitationTrainer->RecvEncoder(*Policy->GetEncoderNetworkAsset()->NeuralNetworkData, TrainerTimeout);
+		if (Response != UE::Learning::ETrainerResponse::Success)
+		{
+			UE_LOG(LogLearning, Error, TEXT("%s: Error receiving encoder from trainer: %s. Check log for additional errors."), *GetName(), UE::Learning::Trainer::GetResponseString(Response));
 			bHasTrainingFailed = true;
 			EndTraining();
 			return;
 		}
+		Policy->GetEncoderNetworkAsset()->ForceMarkDirty();
+
+		Response = ImitationTrainer->RecvDecoder(*Policy->GetDecoderNetworkAsset()->NeuralNetworkData, TrainerTimeout);
+		if (Response != UE::Learning::ETrainerResponse::Success)
+		{
+			UE_LOG(LogLearning, Error, TEXT("%s: Error receiving decoder from trainer: %s. Check log for additional errors."), *GetName(), UE::Learning::Trainer::GetResponseString(Response));
+			bHasTrainingFailed = true;
+			EndTraining();
+			return;
+		}
+		Policy->GetDecoderNetworkAsset()->ForceMarkDirty();
 	}
 }
 
 void ULearningAgentsImitationTrainer::RunTraining(
-	ULearningAgentsPolicy* InPolicy,
 	const ULearningAgentsRecording* Recording,
 	const FLearningAgentsImitationTrainerSettings& ImitationTrainerSettings,
 	const FLearningAgentsImitationTrainerTrainingSettings& ImitationTrainerTrainingSettings,
-	const FLearningAgentsTrainerPathSettings& ImitationTrainerPathSettings,
-	const bool bReinitializePolicyNetwork)
+	const FLearningAgentsTrainerPathSettings& ImitationTrainerPathSettings)
 {
+	if (!IsSetup())
+	{
+		UE_LOG(LogLearning, Error, TEXT("%s: Setup not complete."), *GetName());
+		return;
+	}
+
 	if (bHasTrainingFailed)
 	{
 		UE_LOG(LogLearning, Error, TEXT("%s: Training has failed. Check log for errors."), *GetName());
@@ -318,12 +443,10 @@ void ULearningAgentsImitationTrainer::RunTraining(
 	if (!IsTraining())
 	{
 		BeginTraining(
-			InPolicy,
 			Recording,
 			ImitationTrainerSettings,
 			ImitationTrainerTrainingSettings,
-			ImitationTrainerPathSettings,
-			bReinitializePolicyNetwork);
+			ImitationTrainerPathSettings);
 
 		if (!IsTraining())
 		{
diff --git a/Engine/Plugins/Experimental/LearningAgents/Source/LearningAgentsTraining/Private/LearningAgentsRecorder.cpp b/Engine/Plugins/Experimental/LearningAgents/Source/LearningAgentsTraining/Private/LearningAgentsRecorder.cpp
--- a/Engine/Plugins/Experimental/LearningAgents/Source/LearningAgentsTraining/Private/LearningAgentsRecorder.cpp
+++ b/Engine/Plugins/Experimental/LearningAgents/Source/LearningAgentsTraining/Private/LearningAgentsRecorder.cpp
@@ -5,8 +5,6 @@
 #include "LearningAgentsManager.h"
 #include "LearningAgentsInteractor.h"
 #include "LearningAgentsRecording.h"
-#include "LearningAgentsHelpers.h"
-#include "LearningFeatureObject.h"
 #include "LearningLog.h"
 #include "LearningTrainer.h"
 
@@ -19,6 +17,12 @@ FLearningAgentsRecorderPathSettings::FLearningAgentsRecorderPathSettings()
 	IntermediateRelativePath.Path = FPaths::ProjectIntermediateDir();
 }
 
+ULearningAgentsRecorder::FAgentRecordBuffer::FAgentRecordBuffer() = default;
+
+ULearningAgentsRecorder::FAgentRecordBuffer::FAgentRecordBuffer(const int32 InObservationCompatibilityHash, const int32 InActionCompatibilityHash) 
+	: ObservationCompatibilityHash(InObservationCompatibilityHash)
+	, ActionCompatibilityHash(InActionCompatibilityHash) {};
+
 TLearningArrayView<1, float> ULearningAgentsRecorder::FAgentRecordBuffer::GetObservation(const int32 SampleIdx)
 {
 	return Observations[SampleIdx / ChunkSize][SampleIdx % ChunkSize];
@@ -73,13 +77,18 @@ void ULearningAgentsRecorder::FAgentRecordBuffer::CopyToRecord(FLearningAgentsRe
 	Record.StepNum = StepNum;
 	Record.ObservationDimNum = GetObservation(0).Num();
 	Record.ActionDimNum = GetAction(0).Num();
-	Record.Observations.SetNumUninitialized({ StepNum, Observations[0].Num<1>() });
-	Record.Actions.SetNumUninitialized({ StepNum, Actions[0].Num<1>() });
+	Record.ObservationCompatibilityHash = ObservationCompatibilityHash;
+	Record.ActionCompatibilityHash = ActionCompatibilityHash;
+	Record.ObservationData.SetNumUninitialized(StepNum * Record.ObservationDimNum);
+	Record.ActionData.SetNumUninitialized(StepNum * Record.ActionDimNum);
+
+	TLearningArrayView<2, float> ObservationsView = TLearningArrayView<2, float>(Record.ObservationData.GetData(), { Record.StepNum, Record.ObservationDimNum });
+	TLearningArrayView<2, float> ActionsView = TLearningArrayView<2, float>(Record.ActionData.GetData(), { Record.StepNum, Record.ActionDimNum });
 
 	for (int32 StepIdx = 0; StepIdx < StepNum; StepIdx++)
 	{
-		UE::Learning::Array::Copy(Record.Observations[StepIdx], GetObservation(StepIdx));
-		UE::Learning::Array::Copy(Record.Actions[StepIdx], GetAction(StepIdx));
+		UE::Learning::Array::Copy(ObservationsView[StepIdx], GetObservation(StepIdx));
+		UE::Learning::Array::Copy(ActionsView[StepIdx], GetAction(StepIdx));
 	}
 }
 
@@ -87,20 +96,58 @@ ULearningAgentsRecorder::ULearningAgentsRecorder() : Super(FObjectInitializer::G
 ULearningAgentsRecorder::ULearningAgentsRecorder(FVTableHelper& Helper) : Super(Helper) {}
 ULearningAgentsRecorder::~ULearningAgentsRecorder() = default;
 
-void ULearningAgentsRecorder::EndPlay(const EEndPlayReason::Type EndPlayReason)
+void ULearningAgentsRecorder::BeginDestroy()
 {
 	if (IsRecording())
 	{
 		EndRecording();
 	}
 
-	Super::EndPlay(EndPlayReason);
+	Super::BeginDestroy();
+}
+
+ULearningAgentsRecorder* ULearningAgentsRecorder::MakeRecorder(
+	ULearningAgentsManager* InManager,
+	ULearningAgentsInteractor* InInteractor,
+	TSubclassOf<ULearningAgentsRecorder> Class,
+	const FName Name,
+	const FLearningAgentsRecorderPathSettings& RecorderPathSettings,
+	ULearningAgentsRecording* RecordingAsset,
+	bool bReinitializeRecording)
+{
+	if (!InManager)
+	{
+		UE_LOG(LogLearning, Error, TEXT("MakeRecorder: InManager is nullptr."));
+		return nullptr;
+	}
+
+	if (!Class)
+	{
+		UE_LOG(LogLearning, Error, TEXT("MakeRecorder: Class is nullptr."));
+		return nullptr;
+	}
+
+	const FName UniqueName = MakeUniqueObjectName(InManager, Class, Name, EUniqueObjectNameOptions::GloballyUnique);
+
+	ULearningAgentsRecorder* Recorder = NewObject<ULearningAgentsRecorder>(InManager, Class, UniqueName);
+	if (!Recorder) { return nullptr; }
+
+	Recorder->SetupRecorder(
+		InManager,
+		InInteractor,
+		RecorderPathSettings,
+		RecordingAsset,
+		bReinitializeRecording);
+
+	return Recorder->IsSetup() ? Recorder : nullptr;
 }
 
 void ULearningAgentsRecorder::SetupRecorder(
+	ULearningAgentsManager* InManager,
 	ULearningAgentsInteractor* InInteractor,
 	const FLearningAgentsRecorderPathSettings& RecorderPathSettings,
-	ULearningAgentsRecording* RecordingAsset)
+	ULearningAgentsRecording* RecordingAsset,
+	bool bReinitializeRecording)
 {
 	if (IsSetup())
 	{
@@ -108,9 +155,9 @@ void ULearningAgentsRecorder::SetupRecorder(
 		return;
 	}
 
-	if (!Manager)
+	if (!InManager)
 	{
-		UE_LOG(LogLearning, Error, TEXT("%s: Must be attached to a LearningAgentsManager Actor."), *GetName());
+		UE_LOG(LogLearning, Error, TEXT("%s: InManager is nullptr."), *GetName());
 		return;
 	}
 
@@ -126,6 +173,7 @@ void ULearningAgentsRecorder::SetupRecorder(
 		return;
 	}
 
+	Manager = InManager;
 	Interactor = InInteractor;
 
 	if (RecordingAsset)
@@ -135,61 +183,68 @@ void ULearningAgentsRecorder::SetupRecorder(
 	else
 	{
 		const FName UniqueName = MakeUniqueObjectName(this, ULearningAgentsRecording::StaticClass(), TEXT("Recording"), EUniqueObjectNameOptions::GloballyUnique);
-
 		Recording = NewObject<ULearningAgentsRecording>(this, UniqueName);
 	}
 
+	if (bReinitializeRecording)
+	{
+		Recording->Records.Empty();
+		Recording->ForceMarkDirty();
+	}
+
 	RecordingDirectory = RecorderPathSettings.IntermediateRelativePath.Path / TEXT("LearningAgents") / RecorderPathSettings.RecordingsSubdirectory;
 
-	RecordBuffers.Empty();
-	RecordBuffers.SetNum(Manager->GetMaxAgentNum());
+	// Find Compatibility Hash
+
+	const int32 ObservationCompatibilityHash = UE::Learning::Observation::GetSchemaObjectsCompatibilityHash(Interactor->GetObservationSchema(), Interactor->GetObservationSchemaElement());
+	const int32 ActionCompatibilityHash = UE::Learning::Action::GetSchemaObjectsCompatibilityHash(Interactor->GetActionSchema(), Interactor->GetActionSchemaElement());
+
+	// Create Record Buffers
+
+	RecordBuffers.Empty(Manager->GetMaxAgentNum());
+	for (int32 RecordBufferIdx = 0; RecordBufferIdx < Manager->GetMaxAgentNum(); RecordBufferIdx++)
+	{
+		RecordBuffers.Emplace(ObservationCompatibilityHash, ActionCompatibilityHash);
+	}
 
 	bIsSetup = true;
 
-	OnAgentsAdded(Manager->GetAllAgentIds());
+	Manager->AddListener(this);
 }
 
-void ULearningAgentsRecorder::OnAgentsRemoved(const TArray<int32>& AgentIds)
+void ULearningAgentsRecorder::OnAgentsRemoved_Implementation(const TArray<int32>& AgentIds)
 {
-	if (IsSetup())
+	if (!IsSetup())
 	{
-		for (const int32 AgentId : AgentIds)
-		{
-			if (!RecordBuffers[AgentId].IsEmpty())
-			{
-				RecordBuffers[AgentId].CopyToRecord(Recording->Records.AddDefaulted_GetRef());
-				RecordBuffers[AgentId].Empty();
-			}
-		}
+		UE_LOG(LogLearning, Error, TEXT("%s: Setup not complete."), *GetName());
+		return;
+	}
 
-		for (ULearningAgentsHelper* Helper : HelperObjects)
+	for (const int32 AgentId : AgentIds)
+	{
+		if (!RecordBuffers[AgentId].IsEmpty())
 		{
-			Helper->OnAgentsRemoved(AgentIds);
+			RecordBuffers[AgentId].CopyToRecord(Recording->Records.AddDefaulted_GetRef());
+			RecordBuffers[AgentId].Empty();
 		}
-
-		AgentsRemoved(AgentIds);
 	}
 }
 
-void ULearningAgentsRecorder::OnAgentsReset(const TArray<int32>& AgentIds)
+void ULearningAgentsRecorder::OnAgentsReset_Implementation(const TArray<int32>& AgentIds)
 {
-	if (IsSetup())
+	if (!IsSetup())
 	{
-		for (const int32 AgentId : AgentIds)
-		{
-			if (!RecordBuffers[AgentId].IsEmpty())
-			{
-				RecordBuffers[AgentId].CopyToRecord(Recording->Records.AddDefaulted_GetRef());
-				RecordBuffers[AgentId].Empty();
-			}
-		}
+		UE_LOG(LogLearning, Error, TEXT("%s: Setup not complete."), *GetName());
+		return;
+	}
 
-		for (ULearningAgentsHelper* Helper : HelperObjects)
+	for (const int32 AgentId : AgentIds)
+	{
+		if (!RecordBuffers[AgentId].IsEmpty())
 		{
-			Helper->OnAgentsReset(AgentIds);
+			RecordBuffers[AgentId].CopyToRecord(Recording->Records.AddDefaulted_GetRef());
+			RecordBuffers[AgentId].Empty();
 		}
-
-		AgentsReset(AgentIds);
 	}
 }
 
@@ -207,25 +262,26 @@ void ULearningAgentsRecorder::AddExperience()
 		return;
 	}
 
+	if (Manager->GetAgentNum() == 0)
+	{
+		UE_LOG(LogLearning, Warning, TEXT("%s: No agents added to Manager."), *GetName());
+	}
+
 	for (const int32 AgentId : Manager->GetAllAgentSet())
 	{
-		if (Interactor->GetObservationEncodingAgentIteration()[AgentId] == 0 ||
-			Interactor->GetActionEncodingAgentIteration()[AgentId] == 0)
+		if (Interactor->ObservationVectorIteration[AgentId] == 0 || Interactor->ActionVectorIteration[AgentId] == 0)
 		{
 			UE_LOG(LogLearning, Warning, TEXT("%s: Agent with id %i has not made observations and taken actions so experience will not be recorded for it."), *GetName(), AgentId);
 			continue;
 		}
 
-		if (Interactor->GetObservationEncodingAgentIteration()[AgentId] !=
-			Interactor->GetActionEncodingAgentIteration()[AgentId])
+		if (Interactor->ObservationVectorIteration[AgentId] != Interactor->ActionVectorIteration[AgentId])
 		{
 			UE_LOG(LogLearning, Warning, TEXT("%s: Agent with id %i does not have matching iteration numbers for observations and actions so experience will not be recorded for it."), *GetName(), AgentId);
 			continue;
 		}
 
-		RecordBuffers[AgentId].Push(
-			Interactor->GetObservationFeature().FeatureBuffer()[AgentId],
-			Interactor->GetActionFeature().FeatureBuffer()[AgentId]);
+		RecordBuffers[AgentId].Push(Interactor->ObservationVectors[AgentId], Interactor->ActionVectors[AgentId]);
 	}
 }
 
@@ -262,96 +318,7 @@ void ULearningAgentsRecorder::EndRecording()
 	bIsRecording = false;
 }
 
-void ULearningAgentsRecorder::LoadRecordingFromFile(const FFilePath& File)
-{
-	if (!IsSetup())
-	{
-		UE_LOG(LogLearning, Error, TEXT("%s: Setup not complete."), *GetName());
-		return;
-	}
-
-	Recording->LoadRecordingFromFile(File);
-}
-
-void ULearningAgentsRecorder::SaveRecordingToFile(const FFilePath& File) const
-{
-	if (!IsSetup())
-	{
-		UE_LOG(LogLearning, Error, TEXT("%s: Setup not complete."), *GetName());
-		return;
-	}
-
-	Recording->SaveRecordingToFile(File);
-}
-
-void ULearningAgentsRecorder::AppendRecordingFromFile(const FFilePath& File)
-{
-	if (!IsSetup())
-	{
-		UE_LOG(LogLearning, Error, TEXT("%s: Setup not complete."), *GetName());
-		return;
-	}
-
-	Recording->AppendRecordingFromFile(File);
-}
-
-void ULearningAgentsRecorder::UseRecordingAsset(ULearningAgentsRecording* RecordingAsset)
-{
-	if (!IsSetup())
-	{
-		UE_LOG(LogLearning, Error, TEXT("%s: Setup not complete."), *GetName());
-		return;
-	}
-
-	if (!RecordingAsset)
-	{
-		UE_LOG(LogLearning, Error, TEXT("%s: Asset is nullptr."), *GetName());
-		return;
-	}
-
-	if (RecordingAsset == Recording)
-	{
-		UE_LOG(LogLearning, Error, TEXT("%s: Asset is same as the current recording."), *GetName());
-		return;
-	}
-
-	Recording = RecordingAsset;
-}
-
-void ULearningAgentsRecorder::LoadRecordingFromAsset(ULearningAgentsRecording* RecordingAsset)
-{
-	if (!IsSetup())
-	{
-		UE_LOG(LogLearning, Error, TEXT("%s: Setup not complete."), *GetName());
-		return;
-	}
-
-	Recording->LoadRecordingFromAsset(RecordingAsset);
-}
-
-void ULearningAgentsRecorder::SaveRecordingToAsset(ULearningAgentsRecording* RecordingAsset)
-{
-	if (!IsSetup())
-	{
-		UE_LOG(LogLearning, Error, TEXT("%s: Setup not complete."), *GetName());
-		return;
-	}
-
-	Recording->SaveRecordingToAsset(RecordingAsset);
-}
-
-void ULearningAgentsRecorder::AppendRecordingToAsset(ULearningAgentsRecording* RecordingAsset)
-{
-	if (!IsSetup())
-	{
-		UE_LOG(LogLearning, Error, TEXT("%s: Setup not complete."), *GetName());
-		return;
-	}
-
-	Recording->AppendRecordingToAsset(RecordingAsset);
-}
-
-void ULearningAgentsRecorder::BeginRecording(bool bReinitializeRecording)
+void ULearningAgentsRecorder::BeginRecording()
 {
 	if (!IsSetup())
 	{
@@ -365,12 +332,6 @@ void ULearningAgentsRecorder::BeginRecording(bool bReinitializeRecording)
 		return;
 	}
 
-	if (bReinitializeRecording)
-	{
-		Recording->Records.Empty();
-		Recording->ForceMarkDirty();
-	}
-
 	for (const int32 AgentId : Manager->GetAllAgentSet())
 	{
 		RecordBuffers[AgentId].Empty();
@@ -379,7 +340,7 @@ void ULearningAgentsRecorder::BeginRecording(bool bReinitializeRecording)
 	bIsRecording = true;
 }
 
-const ULearningAgentsRecording* ULearningAgentsRecorder::GetCurrentRecording() const
+const ULearningAgentsRecording* ULearningAgentsRecorder::GetRecordingAsset() const
 {
 	if (!IsSetup())
 	{
diff --git a/Engine/Plugins/Experimental/LearningAgents/Source/LearningAgentsTraining/Private/LearningAgentsRecording.cpp b/Engine/Plugins/Experimental/LearningAgents/Source/LearningAgentsTraining/Private/LearningAgentsRecording.cpp
--- a/Engine/Plugins/Experimental/LearningAgents/Source/LearningAgentsTraining/Private/LearningAgentsRecording.cpp
+++ b/Engine/Plugins/Experimental/LearningAgents/Source/LearningAgentsTraining/Private/LearningAgentsRecording.cpp
@@ -7,16 +7,6 @@
 #include "UObject/Package.h"
 #include "Misc/FileHelper.h"
 
-bool FLearningAgentsRecord::Serialize(FArchive& Ar)
-{
-	Ar << StepNum;
-	Ar << ObservationDimNum;
-	Ar << ActionDimNum;
-	UE::Learning::Array::Serialize(Ar, Observations);
-	UE::Learning::Array::Serialize(Ar, Actions);
-	return true;
-}
-
 ULearningAgentsRecording::ULearningAgentsRecording() = default;
 ULearningAgentsRecording::ULearningAgentsRecording(FVTableHelper& Helper) : Super(Helper) {}
 ULearningAgentsRecording::~ULearningAgentsRecording() = default;
@@ -69,8 +59,10 @@ void ULearningAgentsRecording::LoadRecordingFromFile(const FFilePath& File)
 			UE::Learning::DeserializeFromBytes(Offset, RecordingData, Records[RecordIdx].StepNum);
 			UE::Learning::DeserializeFromBytes(Offset, RecordingData, Records[RecordIdx].ObservationDimNum);
 			UE::Learning::DeserializeFromBytes(Offset, RecordingData, Records[RecordIdx].ActionDimNum);
-			UE::Learning::Array::DeserializeFromBytes(Offset, RecordingData, Records[RecordIdx].Observations);
-			UE::Learning::Array::DeserializeFromBytes(Offset, RecordingData, Records[RecordIdx].Actions);
+			UE::Learning::DeserializeFromBytes(Offset, RecordingData, Records[RecordIdx].ObservationCompatibilityHash);
+			UE::Learning::DeserializeFromBytes(Offset, RecordingData, Records[RecordIdx].ActionCompatibilityHash);
+			UE::Learning::Array::DeserializeFromBytes<2, float>(Offset, RecordingData, Records[RecordIdx].ObservationData);
+			UE::Learning::Array::DeserializeFromBytes<2, float>(Offset, RecordingData, Records[RecordIdx].ActionData);
 		}
 
 		UE_LEARNING_CHECK(Offset == RecordingData.Num());
@@ -95,11 +87,13 @@ void ULearningAgentsRecording::SaveRecordingToFile(const FFilePath& File) const
 	for (int32 RecordIdx = 0; RecordIdx < Records.Num(); RecordIdx++)
 	{
 		TotalByteNum +=
-			sizeof(int32) + // SampleNum
+			sizeof(int32) + // StepNum
 			sizeof(int32) + // ObservationDimNum
 			sizeof(int32) + // ActionDimNum
-			UE::Learning::Array::SerializationByteNum<2, float>(Records[RecordIdx].Observations.Shape()) + // Observations
-			UE::Learning::Array::SerializationByteNum<2, float>(Records[RecordIdx].Actions.Shape());	   // Actions
+			sizeof(int32) + // ObservationCompatibilityHash
+			sizeof(int32) + // ActionCompatibilityHash
+			UE::Learning::Array::SerializationByteNum<2, float>({ Records[RecordIdx].StepNum, Records[RecordIdx].ObservationDimNum }) + // Observations
+			UE::Learning::Array::SerializationByteNum<2, float>({ Records[RecordIdx].StepNum, Records[RecordIdx].ActionDimNum });	   // Actions
 	}
 
 	RecordingData.SetNumUninitialized(TotalByteNum);
@@ -114,8 +108,10 @@ void ULearningAgentsRecording::SaveRecordingToFile(const FFilePath& File) const
 		UE::Learning::SerializeToBytes(Offset, RecordingData, Records[RecordIdx].StepNum);
 		UE::Learning::SerializeToBytes(Offset, RecordingData, Records[RecordIdx].ObservationDimNum);
 		UE::Learning::SerializeToBytes(Offset, RecordingData, Records[RecordIdx].ActionDimNum);
-		UE::Learning::Array::SerializeToBytes(Offset, RecordingData, Records[RecordIdx].Observations);
-		UE::Learning::Array::SerializeToBytes(Offset, RecordingData, Records[RecordIdx].Actions);
+		UE::Learning::SerializeToBytes(Offset, RecordingData, Records[RecordIdx].ObservationCompatibilityHash);
+		UE::Learning::SerializeToBytes(Offset, RecordingData, Records[RecordIdx].ActionCompatibilityHash);
+		UE::Learning::Array::SerializeToBytes<2, float>(Offset, RecordingData, { Records[RecordIdx].StepNum, Records[RecordIdx].ObservationDimNum }, Records[RecordIdx].ObservationData);
+		UE::Learning::Array::SerializeToBytes<2, float>(Offset, RecordingData, { Records[RecordIdx].StepNum, Records[RecordIdx].ActionDimNum }, Records[RecordIdx].ActionData);
 	}
 
 	UE_LEARNING_CHECK(Offset == RecordingData.Num());
@@ -189,6 +185,67 @@ void ULearningAgentsRecording::AppendRecordingToAsset(ULearningAgentsRecording*
 	RecordingAsset->ForceMarkDirty();
 }
 
+int32 ULearningAgentsRecording::GetRecordNum() const
+{
+	return Records.Num();
+}
+
+int32 ULearningAgentsRecording::GetRecordStepNum(const int32 Record) const
+{
+	if (Record < 0 || Record >= Records.Num())
+	{
+		UE_LOG(LogLearning, Error, TEXT("%s: Record out of range. Asked for record %i but recording only has %i records."), *GetName(), Record, Records.Num());
+		return 0;
+	}
+
+	return Records[Record].StepNum;
+}
+
+
+void ULearningAgentsRecording::GetObservationVector(TArray<float>& OutObservationVector, int32& OutObservationCompatibilityHash, const int32 Record, const int32 Step)
+{
+	if (Record < 0 || Record >= Records.Num())
+	{
+		UE_LOG(LogLearning, Error, TEXT("%s: Record out of range. Asked for record %i but recording only has %i records."), *GetName(), Record, Records.Num());
+		OutObservationVector.Empty();
+		OutObservationCompatibilityHash = 0;
+		return;
+	}
+
+	if (Step < 0 || Step >= Records[Record].StepNum)
+	{
+		UE_LOG(LogLearning, Error, TEXT("%s: Step out of range. Asked for step %i but recording only has %i steps."), *GetName(), Step, Records[Record].StepNum);
+		OutObservationVector.Empty();
+		OutObservationCompatibilityHash = 0;
+		return;
+	}
+
+	OutObservationVector = MakeArrayView(Records[Record].ObservationData).Slice(Records[Record].ObservationDimNum * Step, Records[Record].ObservationDimNum);
+	OutObservationCompatibilityHash = Records[Record].ObservationCompatibilityHash;
+}
+
+void ULearningAgentsRecording::GetActionVector(TArray<float>& OutActionVector, int32& OutActionCompatibilityHash, const int32 Record, const int32 Step)
+{
+	if (Record < 0 || Record >= Records.Num())
+	{
+		UE_LOG(LogLearning, Error, TEXT("%s: Record out of range. Asked for record %i but recording only has %i records."), *GetName(), Record, Records.Num());
+		OutActionVector.Empty();
+		OutActionCompatibilityHash = 0;
+		return;
+	}
+
+	if (Step < 0 || Step >= Records[Record].StepNum)
+	{
+		UE_LOG(LogLearning, Error, TEXT("%s: Step out of range. Asked for step %i but recording only has %i steps."), *GetName(), Step, Records[Record].StepNum);
+		OutActionVector.Empty();
+		OutActionCompatibilityHash = 0;
+		return;
+	}
+
+	OutActionVector = MakeArrayView(Records[Record].ActionData).Slice(Records[Record].ActionDimNum * Step, Records[Record].ActionDimNum);
+	OutActionCompatibilityHash = Records[Record].ActionCompatibilityHash;
+}
+
 void ULearningAgentsRecording::ForceMarkDirty()
 {
 	// Manually mark the package as dirty since just using `Modify` prevents 
