LearningAgents: Large restructuring to allow for structured actions and observations.

diff --git a/Engine/Plugins/Experimental/LearningAgents/Source/LearningAgentsTraining/Private/LearningAgentsRewards.cpp b/Engine/Plugins/Experimental/LearningAgents/Source/LearningAgentsTraining/Private/LearningAgentsRewards.cpp
--- a/Engine/Plugins/Experimental/LearningAgents/Source/LearningAgentsTraining/Private/LearningAgentsRewards.cpp
+++ b/Engine/Plugins/Experimental/LearningAgents/Source/LearningAgentsTraining/Private/LearningAgentsRewards.cpp
@@ -2,535 +2,130 @@
 
 #include "LearningAgentsRewards.h"
 
-#include "LearningAgentsTrainer.h"
-#include "LearningAgentsManager.h"
-#include "LearningArray.h"
-#include "LearningArrayMap.h"
-#include "LearningRewardObject.h"
 #include "LearningLog.h"
 
-#include "GameFramework/Actor.h"
+#include "Components/SplineComponent.h"
 
-namespace UE::Learning::Agents::Rewards::Private
+float ULearningAgentsRewards::RewardOnCondition(const bool bCondition, const float RewardScale)
 {
-	template<typename RewardUObject, typename RewardFObject, typename... InArgTypes>
-	RewardUObject* AddReward(ULearningAgentsTrainer* InAgentTrainer, const FName Name, const TCHAR* FunctionName, InArgTypes&& ...Args)
-	{
-		if (!InAgentTrainer)
-		{
-			UE_LOG(LogLearning, Error, TEXT("%s: InAgentTrainer is nullptr."), FunctionName);
-			return nullptr;
-		}
-
-		if (!InAgentTrainer->HasAgentManager())
-		{
-			UE_LOG(LogLearning, Error, TEXT("%s: Must be attached to a LearningAgentsManager Actor."), *InAgentTrainer->GetName());
-			return nullptr;
-		}
-
-		const FName UniqueName = MakeUniqueObjectName(InAgentTrainer, RewardUObject::StaticClass(), Name, EUniqueObjectNameOptions::GloballyUnique);
-
-		RewardUObject* Reward = NewObject<RewardUObject>(InAgentTrainer, UniqueName);
-		Reward->Init(InAgentTrainer->GetAgentManager()->GetMaxAgentNum());
-		Reward->AgentTrainer = InAgentTrainer;
-		Reward->RewardObject = MakeShared<RewardFObject>(
-			Reward->GetFName(),
-			InAgentTrainer->GetAgentManager()->GetInstanceData().ToSharedRef(),
-			InAgentTrainer->GetAgentManager()->GetMaxAgentNum(),
-			Forward<InArgTypes>(Args)...);
-
-		InAgentTrainer->AddReward(Reward, Reward->RewardObject.ToSharedRef());
-
-		return Reward;
-	}
+	return bCondition ? RewardScale : 0.0f;
 }
 
-//------------------------------------------------------------------
-
-void ULearningAgentsReward::Init(const int32 MaxAgentNum)
+float ULearningAgentsRewards::RewardOnTranslationDifferenceBelowThreshold(const FVector TranslationA, const FVector TranslationB, const float DistanceThreshold, const float RewardScale)
 {
-	AgentIteration.SetNumUninitialized({ MaxAgentNum });
-	UE::Learning::Array::Set<1, uint64>(AgentIteration, INDEX_NONE);
+	return RewardOnCondition(FVector::Distance(TranslationA, TranslationB) < DistanceThreshold, RewardScale);
 }
 
-void ULearningAgentsReward::OnAgentsAdded(const TArray<int32>& AgentIds)
+float ULearningAgentsRewards::RewardFromTranslationSimilarity(const FVector TranslationA, const FVector TranslationB, const float TranslationScale, const float RewardScale)
 {
-	UE::Learning::Array::Set<1, uint64>(AgentIteration, 0, AgentIds);
+	const float TranslationDifference = FVector::Dist(TranslationA, TranslationB);
+	return RewardScale * FMath::InvExpApprox(FMath::Square(TranslationDifference / FMath::Max(TranslationScale, UE_SMALL_NUMBER)));
 }
 
-void ULearningAgentsReward::OnAgentsRemoved(const TArray<int32>& AgentIds)
+float ULearningAgentsRewards::RewardFromAngleSimilarity(const float AngleA, const float AngleB, const float AngleScale, const float RewardScale)
 {
-	UE::Learning::Array::Set<1, uint64>(AgentIteration, INDEX_NONE, AgentIds);
+	const float AngleDifference = FMath::FindDeltaAngleRadians(FMath::DegreesToRadians(AngleA), FMath::DegreesToRadians(AngleB));
+	return RewardScale * FMath::InvExpApprox(FMath::Square(AngleDifference / FMath::Max(FMath::DegreesToRadians(AngleScale), UE_SMALL_NUMBER)));
 }
 
-void ULearningAgentsReward::OnAgentsReset(const TArray<int32>& AgentIds)
+float ULearningAgentsRewards::RewardFromRotationSimilarityAsQuats(const FQuat RotationA, const FQuat RotationB, const float AngleScale, const float RewardScale)
 {
-	UE::Learning::Array::Set<1, uint64>(AgentIteration, 0, AgentIds);
+	FQuat Difference = RotationA.Inverse() * RotationB;
+	Difference.EnforceShortestArcWith(FQuat::Identity);
+	return RewardScale * FMath::InvExpApprox(FMath::Square(Difference.GetAngle() / (FMath::Max(FMath::DegreesToRadians(AngleScale), UE_SMALL_NUMBER))));
 }
 
-uint64 ULearningAgentsReward::GetAgentIteration(const int32 AgentId) const
+float ULearningAgentsRewards::RewardFromRotationSimilarity(const FRotator RotationA, const FRotator RotationB, const float AngleScale, const float RewardScale)
 {
-	return AgentIteration[AgentId];
+	return RewardScale * RewardFromRotationSimilarityAsQuats(RotationA.Quaternion(), RotationB.Quaternion(), AngleScale);
 }
 
-//------------------------------------------------------------------
-
-UFloatReward* UFloatReward::AddFloatReward(ULearningAgentsTrainer* InAgentTrainer, const FName Name, const float Weight)
+float ULearningAgentsRewards::RewardFromDirectionSimilarity(const FVector DirectionA, const FVector DirectionB, const float AngleScale, const float RewardScale)
 {
-	return UE::Learning::Agents::Rewards::Private::AddReward<UFloatReward, UE::Learning::FFloatReward>(InAgentTrainer, Name, TEXT("AddFloatReward"), Weight);
+	const float AngleDifference = FMath::Acos(DirectionA.Dot(DirectionB));
+	return RewardScale * FMath::InvExpApprox(FMath::Square(AngleDifference / FMath::Max(FMath::DegreesToRadians(AngleScale), UE_SMALL_NUMBER)));
 }
 
-void UFloatReward::SetFloatReward(const int32 AgentId, const float Reward)
-{
-	if (!AgentTrainer->HasAgent(AgentId))
-	{
-		UE_LOG(LogLearning, Error, TEXT("%s: AgentId %d not found in the agents set."), *GetName(), AgentId);
-		return;
-	}
-
-	RewardObject->InstanceData->View(RewardObject->ValueHandle)[AgentId] = Reward;
-	AgentIteration[AgentId]++;
-}
+// Spline Rewards
 
-#if UE_LEARNING_AGENTS_ENABLE_VISUAL_LOG
-void UFloatReward::VisualLog(const UE::Learning::FIndexSet Instances) const
+float ULearningAgentsRewards::RewardFromVelocityAlongSpline(const USplineComponent* SplineComponent, const FVector Position, const FVector Velocity, const float VelocityScale, const float RewardScale, const float FiniteDifferenceDelta)
 {
-	UE_LEARNING_TRACE_CPUPROFILER_EVENT_SCOPE(UFloatReward::VisualLog);
-
-	const TLearningArrayView<1, const float> ValueView = RewardObject->InstanceData->ConstView(RewardObject->ValueHandle);
-	const TLearningArrayView<1, const float> WeightView = RewardObject->InstanceData->ConstView(RewardObject->WeightHandle);
-	const TLearningArrayView<1, const float> RewardView = RewardObject->InstanceData->ConstView(RewardObject->RewardHandle);
-
-	for (const int32 Instance : Instances)
+	if (!SplineComponent)
 	{
-		if (const AActor* Actor = Cast<AActor>(AgentTrainer->GetAgent(Instance)))
-		{
-			UE_LEARNING_AGENTS_VLOG_STRING(this, LogLearning, Display,
-				Actor->GetActorLocation(),
-				VisualLogColor.ToFColor(true),
-				TEXT("Agent %i\nWeight: [% 6.2f]\nValue: [% 6.2f]\nReward: [% 6.3f]"),
-				Instance,
-				WeightView[Instance],
-				ValueView[Instance],
-				RewardView[Instance]);
-		}
+		UE_LOG(LogLearning, Error, TEXT("RewardFromVelocityAlongSpline: SplineComponent is nullptr."));
+		return 0.0f;
 	}
-}
-#endif
-
-//------------------------------------------------------------------
 
-UConditionalReward* UConditionalReward::AddConditionalReward(ULearningAgentsTrainer* InAgentTrainer, const FName Name, const float Value)
-{
-	return UE::Learning::Agents::Rewards::Private::AddReward<UConditionalReward, UE::Learning::FConditionalConstantReward>(InAgentTrainer, Name, TEXT("AddConditionalReward"), Value);
-}
+	float FiniteDiff = FiniteDifferenceDelta;
 
-void UConditionalReward::SetConditionalReward(const int32 AgentId, const bool bCondition)
-{
-	if (!AgentTrainer->HasAgent(AgentId))
+	if (FiniteDiff < UE_KINDA_SMALL_NUMBER)
 	{
-		UE_LOG(LogLearning, Error, TEXT("%s: AgentId %d not found in the agents set."), *GetName(), AgentId);
-		return;
+		UE_LOG(LogLearning, Warning, TEXT("RewardFromVelocityAlongSpline: FiniteDifferenceDelta is too small (%6.4f). Clamping to %6.4f."), FiniteDifferenceDelta, UE_KINDA_SMALL_NUMBER);
+		FiniteDiff = UE_KINDA_SMALL_NUMBER;
 	}
 
-	RewardObject->InstanceData->View(RewardObject->ConditionHandle)[AgentId] = bCondition;
-	AgentIteration[AgentId]++;
-}
-
-#if UE_LEARNING_AGENTS_ENABLE_VISUAL_LOG
-void UConditionalReward::VisualLog(const UE::Learning::FIndexSet Instances) const
-{
-	UE_LEARNING_TRACE_CPUPROFILER_EVENT_SCOPE(UBoolReward::VisualLog);
+	const float RawDistance0 = SplineComponent->GetDistanceAlongSplineAtLocation(Position, ESplineCoordinateSpace::World);
+	const float RawDistance1 = SplineComponent->GetDistanceAlongSplineAtLocation(Position + FiniteDiff * Velocity.GetSafeNormal(), ESplineCoordinateSpace::World);
 
-	const TLearningArrayView<1, const bool> ConditionView = RewardObject->InstanceData->ConstView(RewardObject->ConditionHandle);
-	const TLearningArrayView<1, const float> RewardView = RewardObject->InstanceData->ConstView(RewardObject->RewardHandle);
+	float Distance0 = RawDistance0, Distance1 = RawDistance1;
 
-	for (const int32 Instance : Instances)
+	if (SplineComponent->IsClosedLoop())
 	{
-		if (const AActor* Actor = Cast<AActor>(AgentTrainer->GetAgent(Instance)))
+		const float SplineDistance = SplineComponent->GetSplineLength();
+
+		if (FMath::Abs(Distance0 - (Distance1 + SplineDistance)) < FMath::Abs(Distance0 - Distance1))
 		{
-			UE_LEARNING_AGENTS_VLOG_STRING(this, LogLearning, Display,
-				Actor->GetActorLocation(),
-				VisualLogColor.ToFColor(true),
-				TEXT("Agent %i\nValue: [% 6.2f]\nCondition: %s\nReward: [% 6.3f]"),
-				Instance,
-				RewardObject->Value,
-				ConditionView[Instance] ? TEXT("true") : TEXT("false"),
-				RewardView[Instance]);
+			Distance1 = Distance1 + SplineDistance;
 		}
-	}
-}
-#endif
-
-//------------------------------------------------------------------
-
-UScalarVelocityReward* UScalarVelocityReward::AddScalarVelocityReward(ULearningAgentsTrainer* InAgentTrainer, const FName Name, const float Weight, const float Scale)
-{
-	return UE::Learning::Agents::Rewards::Private::AddReward<UScalarVelocityReward, UE::Learning::FScalarVelocityReward>(InAgentTrainer, Name, TEXT("AddScalarVelocityReward"), Weight, Scale);
-}
-
-void UScalarVelocityReward::SetScalarVelocityReward(int32 AgentId, float Velocity)
-{
-	if (!AgentTrainer->HasAgent(AgentId))
-	{
-		UE_LOG(LogLearning, Error, TEXT("%s: AgentId %d not found in the agents set."), *GetName(), AgentId);
-		return;
-	}
-
-	RewardObject->InstanceData->View(RewardObject->VelocityHandle)[AgentId] = Velocity;
-	AgentIteration[AgentId]++;
-}
-
-#if UE_LEARNING_AGENTS_ENABLE_VISUAL_LOG
-void UScalarVelocityReward::VisualLog(const UE::Learning::FIndexSet Instances) const
-{
-	UE_LEARNING_TRACE_CPUPROFILER_EVENT_SCOPE(UScalarVelocityReward::VisualLog);
-
-	const TLearningArrayView<1, const float> VelocityView = RewardObject->InstanceData->ConstView(RewardObject->VelocityHandle);
-	const TLearningArrayView<1, const float> WeightView = RewardObject->InstanceData->ConstView(RewardObject->WeightHandle);
-	const TLearningArrayView<1, const float> ScaleView = RewardObject->InstanceData->ConstView(RewardObject->ScaleHandle);
-	const TLearningArrayView<1, const float> RewardView = RewardObject->InstanceData->ConstView(RewardObject->RewardHandle);
-
-	for (const int32 Instance : Instances)
-	{
-		if (const AActor* Actor = Cast<const AActor>(AgentTrainer->GetAgent(Instance)))
+		else if (FMath::Abs((Distance0 + SplineDistance) - Distance1) < FMath::Abs(Distance0 - Distance1))
 		{
-			UE_LEARNING_AGENTS_VLOG_STRING(this, LogLearning, Display,
-				Actor->GetActorLocation(),
-				VisualLogColor.ToFColor(true),
-				TEXT("Agent %i\nWeight: [% 6.2f]\nScale: [% 6.2f]\nVelocity: [% 6.2f]\nReward: [% 6.3f]"),
-				Instance,
-				WeightView[Instance],
-				ScaleView[Instance],
-				VelocityView[Instance],
-				RewardView[Instance]);
+			Distance0 = Distance0 + SplineDistance;
 		}
 	}
-}
-#endif
 
-ULocalDirectionalVelocityReward* ULocalDirectionalVelocityReward::AddLocalDirectionalVelocityReward(ULearningAgentsTrainer* InAgentTrainer, const FName Name, const float Weight, const float Scale, const FVector Axis)
-{
-	return UE::Learning::Agents::Rewards::Private::AddReward<ULocalDirectionalVelocityReward, UE::Learning::FLocalDirectionalVelocityReward>(InAgentTrainer, Name, TEXT("AddLocalDirectionalVelocityReward"), Weight, Scale, Axis.GetSafeNormal(UE_SMALL_NUMBER, FVector::ForwardVector));
+	return RewardScale * ((Distance1 - Distance0) / FiniteDiff) * (Velocity.Length() / FMath::Max(VelocityScale, UE_SMALL_NUMBER));
 }
 
-void ULocalDirectionalVelocityReward::SetLocalDirectionalVelocityReward(const int32 AgentId, const FVector Velocity, const FRotator RelativeRotation)
-{
-	if (!AgentTrainer->HasAgent(AgentId))
-	{
-		UE_LOG(LogLearning, Error, TEXT("%s: AgentId %d not found in the agents set."), *GetName(), AgentId);
-		return;
-	}
 
-	RewardObject->InstanceData->View(RewardObject->VelocityHandle)[AgentId] = Velocity;
-	RewardObject->InstanceData->View(RewardObject->RelativeRotationHandle)[AgentId] = RelativeRotation.Quaternion();
-	AgentIteration[AgentId]++;
-}
+// Penalties
 
-#if UE_LEARNING_AGENTS_ENABLE_VISUAL_LOG
-void ULocalDirectionalVelocityReward::VisualLog(const UE::Learning::FIndexSet Instances) const
+float ULearningAgentsRewards::PenaltyOnCondition(const bool bCondition, const float PenaltyScale)
 {
-	UE_LEARNING_TRACE_CPUPROFILER_EVENT_SCOPE(ULocalDirectionalVelocityReward::VisualLog);
-
-	const TLearningArrayView<1, const FVector> VelocityView = RewardObject->InstanceData->View(RewardObject->VelocityHandle);
-	const TLearningArrayView<1, const FQuat> RelativeRotationView = RewardObject->InstanceData->View(RewardObject->RelativeRotationHandle);
-	const TLearningArrayView<1, const float> WeightView = RewardObject->InstanceData->ConstView(RewardObject->WeightHandle);
-	const TLearningArrayView<1, const float> ScaleView = RewardObject->InstanceData->ConstView(RewardObject->ScaleHandle);
-	const TLearningArrayView<1, const float> RewardView = RewardObject->InstanceData->ConstView(RewardObject->RewardHandle);
-
-	for (const int32 Instance : Instances)
-	{
-		if (const AActor* Actor = Cast<AActor>(AgentTrainer->GetAgent(Instance)))
-		{
-			const FVector Velocity = VelocityView[Instance];
-			const FQuat RelativeRotation = RelativeRotationView[Instance];
-			const FVector LocalVelocity = RelativeRotation.UnrotateVector(Velocity);
-			const FVector Direction = RelativeRotation.RotateVector(RewardObject->Axis);
-
-			UE_LEARNING_AGENTS_VLOG_ARROW(this, LogLearning, Display,
-				Actor->GetActorLocation(),
-				Actor->GetActorLocation() + Velocity,
-				VisualLogColor.ToFColor(true),
-				TEXT(""));
-
-			UE_LEARNING_AGENTS_VLOG_STRING(this, LogLearning, Display,
-				Actor->GetActorLocation() + Velocity,
-				VisualLogColor.ToFColor(true),
-				TEXT("Velocity: [% 6.3f % 6.3f % 6.3f]\nLocal Velocity: [% 6.3f % 6.3f % 6.3f]"),
-				Velocity.X, Velocity.Y, Velocity.Z,
-				LocalVelocity.X, LocalVelocity.Y, LocalVelocity.Z);
-
-			UE_LEARNING_AGENTS_VLOG_ARROW(this, LogLearning, Display,
-				Actor->GetActorLocation(),
-				Actor->GetActorLocation() + 100.0f * Direction,
-				VisualLogColor.ToFColor(true),
-				TEXT(""));
-
-			UE_LEARNING_AGENTS_VLOG_STRING(this, LogLearning, Display,
-				Actor->GetActorLocation() + 100.0f * Direction,
-				VisualLogColor.ToFColor(true),
-				TEXT("Direction: [% 6.3f % 6.3f % 6.3f]\nLocal Direction: [% 6.3f % 6.3f % 6.3f]"),
-				Direction.X, Direction.Y, Direction.Z,
-				RewardObject->Axis.X, RewardObject->Axis.Y, RewardObject->Axis.Z);
-
-			UE_LEARNING_AGENTS_VLOG_TRANSFORM(this, LogLearning, Display,
-				Actor->GetActorLocation(),
-				RelativeRotation,
-				VisualLogColor.ToFColor(true),
-				TEXT("Agent %i\nDot Product: [% 6.3f]\nWeight: [% 6.2f]\nScale: [% 6.2f]\nReward: [% 6.3f]"),
-				Instance,
-				LocalVelocity.Dot(RewardObject->Axis),
-				WeightView[Instance],
-				ScaleView[Instance],
-				RewardView[Instance]);
-		}
-	}
+	return bCondition ? -PenaltyScale : 0.0f;
 }
-#endif
 
-//------------------------------------------------------------------
-
-UPlanarPositionDifferencePenalty* UPlanarPositionDifferencePenalty::AddPlanarPositionDifferencePenalty(
-	ULearningAgentsTrainer* InAgentTrainer,
-	const FName Name,
-	const float Weight,
-	const float Scale,
-	const float Threshold,
-	const FVector Axis0,
-	const FVector Axis1)
+float ULearningAgentsRewards::PenaltyOnTranslationDifferenceAboveThreshold(const FVector TranslationA, const FVector TranslationB, const float DistanceThreshold, const float PenaltyScale)
 {
-	return UE::Learning::Agents::Rewards::Private::AddReward<UPlanarPositionDifferencePenalty, UE::Learning::FPlanarPositionDifferencePenalty>(
-		InAgentTrainer, 
-		Name, 
-		TEXT("AddPlanarPositionDifferencePenalty"), 
-		Weight, 
-		Scale, 
-		Threshold, 
-		Axis0.GetSafeNormal(UE_SMALL_NUMBER, FVector::ForwardVector),
-		Axis1.GetSafeNormal(UE_SMALL_NUMBER, FVector::RightVector));
+	return PenaltyOnCondition(FVector::Distance(TranslationA, TranslationB) > DistanceThreshold, PenaltyScale);
 }
 
-void UPlanarPositionDifferencePenalty::SetPlanarPositionDifferencePenalty(const int32 AgentId, const FVector Position0, const FVector Position1)
+float ULearningAgentsRewards::PenaltyFromTranslationDifference(const FVector TranslationA, const FVector TranslationB, const float TranslationScale, const float PenaltyScale)
 {
-	if (!AgentTrainer->HasAgent(AgentId))
-	{
-		UE_LOG(LogLearning, Error, TEXT("%s: AgentId %d not found in the agents set."), *GetName(), AgentId);
-		return;
-	}
-
-	RewardObject->InstanceData->View(RewardObject->Position0Handle)[AgentId] = Position0;
-	RewardObject->InstanceData->View(RewardObject->Position1Handle)[AgentId] = Position1;
-	AgentIteration[AgentId]++;
+	const float TranslationDifference = FVector::Dist(TranslationA, TranslationB);
+	return PenaltyScale  * (-TranslationDifference / FMath::Max(TranslationScale, UE_SMALL_NUMBER));
 }
 
-#if UE_LEARNING_AGENTS_ENABLE_VISUAL_LOG
-void UPlanarPositionDifferencePenalty::VisualLog(const UE::Learning::FIndexSet Instances) const
+float ULearningAgentsRewards::PenaltyFromAngleDifference(const float AngleA, const float AngleB, const float AngleScale, const float PenaltyScale)
 {
-	UE_LEARNING_TRACE_CPUPROFILER_EVENT_SCOPE(UPlanarPositionDifferencePenalty::VisualLog);
-
-	const TLearningArrayView<1, const FVector> Position0View = RewardObject->InstanceData->ConstView(RewardObject->Position0Handle);
-	const TLearningArrayView<1, const FVector> Position1View = RewardObject->InstanceData->ConstView(RewardObject->Position1Handle);
-	const TLearningArrayView<1, const float> WeightView = RewardObject->InstanceData->ConstView(RewardObject->WeightHandle);
-	const TLearningArrayView<1, const float> ScaleView = RewardObject->InstanceData->ConstView(RewardObject->ScaleHandle);
-	const TLearningArrayView<1, const float> ThresholdView = RewardObject->InstanceData->ConstView(RewardObject->ThresholdHandle);
-	const TLearningArrayView<1, const float> RewardView = RewardObject->InstanceData->ConstView(RewardObject->RewardHandle);
-
-	for (const int32 Instance : Instances)
-	{
-		if (const AActor* Actor = Cast<AActor>(AgentTrainer->GetAgent(Instance)))
-		{
-			const FVector Position0 = Position0View[Instance];
-			const FVector Position1 = Position1View[Instance];
-
-			const FVector PlanarPosition0 = FVector(RewardObject->Axis0.Dot(Position0), RewardObject->Axis1.Dot(Position0), 0.0f);
-			const FVector PlanarPosition1 = FVector(RewardObject->Axis0.Dot(Position1), RewardObject->Axis1.Dot(Position1), 0.0f);
-
-			UE_LEARNING_AGENTS_VLOG_LOCATION(this, LogLearning, Display,
-				Position0,
-				10,
-				VisualLogColor.ToFColor(true),
-				TEXT("Position0: [% 6.1f % 6.1f % 6.1f]\nPlanar Position0: [% 6.1f % 6.1f]"),
-				Position0.X, Position0.Y, Position0.Z,
-				PlanarPosition0.X, PlanarPosition0.Y);
-
-			UE_LEARNING_AGENTS_VLOG_PLANE(this, LogLearning, Display,
-				Position0,
-				FQuat::Identity,
-				RewardObject->Axis0,
-				RewardObject->Axis1,
-				VisualLogColor.ToFColor(true),
-				TEXT(""));
-
-			UE_LEARNING_AGENTS_VLOG_LOCATION(this, LogLearning, Display,
-				Position1,
-				10,
-				VisualLogColor.ToFColor(true),
-				TEXT("Position1: [% 6.1f % 6.1f % 6.1f]\nPlanar Position1: [% 6.1f % 6.1f]"),
-				Position1.X, Position1.Y, Position1.Z,
-				PlanarPosition1.X, PlanarPosition1.Y);
-
-			UE_LEARNING_AGENTS_VLOG_PLANE(this, LogLearning, Display,
-				Position1,
-				FQuat::Identity,
-				RewardObject->Axis0,
-				RewardObject->Axis1,
-				VisualLogColor.ToFColor(true),
-				TEXT(""));
-
-			UE_LEARNING_AGENTS_VLOG_SEGMENT(this, LogLearning, Display,
-				Position0,
-				Position1,
-				VisualLogColor.ToFColor(true),
-				TEXT("Agent %i\nDistance: [% 6.3f]\nPlanar Distance: [% 6.3f]\nWeight: [% 6.2f]\nScale: [% 6.2f]\nThreshold: [% 6.2f]\nReward: [% 6.3f]"),
-				Instance,
-				FVector::Distance(Position0, Position1),
-				FVector::Distance(PlanarPosition0, PlanarPosition1),
-				WeightView[Instance],
-				ScaleView[Instance],
-				ThresholdView[Instance],
-				RewardView[Instance]);
-		}
-	}
+	const float AngleDifference = FMath::FindDeltaAngleRadians(FMath::DegreesToRadians(AngleA), FMath::DegreesToRadians(AngleB));
+	return PenaltyScale  * (-AngleDifference / FMath::Max(FMath::DegreesToRadians(AngleScale), UE_SMALL_NUMBER));
 }
-#endif
 
-//------------------------------------------------------------------
-
-UPositionArraySimilarityReward* UPositionArraySimilarityReward::AddPositionArraySimilarityReward(
-	ULearningAgentsTrainer* InAgentTrainer,
-	const FName Name,
-	const int32 PositionNum,
-	const float Weight,
-	const float Scale)
+float ULearningAgentsRewards::PenaltyFromRotationDifferenceAsQuats(const FQuat RotationA, const FQuat RotationB, const float AngleScale, const float PenaltyScale)
 {
-	if (PositionNum < 1)
-	{
-		UE_LOG(LogLearning, Error, TEXT("AddPositionArraySimilarityReward: Number of elements in array must be at least 1, got %i."), PositionNum);
-		return nullptr;
-	}
-
-	return UE::Learning::Agents::Rewards::Private::AddReward<UPositionArraySimilarityReward, UE::Learning::FPositionArraySimilarityReward>(InAgentTrainer, Name, TEXT("AddPositionArraySimilarityReward"), PositionNum, Weight, Scale);
+	FQuat Difference = RotationA.Inverse() * RotationB;
+	Difference.EnforceShortestArcWith(FQuat::Identity);
+	return PenaltyScale  * (-Difference.GetAngle() / (FMath::Max(FMath::DegreesToRadians(AngleScale), UE_SMALL_NUMBER)));
 }
 
-void UPositionArraySimilarityReward::SetPositionArraySimilarityReward(
-	const int32 AgentId,
-	const TArray<FVector>& Positions0,
-	const TArray<FVector>& Positions1,
-	const FVector RelativePosition0,
-	const FVector RelativePosition1,
-	const FRotator RelativeRotation0,
-	const FRotator RelativeRotation1)
+float ULearningAgentsRewards::PenaltyFromRotationDifference(const FRotator RotationA, const FRotator RotationB, const float AngleScale, const float PenaltyScale)
 {
-	if (!AgentTrainer->HasAgent(AgentId))
-	{
-		UE_LOG(LogLearning, Error, TEXT("%s: AgentId %d not found in the agents set."), *GetName(), AgentId);
-		return;
-	}
-
-	const TLearningArrayView<2, FVector> Positions0View = RewardObject->InstanceData->View(RewardObject->Positions0Handle);
-	const TLearningArrayView<2, FVector> Positions1View = RewardObject->InstanceData->View(RewardObject->Positions1Handle);
-	const TLearningArrayView<1, FVector> RelativePosition0View = RewardObject->InstanceData->View(RewardObject->RelativePosition0Handle);
-	const TLearningArrayView<1, FVector> RelativePosition1View = RewardObject->InstanceData->View(RewardObject->RelativePosition1Handle);
-	const TLearningArrayView<1, FQuat> RelativeRotation0View = RewardObject->InstanceData->View(RewardObject->RelativeRotation0Handle);
-	const TLearningArrayView<1, FQuat> RelativeRotation1View = RewardObject->InstanceData->View(RewardObject->RelativeRotation1Handle);
-
-	if (Positions0.Num() != Positions0View.Num<1>() || Positions1.Num() != Positions0View.Num<1>())
-	{
-		UE_LOG(LogLearning, Error, TEXT("%s: Got wrong number of elements in array. Expected %i, got %i and %i."), *GetName(), Positions0View.Num<1>(), Positions0.Num(), Positions1.Num());
-		return;
-	}
-
-	RelativePosition0View[AgentId] = RelativePosition0;
-	RelativePosition1View[AgentId] = RelativePosition1;
-	RelativeRotation0View[AgentId] = RelativeRotation0.Quaternion();
-	RelativeRotation1View[AgentId] = RelativeRotation1.Quaternion();
-	UE::Learning::Array::Copy<1, FVector>(Positions0View[AgentId], Positions0);
-	UE::Learning::Array::Copy<1, FVector>(Positions1View[AgentId], Positions1);
-	AgentIteration[AgentId]++;
+	return PenaltyFromRotationDifferenceAsQuats(RotationA.Quaternion(), RotationB.Quaternion(), AngleScale, PenaltyScale);
 }
 
-#if UE_LEARNING_AGENTS_ENABLE_VISUAL_LOG
-void UPositionArraySimilarityReward::VisualLog(const UE::Learning::FIndexSet Instances) const
+float ULearningAgentsRewards::PenaltyFromDirectionDifference(const FVector DirectionA, const FVector DirectionB, const float AngleScale, const float PenaltyScale)
 {
-	UE_LEARNING_TRACE_CPUPROFILER_EVENT_SCOPE(UPositionArraySimilarityReward::VisualLog);
-
-	const TLearningArrayView<2, const FVector> Positions0View = RewardObject->InstanceData->ConstView(RewardObject->Positions0Handle);
-	const TLearningArrayView<2, const FVector> Positions1View = RewardObject->InstanceData->ConstView(RewardObject->Positions1Handle);
-	const TLearningArrayView<1, const FVector> RelativePosition0View = RewardObject->InstanceData->ConstView(RewardObject->RelativePosition0Handle);
-	const TLearningArrayView<1, const FVector> RelativePosition1View = RewardObject->InstanceData->ConstView(RewardObject->RelativePosition1Handle);
-	const TLearningArrayView<1, const FQuat> RelativeRotation0View = RewardObject->InstanceData->ConstView(RewardObject->RelativeRotation0Handle);
-	const TLearningArrayView<1, const FQuat> RelativeRotation1View = RewardObject->InstanceData->ConstView(RewardObject->RelativeRotation1Handle);
-	const TLearningArrayView<1, const float> WeightView = RewardObject->InstanceData->ConstView(RewardObject->WeightHandle);
-	const TLearningArrayView<1, const float> ScaleView = RewardObject->InstanceData->ConstView(RewardObject->ScaleHandle);
-	const TLearningArrayView<1, const float> ThresholdView = RewardObject->InstanceData->ConstView(RewardObject->ThresholdHandle);
-	const TLearningArrayView<1, const float> RewardView = RewardObject->InstanceData->ConstView(RewardObject->RewardHandle);
-
-	const int32 PositionNum = Positions0View.Num<1>();
-		
-	for (const int32 Instance : Instances)
-	{
-		if (const AActor* Actor = Cast<AActor>(AgentTrainer->GetAgent(Instance)))
-		{
-			const FVector RelativePosition0 = RelativePosition0View[Instance];
-			const FVector RelativePosition1 = RelativePosition1View[Instance];
-			const FQuat RelativeRotation0 = RelativeRotation0View[Instance];
-			const FQuat RelativeRotation1 = RelativeRotation1View[Instance];
-
-			for (int32 PositionIdx = 0; PositionIdx < PositionNum; PositionIdx++)
-			{
-				const FVector Position0 = Positions0View[Instance][PositionIdx];
-				const FVector Position1 = Positions1View[Instance][PositionIdx];
-
-				const FVector LocalPosition0 = RelativeRotation0.UnrotateVector(Positions0View[Instance][PositionIdx] - RelativePosition0);
-				const FVector LocalPosition1 = RelativeRotation1.UnrotateVector(Positions1View[Instance][PositionIdx] - RelativePosition1);
-
-				UE_LEARNING_AGENTS_VLOG_LOCATION(this, LogLearning, Display,
-					Position0,
-					10,
-					VisualLogColor.ToFColor(true),
-					TEXT("Position0: [% 6.1f % 6.1f % 6.1f]\nLocal Position0: [% 6.1f % 6.1f % 6.1f]"),
-					Position0.X, Position0.Y, Position0.Z,
-					LocalPosition0.X, LocalPosition0.Y, LocalPosition0.Z);
-
-				UE_LEARNING_AGENTS_VLOG_LOCATION(this, LogLearning, Display,
-					Position1,
-					10,
-					VisualLogColor.ToFColor(true),
-					TEXT("Position1: [% 6.1f % 6.1f % 6.1f]\nLocal Position1: [% 6.1f % 6.1f % 6.1f]"),
-					Position1.X, Position1.Y, Position1.Z,
-					LocalPosition1.X, LocalPosition1.Y, LocalPosition1.Z);
-
-				UE_LEARNING_AGENTS_VLOG_SEGMENT(this, LogLearning, Display,
-					Position0,
-					Position1,
-					VisualLogColor.ToFColor(true),
-					TEXT("Distance: [% 6.1f]\nLocal Distance: [% 6.1f]"),
-					FVector::Distance(Position0, Position1),
-					FVector::Distance(LocalPosition0, LocalPosition1));
-			}
-
-			UE_LEARNING_AGENTS_VLOG_TRANSFORM(this, LogLearning, Display,
-				RelativePosition0,
-				RelativeRotation0,
-				VisualLogColor.ToFColor(true),
-				TEXT("Relative Transform 0"));
-
-			UE_LEARNING_AGENTS_VLOG_TRANSFORM(this, LogLearning, Display,
-				RelativePosition1,
-				RelativeRotation1,
-				VisualLogColor.ToFColor(true),
-				TEXT("Relative Transform 1"));
-
-			UE_LEARNING_AGENTS_VLOG_STRING(this, LogLearning, Display,
-				Actor->GetActorLocation(),
-				VisualLogColor.ToFColor(true),
-				TEXT("Agent %i\nWeight: [% 6.2f]\nScale: [% 6.2f]\nThreshold: [% 6.2f]\nReward: [% 6.3f]"),
-				Instance,
-				WeightView[Instance],
-				ScaleView[Instance],
-				ThresholdView[Instance],
-				RewardView[Instance]);
-		}
-	}
+	const float AngleDifference = FMath::Acos(DirectionA.Dot(DirectionB));
+	return PenaltyScale  * (-AngleDifference / FMath::Max(FMath::DegreesToRadians(AngleScale), UE_SMALL_NUMBER));
 }
-#endif
+
diff --git a/Engine/Plugins/Experimental/LearningAgents/Source/LearningAgentsTraining/Private/LearningAgentsTrainer.cpp b/Engine/Plugins/Experimental/LearningAgents/Source/LearningAgentsTraining/Private/LearningAgentsTrainer.cpp
--- a/Engine/Plugins/Experimental/LearningAgents/Source/LearningAgentsTraining/Private/LearningAgentsTrainer.cpp
+++ b/Engine/Plugins/Experimental/LearningAgents/Source/LearningAgentsTraining/Private/LearningAgentsTrainer.cpp
@@ -9,17 +9,11 @@
 #include "LearningAgentsObservations.h"
 #include "LearningAgentsPolicy.h"
 #include "LearningAgentsCritic.h"
-#include "LearningAgentsHelpers.h"
 #include "LearningArray.h"
-#include "LearningArrayMap.h"
 #include "LearningExperience.h"
-#include "LearningFeatureObject.h"
 #include "LearningLog.h"
-#include "LearningNeuralNetworkObject.h"
 #include "LearningPPOTrainer.h"
-#include "LearningRewardObject.h"
 #include "LearningCompletion.h"
-#include "LearningCompletionObject.h"
 
 #include "GenericPlatform/GenericPlatformMisc.h"
 #include "HAL/FileManager.h"
@@ -31,30 +25,6 @@
 #include "Engine/GameViewportClient.h"
 #include "EngineDefines.h"
 
-namespace UE::Learning::Agents
-{
-	ELearningAgentsCompletion GetLearningAgentsCompletion(const ECompletionMode CompletionMode)
-	{
-		switch (CompletionMode)
-		{
-		case ECompletionMode::Running: UE_LOG(LogLearning, Error, TEXT("Cannot convert from ECompletionMode::Running to ELearningAgentsCompletion")); return ELearningAgentsCompletion::Termination;
-		case ECompletionMode::Terminated: return ELearningAgentsCompletion::Termination;
-		case ECompletionMode::Truncated: return ELearningAgentsCompletion::Truncation;
-		default:UE_LOG(LogLearning, Error, TEXT("Unknown Completion Mode.")); return ELearningAgentsCompletion::Termination;
-		}
-	}
-
-	ECompletionMode GetCompletionMode(const ELearningAgentsCompletion Completion)
-	{
-		switch (Completion)
-		{
-		case ELearningAgentsCompletion::Termination: return ECompletionMode::Terminated;
-		case ELearningAgentsCompletion::Truncation: return ECompletionMode::Truncated;
-		default:UE_LOG(LogLearning, Error, TEXT("Unknown Completion.")); return ECompletionMode::Running;
-		}
-	}
-}
-
 namespace UE::Learning::Agents
 {
 	ELearningAgentsTrainerDevice GetLearningAgentsTrainerDevice(const ETrainerDevice Device)
@@ -107,7 +77,49 @@ ULearningAgentsTrainer::ULearningAgentsTrainer() : Super(FObjectInitializer::Get
 ULearningAgentsTrainer::ULearningAgentsTrainer(FVTableHelper& Helper) : Super(Helper) {}
 ULearningAgentsTrainer::~ULearningAgentsTrainer() = default;
 
+void ULearningAgentsTrainer::BeginDestroy()
+{
+	if (IsTraining())
+	{
+		EndTraining();
+	}
+
+	Super::BeginDestroy();
+}
+
+ULearningAgentsTrainer* ULearningAgentsTrainer::MakeTrainer(
+	ULearningAgentsManager* InManager,
+	ULearningAgentsInteractor* InInteractor,
+	ULearningAgentsPolicy* InPolicy,
+	ULearningAgentsCritic* InCritic,
+	TSubclassOf<ULearningAgentsTrainer> Class,
+	const FName Name,
+	const FLearningAgentsTrainerSettings& TrainerSettings)
+{
+	if (!InManager)
+	{
+		UE_LOG(LogLearning, Error, TEXT("MakeTrainer: InManager is nullptr."));
+		return nullptr;
+	}
+
+	if (!Class)
+	{
+		UE_LOG(LogLearning, Error, TEXT("MakeTrainer: Class is nullptr."));
+		return nullptr;
+	}
+
+	const FName UniqueName = MakeUniqueObjectName(InManager, Class, Name, EUniqueObjectNameOptions::GloballyUnique);
+
+	ULearningAgentsTrainer* Trainer = NewObject<ULearningAgentsTrainer>(InManager, Class, UniqueName);
+	if (!Trainer) { return nullptr; }
+
+	Trainer->SetupTrainer(InManager, InInteractor, InPolicy, InCritic, TrainerSettings);
+
+	return Trainer->IsSetup() ? Trainer : nullptr;
+}
+
 void ULearningAgentsTrainer::SetupTrainer(
+	ULearningAgentsManager* InManager,
 	ULearningAgentsInteractor* InInteractor,
 	ULearningAgentsPolicy* InPolicy,
 	ULearningAgentsCritic* InCritic,
@@ -119,9 +131,9 @@ void ULearningAgentsTrainer::SetupTrainer(
 		return;
 	}
 
-	if (!Manager)
+	if (!InManager)
 	{
-		UE_LOG(LogLearning, Error, TEXT("%s: Must be attached to a LearningAgentsManager Actor."), *GetName());
+		UE_LOG(LogLearning, Error, TEXT("%s: InManager is nullptr."), *GetName());
 		return;
 	}
 
@@ -143,15 +155,15 @@ void ULearningAgentsTrainer::SetupTrainer(
 		return;
 	}
 
-	if (!InCritic)
+	if (!InPolicy->IsSetup())
 	{
-		UE_LOG(LogLearning, Error, TEXT("%s: InCritic is nullptr."), *GetName());
+		UE_LOG(LogLearning, Error, TEXT("%s: %s's Setup must be run before it can be used."), *GetName(), *InPolicy->GetName());
 		return;
 	}
 
-	if (!InPolicy->IsSetup())
+	if (!InCritic)
 	{
-		UE_LOG(LogLearning, Error, TEXT("%s: %s's Setup must be run before it can be used."), *GetName(), *InPolicy->GetName());
+		UE_LOG(LogLearning, Error, TEXT("%s: InCritic is nullptr."), *GetName());
 		return;
 	}
 
@@ -161,208 +173,195 @@ void ULearningAgentsTrainer::SetupTrainer(
 		return;
 	}
 
+	Manager = InManager;
 	Interactor = InInteractor;
 	Policy = InPolicy;
 	Critic = InCritic;
 
-	// Setup Rewards
-	RewardObjects.Empty();
-	RewardFeatures.Empty();
-	SetupRewards();
-
-	if (RewardObjects.Num() == 0)
-	{
-		UE_LOG(LogLearning, Error, TEXT("%s: No rewards added to Trainer during SetupRewards."), *GetName());
-		return;
-	}
-
-	Rewards = MakeShared<UE::Learning::FSumReward>(TEXT("Rewards"),
-		TLearningArrayView<1, const TSharedRef<UE::Learning::FRewardObject>>(RewardFeatures),
-		Manager->GetInstanceData().ToSharedRef(),
-		Manager->GetMaxAgentNum());
-
-	// Setup Completions
-	CompletionObjects.Empty();
-	CompletionFeatures.Empty();
-	SetupCompletions();
-
-	if (CompletionObjects.Num() == 0)
-	{
-		// Not an error or warning because it's fine to run training without any completions.
-		UE_LOG(LogLearning, Display, TEXT("%s: No completions added to Trainer during SetupCompletions."), *GetName());
-	}
-
-	Completions = MakeShared<UE::Learning::FAnyCompletion>(TEXT("Completions"),
-		TLearningArrayView<1, const TSharedRef<UE::Learning::FCompletionObject>>(CompletionFeatures),
-		Manager->GetInstanceData().ToSharedRef(),
-		Manager->GetMaxAgentNum());
-
 	// Create Episode Buffer
 	EpisodeBuffer = MakeUnique<UE::Learning::FEpisodeBuffer>();
 	EpisodeBuffer->Resize(
 		Manager->GetMaxAgentNum(),
-		TrainerSettings.MaxStepNum,
-		Interactor->GetObservationFeature().DimNum(),
-		Interactor->GetActionFeature().DimNum(),
+		TrainerSettings.MaxEpisodeStepNum,
+		Interactor->GetObservationVectorSize(),
+		Interactor->GetActionVectorSize(),
 		Policy->GetMemoryStateSize());
 
-	MaxStepsCompletion = TrainerSettings.MaxStepsCompletion;
-
 	// Create Replay Buffer
 	ReplayBuffer = MakeUnique<UE::Learning::FReplayBuffer>();
 	ReplayBuffer->Resize(
-		Interactor->GetObservationFeature().DimNum(),
-		Interactor->GetActionFeature().DimNum(),
+		Interactor->GetObservationVectorSize(),
+		Interactor->GetActionVectorSize(),
 		Policy->GetMemoryStateSize(),
 		TrainerSettings.MaximumRecordedEpisodesPerIteration,
 		TrainerSettings.MaximumRecordedStepsPerIteration);
 
 	// Create Reset Buffer
 	ResetBuffer = MakeUnique<UE::Learning::FResetInstanceBuffer>();
-	ResetBuffer->Resize(Manager->GetMaxAgentNum());
+	ResetBuffer->Reserve(Manager->GetMaxAgentNum());
 
 	// Record Timeout Setting
 	TrainerTimeout = TrainerSettings.TrainerCommunicationTimeout;
 
+	// Rewards and Completions Buffer
+
+	Rewards.SetNumUninitialized({ Manager->GetMaxAgentNum() });
+	AgentCompletions.SetNumUninitialized({ Manager->GetMaxAgentNum() });
+	EpisodeCompletions.SetNumUninitialized({ Manager->GetMaxAgentNum() });
+	AllCompletions.SetNumUninitialized({ Manager->GetMaxAgentNum() });
+	EpisodeTimes.SetNumUninitialized({ Manager->GetMaxAgentNum() });
+	UE::Learning::Array::Set<1, float>(Rewards, FLT_MAX);
+	UE::Learning::Array::Set<1, UE::Learning::ECompletionMode>(AgentCompletions, UE::Learning::ECompletionMode::Terminated);
+	UE::Learning::Array::Set<1, UE::Learning::ECompletionMode>(EpisodeCompletions, UE::Learning::ECompletionMode::Terminated);
+	UE::Learning::Array::Set<1, UE::Learning::ECompletionMode>(AllCompletions, UE::Learning::ECompletionMode::Terminated);
+	UE::Learning::Array::Set<1, float>(EpisodeTimes, FLT_MAX);
+
 	// Reset Agent iteration
-	RewardEvaluatedAgentIteration.SetNumUninitialized({ Manager->GetMaxAgentNum() });
-	CompletionEvaluatedAgentIteration.SetNumUninitialized({ Manager->GetMaxAgentNum() });
-	UE::Learning::Array::Set<1, uint64>(RewardEvaluatedAgentIteration, INDEX_NONE);
-	UE::Learning::Array::Set<1, uint64>(CompletionEvaluatedAgentIteration, INDEX_NONE);
+	RewardIteration.SetNumUninitialized({ Manager->GetMaxAgentNum() });
+	CompletionIteration.SetNumUninitialized({ Manager->GetMaxAgentNum() });
+	UE::Learning::Array::Set<1, uint64>(RewardIteration, INDEX_NONE);
+	UE::Learning::Array::Set<1, uint64>(CompletionIteration, INDEX_NONE);
 
 	bIsSetup = true;
 
-	OnAgentsAdded(Manager->GetAllAgentIds());
+	Manager->AddListener(this);
 }
 
-void ULearningAgentsTrainer::OnAgentsAdded(const TArray<int32>& AgentIds)
+void ULearningAgentsTrainer::OnAgentsAdded_Implementation(const TArray<int32>& AgentIds)
 {
-	if (IsSetup())
+	if (!IsSetup())
 	{
-		ResetEpisodes(AgentIds);
-		EpisodeBuffer->Reset(AgentIds);
-
-		UE::Learning::Array::Set<1, uint64>(RewardEvaluatedAgentIteration, 0, AgentIds);
-		UE::Learning::Array::Set<1, uint64>(CompletionEvaluatedAgentIteration, 0, AgentIds);
-
-		for (ULearningAgentsReward* RewardObject : RewardObjects)
-		{
-			RewardObject->OnAgentsAdded(AgentIds);
-		}
-
-		for (ULearningAgentsCompletion* CompletionObject : CompletionObjects)
-		{
-			CompletionObject->OnAgentsAdded(AgentIds);
-		}
+		UE_LOG(LogLearning, Error, TEXT("%s: Setup not complete."), *GetName());
+		return;
+	}
 
-		for (ULearningAgentsHelper* Helper : HelperObjects)
-		{
-			Helper->OnAgentsAdded(AgentIds);
-		}
+	EpisodeBuffer->Reset(AgentIds);
 
-		AgentsAdded(AgentIds);
-	}
+	UE::Learning::Array::Set<1, float>(Rewards, 0.0f, AgentIds);
+	UE::Learning::Array::Set<1, UE::Learning::ECompletionMode>(AgentCompletions, UE::Learning::ECompletionMode::Running, AgentIds);
+	UE::Learning::Array::Set<1, UE::Learning::ECompletionMode>(EpisodeCompletions, UE::Learning::ECompletionMode::Running, AgentIds);
+	UE::Learning::Array::Set<1, UE::Learning::ECompletionMode>(AllCompletions, UE::Learning::ECompletionMode::Running, AgentIds);
+	UE::Learning::Array::Set<1, uint64>(RewardIteration, 0, AgentIds);
+	UE::Learning::Array::Set<1, uint64>(CompletionIteration, 0, AgentIds);
+	UE::Learning::Array::Set<1, float>(EpisodeTimes, 0.0f, AgentIds);
 }
 
-void ULearningAgentsTrainer::OnAgentsRemoved(const TArray<int32>& AgentIds)
+void ULearningAgentsTrainer::OnAgentsRemoved_Implementation(const TArray<int32>& AgentIds)
 {
-	if (IsSetup())
+	if (!IsSetup())
 	{
-		UE::Learning::Array::Set<1, uint64>(RewardEvaluatedAgentIteration, INDEX_NONE, AgentIds);
-		UE::Learning::Array::Set<1, uint64>(CompletionEvaluatedAgentIteration, INDEX_NONE, AgentIds);
-
-		for (ULearningAgentsReward* RewardObject : RewardObjects)
-		{
-			RewardObject->OnAgentsRemoved(AgentIds);
-		}
-
-		for (ULearningAgentsCompletion* CompletionObject : CompletionObjects)
-		{
-			CompletionObject->OnAgentsRemoved(AgentIds);
-		}
+		UE_LOG(LogLearning, Error, TEXT("%s: Setup not complete."), *GetName());
+		return;
+	}
 
-		for (ULearningAgentsHelper* Helper : HelperObjects)
-		{
-			Helper->OnAgentsRemoved(AgentIds);
-		}
+	EpisodeBuffer->Reset(AgentIds);
 
-		AgentsRemoved(AgentIds);
-	}
+	UE::Learning::Array::Set<1, float>(Rewards, FLT_MAX, AgentIds);
+	UE::Learning::Array::Set<1, UE::Learning::ECompletionMode>(AgentCompletions, UE::Learning::ECompletionMode::Terminated, AgentIds);
+	UE::Learning::Array::Set<1, UE::Learning::ECompletionMode>(EpisodeCompletions, UE::Learning::ECompletionMode::Terminated, AgentIds);
+	UE::Learning::Array::Set<1, UE::Learning::ECompletionMode>(AllCompletions, UE::Learning::ECompletionMode::Terminated, AgentIds);
+	UE::Learning::Array::Set<1, uint64>(RewardIteration, INDEX_NONE, AgentIds);
+	UE::Learning::Array::Set<1, uint64>(CompletionIteration, INDEX_NONE, AgentIds);
+	UE::Learning::Array::Set<1, float>(EpisodeTimes, FLT_MAX, AgentIds);
 }
 
-void ULearningAgentsTrainer::OnAgentsReset(const TArray<int32>& AgentIds)
+void ULearningAgentsTrainer::OnAgentsReset_Implementation(const TArray<int32>& AgentIds)
 {
-	if (IsSetup())
+	if (!IsSetup())
 	{
-		ResetEpisodes(AgentIds);
-		EpisodeBuffer->Reset(AgentIds);
-
-		UE::Learning::Array::Set<1, uint64>(RewardEvaluatedAgentIteration, 0, AgentIds);
-		UE::Learning::Array::Set<1, uint64>(CompletionEvaluatedAgentIteration, 0, AgentIds);
-
-		for (ULearningAgentsReward* RewardObject : RewardObjects)
-		{
-			RewardObject->OnAgentsReset(AgentIds);
-		}
-
-		for (ULearningAgentsCompletion* CompletionObject : CompletionObjects)
-		{
-			CompletionObject->OnAgentsReset(AgentIds);
-		}
+		UE_LOG(LogLearning, Error, TEXT("%s: Setup not complete."), *GetName());
+		return;
+	}
 
-		for (ULearningAgentsHelper* Helper : HelperObjects)
-		{
-			Helper->OnAgentsReset(AgentIds);
-		}
+	ResetAgentEpisodes(AgentIds);
+	EpisodeBuffer->Reset(AgentIds);
 
-		AgentsReset(AgentIds);
-	}
+	UE::Learning::Array::Set<1, float>(Rewards, 0.0f, AgentIds);
+	UE::Learning::Array::Set<1, UE::Learning::ECompletionMode>(AgentCompletions, UE::Learning::ECompletionMode::Running, AgentIds);
+	UE::Learning::Array::Set<1, UE::Learning::ECompletionMode>(EpisodeCompletions, UE::Learning::ECompletionMode::Running, AgentIds);
+	UE::Learning::Array::Set<1, UE::Learning::ECompletionMode>(AllCompletions, UE::Learning::ECompletionMode::Running, AgentIds);
+	UE::Learning::Array::Set<1, uint64>(RewardIteration, 0, AgentIds);
+	UE::Learning::Array::Set<1, uint64>(CompletionIteration, 0, AgentIds);
+	UE::Learning::Array::Set<1, float>(EpisodeTimes, 0.0f, AgentIds);
 }
 
-void ULearningAgentsTrainer::EndPlay(const EEndPlayReason::Type EndPlayReason)
+void ULearningAgentsTrainer::OnAgentsManagerTick_Implementation(const TArray<int32>& AgentIds, const float DeltaTime)
 {
-	if (bIsTraining)
+	if (!IsSetup())
 	{
-		EndTraining();
+		UE_LOG(LogLearning, Error, TEXT("%s: Setup not complete."), *GetName());
+		return;
 	}
 
-	Super::EndPlay(EndPlayReason);
+	for (const int32 AgentId : AgentIds)
+	{
+		EpisodeTimes[AgentId] += DeltaTime;
+	}
 }
 
-void ULearningAgentsTrainer::SetupRewards_Implementation()
+void ULearningAgentsTrainer::GatherAgentRewards_Implementation(TArray<float>& OutRewards, const TArray<int32>& AgentIds)
 {
-	// Can be overridden to setup rewards without blueprints
+	UE_LOG(LogLearning, Error, TEXT("%s: GatherAgentRewards function must be overridden!"), *GetName());
+	OutRewards.Empty();
 }
 
-void ULearningAgentsTrainer::SetRewards_Implementation(const TArray<int32>& AgentIds)
+void ULearningAgentsTrainer::GatherAgentRewardsUsingDelegate(TArray<float>& OutRewards, const TArray<int32>& AgentIds, const FGatherAgentRewardDelegate& Delegate)
 {
-	// Can be overridden to set rewards without blueprints
+	if (!Delegate.IsBound())
+	{
+		UE_LOG(LogLearning, Error, TEXT("%s: Delegate Not Bound."), *GetName());
+		return;
+	}
+
+	const int32 AgentNum = AgentIds.Num();
+
+	OutRewards.Empty(AgentNum);
+	for (int32 AgentIdx = 0; AgentIdx < AgentNum; AgentIdx++)
+	{
+		OutRewards.Add(Delegate.Execute(AgentIds[AgentIdx]));
+	}
 }
 
-void ULearningAgentsTrainer::AddReward(TObjectPtr<ULearningAgentsReward> Object, const TSharedRef<UE::Learning::FRewardObject>& Reward)
+void ULearningAgentsTrainer::GatherAgentCompletions_Implementation(TArray<ELearningAgentsCompletion>& OutCompletions, const TArray<int32>& AgentIds)
 {
-	UE_LEARNING_CHECK(!IsSetup());
-	UE_LEARNING_CHECK(Object);
-	RewardObjects.Add(Object);
-	RewardFeatures.Add(Reward);
+	UE_LOG(LogLearning, Error, TEXT("%s: GatherAgentCompletions function must be overridden!"), *GetName());
+	OutCompletions.Empty();
 }
 
-void ULearningAgentsTrainer::SetupCompletions_Implementation()
+void ULearningAgentsTrainer::GatherAgentCompletionsDelegate(TArray<ELearningAgentsCompletion>& OutCompletions, const TArray<int32>& AgentIds, const FGatherAgentCompletionDelegate& Delegate)
 {
-	// Can be overridden to setup completions without blueprints
+	if (!Delegate.IsBound())
+	{
+		UE_LOG(LogLearning, Error, TEXT("%s: Delegate Not Bound."), *GetName());
+		return;
+	}
+
+	const int32 AgentNum = AgentIds.Num();
+
+	OutCompletions.Empty(AgentNum);
+	for (int32 AgentIdx = 0; AgentIdx < AgentNum; AgentIdx++)
+	{
+		OutCompletions.Add(Delegate.Execute(AgentIds[AgentIdx]));
+	}
 }
 
-void ULearningAgentsTrainer::SetCompletions_Implementation(const TArray<int32>& AgentIds)
+void ULearningAgentsTrainer::ResetAgentEpisodes_Implementation(const TArray<int32>& AgentId)
 {
-	// Can be overridden to evaluate completions without blueprints
+	UE_LOG(LogLearning, Error, TEXT("%s: ResetAgentEpisodes function must be overridden!"), *GetName());
 }
 
-void ULearningAgentsTrainer::AddCompletion(TObjectPtr<ULearningAgentsCompletion> Object, const TSharedRef<UE::Learning::FCompletionObject>& Completion)
+void ULearningAgentsTrainer::ResetAgentEpisodesDelegate(const TArray<int32>& AgentIds, const FResetAgentEpisodeDelegate& Delegate)
 {
-	UE_LEARNING_CHECK(!IsSetup());
-	UE_LEARNING_CHECK(Object);
-	CompletionObjects.Add(Object);
-	CompletionFeatures.Add(Completion);
+	if (!Delegate.IsBound())
+	{
+		UE_LOG(LogLearning, Error, TEXT("%s: Delegate Not Bound."), *GetName());
+		return;
+	}
+
+	for (const int32 AgentId : AgentIds)
+	{
+		Delegate.Execute(AgentId);
+	}
 }
 
 const bool ULearningAgentsTrainer::IsTraining() const
@@ -374,8 +373,6 @@ void ULearningAgentsTrainer::BeginTraining(
 	const FLearningAgentsTrainerTrainingSettings& TrainerTrainingSettings,
 	const FLearningAgentsTrainerGameSettings& TrainerGameSettings,
 	const FLearningAgentsTrainerPathSettings& TrainerPathSettings,
-	const bool bReinitializePolicyNetwork,
-	const bool bReinitializeCriticNetwork,
 	const bool bResetAgentsOnBegin)
 {
 	if (!IsSetup())
@@ -477,7 +474,6 @@ void ULearningAgentsTrainer::BeginTraining(
 	PPOTrainingSettings.LearningRateCritic = TrainerTrainingSettings.LearningRateCritic;
 	PPOTrainingSettings.LearningRateDecay = TrainerTrainingSettings.LearningRateDecay;
 	PPOTrainingSettings.WeightDecay = TrainerTrainingSettings.WeightDecay;
-	PPOTrainingSettings.InitialActionScale = TrainerTrainingSettings.InitialActionScale;
 	PPOTrainingSettings.PolicyBatchSize = TrainerTrainingSettings.PolicyBatchSize;
 	PPOTrainingSettings.CriticBatchSize = TrainerTrainingSettings.CriticBatchSize;
 	PPOTrainingSettings.PolicyWindow = TrainerTrainingSettings.PolicyWindowSize;
@@ -495,14 +491,7 @@ void ULearningAgentsTrainer::BeginTraining(
 	PPOTrainingSettings.DiscountFactor = TrainerTrainingSettings.DiscountFactor;
 	PPOTrainingSettings.Device = UE::Learning::Agents::GetTrainerDevice(TrainerTrainingSettings.Device);
 	PPOTrainingSettings.bUseTensorboard = TrainerTrainingSettings.bUseTensorboard;
-
-	UE::Learning::FPPOTrainerNetworkSettings PPONetworkSettings;
-	PPONetworkSettings.PolicyActionNoiseMin = Policy->GetPolicyObject().Settings.ActionNoiseMin;
-	PPONetworkSettings.PolicyActionNoiseMax = Policy->GetPolicyObject().Settings.ActionNoiseMax;
-
-	UE::Learning::EPPOTrainerFlags TrainerFlags = UE::Learning::EPPOTrainerFlags::None;
-	if (!bReinitializePolicyNetwork) { TrainerFlags |= UE::Learning::EPPOTrainerFlags::UseInitialPolicyNetwork; }
-	if (!bReinitializeCriticNetwork) { TrainerFlags |= UE::Learning::EPPOTrainerFlags::UseInitialCriticNetwork; }
+	PPOTrainingSettings.bSaveSnapshots = TrainerTrainingSettings.bSaveSnapshots;
 
 	// Start Python Training Process (this must be done on game thread)
 	Trainer = MakeUnique<UE::Learning::FSharedMemoryPPOTrainer>(
@@ -512,64 +501,54 @@ void ULearningAgentsTrainer::BeginTraining(
 		PythonContentPath,
 		IntermediatePath,
 		*ReplayBuffer,
-		Policy->GetPolicyNetwork(),
-		Critic->GetCriticNetwork(),
-		PPOTrainingSettings,
-		PPONetworkSettings,
-		TrainerFlags);
+		*Policy->GetPolicyNetworkAsset()->NeuralNetworkData,
+		*Critic->GetCriticNetworkAsset()->NeuralNetworkData,
+		*Policy->GetEncoderNetworkAsset()->NeuralNetworkData,
+		*Policy->GetDecoderNetworkAsset()->NeuralNetworkData,
+		Interactor->GetObservationSchema(),
+		Interactor->GetObservationSchemaElement(),
+		Interactor->GetActionSchema(),
+		Interactor->GetActionSchemaElement(),
+		PPOTrainingSettings);
 
 	UE_LOG(LogLearning, Display, TEXT("%s: Sending / Receiving initial policy..."), *GetName());
 
 	UE::Learning::ETrainerResponse Response = UE::Learning::ETrainerResponse::Success;
 
-	Response = Trainer->SendPolicy(Policy->GetPolicyNetwork(), TrainerTimeout);
-
+	Response = Trainer->SendPolicy(*Policy->GetPolicyNetworkAsset()->NeuralNetworkData, TrainerTimeout);
 	if (Response != UE::Learning::ETrainerResponse::Success)
 	{
-		UE_LOG(LogLearning, Error, TEXT("%s: Error sending policy to trainer: %s. Check log for errors."), *GetName(), UE::Learning::Trainer::GetResponseString(Response));
+		UE_LOG(LogLearning, Error, TEXT("%s: Error sending policy to trainer: %s. Check log for additional errors."), *GetName(), UE::Learning::Trainer::GetResponseString(Response));
 		bHasTrainingFailed = true;
 		Trainer->Terminate();
 		return;
 	}
 
-	Response = Trainer->SendCritic(Critic->GetCriticNetwork(), TrainerTimeout);
-
+	Response = Trainer->SendCritic(*Critic->GetCriticNetworkAsset()->NeuralNetworkData, TrainerTimeout);
 	if (Response != UE::Learning::ETrainerResponse::Success)
 	{
-		UE_LOG(LogLearning, Error, TEXT("%s: Error sending critic to trainer: %s. Check log for errors."), *GetName(), UE::Learning::Trainer::GetResponseString(Response));
+		UE_LOG(LogLearning, Error, TEXT("%s: Error sending critic to trainer: %s. Check log for additional errors."), *GetName(), UE::Learning::Trainer::GetResponseString(Response));
 		bHasTrainingFailed = true;
 		Trainer->Terminate();
 		return;
 	}
 
-	if (!(bool)(TrainerFlags & UE::Learning::EPPOTrainerFlags::UseInitialPolicyNetwork))
+	Response = Trainer->SendEncoder(*Policy->GetEncoderNetworkAsset()->NeuralNetworkData, TrainerTimeout);
+	if (Response != UE::Learning::ETrainerResponse::Success)
 	{
-		Response = Trainer->RecvPolicy(Policy->GetPolicyNetwork(), TrainerTimeout);
-
-		if (Response != UE::Learning::ETrainerResponse::Success)
-		{
-			UE_LOG(LogLearning, Error, TEXT("%s: Error receiving policy from trainer: %s. Check log for errors."), *GetName(), UE::Learning::Trainer::GetResponseString(Response));
-			bHasTrainingFailed = true;
-			Trainer->Terminate();
-			return;
-		}
-
-		Policy->GetNetworkAsset()->ForceMarkDirty();
+		UE_LOG(LogLearning, Error, TEXT("%s: Error sending encoder to trainer: %s. Check log for additional errors."), *GetName(), UE::Learning::Trainer::GetResponseString(Response));
+		bHasTrainingFailed = true;
+		Trainer->Terminate();
+		return;
 	}
 
-	if (!(bool)(TrainerFlags & UE::Learning::EPPOTrainerFlags::UseInitialCriticNetwork))
+	Response = Trainer->SendDecoder(*Policy->GetDecoderNetworkAsset()->NeuralNetworkData, TrainerTimeout);
+	if (Response != UE::Learning::ETrainerResponse::Success)
 	{
-		Response = Trainer->RecvCritic(Critic->GetCriticNetwork(), TrainerTimeout);
-
-		if (Response != UE::Learning::ETrainerResponse::Success)
-		{
-			UE_LOG(LogLearning, Error, TEXT("%s: Error receiving critic from trainer: %s. Check log for errors."), *GetName(), UE::Learning::Trainer::GetResponseString(Response));
-			bHasTrainingFailed = true;
-			Trainer->Terminate();
-			return;
-		}
-
-		Critic->GetNetworkAsset()->ForceMarkDirty();
+		UE_LOG(LogLearning, Error, TEXT("%s: Error sending decoder to trainer: %s. Check log for additional errors."), *GetName(), UE::Learning::Trainer::GetResponseString(Response));
+		bHasTrainingFailed = true;
+		Trainer->Terminate();
+		return;
 	}
 
 	// Reset Agents, Replay Buffer
@@ -630,9 +609,9 @@ void ULearningAgentsTrainer::EndTraining()
 	}
 }
 
-void ULearningAgentsTrainer::EvaluateRewards()
+void ULearningAgentsTrainer::GatherRewards()
 {
-	UE_LEARNING_TRACE_CPUPROFILER_EVENT_SCOPE(ULearningAgentsTrainer::EvaluateRewards);
+	UE_LEARNING_TRACE_CPUPROFILER_EVENT_SCOPE(ULearningAgentsTrainer::GatherRewards);
 
 	if (!IsSetup())
 	{
@@ -640,95 +619,45 @@ void ULearningAgentsTrainer::EvaluateRewards()
 		return;
 	}
 
-	// Check agents have actually made observations and taken actions.
-
-	ValidAgentIds.Empty(Manager->GetAgentNum());
-
-	for (const int32 AgentId : Manager->GetAllAgentSet())
+	if (Manager->GetAgentNum() == 0)
 	{
-		if (Interactor->GetObservationEncodingAgentIteration()[AgentId] == 0 || 
-			Interactor->GetActionEncodingAgentIteration()[AgentId] == 0)
-		{
-			UE_LOG(LogLearning, Display, TEXT("%s: Agent with id %i has not made observations and taken actions so rewards will not be evaluated for it."), *GetName(), AgentId);
-			continue;
-		}
-
-		ValidAgentIds.Add(AgentId);
+		UE_LOG(LogLearning, Warning, TEXT("%s: No agents added to Manager."), *GetName());
 	}
 
-	ValidAgentSet = ValidAgentIds;
-	ValidAgentSet.TryMakeSlice();
-
-	// Set Rewards
-
-	SetRewards(ValidAgentIds);
-
-	// Check that all rewards have had their setter run
+	// Get Rewards
 
-	ValidAgentStatus.SetNumUninitialized(Manager->GetMaxAgentNum());
-	ValidAgentStatus.SetRange(0, Manager->GetMaxAgentNum(), true);
+	ValidAgentIds = Manager->GetAllAgentIds();
+	ValidAgentSet = Manager->GetAllAgentSet();
 
-	for (ULearningAgentsReward* RewardObject : RewardObjects)
-	{
-		for (const int32 AgentId : ValidAgentSet)
-		{
-			if (RewardObject->GetAgentIteration(AgentId) == RewardEvaluatedAgentIteration[AgentId])
-			{
-				UE_LOG(LogLearning, Error, TEXT("%s: Reward %s for agent with id %i has not been set (got iteration %i, expected iteration %i) and so agent will not have rewards evaluated."), *GetName(), *RewardObject->GetName(), AgentId, RewardObject->GetAgentIteration(AgentId), RewardEvaluatedAgentIteration[AgentId] + 1);
-				ValidAgentStatus[AgentId] = false;
-				continue;
-			}
-
-			if (RewardObject->GetAgentIteration(AgentId) > RewardEvaluatedAgentIteration[AgentId] + 1)
-			{
-				UE_LOG(LogLearning, Error, TEXT("%s: Reward %s for agent with id %i appears to have been set multiple times (got iteration %i, expected iteration %i) and so agent will not have rewards evaluated."), *GetName(), *RewardObject->GetName(), AgentId, RewardObject->GetAgentIteration(AgentId), RewardEvaluatedAgentIteration[AgentId] + 1);
-				ValidAgentStatus[AgentId] = false;
-				continue;
-			}
-
-			if (RewardObject->GetAgentIteration(AgentId) != RewardEvaluatedAgentIteration[AgentId] + 1)
-			{
-				UE_LOG(LogLearning, Error, TEXT("%s: Reward %s for agent with id %i does not have a matching iteration number (got iteration %i, expected iteration %i) and so agent will not have rewards evaluated."), *GetName(), *RewardObject->GetName(), AgentId, RewardObject->GetAgentIteration(AgentId), RewardEvaluatedAgentIteration[AgentId] + 1);
-				ValidAgentStatus[AgentId] = false;
-				continue;
-			}
-		}
-	}
-
-	FinalValidAgentIds.Empty(ValidAgentSet.Num());
+	RewardBuffer.Empty(Manager->GetMaxAgentNum());
+	GatherAgentRewards(RewardBuffer, ValidAgentIds);
 
-	for (const int32 AgentId : ValidAgentSet)
+	if (ValidAgentSet.Num() != RewardBuffer.Num())
 	{
-		if (ValidAgentStatus[AgentId]) { FinalValidAgentIds.Add(AgentId); }
+		UE_LOG(LogLearning, Warning, TEXT("%s: Not enough rewards added by GetAgentRewards. Expected %i, Got %i."), *GetName(), ValidAgentSet.Num(), RewardBuffer.Num());
+		return;
 	}
 
-	FinalValidAgentSet = FinalValidAgentIds;
-	FinalValidAgentSet.TryMakeSlice();
-
-	// Evaluate Rewards
-
-	Rewards->Evaluate(FinalValidAgentSet);
-
-	// Increment Reward Evaluation Iteration
-
-	for (const int32 AgentId : FinalValidAgentSet)
+	for (int32 AgentIdx = 0; AgentIdx < RewardBuffer.Num(); AgentIdx++)
 	{
-		RewardEvaluatedAgentIteration[AgentId]++;
-	}
-
-	// Visual Logger
+		const float RewardValue = RewardBuffer[AgentIdx];
 
-#if UE_LEARNING_AGENTS_ENABLE_VISUAL_LOG
-	for (const ULearningAgentsReward* RewardObject : RewardObjects)
-	{
-		RewardObject->VisualLog(FinalValidAgentSet);
+		if (FMath::IsFinite(RewardValue) && RewardValue != MAX_flt && RewardValue != -MAX_flt)
+		{
+			Rewards[ValidAgentSet[AgentIdx]] = RewardValue;
+			RewardIteration[ValidAgentSet[AgentIdx]]++;
+		}
+		else
+		{
+			UE_LOG(LogLearning, Warning, TEXT("%s: Got invalid reward for agent %i: %f."), *GetName(), ValidAgentSet[AgentIdx], RewardValue);
+			continue;
+		}
 	}
-#endif
 }
 
-void ULearningAgentsTrainer::EvaluateCompletions()
+void ULearningAgentsTrainer::GatherCompletions()
 {
-	UE_LEARNING_TRACE_CPUPROFILER_EVENT_SCOPE(ULearningAgentsTrainer::EvaluateCompletions);
+	UE_LEARNING_TRACE_CPUPROFILER_EVENT_SCOPE(ULearningAgentsTrainer::GatherCompletions);
 
 	if (!IsSetup())
 	{
@@ -736,93 +665,33 @@ void ULearningAgentsTrainer::EvaluateCompletions()
 		return;
 	}
 
-	// Check agents have actually made observations and taken actions.
-
-	ValidAgentIds.Empty(Manager->GetAgentNum());
-
-	for (const int32 AgentId : Manager->GetAllAgentSet())
-	{
-		if (Interactor->GetObservationEncodingAgentIteration()[AgentId] == 0 || 
-			Interactor->GetActionEncodingAgentIteration()[AgentId] == 0)
-		{
-			UE_LOG(LogLearning, Display, TEXT("%s: Agent with id %i has not made observations and taken actions so completions will not be evaluated for it."), *GetName(), AgentId);
-			continue;
-		}
-
-		ValidAgentIds.Add(AgentId);
-	}
-		
-	ValidAgentSet = ValidAgentIds;
-	ValidAgentSet.TryMakeSlice();
-
-	// Set Completions
-
-	SetCompletions(ValidAgentIds);
-
-	// Check that all completions have had their setter run
-
-	ValidAgentStatus.SetNumUninitialized(Manager->GetMaxAgentNum());
-	ValidAgentStatus.SetRange(0, Manager->GetMaxAgentNum(), true);
-
-	for (ULearningAgentsCompletion* CompletionObject : CompletionObjects)
-	{
-		for (const int32 AgentId : ValidAgentSet)
-		{
-			if (CompletionObject->GetAgentIteration(AgentId) == CompletionEvaluatedAgentIteration[AgentId])
-			{
-				UE_LOG(LogLearning, Error, TEXT("%s: Completion %s for agent with id %i has not been set (got iteration %i, expected iteration %i) and so agent will not have completions evaluated."), *GetName(), *CompletionObject->GetName(), AgentId, CompletionObject->GetAgentIteration(AgentId), CompletionEvaluatedAgentIteration[AgentId] + 1);
-				ValidAgentStatus[AgentId] = false;
-				continue;
-			}
-
-			if (CompletionObject->GetAgentIteration(AgentId) > CompletionEvaluatedAgentIteration[AgentId] + 1)
-			{
-				UE_LOG(LogLearning, Error, TEXT("%s: Completion %s for agent with id %i appears to have been set multiple times (got iteration %i, expected iteration %i) and so agent will not have completions evaluated."), *GetName(), *CompletionObject->GetName(), AgentId, CompletionObject->GetAgentIteration(AgentId), CompletionEvaluatedAgentIteration[AgentId] + 1);
-				ValidAgentStatus[AgentId] = false;
-				continue;
-			}
-
-			if (CompletionObject->GetAgentIteration(AgentId) != CompletionEvaluatedAgentIteration[AgentId] + 1)
-			{
-				UE_LOG(LogLearning, Error, TEXT("%s: Completion %s for agent with id %i does not have a matching iteration number (got iteration %i, expected iteration %i) and so agent will not have completions evaluated."), *GetName(), *CompletionObject->GetName(), AgentId, CompletionObject->GetAgentIteration(AgentId), CompletionEvaluatedAgentIteration[AgentId] + 1);
-				ValidAgentStatus[AgentId] = false;
-				continue;
-			}
-		}
-	}
-
-	FinalValidAgentIds.Empty(ValidAgentSet.Num());
-
-	for (const int32 AgentId : ValidAgentSet)
+	if (Manager->GetAgentNum() == 0)
 	{
-		if (ValidAgentStatus[AgentId]) { FinalValidAgentIds.Add(AgentId); }
+		UE_LOG(LogLearning, Warning, TEXT("%s: No agents added to Manager."), *GetName());
 	}
 
-	FinalValidAgentSet = FinalValidAgentIds;
-	FinalValidAgentSet.TryMakeSlice();
+	// Get Completions
 
-	// Evaluate Completions
+	ValidAgentIds = Manager->GetAllAgentIds();
+	ValidAgentSet = Manager->GetAllAgentSet();
 
-	Completions->Evaluate(FinalValidAgentSet);
-	
-	// Increment Completion Evaluation Iteration
+	CompletionBuffer.Empty(Manager->GetMaxAgentNum());
+	GatherAgentCompletions(CompletionBuffer, ValidAgentIds);
 
-	for (const int32 AgentId : FinalValidAgentSet)
+	if (ValidAgentSet.Num() != CompletionBuffer.Num())
 	{
-		CompletionEvaluatedAgentIteration[AgentId]++;
+		UE_LOG(LogLearning, Warning, TEXT("%s: Not enough completions added by GetAgentCompletions. Expected %i, Got %i."), *GetName(), ValidAgentSet.Num(), CompletionBuffer.Num());
+		return;
 	}
 
-	// Visual Logger
-
-#if UE_LEARNING_AGENTS_ENABLE_VISUAL_LOG
-	for (const ULearningAgentsCompletion* CompletionObject : CompletionObjects)
+	for (int32 AgentIdx = 0; AgentIdx < CompletionBuffer.Num(); AgentIdx++)
 	{
-		CompletionObject->VisualLog(FinalValidAgentSet);
+		AgentCompletions[ValidAgentSet[AgentIdx]] = UE::Learning::Agents::GetCompletionMode(CompletionBuffer[AgentIdx]);
+		CompletionIteration[ValidAgentSet[AgentIdx]]++;
 	}
-#endif
 }
 
-void ULearningAgentsTrainer::ProcessExperience()
+void ULearningAgentsTrainer::ProcessExperience(const bool bResetAgentsOnUpdate)
 {
 	UE_LEARNING_TRACE_CPUPROFILER_EVENT_SCOPE(ULearningAgentsTrainer::ProcessExperience);
 
@@ -838,35 +707,35 @@ void ULearningAgentsTrainer::ProcessExperience()
 		return;
 	}
 
-	if (Policy->GetActionNoiseScale() != 1.0f)
+	if (Manager->GetAgentNum() == 0)
 	{
-		UE_LOG(LogLearning, Warning, TEXT("%s: Policy Action Noise Scale should be set to 1.0 during training."), *GetName());
+		UE_LOG(LogLearning, Warning, TEXT("%s: No agents added to Manager."), *GetName());
 	}
 
 	// Check Observations, Actions, Rewards, and Completions have been completed and have matching iteration number
 
-	ValidAgentIds.Empty(Manager->GetAgentNum());
+	ValidAgentIds.Empty(Manager->GetMaxAgentNum());
 
 	for (const int32 AgentId : Manager->GetAllAgentSet())
 	{
-		if (Interactor->GetObservationEncodingAgentIteration()[AgentId] == 0 ||
-			Interactor->GetActionEncodingAgentIteration()[AgentId] == 0 ||
-			RewardEvaluatedAgentIteration[AgentId] == 0 ||
-			CompletionEvaluatedAgentIteration[AgentId] == 0)
+		if (Interactor->ObservationVectorIteration[AgentId] == 0 ||
+			Interactor->ActionVectorIteration[AgentId] == 0 ||
+			RewardIteration[AgentId] == 0 ||
+			CompletionIteration[AgentId] == 0)
 		{
 			UE_LOG(LogLearning, Display, TEXT("%s: Agent with id %i has not completed a full step of observations, actions, rewards, completions and so experience will not be processed for it."), *GetName(), AgentId);
 			continue;
 		}
 
-		if (Interactor->GetObservationEncodingAgentIteration()[AgentId] != Interactor->GetActionEncodingAgentIteration()[AgentId] ||
-			Interactor->GetObservationEncodingAgentIteration()[AgentId] != RewardEvaluatedAgentIteration[AgentId] ||
-			Interactor->GetObservationEncodingAgentIteration()[AgentId] != CompletionEvaluatedAgentIteration[AgentId])
+		if (Interactor->ObservationVectorIteration[AgentId] != Interactor->ActionVectorIteration[AgentId] ||
+			Interactor->ObservationVectorIteration[AgentId] != RewardIteration[AgentId] ||
+			Interactor->ObservationVectorIteration[AgentId] != CompletionIteration[AgentId])
 		{
 			UE_LOG(LogLearning, Warning, TEXT("%s: Agent with id %i has non-matching iteration numbers (observation: %i, action: %i, reward: %i, completion: %i). Experience will not be processed for it."), *GetName(), AgentId,
-				Interactor->GetObservationEncodingAgentIteration()[AgentId],
-				Interactor->GetActionEncodingAgentIteration()[AgentId],
-				RewardEvaluatedAgentIteration[AgentId],
-				CompletionEvaluatedAgentIteration[AgentId]);
+				Interactor->ObservationVectorIteration[AgentId],
+				Interactor->ActionVectorIteration[AgentId],
+				RewardIteration[AgentId],
+				CompletionIteration[AgentId]);
 			continue;
 		}
 
@@ -874,104 +743,144 @@ void ULearningAgentsTrainer::ProcessExperience()
 	}
 
 	ValidAgentSet = ValidAgentIds;
+	ValidAgentSet.TryMakeSlice();
 
 	// Check for episodes that have been immediately completed
 
 	for (const int32 AgentId : ValidAgentSet)
 	{
-		if (Completions->CompletionBuffer()[AgentId] != UE::Learning::ECompletionMode::Running && 
-			EpisodeBuffer->GetEpisodeStepNums()[AgentId] == 0)
+		if (AgentCompletions[AgentId] != UE::Learning::ECompletionMode::Running && EpisodeBuffer->GetEpisodeStepNums()[AgentId] == 0)
 		{
 			UE_LOG(LogLearning, Warning, TEXT("%s: Agent with id %i has completed episode and will be reset but has not generated any experience."), *GetName(), AgentId);
 		}
 	}
 
-	// Get Feature Buffers
-
-	UE::Learning::FFeatureObject& Observations = Interactor->GetObservationFeature();
-	UE::Learning::FFeatureObject& Actions = Interactor->GetActionFeature();
-
 	// Add Experience to Episode Buffer
 	EpisodeBuffer->Push(
-		Observations.FeatureBuffer(),
-		Actions.FeatureBuffer(),
-		Policy->GetPreEvaluationMemoryStateView(),
-		Rewards->RewardBuffer(),
+		Interactor->ObservationVectors,
+		Interactor->ActionVectors,
+		Policy->PreEvaluationMemoryState,
+		Rewards,
 		ValidAgentSet);
 
-	// Check for completion based on reaching the maximum episode length
+	// Find the set of agents which have reached the maximum episode length and mark them as truncated
 	UE::Learning::Completion::EvaluateEndOfEpisodeCompletions(
-		Completions->CompletionBuffer(),
+		EpisodeCompletions,
 		EpisodeBuffer->GetEpisodeStepNums(),
 		EpisodeBuffer->GetMaxStepNum(),
-		MaxStepsCompletion == ELearningAgentsCompletion::Truncation
-		? UE::Learning::ECompletionMode::Truncated
-		: UE::Learning::ECompletionMode::Terminated,
 		ValidAgentSet);
 
-	// Find the set of Instances that need to be reset
-	ResetBuffer->SetResetInstancesFromCompletions(Completions->CompletionBuffer(), ValidAgentSet);
+	// Compute a combined completion buffer for agents that have been completed manually and those which have reached the maximum episode length
+	for (const int32 AgentIdx : ValidAgentSet)
+	{
+		AllCompletions[AgentIdx] = UE::Learning::Completion::Or(AgentCompletions[AgentIdx], EpisodeCompletions[AgentIdx]);
+	}
 
-	if (ResetBuffer->GetResetInstanceNum() > 0)
+	ResetBuffer->SetResetInstancesFromCompletions(AllCompletions, ValidAgentSet);
+
+	// If there are no agents completed we are done
+	if (ResetBuffer->GetResetInstanceNum() == 0)
+	{
+		return;
+	}
+
+	// Otherwise Gather Observations for completed Instances without incrementing iteration number
+	Interactor->GatherObservations(ResetBuffer->GetResetInstances(), false);
+
+	// And push those episodes to the Replay Buffer
+	const bool bReplayBufferFull = ReplayBuffer->AddEpisodes(
+		AllCompletions,
+		Interactor->ObservationVectors,
+		Policy->MemoryState,
+		*EpisodeBuffer,
+		ResetBuffer->GetResetInstances());
+
+	if (bReplayBufferFull)
 	{
-		// Encode Observations for completed Instances
-		Interactor->EncodeObservations(ResetBuffer->GetResetInstances());
+		UE::Learning::ETrainerResponse Response = Trainer->SendExperience(*ReplayBuffer, TrainerTimeout);
+
+		if (Response != UE::Learning::ETrainerResponse::Success)
+		{
+			UE_LOG(LogLearning, Error, TEXT("%s: Error waiting to push experience to trainer: %s. Check log for additional errors."), *GetName(), UE::Learning::Trainer::GetResponseString(Response));
+			EndTraining();
+			return;
+		}
+
+		ReplayBuffer->Reset();
+
+		// Get Updated Policy
+		Response = Trainer->RecvPolicy(*Policy->GetPolicyNetworkAsset()->NeuralNetworkData, TrainerTimeout);
+		Policy->GetPolicyNetworkAsset()->ForceMarkDirty();
+
+		if (Response == UE::Learning::ETrainerResponse::Completed)
+		{
+			UE_LOG(LogLearning, Display, TEXT("%s: Trainer completed training."), *GetName());
+			DoneTraining();
+			return;
+		}
+		else if (Response != UE::Learning::ETrainerResponse::Success)
+		{
+			UE_LOG(LogLearning, Error, TEXT("%s: Error waiting for policy from trainer: %s. Check log for additional errors."), *GetName(), UE::Learning::Trainer::GetResponseString(Response));
+			bHasTrainingFailed = true;
+			EndTraining();
+			return;
+		}
+
+		// Get Updated Critic
+		Response = Trainer->RecvCritic(*Critic->GetCriticNetworkAsset()->NeuralNetworkData, TrainerTimeout);
+		Critic->GetCriticNetworkAsset()->ForceMarkDirty();
+
+		if (Response != UE::Learning::ETrainerResponse::Success)
+		{
+			UE_LOG(LogLearning, Error, TEXT("%s: Error waiting for critic from trainer: %s. Check log for additional errors."), *GetName(), UE::Learning::Trainer::GetResponseString(Response));
+			bHasTrainingFailed = true;
+			EndTraining();
+			return;
+		}
+
+		// Get Updated Encoder
+		Response = Trainer->RecvEncoder(*Policy->GetEncoderNetworkAsset()->NeuralNetworkData, TrainerTimeout);
+		Policy->GetEncoderNetworkAsset()->ForceMarkDirty();
+
+		if (Response != UE::Learning::ETrainerResponse::Success)
+		{
+			UE_LOG(LogLearning, Error, TEXT("%s: Error waiting for encoder from trainer: %s. Check log for additional errors."), *GetName(), UE::Learning::Trainer::GetResponseString(Response));
+			bHasTrainingFailed = true;
+			EndTraining();
+			return;
+		}
+
+		// Get Updated Decoder
+		Response = Trainer->RecvDecoder(*Policy->GetDecoderNetworkAsset()->NeuralNetworkData, TrainerTimeout);
+		Policy->GetDecoderNetworkAsset()->ForceMarkDirty();
 
-		const bool bReplayBufferFull = ReplayBuffer->AddEpisodes(
-			Completions->CompletionBuffer(),
-			Observations.FeatureBuffer(),
-			Policy->GetMemoryStateView(),
-			*EpisodeBuffer,
-			ResetBuffer->GetResetInstances());
+		if (Response != UE::Learning::ETrainerResponse::Success)
+		{
+			UE_LOG(LogLearning, Error, TEXT("%s: Error waiting for decoder from trainer: %s. Check log for additional errors."), *GetName(), UE::Learning::Trainer::GetResponseString(Response));
+			bHasTrainingFailed = true;
+			EndTraining();
+			return;
+		}
 
-		if (bReplayBufferFull)
+		if (bResetAgentsOnUpdate)
 		{
-			UE::Learning::ETrainerResponse Response = Trainer->SendExperience(*ReplayBuffer, TrainerTimeout);
-
-			if (Response != UE::Learning::ETrainerResponse::Success)
-			{
-				UE_LOG(LogLearning, Error, TEXT("%s: Error waiting to push experience to trainer. Check log for errors."), *GetName());
-				EndTraining();
-				return;
-			}
-
-			ReplayBuffer->Reset();
-
-			// Get Updated Policy
-			Response = Trainer->RecvPolicy(Policy->GetPolicyNetwork(), TrainerTimeout);
-			Policy->GetNetworkAsset()->ForceMarkDirty();
-
-			if (Response == UE::Learning::ETrainerResponse::Completed)
-			{
-				UE_LOG(LogLearning, Display, TEXT("%s: Trainer completed training."), *GetName());
-				DoneTraining();
-				return;
-			}
-			else if (Response != UE::Learning::ETrainerResponse::Success)
-			{
-				UE_LOG(LogLearning, Error, TEXT("%s: Error waiting for policy from trainer. Check log for errors."), *GetName());
-				bHasTrainingFailed = true;
-				EndTraining();
-				return;
-			}
-
-			// Get Updated Critic
-			Response = Trainer->RecvCritic(Critic->GetCriticNetwork(), TrainerTimeout);
-			Critic->GetNetworkAsset()->ForceMarkDirty();
-
-			if (Response != UE::Learning::ETrainerResponse::Success)
-			{
-				UE_LOG(LogLearning, Error, TEXT("%s: Error waiting for critic from trainer. Check log for errors."), *GetName());
-				bHasTrainingFailed = true;
-				EndTraining();
-				return;
-			}
-
-			// Mark all agents for reset since we have a new policy
+			// Reset all agents since we have a new policy
 			ResetBuffer->SetResetInstances(Manager->GetAllAgentSet());
+			Manager->ResetAgents(ResetBuffer->GetResetInstancesArray());
+			return;
 		}
+	}
+
+	// Manually reset Episode Buffer for agents who have reached the maximum episode length as 
+	// they wont get it reset via the agent manager's call to ResetAgents
+	ResetBuffer->SetResetInstancesFromCompletions(EpisodeCompletions, ValidAgentSet);
+	EpisodeBuffer->Reset(ResetBuffer->GetResetInstances());
 
-		Manager->ResetAgents(ResetBuffer->GetResetInstances().ToArray());
+	// Call ResetAgents for agents which have manually signaled a completion
+	ResetBuffer->SetResetInstancesFromCompletions(AgentCompletions, ValidAgentSet);
+	if (ResetBuffer->GetResetInstanceNum() > 0)
+	{
+		Manager->ResetAgents(ResetBuffer->GetResetInstancesArray());
 	}
 }
 
@@ -979,9 +888,8 @@ void ULearningAgentsTrainer::RunTraining(
 	const FLearningAgentsTrainerTrainingSettings& TrainerTrainingSettings,
 	const FLearningAgentsTrainerGameSettings& TrainerGameSettings,
 	const FLearningAgentsTrainerPathSettings& TrainerPathSettings,
-	const bool bReinitializePolicyNetwork,
-	const bool bReinitializeCriticNetwork,
-	const bool bResetAgentsOnBegin)
+	const bool bResetAgentsOnBegin,
+	const bool bResetAgentsOnUpdate)
 {
 	if (!IsSetup())
 	{
@@ -1002,8 +910,6 @@ void ULearningAgentsTrainer::RunTraining(
 			TrainerTrainingSettings,
 			TrainerGameSettings, 
 			TrainerPathSettings, 
-			bReinitializePolicyNetwork, 
-			bReinitializeCriticNetwork,
 			bResetAgentsOnBegin);
 
 		if (!IsTraining())
@@ -1014,12 +920,48 @@ void ULearningAgentsTrainer::RunTraining(
 
 		Policy->RunInference();
 	}
-
 	// Otherwise, do the regular training process.
-	EvaluateCompletions();
-	EvaluateRewards();
-	ProcessExperience();
-	Policy->RunInference();
+	else
+	{
+		GatherCompletions();
+		GatherRewards();
+		ProcessExperience(bResetAgentsOnUpdate);
+		Policy->RunInference();
+	}
+}
+
+bool ULearningAgentsTrainer::HasReward(const int32 AgentId) const
+{
+	if (!IsSetup())
+	{
+		UE_LOG(LogLearning, Error, TEXT("%s: Setup not complete."), *GetName());
+		return false;
+	}
+
+	if (!HasAgent(AgentId))
+	{
+		UE_LOG(LogLearning, Error, TEXT("%s: AgentId %d not found in the agents set."), *GetName(), AgentId);
+		return false;
+	}
+
+	return RewardIteration[AgentId] > 0;
+}
+
+bool ULearningAgentsTrainer::HasCompletion(const int32 AgentId) const
+{
+	if (!IsSetup())
+	{
+		UE_LOG(LogLearning, Error, TEXT("%s: Setup not complete."), *GetName());
+		return false;
+	}
+
+	if (!HasAgent(AgentId))
+	{
+		UE_LOG(LogLearning, Error, TEXT("%s: AgentId %d not found in the agents set."), *GetName(), AgentId);
+		return false;
+	}
+
+	return CompletionIteration[AgentId] > 0;
 }
 
 float ULearningAgentsTrainer::GetReward(const int32 AgentId) const
@@ -1036,58 +978,74 @@ float ULearningAgentsTrainer::GetReward(const int32 AgentId) const
 		return 0.0f;
 	}
 
-	if (RewardEvaluatedAgentIteration[AgentId] == 0)
+	if (RewardIteration[AgentId] == 0)
 	{
 		UE_LOG(LogLearning, Error, TEXT("%s: Agent with id %d has not evaluated rewards. Did you run EvaluateRewards?"), *GetName(), AgentId);
 		return 0.0f;
 	}
 
-	return Rewards->RewardBuffer()[AgentId];
+	return Rewards[AgentId];
 }
 
-bool ULearningAgentsTrainer::IsCompleted(const int32 AgentId, ELearningAgentsCompletion& OutCompletion) const
+ELearningAgentsCompletion ULearningAgentsTrainer::GetCompletion(const int32 AgentId) const
 {
 	if (!IsSetup())
 	{
 		UE_LOG(LogLearning, Error, TEXT("%s: Setup not complete."), *GetName());
-		OutCompletion = ELearningAgentsCompletion::Termination;
-		return false;
+		return ELearningAgentsCompletion::Running;
 	}
 
 	if (!HasAgent(AgentId))
 	{
 		UE_LOG(LogLearning, Error, TEXT("%s: AgentId %d not found in the agents set."), *GetName(), AgentId);
-		OutCompletion = ELearningAgentsCompletion::Termination;
-		return false;
+		return ELearningAgentsCompletion::Running;
 	}
 
-	if (CompletionEvaluatedAgentIteration[AgentId] == 0)
+	if (CompletionIteration[AgentId] == 0)
 	{
 		UE_LOG(LogLearning, Error, TEXT("%s: Agent with id %d has not evaluated completions. Did you run EvaluateCompletions?"), *GetName(), AgentId);
-		OutCompletion = ELearningAgentsCompletion::Termination;
-		return false;
+		return ELearningAgentsCompletion::Running;
 	}
 
-	const UE::Learning::ECompletionMode CompletionMode = Completions->CompletionBuffer()[AgentId];
+	return UE::Learning::Agents::GetLearningAgentsCompletion(AgentCompletions[AgentId]);
+}
 
-	if (CompletionMode == UE::Learning::ECompletionMode::Running)
+float ULearningAgentsTrainer::GetEpisodeTime(const int32 AgentId) const
+{
+	if (!IsSetup())
 	{
-		OutCompletion = ELearningAgentsCompletion::Termination;
-		return false;
+		UE_LOG(LogLearning, Error, TEXT("%s: Setup not complete."), *GetName());
+		return 0.0f;
 	}
-	else
+
+	if (!HasAgent(AgentId))
 	{
-		OutCompletion = UE::Learning::Agents::GetLearningAgentsCompletion(CompletionMode);
-		return true;
+		UE_LOG(LogLearning, Error, TEXT("%s: AgentId %d not found in the agents set."), *GetName(), AgentId);
+		return 0.0f;
 	}
+
+	return EpisodeTimes[AgentId];
 }
 
-bool ULearningAgentsTrainer::HasTrainingFailed() const
+int32 ULearningAgentsTrainer::GetEpisodeStepNum(const int32 AgentId) const
 {
-	return bHasTrainingFailed;
+	if (!IsSetup())
+	{
+		UE_LOG(LogLearning, Error, TEXT("%s: Setup not complete."), *GetName());
+		return 0.0f;
+	}
+
+	if (!HasAgent(AgentId))
+	{
+		UE_LOG(LogLearning, Error, TEXT("%s: AgentId %d not found in the agents set."), *GetName(), AgentId);
+		return 0.0f;
+	}
+
+	return EpisodeBuffer->GetEpisodeStepNums()[AgentId];
 }
 
-void ULearningAgentsTrainer::ResetEpisodes_Implementation(const TArray<int32>& AgentId)
+bool ULearningAgentsTrainer::HasTrainingFailed() const
 {
-	// Can be overridden to reset agent without blueprints
+	return bHasTrainingFailed;
 }
+
diff --git a/Engine/Plugins/Experimental/LearningAgents/Source/LearningAgentsTraining/Public/LearningAgentsCompletions.h b/Engine/Plugins/Experimental/LearningAgents/Source/LearningAgentsTraining/Public/LearningAgentsCompletions.h
--- a/Engine/Plugins/Experimental/LearningAgents/Source/LearningAgentsTraining/Public/LearningAgentsCompletions.h
+++ b/Engine/Plugins/Experimental/LearningAgents/Source/LearningAgentsTraining/Public/LearningAgentsCompletions.h
@@ -2,258 +2,86 @@
 
 #pragma once
 
-#include "LearningArray.h"
-#include "LearningAgentsTrainer.h" // Required for ELearningAgentsCompletion::Termination
-#include "LearningAgentsDebug.h"
-
-#include "Templates/SharedPointer.h"
-#include "UObject/Object.h"
+#include "Kismet/BlueprintFunctionLibrary.h"
 
 #include "LearningAgentsCompletions.generated.h"
 
 namespace UE::Learning
 {
-	struct FCompletionObject;
-	struct FConditionalCompletion;
-	struct FTimeElapsedCompletion;
-	struct FPlanarPositionDifferenceCompletion;
-	struct FPlanarPositionSimilarityCompletion;
+	enum class ECompletionMode : uint8;
 }
 
-// For functions in this file, we are favoring having more verbose names such as "AddConditionalCompletion" vs simply "Add" in 
-// order to keep it easy to find the correct function in blueprints.
-
-//------------------------------------------------------------------
-
-/**
- * The base class for all completions. Completions contain logic that determines if an agent's current episode should
- * end, e.g. because the agent achieved the normal win/loss condition for the game. Additionally, completions can speed
- * up training by ending episodes early if the agent has gotten into a state where training data is no longer useful,
- * e.g. the agent is stuck somewhere. These two modes of completions are expressed with the following enum values:
- *   ELearningAgentsCompletion::Termination - used when the episode ends in an expected way and no further rewards
- *       should be expected, i.e. do not use the value function to estimate future rewards.
- *   ELearningAgentsCompletion::Truncation - used when the episode ends in an unexpected way, mainly to speed up the 
- *       training process. The agent should expect additional rewards if training were to continue, so it should use
- *       its value function to estimate future rewards.
- */
-UCLASS(Abstract, BlueprintType)
-class LEARNINGAGENTSTRAINING_API ULearningAgentsCompletion : public UObject
+/** Completion modes for episodes. */
+UENUM(BlueprintType, Category = "LearningAgents", meta = (ScriptName = "LearningAgentsCompletionEnum"))
+enum class ELearningAgentsCompletion : uint8
 {
-	GENERATED_BODY()
-
-public:
-
-	/** Reference to the Trainer this completion is associated with. */
-	UPROPERTY(VisibleAnywhere, Transient, Category = "LearningAgents")
-	TObjectPtr<ULearningAgentsTrainer> AgentTrainer;
-
-public:
-
-	/** Initialize the internal state for a given maximum number of agents */
-	void Init(const int32 MaxAgentNum);
+	/** Episode is still running. */
+	Running 	UMETA(DisplayName = "Running"),
 
-	/**
-	 * Called whenever agents are added to the associated ULearningAgentsTrainer object.
-	 * @param AgentIds Array of agent ids which have been added
-	 */
-	virtual void OnAgentsAdded(const TArray<int32>& AgentIds);
+	/** Episode ended while in progress. Critic will be used to estimate final return. */
+	Truncation	UMETA(DisplayName = "Truncation"),
 
-	/**
-	 * Called whenever agents are removed from the associated ULearningAgentsTrainer object.
-	 * @param AgentIds Array of agent ids which have been removed
-	 */
-	virtual void OnAgentsRemoved(const TArray<int32>& AgentIds);
-
-	/**
-	 * Called whenever agents are reset on the associated ULearningAgentsTrainer object.
-	 * @param AgentIds Array of agent ids which have been reset
-	 */
-	virtual void OnAgentsReset(const TArray<int32>& AgentIds);
-
-	/** Get the number of times a completion has been set for the given agent id. */
-	uint64 GetAgentIteration(const int32 AgentId) const;
-
-public:
-#if UE_LEARNING_AGENTS_ENABLE_VISUAL_LOG
-	/** Color used to draw this completion in the visual log */
-	FLinearColor VisualLogColor = FColor::Yellow;
-
-	/** Describes this completion to the visual logger for debugging purposes. */
-	virtual void VisualLog(const UE::Learning::FIndexSet Instances) const {}
-#endif
-
-protected:
-
-	/** Number of times this completion has been set for all agents */
-	TLearningArray<1, uint64, TInlineAllocator<32>> AgentIteration;
+	/** Episode ended and zero reward was expected for all future steps. */
+	Termination	UMETA(DisplayName = "Termination"),
 };
 
-//------------------------------------------------------------------
-
-/** A simple boolean completion. Used as a catch-all for situations where a more type-specific completion does not exist yet. */
-UCLASS()
-class LEARNINGAGENTSTRAINING_API UConditionalCompletion : public ULearningAgentsCompletion
+namespace UE::Learning::Agents
 {
-	GENERATED_BODY()
+	/** Get the learning agents completion from the UE::Learning completion. */
+	LEARNINGAGENTSTRAINING_API ELearningAgentsCompletion GetLearningAgentsCompletion(const ECompletionMode CompletionMode);
 
-public:
-
-	/**
-	 * Adds a new conditional completion to the given trainer. Call during ULearningAgentsTrainer::SetupCompletions event.
-	 * @param InAgentTrainer The trainer to add this completion to.
-	 * @param Name The name of this new completion. Used for debugging.
-	 * @param InCompletionMode The completion mode.
-	 * @return The newly created completion.
-	 */
-	UFUNCTION(BlueprintCallable, Category = "LearningAgents", meta = (DefaultToSelf = "InAgentTrainer"))
-	static UConditionalCompletion* AddConditionalCompletion(ULearningAgentsTrainer* InAgentTrainer, const FName Name = NAME_None, const ELearningAgentsCompletion InCompletionMode = ELearningAgentsCompletion::Termination);
-
-	/**
-	 * Sets the data for this completion. Call during ULearningAgentsTrainer::SetCompletions event.
-	 * @param AgentId The agent id this data corresponds to.
-	 * @param bIsCompleted Pass in true if condition is met. Otherwise, false.
-	 */
-	UFUNCTION(BlueprintCallable, Category = "LearningAgents", meta = (AgentId = "-1"))
-	void SetConditionalCompletion(const int32 AgentId, const bool bIsCompleted);
-
-#if UE_LEARNING_AGENTS_ENABLE_VISUAL_LOG
-	/** Describes this completion to the visual logger for debugging purposes. */
-	virtual void VisualLog(const UE::Learning::FIndexSet Instances) const override;
-#endif
-
-	TSharedPtr<UE::Learning::FConditionalCompletion> CompletionObject;
-};
-
-//------------------------------------------------------------------
+	/** Get the UE::Learning completion from the learning agents completion. */
+	LEARNINGAGENTSTRAINING_API ECompletionMode GetCompletionMode(const ELearningAgentsCompletion Completion);
+}
 
-/** A completion for if a given amount of time has elapsed. */
-UCLASS()
-class LEARNINGAGENTSTRAINING_API UTimeElapsedCompletion : public ULearningAgentsCompletion
+UCLASS(BlueprintType)
+class LEARNINGAGENTSTRAINING_API ULearningAgentsCompletions : public UBlueprintFunctionLibrary
 {
 	GENERATED_BODY()
 
 public:
 
-	/**
-	 * Adds a new time elapsed completion to the given trainer. Call during ULearningAgentsTrainer::SetupCompletions event.
-	 * @param InAgentTrainer The trainer to add this completion to.
-	 * @param Name The name of this new completion. Used for debugging.
-	 * @param Threshold How much time should be elapsed for the completion to trigger.
-	 * @param InCompletionMode The completion mode.
-	 * @return The newly created completion.
-	 */
-	UFUNCTION(BlueprintCallable, Category = "LearningAgents", meta = (DefaultToSelf = "InAgentTrainer"))
-	static UTimeElapsedCompletion* AddTimeElapsedCompletion(ULearningAgentsTrainer* InAgentTrainer, const FName Name = NAME_None, const float Threshold = 10.0f, const ELearningAgentsCompletion InCompletionMode = ELearningAgentsCompletion::Termination);
+	UFUNCTION(BlueprintPure, Category = "LearningAgents")
+	static bool CompletionIsRunning(const ELearningAgentsCompletion Completion);
 
-	/**
-	 * Sets the data for this completion. Call during ULearningAgentsTrainer::SetCompletions event.
-	 * @param AgentId The agent id this data corresponds to.
-	 * @param Time The amount of time that has passed
-	 */
-	UFUNCTION(BlueprintCallable, Category = "LearningAgents", meta = (AgentId = "-1"))
-	void SetTimeElapsedCompletion(const int32 AgentId, const float Time);
+	UFUNCTION(BlueprintPure, Category = "LearningAgents")
+	static bool CompletionIsCompleted(const ELearningAgentsCompletion Completion);
 
-#if UE_LEARNING_AGENTS_ENABLE_VISUAL_LOG
-	/** Describes this completion to the visual logger for debugging purposes. */
-	virtual void VisualLog(const UE::Learning::FIndexSet Instances) const override;
-#endif
+	UFUNCTION(BlueprintPure, Category = "LearningAgents")
+	static bool CompletionIsTruncation(const ELearningAgentsCompletion Completion);
 
-	TSharedPtr<UE::Learning::FTimeElapsedCompletion> CompletionObject;
-};
+	UFUNCTION(BlueprintPure, Category = "LearningAgents")
+	static bool CompletionIsTermination(const ELearningAgentsCompletion Completion);
 
-//------------------------------------------------------------------
+	UFUNCTION(BlueprintPure, Category = "LearningAgents", meta = (CommutativeAssociativeBinaryOperator))
+	static ELearningAgentsCompletion CompletionOr(ELearningAgentsCompletion A, ELearningAgentsCompletion B);
 
-/**
- * A completion for if two positions differ by some threshold in a plane, e.g. if the agent gets too far from a
- * starting position.
- */
-UCLASS()
-class LEARNINGAGENTSTRAINING_API UPlanarPositionDifferenceCompletion : public ULearningAgentsCompletion
-{
-	GENERATED_BODY()
+	UFUNCTION(BlueprintPure, Category = "LearningAgents", meta = (CommutativeAssociativeBinaryOperator))
+	static ELearningAgentsCompletion CompletionAnd(ELearningAgentsCompletion A, ELearningAgentsCompletion B);
 
-public:
-
-	/**
-	 * Adds a new planar position difference completion to the given trainer. The axis parameters define the plane.
-	 * Call during ULearningAgentsTrainer::SetupCompletions event.
-	 * @param InAgentTrainer The trainer to add this completion to.
-	 * @param Name The name of this new completion. Used for debugging.
-	 * @param Threshold If the distance becomes greater than this threshold, then the episode will complete.
-	 * @param InCompletionMode The completion mode.
-	 * @param Axis0 The forward axis of the plane.
-	 * @param Axis1 The right axis of the plane.
-	 * @return The newly created completion.
-	 */
-	UFUNCTION(BlueprintCallable, Category = "LearningAgents", meta = (DefaultToSelf = "InAgentTrainer"))
-	static UPlanarPositionDifferenceCompletion* AddPlanarPositionDifferenceCompletion(
-		ULearningAgentsTrainer* InAgentTrainer,
-		const FName Name = NAME_None,
-		const float Threshold = 100.0f,
-		const ELearningAgentsCompletion InCompletionMode = ELearningAgentsCompletion::Termination,
-		const FVector Axis0 = FVector::ForwardVector,
-		const FVector Axis1 = FVector::RightVector);
-
-	/**
-	 * Sets the data for this completion. Call during ULearningAgentsTrainer::SetCompletions event.
-	 * @param AgentId The agent id this data corresponds to.
-	 * @param Position0 The first position.
-	 * @param Position1 The second position.
-	 */
-	UFUNCTION(BlueprintCallable, Category = "LearningAgents", meta = (AgentId = "-1"))
-	void SetPlanarPositionDifferenceCompletion(const int32 AgentId, const FVector Position0, const FVector Position1);
+	UFUNCTION(BlueprintPure, Category = "LearningAgents")
+	static ELearningAgentsCompletion CompletionNot(const ELearningAgentsCompletion A, const ELearningAgentsCompletion NotRunningType = ELearningAgentsCompletion::Termination);
 
-#if UE_LEARNING_AGENTS_ENABLE_VISUAL_LOG
-	/** Describes this completion to the visual logger for debugging purposes. */
-	virtual void VisualLog(const UE::Learning::FIndexSet Instances) const override;
-#endif
-
-	TSharedPtr<UE::Learning::FPlanarPositionDifferenceCompletion> CompletionObject;
-};
+	UFUNCTION(BlueprintPure, Category = "LearningAgents")
+	static ELearningAgentsCompletion CompletionOnCondition(const bool bCondition, const ELearningAgentsCompletion CompletionType = ELearningAgentsCompletion::Termination);
 
-/**
- * A completion for if two positions are near by some threshold in a plane, e.g. if the agent gets close to a position.
- */
-UCLASS()
-class LEARNINGAGENTSTRAINING_API UPlanarPositionSimilarityCompletion : public ULearningAgentsCompletion
-{
-	GENERATED_BODY()
-
-public:
+	UFUNCTION(BlueprintPure, Category = "LearningAgents")
+	static ELearningAgentsCompletion CompletionOnTimeElapsed(const float Time, const float TimeThreshold = 10.0f, const ELearningAgentsCompletion CompletionType = ELearningAgentsCompletion::Truncation);
 
-	/**
-	 * Adds a new planar position difference completion to the given trainer. The axis parameters define the plane.
-	 * Call during ULearningAgentsTrainer::SetupCompletions event.
-	 * @param InAgentTrainer The trainer to add this completion to.
-	 * @param Name The name of this new completion. Used for debugging.
-	 * @param Threshold If the distance becomes greater than this threshold, then the episode will complete.
-	 * @param InCompletionMode The completion mode.
-	 * @param Axis0 The forward axis of the plane.
-	 * @param Axis1 The right axis of the plane.
-	 * @return The newly created completion.
-	 */
-	UFUNCTION(BlueprintCallable, Category = "LearningAgents", meta = (DefaultToSelf = "InAgentTrainer"))
-	static UPlanarPositionSimilarityCompletion* AddPlanarPositionSimilarityCompletion(
-		ULearningAgentsTrainer* InAgentTrainer,
-		const FName Name = NAME_None,
-		const float Threshold = 100.0f,
-		const ELearningAgentsCompletion InCompletionMode = ELearningAgentsCompletion::Termination,
-		const FVector Axis0 = FVector::ForwardVector,
-		const FVector Axis1 = FVector::RightVector);
+	UFUNCTION(BlueprintPure, Category = "LearningAgents")
+	static ELearningAgentsCompletion CompletionOnEpisodeStepsRecorded(const int32 EpisodeSteps, const int32 MaxEpisodeSteps = 64, const ELearningAgentsCompletion CompletionType = ELearningAgentsCompletion::Truncation);
 
-	/**
-	 * Sets the data for this completion. Call during ULearningAgentsTrainer::SetCompletions event.
-	 * @param AgentId The agent id this data corresponds to.
-	 * @param Position0 The first position.
-	 * @param Position1 The second position.
-	 */
-	UFUNCTION(BlueprintCallable, Category = "LearningAgents", meta = (AgentId = "-1"))
-	void SetPlanarPositionSimilarityCompletion(const int32 AgentId, const FVector Position0, const FVector Position1);
+	UFUNCTION(BlueprintPure, Category = "LearningAgents")
+	static ELearningAgentsCompletion CompletionOnTranslationDifferenceBelowThreshold(const FVector TranslationA, const FVector TranslationB, const float DistanceThreshold = 100.0f, const ELearningAgentsCompletion CompletionType = ELearningAgentsCompletion::Termination);
 
-#if UE_LEARNING_AGENTS_ENABLE_VISUAL_LOG
-	/** Describes this completion to the visual logger for debugging purposes. */
-	virtual void VisualLog(const UE::Learning::FIndexSet Instances) const override;
-#endif
+	UFUNCTION(BlueprintPure, Category = "LearningAgents")
+	static ELearningAgentsCompletion CompletionOnTranslationDifferenceAboveThreshold(const FVector TranslationA, const FVector TranslationB, const float DistanceThreshold = 100.0f, const ELearningAgentsCompletion CompletionType = ELearningAgentsCompletion::Termination);
 
-	TSharedPtr<UE::Learning::FPlanarPositionSimilarityCompletion> CompletionObject;
+	UFUNCTION(BlueprintPure, Category = "LearningAgents")
+	static ELearningAgentsCompletion CompletionOnTranslationOutsideBounds(
+		const FVector Translation, 
+		const FTransform BoundsTransform = FTransform(),
+		const FVector BoundsMins = FVector(-100.0f, -100.0f, -100.0f),
+		const FVector BoundsMaxs = FVector(+100.0f, +100.0f, +100.0f),
+		const ELearningAgentsCompletion CompletionType = ELearningAgentsCompletion::Termination);
 };
