LearningAgents: Large restructuring to allow for structured actions and observations.

diff --git a/Engine/Plugins/Experimental/LearningAgents/Source/LearningAgentsTraining/Public/LearningAgentsImitationTrainer.h b/Engine/Plugins/Experimental/LearningAgents/Source/LearningAgentsTraining/Public/LearningAgentsImitationTrainer.h
--- a/Engine/Plugins/Experimental/LearningAgents/Source/LearningAgentsTraining/Public/LearningAgentsImitationTrainer.h
+++ b/Engine/Plugins/Experimental/LearningAgents/Source/LearningAgentsTraining/Public/LearningAgentsImitationTrainer.h
@@ -15,6 +15,7 @@ namespace UE::Learning
 	struct FSharedMemoryImitationTrainer;
 }
 
+class ULearningAgentsInteractor;
 class ULearningAgentsPolicy;
 class ULearningAgentsRecording;
 
@@ -45,7 +46,7 @@ public:
 
 	/** Learning rate of the policy network. Typical values are between 0.001 and 0.0001. */
 	UPROPERTY(EditAnywhere, Category = "LearningAgents", meta = (ClampMin = "0.0", ClampMax = "1.0", UIMin = "0.0", UIMax = "1.0"))
-	float LearningRate = 0.0001f;
+	float LearningRate = 0.001f;
 
 	/** Ratio by which to decay the learning rate every 1000 iterations. */
 	UPROPERTY(EditAnywhere, Category = "LearningAgents", meta = (ClampMin = "0.0", ClampMax = "1.0", UIMin = "0.0", UIMax = "1.0"))
@@ -65,6 +66,27 @@ public:
 	UPROPERTY(EditAnywhere, Category = "LearningAgents", meta = (ClampMin = "0", ClampMax = "4096", UIMin = "0", UIMax = "4096"))
 	uint32 BatchSize = 128;
 
+	/**
+	 * The window of observations and actions over which to do the training.Increasing this value will encourage the policy to use its memory 
+	 * effectively. Too large and training can become unstable. Given we don't know the memory state during imitation learning it is better this is 
+	 * slightly larger than when we are doing reinforcement learning.
+	 */
+	UPROPERTY(EditAnywhere, Category = "LearningAgents", meta = (ClampMin = "1", ClampMax = "128", UIMin = "1", UIMax = "512"))
+	uint32 Window = 64;
+
+	/**
+	 * Weight used to regularize actions. Larger values will encourage smaller actions but too large will cause actions to become always zero.
+	 */
+	UPROPERTY(EditAnywhere, Category = "LearningAgents", meta = (ClampMin = "0", UIMin = "0"))
+	float ActionRegularizationWeight = 0.001f;
+
+	/**
+	 * Weighting used for the entropy bonus. Larger values encourage larger action noise and therefore greater exploration but can make actions very 
+	 * noisy.
+	 */
+	UPROPERTY(EditAnywhere, Category = "LearningAgents", meta = (ClampMin = "0", UIMin = "0"))
+	float ActionEntropyWeight = 0.0f;
+
 	/** The seed used for any random sampling the trainer will perform, e.g. for weight initialization. */
 	UPROPERTY(EditAnywhere, Category = "LearningAgents", meta = (ClampMin = "0", UIMin = "0"))
 	int32 RandomSeed = 1234;
@@ -76,6 +98,10 @@ public:
 	/** If true, TensorBoard logs will be emitted to the intermediate directory. */
 	UPROPERTY(EditAnywhere, Category = "LearningAgents")
 	bool bUseTensorboard = false;
+
+	/** If true, snapshots of the trained networks will be emitted to the intermediate directory. */
+	UPROPERTY(EditAnywhere, Category = "LearningAgents")
+	bool bSaveSnapshots = false;
 };
 
 /**
@@ -88,7 +114,7 @@ public:
  * @see ULearningAgentsRecorder to understand how to make new recordings.
  */
 UCLASS(BlueprintType, Blueprintable)
-class LEARNINGAGENTSTRAINING_API ULearningAgentsImitationTrainer : public UActorComponent
+class LEARNINGAGENTSTRAINING_API ULearningAgentsImitationTrainer : public ULearningAgentsManagerListener
 {
 	GENERATED_BODY()
 
@@ -100,8 +126,35 @@ public:
 	ULearningAgentsImitationTrainer(FVTableHelper& Helper);
 	virtual ~ULearningAgentsImitationTrainer();
 
-	/** Will automatically call EndTraining if training is still in-progress when play is ending. */
-	virtual void EndPlay(const EEndPlayReason::Type EndPlayReason) override;
+	/** Will automatically call EndTraining if training is still in-progress when the object is destroyed. */
+	virtual void BeginDestroy() override;
+
+	/**
+	 * Constructs the imitation trainer and runs the setup functions.
+	 * 
+	 * @param InManager			The agent manager we are using.
+	 * @param InInteractor		The agent interactor we are recording with.
+	 * @param InPolicy			The policy we are using.
+	 * @param Class				The imitation trainer class
+	 */
+	UFUNCTION(BlueprintCallable, Category = "LearningAgents", meta = (Class = "/Script/LearningAgentsTraining.LearningAgentsImitationTrainer", DeterminesOutputType = "Class"))
+	static ULearningAgentsImitationTrainer* MakeImitationTrainer(
+		ULearningAgentsManager* InManager,
+		ULearningAgentsInteractor* InInteractor,
+		ULearningAgentsPolicy* InPolicy,
+		TSubclassOf<ULearningAgentsImitationTrainer> Class);
+
+	/**
+	 * Initializes the imitation trainer and runs the setup functions.
+	 * 
+	 * @param InManager			The agent manager we are using.
+	 * @param InInteractor		The agent interactor we are recording with.
+	 * @param InPolicy			The policy we are using.
+	 */
+	void SetupImitationTrainer(
+		ULearningAgentsManager* InManager,
+		ULearningAgentsInteractor* InInteractor,
+		ULearningAgentsPolicy* InPolicy);
 
 	/** Returns true if the trainer is currently training; Otherwise, false. */
 	UFUNCTION(BlueprintPure, Category = "LearningAgents")
@@ -118,21 +171,18 @@ public:
 
 	/**
 	 * Begins the training process with the provided settings.
-	 * @param InPolicy The policy to train.
-	 * @param Recording The data to train on.
-	 * @param ImitationTrainerSettings The settings for this trainer.
-	 * @param ImitationTrainerTrainingSettings The training settings for this network.
-	 * @param ImitationTrainerPathSettings The path settings used by the imitation trainer.
-	 * @param bReinitializePolicyNetwork If true, reinitialize the policy. Set this to false if your policy is pre-trained.
+	 * 
+	 * @param Recording							The data to train on.
+	 * @param ImitationTrainerSettings			The settings for this trainer.
+	 * @param ImitationTrainerTrainingSettings	The training settings for this network.
+	 * @param ImitationTrainerPathSettings		The path settings used by the imitation trainer.
 	 */
 	UFUNCTION(BlueprintCallable, Category = "LearningAgents")
 	void BeginTraining(
-		ULearningAgentsPolicy* InPolicy,
 		const ULearningAgentsRecording* Recording,
 		const FLearningAgentsImitationTrainerSettings& ImitationTrainerSettings = FLearningAgentsImitationTrainerSettings(),
 		const FLearningAgentsImitationTrainerTrainingSettings& ImitationTrainerTrainingSettings = FLearningAgentsImitationTrainerTrainingSettings(),
-		const FLearningAgentsTrainerPathSettings& ImitationTrainerPathSettings = FLearningAgentsTrainerPathSettings(),
-		const bool bReinitializePolicyNetwork = true);
+		const FLearningAgentsTrainerPathSettings& ImitationTrainerPathSettings = FLearningAgentsTrainerPathSettings());
 
 	/** Stops the training process. */
 	UFUNCTION(BlueprintCallable, Category = "LearningAgents")
@@ -145,25 +195,26 @@ public:
 	/**
 	 * Convenience function that runs a basic training loop. If training has not been started, it will start it. On 
 	 * each following call to this function, it will call IterateTraining.
-	 * @param InPolicy The policy to train.
-	 * @param Recording The data to train on.
-	 * @param ImitationTrainerSettings The settings for this trainer.
-	 * @param ImitationTrainerTrainingSettings The training settings for this network.
-	 * @param ImitationTrainerPathSettings The path settings used by the imitation trainer.
-	 * @param bReinitializePolicyNetwork If true, reinitialize the policy. Set this to false if your policy is pre-trained.
+	 * 
+	 * @param Recording							The data to train on.
+	 * @param ImitationTrainerSettings			The settings for this trainer.
+	 * @param ImitationTrainerTrainingSettings	The training settings for this network.
+	 * @param ImitationTrainerPathSettings		The path settings used by the imitation trainer.
 	 */
 	UFUNCTION(BlueprintCallable, Category = "LearningAgents")
 	void RunTraining(
-		ULearningAgentsPolicy* InPolicy,
 		const ULearningAgentsRecording* Recording,
 		const FLearningAgentsImitationTrainerSettings& ImitationTrainerSettings = FLearningAgentsImitationTrainerSettings(),
 		const FLearningAgentsImitationTrainerTrainingSettings& ImitationTrainerTrainingSettings = FLearningAgentsImitationTrainerTrainingSettings(),
-		const FLearningAgentsTrainerPathSettings& ImitationTrainerPathSettings = FLearningAgentsTrainerPathSettings(),
-		const bool bReinitializePolicyNetwork = true);
+		const FLearningAgentsTrainerPathSettings& ImitationTrainerPathSettings = FLearningAgentsTrainerPathSettings());
 
 // ----- Private Data -----
 private:
 
+	/** The interactor being trained. */
+	UPROPERTY(VisibleAnywhere, Transient, Category = "LearningAgents")
+	TObjectPtr<ULearningAgentsInteractor> Interactor;
+
 	/** The policy being trained. */
 	UPROPERTY(VisibleAnywhere, Transient, Category = "LearningAgents")
 	TObjectPtr<ULearningAgentsPolicy> Policy;
@@ -183,10 +234,5 @@ private:
 
 	void DoneTraining();
 
-	TLearningArray<1, int32> RecordedEpisodeStarts;
-	TLearningArray<1, int32> RecordedEpisodeLengths;
-	TLearningArray<2, float> RecordedObservations;
-	TLearningArray<2, float> RecordedActions;
-
 	TUniquePtr<UE::Learning::FSharedMemoryImitationTrainer> ImitationTrainer;
 };
diff --git a/Engine/Plugins/Experimental/LearningAgents/Source/LearningAgentsTraining/Public/LearningAgentsRecorder.h b/Engine/Plugins/Experimental/LearningAgents/Source/LearningAgentsTraining/Public/LearningAgentsRecorder.h
--- a/Engine/Plugins/Experimental/LearningAgents/Source/LearningAgentsTraining/Public/LearningAgentsRecorder.h
+++ b/Engine/Plugins/Experimental/LearningAgents/Source/LearningAgentsTraining/Public/LearningAgentsRecorder.h
@@ -2,7 +2,7 @@
 
 #pragma once
 
-#include "LearningAgentsManagerComponent.h"
+#include "LearningAgentsManagerListener.h"
 
 #include "LearningArray.h"
 #include "Containers/Map.h"
@@ -33,8 +33,8 @@ public:
 };
 
 /** A component that can be used to create recordings of training data for imitation learning. */
-UCLASS(BlueprintType, Blueprintable)
-class LEARNINGAGENTSTRAINING_API ULearningAgentsRecorder : public ULearningAgentsManagerComponent
+UCLASS(BlueprintType, Blueprintable, meta = (BlueprintSpawnableComponent))
+class LEARNINGAGENTSTRAINING_API ULearningAgentsRecorder : public ULearningAgentsManagerListener
 {
 	GENERATED_BODY()
 
@@ -46,42 +46,64 @@ public:
 	ULearningAgentsRecorder(FVTableHelper& Helper);
 	virtual ~ULearningAgentsRecorder();
 
-	/** Will automatically call EndRecording if recording is still in-progress when play is ending. */
-	virtual void EndPlay(const EEndPlayReason::Type EndPlayReason) override;
+	/** Will automatically call EndRecording if recording is still in-progress when the object is destroyed. */
+	virtual void BeginDestroy() override;
+
+	/**
+	 * Constructs this object and runs the setup functions for the underlying data storage.
+	 * 
+	 * @param InManager					The agent manager we are using.
+	 * @param InInteractor				The agent interactor we are recording with.
+	 * @param Class						The recorder class
+	 * @param Name						The recorder class
+	 * @param RecorderPathSettings		The path settings used by the recorder.
+	 * @param RecordingAsset			Optional recording asset to use. If not provided or bReinitializeRecording is 
+	 *									set then a new recording object will be created.
+	 * @pram bReinitializeRecording		If to reinitialize the recording asset
+	 */
+	UFUNCTION(BlueprintCallable, Category = "LearningAgents", meta = (Class = "/Script/LearningAgentsTraining.LearningAgentsRecorder", DeterminesOutputType = "Class"))
+	static ULearningAgentsRecorder* MakeRecorder(
+		ULearningAgentsManager* InManager,
+		ULearningAgentsInteractor* InInteractor,
+		TSubclassOf<ULearningAgentsRecorder> Class,
+		const FName Name = TEXT("Recorder"),
+		const FLearningAgentsRecorderPathSettings& RecorderPathSettings = FLearningAgentsRecorderPathSettings(),
+		ULearningAgentsRecording* RecordingAsset = nullptr,
+		bool bReinitializeRecording = true);
 
 	/**
 	 * Initializes this object and runs the setup functions for the underlying data storage.
-	 * @param InInteractor The agent interactor we are recording with.
-	 * @param RecorderPathSettings The path settings used by the recorder.
-	 * @param RecordingAsset Optional recording asset to use. If not provided then a new recording object will be 
-	 * created.
+	 * 
+	 * @param InManager					The agent manager we are using.
+	 * @param InInteractor				The agent interactor we are recording with.
+	 * @param RecorderPathSettings		The path settings used by the recorder.
+	 * @param RecordingAsset			Optional recording asset to use. If not provided or bReinitializeRecording is
+	 *									set then a new recording object will be created.
+	 * @pram bReinitializeRecording		If to reinitialize the recording asset
 	 */
 	UFUNCTION(BlueprintCallable, Category = "LearningAgents")
 	void SetupRecorder(
+		ULearningAgentsManager* InManager,
 		ULearningAgentsInteractor* InInteractor,
 		const FLearningAgentsRecorderPathSettings& RecorderPathSettings = FLearningAgentsRecorderPathSettings(),
-		ULearningAgentsRecording* RecordingAsset = nullptr);
+		ULearningAgentsRecording* RecordingAsset = nullptr,
+		bool bReinitializeRecording = true);
 
 public:
 
-	//~ Begin ULearningAgentsManagerComponent Interface
-	virtual void OnAgentsRemoved(const TArray<int32>& AgentIds) override;
-	virtual void OnAgentsReset(const TArray<int32>& AgentIds) override;
-	//~ End ULearningAgentsManagerComponent Interface
+	//~ Begin ULearningAgentsManagerListener Interface
+	virtual void OnAgentsRemoved_Implementation(const TArray<int32>& AgentIds) override;
+	virtual void OnAgentsReset_Implementation(const TArray<int32>& AgentIds) override;
+	//~ End ULearningAgentsManagerListener Interface
 
 // ----- Recording Process -----
 public:
 
-	/**
-	 * Begins the recording of the observations and actions of each added agent.
-	 * @param bReinitializeRecording If to clear all records from the recording object in use at the start of recording.
-	 */
+	/** Begins the recording of the observations and actions of each added agent. */
 	UFUNCTION(BlueprintCallable, Category = "LearningAgents")
-	void BeginRecording(bool bReinitializeRecording = true);
+	void BeginRecording();
 
-	/**
-	 * Ends the recording of the observations and actions of each agent and stores them in the current recording object.
-	 */
+	/** Ends the recording of the observations and actions of each agent and stores them in the current recording object. */
 	UFUNCTION(BlueprintCallable, Category = "LearningAgents")
 	void EndRecording();
 
@@ -90,47 +112,16 @@ public:
 	bool IsRecording() const;
 
 	/**
-	 * While recording, adds the current observations and actions of the added agents to the internal buffer. Call this 
-	 * after ULearningAgentsInteractor::EncodeObservations and either ULearningAgentsController::EncodeActions (if 
-	 * recording a human/AI demonstration) or ULearningAgentsInteractor::DecodeActions (if recording another policy).
+	 * While recording, adds the current buffered observations and actions of the added agents to the internal buffer. Call this after 
+	 * GatherObservations and either EvaluateAgentController (if recording a human/AI demonstration) or DecodeAndSampleActions (if recording from 
+	 * another policy).
 	 */
 	UFUNCTION(BlueprintCallable, Category = "LearningAgents")
 	void AddExperience();
 
 	/** Gets the current recording object. Note: this may be empty until EndRecording has been called. */
 	UFUNCTION(BlueprintPure, Category = "LearningAgents")
-	const ULearningAgentsRecording* GetCurrentRecording() const;
-
-// ----- Load / Save -----
-public:
-
-	/** Loads the current recording object from a file */
-	UFUNCTION(BlueprintCallable, Category = "LearningAgents", meta = (RelativePath))
-	void LoadRecordingFromFile(const FFilePath& File);
-
-	/** Saves the current recording object to a file */
-	UFUNCTION(BlueprintCallable, BlueprintPure = false, Category = "LearningAgents", meta = (RelativePath))
-	void SaveRecordingToFile(const FFilePath& File) const;
-
-	/** Append to the current recording object from a file. */
-	UFUNCTION(BlueprintCallable, Category = "LearningAgents", meta = (RelativePath))
-	void AppendRecordingFromFile(const FFilePath& File);
-
-	/** Uses the given recording asset. New recordings will be appended to this. */
-	UFUNCTION(BlueprintCallable, Category = "LearningAgents")
-	void UseRecordingAsset(ULearningAgentsRecording* RecordingAsset);
-
-	/** Loads the current recording object from the given recording asset */
-	UFUNCTION(BlueprintCallable, Category = "LearningAgents")
-	void LoadRecordingFromAsset(ULearningAgentsRecording* RecordingAsset);
-
-	/** Saves the current recording object to the given recording asset */
-	UFUNCTION(BlueprintCallable, Category = "LearningAgents", meta = (DevelopmentOnly))
-	void SaveRecordingToAsset(ULearningAgentsRecording* RecordingAsset);
-
-	/** Appends the current recording object to the given recording asset */
-	UFUNCTION(BlueprintCallable, Category = "LearningAgents", meta = (DevelopmentOnly))
-	void AppendRecordingToAsset(ULearningAgentsRecording* RecordingAsset);
+	const ULearningAgentsRecording* GetRecordingAsset() const;
 
 
 // ----- Private Data ----- 
@@ -157,8 +148,13 @@ private:
 	/** Basic structure used to buffer the observations and actions of each agent. */
 	struct FAgentRecordBuffer
 	{
+		FAgentRecordBuffer();
+		FAgentRecordBuffer(const int32 InObservationCompatibilityHash, const int32 InActionCompatibilityHash);
+
 		const int32 ChunkSize = 1024;
 		int32 StepNum = 0;
+		int32 ObservationCompatibilityHash = 0;
+		int32 ActionCompatibilityHash = 0;
 		TArray<TLearningArray<2, float>, TInlineAllocator<16>> Observations;
 		TArray<TLearningArray<2, float>, TInlineAllocator<16>> Actions;
 
diff --git a/Engine/Plugins/Experimental/LearningAgents/Source/LearningAgentsTraining/Public/LearningAgentsRecording.h b/Engine/Plugins/Experimental/LearningAgents/Source/LearningAgentsTraining/Public/LearningAgentsRecording.h
--- a/Engine/Plugins/Experimental/LearningAgents/Source/LearningAgentsTraining/Public/LearningAgentsRecording.h
+++ b/Engine/Plugins/Experimental/LearningAgents/Source/LearningAgentsTraining/Public/LearningAgentsRecording.h
@@ -30,23 +30,26 @@ public:
 	UPROPERTY(BlueprintReadOnly, VisibleAnywhere, Category = "LearningAgents")
 	int32 ActionDimNum = 0;
 
-	bool Serialize(FArchive& Ar);
+	/** The compatibility hash for the recorded observation vectors */
+	UPROPERTY(BlueprintReadOnly, VisibleAnywhere, Category = "LearningAgents")
+	int32 ObservationCompatibilityHash = 0;
 
-	TLearningArray<2, float> Observations;
-	TLearningArray<2, float> Actions;
-};
+	/** The compatibility hash for the recorded action vectors */
+	UPROPERTY(BlueprintReadOnly, VisibleAnywhere, Category = "LearningAgents")
+	int32 ActionCompatibilityHash = 0;
+
+	/** Observation data */
+	UPROPERTY(BlueprintReadOnly, Category = "LearningAgents")
+	TArray<float> ObservationData;
+
+	/** Action data */
+	UPROPERTY(BlueprintReadOnly, Category = "LearningAgents")
+	TArray<float> ActionData;
 
-template<>
-struct TStructOpsTypeTraits<FLearningAgentsRecord> : public TStructOpsTypeTraitsBase2<FLearningAgentsRecord>
-{
-	enum
-	{
-		WithSerializer = true,
-	};
 };
 
 /** Data asset representing an array of records. */
-UCLASS(BlueprintType)
+UCLASS(BlueprintType, Blueprintable)
 class LEARNINGAGENTSTRAINING_API ULearningAgentsRecording : public UDataAsset
 {
 	GENERATED_BODY()
@@ -76,13 +79,43 @@ public:
 	void LoadRecordingFromAsset(ULearningAgentsRecording* RecordingAsset);
 
 	/** Saves this recording to the given recording asset */
-	UFUNCTION(BlueprintCallable, Category = "LearningAgents", meta = (DevelopmentOnly))
+	UFUNCTION(BlueprintCallable, Category = "LearningAgents")
 	void SaveRecordingToAsset(ULearningAgentsRecording* RecordingAsset);
 
 	/** Appends this recording to the given recording asset */
-	UFUNCTION(BlueprintCallable, Category = "LearningAgents", meta = (DevelopmentOnly))
+	UFUNCTION(BlueprintCallable, Category = "LearningAgents")
 	void AppendRecordingToAsset(ULearningAgentsRecording* RecordingAsset);
 
+	/** Get the number of records */
+	UFUNCTION(BlueprintPure, Category = "LearningAgents")
+	int32 GetRecordNum() const;
+
+	/** Get the number of steps in a given record */
+	UFUNCTION(BlueprintPure, Category = "LearningAgents")
+	int32 GetRecordStepNum(const int32 Record) const;
+
+	/**
+	 * Get the Observation Vector associated with a particular step of a given recording
+	 * 
+	 * @param OutObservationVector				Output Observation Vector
+	 * @param OutObservationCompatibilityHash	Output Compatibility Hash for the given Observation Vector
+	 * @param Record							Index of the record in the array of records.
+	 * @param Step								Step of the recording
+	 */
+	UFUNCTION(BlueprintCallable, Category = "LearningAgents")
+	void GetObservationVector(TArray<float>& OutObservationVector, int32& OutObservationCompatibilityHash, const int32 Record, const int32 Step);
+
+	/**
+	 * Get the Action Vector associated with a particular step of a given recording
+	 *
+	 * @param OutActionVector					Output Action Vector
+	 * @param OutActionCompatibilityHash		Output Compatibility Hash for the given Action Vector
+	 * @param Record							Index of the record in the array of records.
+	 * @param Step								Step of the recording
+	 */
+	UFUNCTION(BlueprintCallable, Category = "LearningAgents")
+	void GetActionVector(TArray<float>& OutActionVector, int32& OutActionCompatibilityHash, const int32 Record, const int32 Step);
+
 public:
 
 	/** Marks this asset as modified even during PIE */
diff --git a/Engine/Plugins/Experimental/LearningAgents/Source/LearningAgentsTraining/Public/LearningAgentsRewards.h b/Engine/Plugins/Experimental/LearningAgents/Source/LearningAgentsTraining/Public/LearningAgentsRewards.h
--- a/Engine/Plugins/Experimental/LearningAgents/Source/LearningAgentsTraining/Public/LearningAgentsRewards.h
+++ b/Engine/Plugins/Experimental/LearningAgents/Source/LearningAgentsTraining/Public/LearningAgentsRewards.h
@@ -2,353 +2,65 @@
 
 #pragma once
 
-#include "LearningArray.h"
-#include "LearningAgentsDebug.h"
-
-#include "Templates/SharedPointer.h"
-#include "UObject/Object.h"
+#include "Kismet/BlueprintFunctionLibrary.h"
 
 #include "LearningAgentsRewards.generated.h"
 
-namespace UE::Learning
-{
-	struct FRewardObject;
-	struct FFloatReward;
-	struct FConditionalConstantReward;
-	struct FScalarVelocityReward;
-	struct FLocalDirectionalVelocityReward;
-	struct FPlanarPositionDifferencePenalty;
-	struct FPositionDifferencePenalty;
-	struct FPositionArraySimilarityReward;
-}
-
-class ULearningAgentsTrainer;
-
-// For functions in this file, we are favoring having more verbose names such as "AddFloatReward" vs simply "Add" in 
-// order to keep it easy to find the correct function in blueprints.
-
-//------------------------------------------------------------------
+class USplineComponent;
 
-/**
- * Base class for all rewards/penalties. Rewards are used during reinforcement learning to encourage/discourage
- * certain behaviors from occurring.
- */
-UCLASS(Abstract, BlueprintType)
-class LEARNINGAGENTSTRAINING_API ULearningAgentsReward : public UObject
+UCLASS(BlueprintType)
+class LEARNINGAGENTSTRAINING_API ULearningAgentsRewards : public UBlueprintFunctionLibrary
 {
 	GENERATED_BODY()
 
 public:
 
-	/** Reference to the Trainer this reward is associated with. */
-	UPROPERTY(VisibleAnywhere, Transient, Category = "LearningAgents")
-	TObjectPtr<ULearningAgentsTrainer> AgentTrainer;
+	// Rewards
 
-public:
+	UFUNCTION(BlueprintPure, Category = "LearningAgents")
+	static float RewardOnCondition(const bool bCondition, const float RewardScale = 1.0f);
 
-	/** Initialize the internal state for a given maximum number of agents */
-	void Init(const int32 MaxAgentNum);
+	UFUNCTION(BlueprintPure, Category = "LearningAgents")
+	static float RewardOnTranslationDifferenceBelowThreshold(const FVector TranslationA, const FVector TranslationB, const float DistanceThreshold = 100.0f, const float RewardScale = 1.0f);
 
-	/**
-	 * Called whenever agents are added to the associated ULearningAgentsTrainer object.
-	 * @param AgentIds Array of agent ids which have been added
-	 */
-	virtual void OnAgentsAdded(const TArray<int32>& AgentIds);
+	UFUNCTION(BlueprintPure, Category = "LearningAgents")
+	static float RewardFromTranslationSimilarity(const FVector TranslationA, const FVector TranslationB, const float TranslationScale = 100.0f, const float RewardScale = 1.0f);
 
-	/**
-	 * Called whenever agents are removed from the associated ULearningAgentsTrainer object.
-	 * @param AgentIds Array of agent ids which have been removed
-	 */
-	virtual void OnAgentsRemoved(const TArray<int32>& AgentIds);
+	UFUNCTION(BlueprintPure, Category = "LearningAgents")
+	static float RewardFromRotationSimilarity(const FRotator RotationA, const FRotator RotationB, const float AngleScale = 90.0f, const float RewardScale = 1.0f);
+	static float RewardFromRotationSimilarityAsQuats(const FQuat RotationA, const FQuat RotationB, const float AngleScale = 90.0f, const float RewardScale = 1.0f);
 
-	/**
-	 * Called whenever agents are reset on the associated ULearningAgentsTrainer object.
-	 * @param AgentIds Array of agent ids which have been reset
-	 */
-	virtual void OnAgentsReset(const TArray<int32>& AgentIds);
+	UFUNCTION(BlueprintPure, Category = "LearningAgents")
+	static float RewardFromAngleSimilarity(const float AngleA, const float AngleB, const float AngleScale = 90.0f, const float RewardScale = 1.0f);
 
-	/** Get the number of times a reward has been set for the given agent id. */
-	uint64 GetAgentIteration(const int32 AgentId) const;
+	UFUNCTION(BlueprintPure, Category = "LearningAgents")
+	static float RewardFromDirectionSimilarity(const FVector DirectionA, const FVector DirectionB, const float AngleScale = 90.0f, const float RewardScale = 1.0f);
 
-public:
-#if UE_LEARNING_AGENTS_ENABLE_VISUAL_LOG
-	/** Color used to draw this reward in the visual log */
-	FLinearColor VisualLogColor = FColor::Green;
+	// Spline Rewards
 
-	/** Describes this reward to the visual logger for debugging purposes. */
-	virtual void VisualLog(const UE::Learning::FIndexSet Instances) const {}
-#endif
+	UFUNCTION(BlueprintPure, Category = "LearningAgents")
+	static float RewardFromVelocityAlongSpline(const USplineComponent* SplineComponent, const FVector Position, const FVector Velocity, const float VelocityScale = 200.0f, const float RewardScale = 1.0f, const float FiniteDifferenceDelta = 10.0f);
 
-protected:
+	// Penalties
 
-	/** Number of times this reward has been set for all agents */
-	TLearningArray<1, uint64, TInlineAllocator<32>> AgentIteration;
-};
+	UFUNCTION(BlueprintPure, Category = "LearningAgents")
+	static float PenaltyOnCondition(const bool bCondition, const float PenaltyScale = 1.0f);
 
-//------------------------------------------------------------------
+	UFUNCTION(BlueprintPure, Category = "LearningAgents")
+	static float PenaltyOnTranslationDifferenceAboveThreshold(const FVector TranslationA, const FVector TranslationB, const float DistanceThreshold = 100.0f, const float PenaltyScale = 1.0f);
 
-/** A simple float reward. Used as a catch-all for situations where a more type-specific reward does not exist yet. */
-UCLASS()
-class LEARNINGAGENTSTRAINING_API UFloatReward : public ULearningAgentsReward
-{
-	GENERATED_BODY()
-
-public:
-
-	/**
-	 * Adds a new float reward to the given trainer. Call during ULearningAgentsTrainer::SetupRewards event.
-	 * @param InAgentTrainer The trainer to add this reward to.
-	 * @param Name The name of this new reward. Used for debugging.
-	 * @param Weight Multiplier for this reward when being summed up for the total reward.
-	 * @return The newly created reward.
-	 */
-	UFUNCTION(BlueprintCallable, Category = "LearningAgents", meta = (DefaultToSelf = "InAgentTrainer"))
-	static UFloatReward* AddFloatReward(
-		ULearningAgentsTrainer* InAgentTrainer,
-		const FName Name = NAME_None,
-		const float Weight = 1.0f);
-
-	/**
-	 * Sets the data for this reward. Call during ULearningAgentsTrainer::SetRewards event.
-	 * @param AgentId The agent id this data corresponds to.
-	 * @param Reward The value currently being rewarded.
-	 */
-	UFUNCTION(BlueprintCallable, Category = "LearningAgents", meta = (AgentId = "-1"))
-	void SetFloatReward(const int32 AgentId, const float Reward);
-
-#if UE_LEARNING_AGENTS_ENABLE_VISUAL_LOG
-	/** Describes this reward to the visual logger for debugging purposes. */
-	virtual void VisualLog(const UE::Learning::FIndexSet Instances) const override;
-#endif
-
-	TSharedPtr<UE::Learning::FFloatReward> RewardObject;
-};
+	UFUNCTION(BlueprintPure, Category = "LearningAgents")
+	static float PenaltyFromTranslationDifference(const FVector TranslationA, const FVector TranslationB, const float TranslationScale = 100.0f, const float PenaltyScale = 1.0f);
 
-//------------------------------------------------------------------
-
-/** A simple conditional reward that gives some constant reward value when a condition is true. */
-UCLASS()
-class LEARNINGAGENTSTRAINING_API UConditionalReward : public ULearningAgentsReward
-{
-	GENERATED_BODY()
-
-public:
-
-	/**
-	 * Adds a new conditional reward to the given trainer. Call during ULearningAgentsTrainer::SetupRewards event.
-	 * @param InAgentTrainer The trainer to add this reward to.
-	 * @param Name The name of this new reward. Used for debugging.
-	 * @param Value The amount of reward to give the agent when the provided condition is true.
-	 * @return The newly created reward.
-	 */
-	UFUNCTION(BlueprintCallable, Category = "LearningAgents", meta = (DefaultToSelf = "InAgentTrainer"))
-	static UConditionalReward* AddConditionalReward(
-		ULearningAgentsTrainer* InAgentTrainer,
-		const FName Name = NAME_None,
-		const float Value = 1.0f);
-
-	/**
-	 * Sets if the agent should receive a reward. Call during ULearningAgentsTrainer::SetRewards event.
-	 * @param AgentId The agent id this data corresponds to.
-	 * @param bCondition If the agent should receive a reward this iteration.
-	 */
-	UFUNCTION(BlueprintCallable, Category = "LearningAgents", meta = (AgentId = "-1"))
-	void SetConditionalReward(const int32 AgentId, const bool bCondition);
-
-#if UE_LEARNING_AGENTS_ENABLE_VISUAL_LOG
-	/** Describes this reward to the visual logger for debugging purposes. */
-	virtual void VisualLog(const UE::Learning::FIndexSet Instances) const override;
-#endif
-
-	TSharedPtr<UE::Learning::FConditionalConstantReward> RewardObject;
-};
-
-//------------------------------------------------------------------
-
-/** A reward for maximizing speed. */
-UCLASS()
-class LEARNINGAGENTSTRAINING_API UScalarVelocityReward : public ULearningAgentsReward
-{
-	GENERATED_BODY()
-
-public:
-
-	/**
-	 * Adds a new scalar velocity reward to the given trainer. Call during ULearningAgentsTrainer::SetupRewards event.
-	 * @param InAgentTrainer The trainer to add this reward to.
-	 * @param Name The name of this new reward. Used for debugging.
-	 * @param Weight Multiplier for this reward when being summed up for the total reward.
-	 * @param Scale Used to normalize the data for the reward.
-	 * @return The newly created reward.
-	 */
-	UFUNCTION(BlueprintCallable, Category = "LearningAgents", meta = (DefaultToSelf = "InAgentTrainer"))
-	static UScalarVelocityReward* AddScalarVelocityReward(
-		ULearningAgentsTrainer* InAgentTrainer,
-		const FName Name = NAME_None,
-		const float Weight = 1.0f,
-		const float Scale = 200.0f);
-
-	/**
-	 * Sets the data for this reward. Call during ULearningAgentsTrainer::SetRewards event.
-	 * @param AgentId The agent id this data corresponds to.
-	 * @param Velocity The current scalar velocity.
-	 */
-	UFUNCTION(BlueprintCallable, Category = "LearningAgents", meta = (AgentId = "-1"))
-	void SetScalarVelocityReward(const int32 AgentId, const float Velocity);
-
-#if UE_LEARNING_AGENTS_ENABLE_VISUAL_LOG
-	/** Describes this reward to the visual logger for debugging purposes. */
-	virtual void VisualLog(const UE::Learning::FIndexSet Instances) const override;
-#endif
-
-	TSharedPtr<UE::Learning::FScalarVelocityReward> RewardObject;
-};
-
-//------------------------------------------------------------------
-
-/** A reward for maximizing velocity along a given local axis. */
-UCLASS()
-class LEARNINGAGENTSTRAINING_API ULocalDirectionalVelocityReward : public ULearningAgentsReward
-{
-	GENERATED_BODY()
-
-public:
-
-	/**
-	 * Adds a new directional velocity reward to the given trainer. Call during ULearningAgentsTrainer::SetupRewards event.
-	 * @param InAgentTrainer The trainer to add this reward to.
-	 * @param Name The name of this new reward. Used for debugging.
-	 * @param Weight Multiplier for this reward when being summed up for the total reward.
-	 * @param Scale Used to normalize the data for the reward.
-	 * @param Axis The local direction we want to maximize velocity in.
-	 * @return The newly created reward.
-	 */
-	UFUNCTION(BlueprintCallable, Category = "LearningAgents", meta = (DefaultToSelf = "InAgentTrainer"))
-	static ULocalDirectionalVelocityReward* AddLocalDirectionalVelocityReward(
-		ULearningAgentsTrainer* InAgentTrainer,
-		const FName Name = NAME_None,
-		const float Weight = 1.0f,
-		const float Scale = 200.0f,
-		const FVector Axis = FVector::ForwardVector);
-
-	/**
-	 * Sets the data for this reward. Call during ULearningAgentsTrainer::SetRewards event.
-	 * @param AgentId The agent id this data corresponds to.
-	 * @param Velocity The current velocity.
-	 * @param RelativeRotation The frame of reference rotation.
-	 */
-	UFUNCTION(BlueprintCallable, Category = "LearningAgents", meta = (AgentId = "-1"))
-	void SetLocalDirectionalVelocityReward(
-		const int32 AgentId,
-		const FVector Velocity,
-		const FRotator RelativeRotation = FRotator::ZeroRotator);
-
-#if UE_LEARNING_AGENTS_ENABLE_VISUAL_LOG
-	/** Describes this reward to the visual logger for debugging purposes. */
-	virtual void VisualLog(const UE::Learning::FIndexSet Instances) const override;
-#endif
-
-	TSharedPtr<UE::Learning::FLocalDirectionalVelocityReward> RewardObject;
-};
-
-//------------------------------------------------------------------
-
-/** A penalty for being far from a goal position in a plane. */
-UCLASS()
-class LEARNINGAGENTSTRAINING_API UPlanarPositionDifferencePenalty : public ULearningAgentsReward
-{
-	GENERATED_BODY()
-
-public:
-
-	/**
-	 * Adds a new planar difference penalty to the given trainer. The axis parameters define the plane.
-	 * Call during ULearningAgentsTrainer::SetupRewards event.
-	 * @param InAgentTrainer The trainer to add this penalty to.
-	 * @param Name The name of this new penalty. Used for debugging.
-	 * @param Weight Multiplier for this penalty when being summed up for the total reward.
-	 * @param Scale Used to normalize the data for the penalty.
-	 * @param Threshold Minimal distance to apply this penalty.
-	 * @param Axis0 The forward axis of the plane.
-	 * @param Axis1 The right axis of the plane.
-	 * @return The newly created reward.
-	 */
-	UFUNCTION(BlueprintCallable, Category = "LearningAgents", meta = (DefaultToSelf = "InAgentTrainer"))
-	static UPlanarPositionDifferencePenalty* AddPlanarPositionDifferencePenalty(
-		ULearningAgentsTrainer* InAgentTrainer,
-		const FName Name = NAME_None,
-		const float Weight = 1.0f,
-		const float Scale = 100.0f,
-		const float Threshold = 0.0f,
-		const FVector Axis0 = FVector::ForwardVector,
-		const FVector Axis1 = FVector::RightVector);
-
-	/**
-	 * Sets the data for this penalty. Call during ULearningAgentsTrainer::SetRewards event.
-	 * @param AgentId The agent id this data corresponds to.
-	 * @param Position0 The current position.
-	 * @param Position1 The goal position.
-	 */
-	UFUNCTION(BlueprintCallable, Category = "LearningAgents", meta = (AgentId = "-1"))
-	void SetPlanarPositionDifferencePenalty(const int32 AgentId, const FVector Position0, const FVector Position1);
-
-#if UE_LEARNING_AGENTS_ENABLE_VISUAL_LOG
-	/** Describes this reward to the visual logger for debugging purposes. */
-	virtual void VisualLog(const UE::Learning::FIndexSet Instances) const override;
-#endif
-
-	TSharedPtr<UE::Learning::FPlanarPositionDifferencePenalty> RewardObject;
-};
-
-/** A reward for minimizing the distances of positions in the given arrays. */
-UCLASS()
-class LEARNINGAGENTSTRAINING_API UPositionArraySimilarityReward : public ULearningAgentsReward
-{
-	GENERATED_BODY()
-
-public:
+	UFUNCTION(BlueprintPure, Category = "LearningAgents")
+	static float PenaltyFromAngleDifference(const float AngleA, const float AngleB, const float AngleScale = 90.0f, const float PenaltyScale = 1.0f);
 
-	/**
-	 * Adds a new position array similarity reward to the given trainer. Call during ULearningAgentsTrainer::SetupRewards event.
-	 * @param InAgentTrainer The trainer to add this reward to.
-	 * @param Name The name of this new reward. Used for debugging.
-	 * @param PositionNum The number of positions in the array.
-	 * @param Scale Used to normalize the data for the reward.
-	 * @param Weight Multiplier for this reward when being summed up for the total reward.
-	 * @return The newly created reward.
-	 */
-	UFUNCTION(BlueprintCallable, Category = "LearningAgents", meta = (DefaultToSelf = "InAgentTrainer"))
-	static UPositionArraySimilarityReward* AddPositionArraySimilarityReward(
-		ULearningAgentsTrainer* InAgentTrainer,
-		const FName Name = NAME_None,
-		const int32 PositionNum = 0,
-		const float Scale = 100.0f,
-		const float Weight = 1.0f);
+	UFUNCTION(BlueprintPure, Category = "LearningAgents")
+	static float PenaltyFromRotationDifference(const FRotator RotationA, const FRotator RotationB, const float AngleScale = 90.0f, const float PenaltyScale = 1.0f);
+	static float PenaltyFromRotationDifferenceAsQuats(const FQuat RotationA, const FQuat RotationB, const float AngleScale = 90.0f, const float PenaltyScale = 1.0f);
 
-	/**
-	 * Sets the data for this reward. Call during ULearningAgentsTrainer::SetRewards event.
-	 * @param AgentId The agent id this data corresponds to.
-	 * @param Positions0 The current positions.
-	 * @param Positions1 The goal positions.
-	 * @param RelativePosition0 The vector Positions0 will be offset from.
-	 * @param RelativePosition1 The vector Positions1 will be offset from.
-	 * @param RelativeRotation0 The frame of reference rotation for Positions0.
-	 * @param RelativeRotation1 The frame of reference rotation for Positions1.
-	 */
-	UFUNCTION(BlueprintCallable, Category = "LearningAgents", meta = (AgentId = "-1"))
-	void SetPositionArraySimilarityReward(
-		const int32 AgentId,
-		const TArray<FVector>& Positions0, 
-		const TArray<FVector>& Positions1, 
-		const FVector RelativePosition0 = FVector::ZeroVector,
-		const FVector RelativePosition1 = FVector::ZeroVector,
-		const FRotator RelativeRotation0 = FRotator::ZeroRotator,
-		const FRotator RelativeRotation1 = FRotator::ZeroRotator);
+	UFUNCTION(BlueprintPure, Category = "LearningAgents")
+	static float PenaltyFromDirectionDifference(const FVector DirectionA, const FVector DirectionB, const float AngleScale = 90.0f, const float PenaltyScale = 1.0f);
 
-#if UE_LEARNING_AGENTS_ENABLE_VISUAL_LOG
-	/** Describes this reward to the visual logger for debugging purposes. */
-	virtual void VisualLog(const UE::Learning::FIndexSet Instances) const override;
-#endif
 
-	TSharedPtr<UE::Learning::FPositionArraySimilarityReward> RewardObject;
 };
diff --git a/Engine/Plugins/Experimental/LearningAgents/Source/LearningAgentsTraining/Public/LearningAgentsTrainer.h b/Engine/Plugins/Experimental/LearningAgents/Source/LearningAgentsTraining/Public/LearningAgentsTrainer.h
--- a/Engine/Plugins/Experimental/LearningAgents/Source/LearningAgentsTraining/Public/LearningAgentsTrainer.h
+++ b/Engine/Plugins/Experimental/LearningAgents/Source/LearningAgentsTraining/Public/LearningAgentsTrainer.h
@@ -2,8 +2,8 @@
 
 #pragma once
 
-#include "LearningAgentsManagerComponent.h"
-#include "LearningAgentsCritic.h" // Included for FLearningAgentsCriticSettings()
+#include "LearningAgentsManagerListener.h"
+#include "LearningAgentsCompletions.h"
 
 #include "Templates/SharedPointer.h"
 #include "UObject/ObjectPtr.h"
@@ -13,45 +13,21 @@
 
 namespace UE::Learning
 {
-	struct FAnyCompletion;
-	struct FArrayMap;
 	struct FEpisodeBuffer;
 	struct FReplayBuffer;
 	struct FResetInstanceBuffer;
-	struct FRewardObject;
 	struct FSharedMemoryPPOTrainer;
-	struct FSumReward;
-	struct FCompletionObject;
-	enum class ECompletionMode : uint8;
 	enum class ETrainerDevice : uint8;
 }
 
 class ALearningAgentsManager;
 class ULearningAgentsInteractor;
-class ULearningAgentsCompletion;
-class ULearningAgentsReward;
 class ULearningAgentsPolicy;
 class ULearningAgentsCritic;
 
-/** Completion modes for episodes. */
-UENUM(BlueprintType, Category = "LearningAgents", meta = (ScriptName = "LearningAgentsCompletionEnum"))
-enum class ELearningAgentsCompletion : uint8
-{
-	/** Episode ended early but was still in progress. Critic will be used to estimate final return. */
-	Truncation	UMETA(DisplayName = "Truncation"),
-
-	/** Episode ended early and zero reward was expected for all future steps. */
-	Termination	UMETA(DisplayName = "Termination"),
-};
-
-namespace UE::Learning::Agents
-{
-	/** Get the learning agents completion from the UE::Learning completion. */
-	LEARNINGAGENTSTRAINING_API ELearningAgentsCompletion GetLearningAgentsCompletion(const ECompletionMode CompletionMode);
-
-	/** Get the UE::Learning completion from the learning agents completion. */
-	LEARNINGAGENTSTRAINING_API ECompletionMode GetCompletionMode(const ELearningAgentsCompletion Completion);
-}
+DECLARE_DYNAMIC_DELEGATE_RetVal_OneParam(float, FGatherAgentRewardDelegate, const int32, AgentId);
+DECLARE_DYNAMIC_DELEGATE_RetVal_OneParam(ELearningAgentsCompletion, FGatherAgentCompletionDelegate, const int32, AgentId);
+DECLARE_DYNAMIC_DELEGATE_OneParam(FResetAgentEpisodeDelegate, const int32, AgentId);
 
 /** The configurable settings for a ULearningAgentsTrainer. */
 USTRUCT(BlueprintType, Category = "LearningAgents")
@@ -61,19 +37,26 @@ struct LEARNINGAGENTSTRAINING_API FLearningAgentsTrainerSettings
 
 public:
 
-	/** Completion type to use when the maximum number of steps for an episode is reached. */
-	UPROPERTY(EditAnywhere, Category = "LearningAgents")
-	ELearningAgentsCompletion MaxStepsCompletion = ELearningAgentsCompletion::Truncation;
-
-	/** Max number of steps to take while training before episode automatically completes. */
+	/**
+	 * Maximum number of steps recorded in an episode before it is added to the replay buffer. This can generally be left at the default value and 
+	 * does not have a large impact on training.
+	 */
 	UPROPERTY(EditAnywhere, Category = "LearningAgents", meta = (ClampMin = "1", UIMin = "1"))
-	int32 MaxStepNum = 300;
+	int32 MaxEpisodeStepNum = 512;
 
-	/** Maximum number of episodes to record before running a training iteration. */
+	/**
+	 * Maximum number of episodes to record before running a training iteration. An iteration of training will be run when either this or
+	 * MaximumRecordedEpisodesPerIteration is reached. Typical values for this should be around 1000. Setting this too small means there is not 
+	 * enough data each iteration for the system to train. Setting it too large means training will be very slow. 
+	 */
 	UPROPERTY(EditAnywhere, Category = "LearningAgents", meta = (ClampMin = "1", UIMin = "1"))
 	int32 MaximumRecordedEpisodesPerIteration = 1000;
 
-	/** Maximum number of steps to record before running a training iteration. */
+	/**
+	 * Maximum number of steps to record before running a training iteration. An iteration of training will be run when either this or
+	 * MaximumRecordedEpisodesPerIteration is reached. Typical values for this should be around 10000. Setting this too small means there is not
+	 * enough data each iteration for the system to train. Setting it too large means training will be very slow.
+	 */
 	UPROPERTY(EditAnywhere, Category = "LearningAgents", meta = (ClampMin = "1", UIMin = "1"))
 	int32 MaximumRecordedStepsPerIteration = 10000;
 
@@ -173,13 +156,6 @@ public:
 	UPROPERTY(EditAnywhere, Category = "LearningAgents", meta = (ClampMin = "0.0", ClampMax = "1.0", UIMin = "0.0", UIMax = "1.0"))
 	float WeightDecay = 0.001f;
 
-	/**
-	 * The initial scaling for the weights of the output layer of the neural network. Typically, you would use this to
-	 * scale down the initial actions as it can stabilize the training and speed up convergence.
-	 */
-	UPROPERTY(EditAnywhere, Category = "LearningAgents", meta = (ClampMin = "0.0", ClampMax = "10.0", UIMin = "0.0", UIMax = "10.0"))
-	float InitialActionScale = 0.1f;
-
 	/**
 	 * Batch size to use for training the policy. Smaller values tend to produce better results at the cost of slowing down 
 	 * training. Large batch sizes are much more computationally efficient when training on the GPU.
@@ -227,7 +203,7 @@ public:
 	 * exploration but can make actions very noisy.
 	 */
 	UPROPERTY(EditAnywhere, Category = "LearningAgents", meta = (ClampMin = "0.0", ClampMax = "1.0", UIMin = "0.0", UIMax = "1.0"))
-	float ActionEntropyWeight = 0.01f;
+	float ActionEntropyWeight = 0.0f;
 
 	/**
 	 * This is used in the Generalized Advantage Estimation as what is essentially an exponential smoothing/decay. 
@@ -286,6 +262,10 @@ public:
 	/** If true, TensorBoard logs will be emitted to the intermediate directory. */
 	UPROPERTY(EditAnywhere, Category = "LearningAgents")
 	bool bUseTensorboard = false;
+
+	/** If true, snapshots of the trained networks will be emitted to the intermediate directory. */
+	UPROPERTY(EditAnywhere, Category = "LearningAgents")
+	bool bSaveSnapshots = false;
 };
 
 /** The path settings for the trainer. */
@@ -337,21 +317,16 @@ public:
 };
 
 /**
- * The ULearningAgentsTrainer is the core class for reinforcement learning training. It has a few responsibilities:
- *   1) It keeps track of which agents are gathering training data.
- *   2) It defines how those agents' rewards, completions, and resets are implemented.
- *   3) It provides methods for orchestrating the training process.
+ * The ULearningAgentsTrainer is the core class for reinforcement learning training. It defines how agents rewards, completions, and resets are 
+ * implemented and provides methods for orchestrating the training process.
  *
- * To use this class, you need to implement the SetupRewards and SetupCompletions functions (as well as their
- * corresponding SetRewards and SetCompletions functions), which will define the rewards and penalties the agent
- * receives and what conditions cause an episode to end. Before you can begin training, you need to call SetupTrainer,
- * which will initialize the underlying data structures, and you need to call AddAgent for each agent you want to gather
- * training data from.
+ * To use this class, you need to implement the GatherAgentRewards and GatherAgentCompletions functions, which will define the rewards and penalties 
+ * the agent receives and what conditions cause an episode to end.
  *
  * @see ULearningAgentsInteractor to understand how observations and actions work.
  */
-UCLASS(Abstract, BlueprintType, Blueprintable)
-class LEARNINGAGENTSTRAINING_API ULearningAgentsTrainer : public ULearningAgentsManagerComponent
+UCLASS(Abstract, HideDropdown, BlueprintType, Blueprintable)
+class LEARNINGAGENTSTRAINING_API ULearningAgentsTrainer : public ULearningAgentsManagerListener
 {
 	GENERATED_BODY()
 
@@ -363,11 +338,34 @@ public:
 	ULearningAgentsTrainer(FVTableHelper& Helper);
 	virtual ~ULearningAgentsTrainer();
 
-	/** Will automatically call EndTraining if training is still in-progress when play is ending. */
-	virtual void EndPlay(const EEndPlayReason::Type EndPlayReason) override;
+	/** Will automatically call EndTraining if training is still in-progress when the object is destroyed. */
+	virtual void BeginDestroy() override;
+
+	/**
+	 * Constructs the trainer and runs the setup functions for rewards and completions.
+	 * 
+	 * @param InManager			The agent manager we are using.
+	 * @param InInteractor		The agent interactor we are training with.
+	 * @param InPolicy			The policy to be trained.
+	 * @param InCritic			The critic to be trained.
+	 * @param Class				The trainer class
+	 * @param Name				The trainer name
+	 * @param TrainerSettings	The trainer settings to use.
+	 */
+	UFUNCTION(BlueprintCallable, Category = "LearningAgents", meta = (DeterminesOutputType = "Class"))
+	static ULearningAgentsTrainer* MakeTrainer(
+		ULearningAgentsManager* InManager,
+		ULearningAgentsInteractor* InInteractor,
+		ULearningAgentsPolicy* InPolicy,
+		ULearningAgentsCritic* InCritic,
+		TSubclassOf<ULearningAgentsTrainer> Class,
+		const FName Name = TEXT("Trainer"),
+		const FLearningAgentsTrainerSettings& TrainerSettings = FLearningAgentsTrainerSettings());
 
 	/**
-	 * Initializes this object and runs the setup functions for rewards and completions.
+	 * Initializes the trainer and runs the setup functions for rewards and completions.
+	 * 
+	 * @param InManager The agent manager we are using.
 	 * @param InInteractor The agent interactor we are training with.
 	 * @param InPolicy The policy to be trained.
 	 * @param InCritic The critic to be trained.
@@ -375,6 +373,7 @@ public:
 	 */
 	UFUNCTION(BlueprintCallable, Category = "LearningAgents")
 	void SetupTrainer(
+		ULearningAgentsManager* InManager,
 		ULearningAgentsInteractor* InInteractor,
 		ULearningAgentsPolicy* InPolicy,
 		ULearningAgentsCritic* InCritic,
@@ -382,72 +381,80 @@ public:
 
 public: 
 
-	//~ Begin ULearningAgentsManagerComponent Interface
-	virtual void OnAgentsAdded(const TArray<int32>& AgentIds) override;
-	virtual void OnAgentsRemoved(const TArray<int32>& AgentIds) override;
-	virtual void OnAgentsReset(const TArray<int32>& AgentIds) override;
-	//~ End ULearningAgentsManagerComponent Interface
+	//~ Begin ULearningAgentsManagerListener Interface
+	virtual void OnAgentsAdded_Implementation(const TArray<int32>& AgentIds) override;
+	virtual void OnAgentsRemoved_Implementation(const TArray<int32>& AgentIds) override;
+	virtual void OnAgentsReset_Implementation(const TArray<int32>& AgentIds) override;
+	virtual void OnAgentsManagerTick_Implementation(const TArray<int32>& AgentIds, const float DeltaTime) override;
+	//~ End ULearningAgentsManagerListener Interface
 
 // ----- Rewards -----
 public:
-	
-	/**
-	 * During this event, all rewards/penalties should be added to this trainer.
-	 * @see LearningAgentsRewards.h for the list of available rewards/penalties.
-	 */
-	UFUNCTION(BlueprintNativeEvent, Category = "LearningAgents")
-	void SetupRewards();
 
 	/**
-	 * During this event, all rewards/penalties should be set for each agent.
-	 * @param AgentIds The list of agent ids to set rewards/penalties for.
-	 * @see LearningAgentsRewards.h for the list of available rewards/penalties.
-	 * @see GetAgent to get the agent corresponding to each id.
+	 * This callback gathers all the reward values for the given set of agents.
+	 *
+	 * @param OutRewards		Output rewards for each agent in AgentIds
+	 * @param AgentIds			Agents to gather rewards for.
 	 */
-	UFUNCTION(BlueprintNativeEvent, Category = "LearningAgents")
-	void SetRewards(const TArray<int32>& AgentIds);
+	UFUNCTION(BlueprintNativeEvent, Category = "LearningAgents", Meta = (ForceAsFunction))
+	void GatherAgentRewards(TArray<float>& OutRewards, const TArray<int32>& AgentIds);
 
 	/**
-	 * Used by objects derived from ULearningAgentsReward to add themselves to this trainer during their creation.
-	 * You shouldn't need to call this directly.
+	 * This is a convenience function that can be used to implement `GatherAgentRewards` in terms of a callback which iterates over every agent
+	 * individually and gathers the rewards.
+	 *
+	 * @param OutRewards		Output rewards for each agent in AgentIds
+	 * @param AgentIds			Agents to gather rewards for.
+	 * @param Delegate			The Delegate used to gather the rewards for each agent.
 	 */
-	void AddReward(TObjectPtr<ULearningAgentsReward> Object, const TSharedRef<UE::Learning::FRewardObject>& Reward);
+	UFUNCTION(BlueprintPure = false, Category = "LearningAgents")
+	void GatherAgentRewardsUsingDelegate(TArray<float>& OutRewards, const TArray<int32>& AgentIds, const FGatherAgentRewardDelegate& Delegate);
 
 // ----- Completions ----- 
 public:
 
 	/**
-	 * During this event, all completions should be added to this trainer.
-	 * @see LearningAgentsCompletions.h for the list of available completions.
+	 * This callback gathers all the completions for the given set of agents.
+	 *
+	 * @param OutCompletions	Output completions for each agent in AgentIds
+	 * @param AgentIds			Agents to gather completions for.
 	 */
-	UFUNCTION(BlueprintNativeEvent, Category = "LearningAgents")
-	void SetupCompletions();
+	UFUNCTION(BlueprintNativeEvent, Category = "LearningAgents", Meta = (ForceAsFunction))
+	void GatherAgentCompletions(TArray<ELearningAgentsCompletion>& OutCompletions, const TArray<int32>& AgentIds);
 
 	/**
-	 * During this event, all completions should be set for each agent.
-	 * @param AgentIds The list of agent ids to set completions for.
-	 * @see LearningAgentsCompletions.h for the list of available completions.
-	 * @see GetAgent to get the agent corresponding to each id.
+	 * This is a convenience function that can be used to implement `GatherAgentCompletions` in terms of a callback which iterates over every agent
+	 * individually and gathers the completions.
+	 *
+	 * @param OutCompletions	Output completions for each agent in AgentIds
+	 * @param AgentIds			Agents to gather completions for.
+	 * @param Delegate			The Delegate used to gather the completions for each agent.
 	 */
-	UFUNCTION(BlueprintNativeEvent, Category = "LearningAgents")
-	void SetCompletions(const TArray<int32>& AgentIds);
+	UFUNCTION(BlueprintPure = false, Category = "LearningAgents")
+	void GatherAgentCompletionsDelegate(TArray<ELearningAgentsCompletion>& OutCompletions, const TArray<int32>& AgentIds, const FGatherAgentCompletionDelegate& Delegate);
 
-	/**
-	 * Used by objects derived from ULearningAgentsCompletion to add themselves to this trainer during their creation.
-	 * You shouldn't need to call this directly.
-	 */
-	void AddCompletion(TObjectPtr<ULearningAgentsCompletion> Object, const TSharedRef<UE::Learning::FCompletionObject>& Completion);
 
 // ----- Resets ----- 
 public:
 
 	/**
 	 * During this event, all episodes should be reset for each agent.
-	 * @param AgentIds The ids of the agents that need resetting.
-	 * @see GetAgent to get the agent corresponding to each id.
+	 * 
+	 * @param AgentIds			The ids of the agents that need resetting.
 	 */
-	UFUNCTION(BlueprintNativeEvent, Category = "LearningAgents")
-	void ResetEpisodes(const TArray<int32>& AgentIds);
+	UFUNCTION(BlueprintNativeEvent, Category = "LearningAgents", Meta = (ForceAsFunction))
+	void ResetAgentEpisodes(const TArray<int32>& AgentIds);
+
+	/**
+	 * This is a convenience function that can be used to implement `ResetAgentEpisodes` in terms of a callback which iterates over every agent
+	 * individually and does the reset.
+	 *
+	 * @param AgentIds			The ids of the agents that need resetting.
+	 * @param Delegate			The Delegate used to reset the episode for each agent.
+	 */
+	UFUNCTION(BlueprintPure = false, Category = "LearningAgents")
+	void ResetAgentEpisodesDelegate(const TArray<int32>& AgentIds, const FResetAgentEpisodeDelegate& Delegate);
 
 // ----- Training Process -----
 public:
@@ -459,20 +466,16 @@ public:
 	/**
 	 * Begins the training process with the provided settings.
 	 * 
-	 * @param TrainerTrainingSettings The settings for this training run.
-	 * @param TrainerGameSettings The settings that will affect the game's simulation.
-	 * @param TrainerPathSettings The path settings used by the trainer.
-	 * @param bReinitializePolicyNetwork If true, reinitialize the policy. Set this to false if your policy is pre-trained, e.g. with imitation learning.
-	 * @param bReinitializeCriticNetwork If true, reinitialize the critic. Set this to false if your critic is pre-trained.
-	 * @param bResetAgentsOnBegin If true, reset all agents at the beginning of training.
+	 * @param TrainerTrainingSettings	The settings for this training run.
+	 * @param TrainerGameSettings		The settings that will affect the game's simulation.
+	 * @param TrainerPathSettings		The path settings used by the trainer.
+	 * @param bResetAgentsOnBegin		If true, reset all agents at the beginning of training.
 	 */
 	UFUNCTION(BlueprintCallable, Category = "LearningAgents")
 	void BeginTraining(
 		const FLearningAgentsTrainerTrainingSettings& TrainerTrainingSettings = FLearningAgentsTrainerTrainingSettings(),
 		const FLearningAgentsTrainerGameSettings& TrainerGameSettings = FLearningAgentsTrainerGameSettings(),
 		const FLearningAgentsTrainerPathSettings& TrainerPathSettings = FLearningAgentsTrainerPathSettings(),
-		const bool bReinitializePolicyNetwork = true,
-		const bool bReinitializeCriticNetwork = true,
 		const bool bResetAgentsOnBegin = true);
 
 	/** Stops the training process. */
@@ -485,7 +488,7 @@ public:
 	 * the next state before evaluating the rewards.
 	 */
 	UFUNCTION(BlueprintCallable, Category = "LearningAgents")
-	void EvaluateRewards();
+	void GatherRewards();
 
 	/**
 	 * Call this function when it is time to evaluate the completions for your agents. This should be done at the beginning
@@ -493,40 +496,52 @@ public:
 	 * the next state before evaluating the completions.
 	 */
 	UFUNCTION(BlueprintCallable, Category = "LearningAgents")
-	void EvaluateCompletions();
+	void GatherCompletions();
 
 	/**
 	 * Call this function at the end of each step of your training loop. This takes the current observations/actions/
 	 * rewards and moves them into the episode experience buffer. All agents with full episode buffers or those which
 	 * have been signaled complete will be reset. If enough experience is gathered, it will be sent to the training 
 	 * process and an iteration of training will be run and the updated policy will be synced back.
+	 *
+	 * @param bResetAgentsOnUpdate				If true, reset all agents whenever an updated policy is received.
 	 */
 	UFUNCTION(BlueprintCallable, Category = "LearningAgents")
-	void ProcessExperience();
+	void ProcessExperience(const bool bResetAgentsOnUpdate = true);
 
 	/**
 	 * Convenience function that runs a basic training loop. If training has not been started, it will start it, and 
-	 * then call RunInference. On each following call to this function, it will call EvaluateRewards, 
-	 * EvaluateCompletions, and ProcessExperience, followed by RunInference.
+	 * then call RunInference. On each following call to this function, it will call GatherRewards, 
+	 * GatherCompletions, and ProcessExperience, followed by RunInference.
 	 * 
-	 * @param TrainerTrainingSettings The settings for this training run.
-	 * @param TrainerGameSettings The settings that will affect the game's simulation.
-	 * @param TrainerPathSettings The path settings used by the trainer.
-	 * @param bReinitializePolicyNetwork If true, reinitialize the policy. Set this to false if your policy is pre-trained, e.g. with imitation learning.
-	 * @param bReinitializeCriticNetwork If true, reinitialize the critic. Set this to false if your critic is pre-trained.
-	 * @param bResetAgentsOnBegin If true, reset all agents at the beginning of training.
+	 * @param TrainerTrainingSettings			The settings for this training run.
+	 * @param TrainerGameSettings				The settings that will affect the game's simulation.
+	 * @param TrainerPathSettings				The path settings used by the trainer.
+	 * @param bResetAgentsOnBegin				If true, reset all agents at the beginning of training.
+	 * @param bResetAgentsOnUpdate				If true, reset all agents whenever an updated policy is received.
 	 */
 	UFUNCTION(BlueprintCallable, Category = "LearningAgents")
 	void RunTraining(
 		const FLearningAgentsTrainerTrainingSettings& TrainerTrainingSettings = FLearningAgentsTrainerTrainingSettings(),
 		const FLearningAgentsTrainerGameSettings& TrainerGameSettings = FLearningAgentsTrainerGameSettings(),
 		const FLearningAgentsTrainerPathSettings& TrainerPathSettings = FLearningAgentsTrainerPathSettings(),
-		const bool bReinitializePolicyNetwork = true,
-		const bool bReinitializeCriticNetwork = true,
-		const bool bResetAgentsOnBegin = true);
+		const bool bResetAgentsOnBegin = true,
+		const bool bResetAgentsOnUpdate = true);
+
+	/**
+	 * Returns true if GatherRewards has been called and the reward already set for the given agent.
+	 */
+	UFUNCTION(BlueprintPure, Category = "LearningAgents", meta = (AgentId = "-1"))
+	bool HasReward(const int32 AgentId) const;
 
 	/**
-	 * Gets the current reward for an agent according to the critic. Should be called only after EvaluateRewards.
+	 * Returns true if GatherCompletions has been called and the completion already set for the given agent.
+	 */
+	UFUNCTION(BlueprintPure, Category = "LearningAgents", meta = (AgentId = "-1"))
+	bool HasCompletion(const int32 AgentId) const;
+
+	/**
+	 * Gets the current reward for an agent. Should be called only after GatherRewards.
 	 *
 	 * @param AgentId	The AgentId to look-up the reward for
 	 * @returns			The reward
@@ -535,15 +550,31 @@ public:
 	float GetReward(const int32 AgentId) const;
 
 	/**
-	 * Gets if the agent will complete the episode or not according to the given set of completions. Should be called 
-	 * only after EvaluateCompletions.
+	 * Gets the current completion for an agent. Should be called only after GatherCompletions.
+	 *
+	 * @param AgentId	The AgentId to look-up the completion for
+	 * @returns			The completion type
+	 */
+	UFUNCTION(BlueprintPure, Category = "LearningAgents", meta = (AgentId = "-1"))
+	ELearningAgentsCompletion GetCompletion(const int32 AgentId) const;
+
+	/**
+	 * Gets the current elapsed episode time for the given agent.
 	 *
-	 * @param AgentId		The AgentId to look-up the completion for
-	 * @param OutCompletion	The completion type if the agent will complete the episode
-	 * @returns				If the agent will complete the episode
+	 * @param AgentId	The AgentId to look-up the episode time for
+	 * @returns			The elapsed episode time
 	 */
 	UFUNCTION(BlueprintPure, Category = "LearningAgents", meta = (AgentId = "-1"))
-	bool IsCompleted(const int32 AgentId, ELearningAgentsCompletion& OutCompletion) const;
+	float GetEpisodeTime(const int32 AgentId) const;
+
+	/**
+	 * Gets the number of step recorded in an episode for the given agent.
+	 *
+	 * @param AgentId	The AgentId to look-up the number of recorded episode steps for
+	 * @returns			The number of recorded episode steps
+	 */
+	UFUNCTION(BlueprintPure, Category = "LearningAgents", meta = (AgentId = "-1"))
+	int32 GetEpisodeStepNum(const int32 AgentId) const;
 
 	/**
 	 * Returns true if the trainer has failed to communicate with the external training process. This can be used in
@@ -553,7 +584,7 @@ public:
 	 */
 	UFUNCTION(BlueprintPure, Category = "LearningAgents")
 	bool HasTrainingFailed() const;
-	
+
 // ----- Private Data ----- 
 private:
 
@@ -580,46 +611,47 @@ private:
 	UPROPERTY(VisibleAnywhere, Transient, Category = "LearningAgents")
 	bool bHasTrainingFailed = false;
 
-	/** The list of current reward objects. */
-	UPROPERTY(VisibleAnywhere, Transient, Category = "LearningAgents")
-	TArray<TObjectPtr<ULearningAgentsReward>> RewardObjects;
+	/** Callback Reward Output */
+	TArray<float> RewardBuffer;
 
-	/** The list of current completion objects. */
-	UPROPERTY(VisibleAnywhere, Transient, Category = "LearningAgents")
-	TArray<TObjectPtr<ULearningAgentsCompletion>> CompletionObjects;
+	/** Callback Completion Output */
+	TArray<ELearningAgentsCompletion> CompletionBuffer;
+
+	/** Reward Buffer */
+	TLearningArray<1, float> Rewards;
+	
+	/** Agent Completions Buffer */
+	TLearningArray<1, UE::Learning::ECompletionMode> AgentCompletions;
 
-	TArray<TSharedRef<UE::Learning::FRewardObject>, TInlineAllocator<16>> RewardFeatures;
-	TArray<TSharedRef<UE::Learning::FCompletionObject>, TInlineAllocator<16>> CompletionFeatures;
+	/** Episode Completions Buffer */
+	TLearningArray<1, UE::Learning::ECompletionMode> EpisodeCompletions;
 
-	TSharedPtr<UE::Learning::FSumReward> Rewards;
-	TSharedPtr<UE::Learning::FAnyCompletion> Completions;
+	/** All Completions Buffer */
+	TLearningArray<1, UE::Learning::ECompletionMode> AllCompletions;
+
+	/** Agent episode times */
+	TLearningArray<1, float> EpisodeTimes;
 
 	TUniquePtr<UE::Learning::FEpisodeBuffer> EpisodeBuffer;
 	TUniquePtr<UE::Learning::FReplayBuffer> ReplayBuffer;
 	TUniquePtr<UE::Learning::FResetInstanceBuffer> ResetBuffer;
 	TUniquePtr<UE::Learning::FSharedMemoryPPOTrainer> Trainer;
 	
-	ELearningAgentsCompletion MaxStepsCompletion = ELearningAgentsCompletion::Truncation;
-
 	float TrainerTimeout = 10.0f;
 
 	void DoneTraining();
 
-// ----- Private Iteration Checks ----- 
-private:
-
 	/** Number of times rewards have been evaluated for all agents */
-	TLearningArray<1, uint64, TInlineAllocator<32>> RewardEvaluatedAgentIteration;
+	TLearningArray<1, uint64, TInlineAllocator<32>> RewardIteration;
 
 	/** Number of times completions have been evaluated for all agents */
-	TLearningArray<1, uint64, TInlineAllocator<32>> CompletionEvaluatedAgentIteration;
+	TLearningArray<1, uint64, TInlineAllocator<32>> CompletionIteration;
 
 	/** Temp buffers used to record the set of agents that are valid for training */
 	TArray<int32> ValidAgentIds;
 	TArray<int32> FinalValidAgentIds;
 	UE::Learning::FIndexSet ValidAgentSet;
 	UE::Learning::FIndexSet FinalValidAgentSet;
-	TBitArray<TInlineAllocator<32>> ValidAgentStatus;
 
 // ----- Private Recording of GameSettings ----- 
 private:
diff --git a/Engine/Plugins/Experimental/LearningAgents/Source/LearningTraining/Private/LearningExperience.cpp b/Engine/Plugins/Experimental/LearningAgents/Source/LearningTraining/Private/LearningExperience.cpp
--- a/Engine/Plugins/Experimental/LearningAgents/Source/LearningTraining/Private/LearningExperience.cpp
+++ b/Engine/Plugins/Experimental/LearningAgents/Source/LearningTraining/Private/LearningExperience.cpp
@@ -257,18 +257,7 @@ namespace UE::Learning
 
 	namespace Experience
 	{
-		void GatherExperienceReset(
-			FReplayBuffer& ReplayBuffer,
-			FEpisodeBuffer& EpisodeBuffer,
-			FResetInstanceBuffer& ResetBuffer,
-			const FIndexSet Instances)
-		{
-			ReplayBuffer.Reset();
-			EpisodeBuffer.Reset(Instances);
-			ResetBuffer.SetResetInstances(Instances);
-		}
-
-		bool GatherExperienceIteration(
+		void GatherExperienceUntilReplayBufferFull(
 			FReplayBuffer& ReplayBuffer,
 			FEpisodeBuffer& EpisodeBuffer,
 			FResetInstanceBuffer& ResetBuffer,
@@ -278,7 +267,8 @@ namespace UE::Learning
 			TLearningArrayView<2, float> MemoryStateVectorBuffer,
 			TLearningArrayView<1, float> RewardBuffer,
 			TLearningArrayView<1, ECompletionMode> CompletionBuffer,
-			const ECompletionMode EpisodeEndCompletionMode,
+			TLearningArrayView<1, ECompletionMode> EpisodeCompletionBuffer,
+			TLearningArrayView<1, ECompletionMode> AllCompletionBuffer,
 			const TFunctionRef<void(const FIndexSet Instances)> ResetFunction,
 			const TFunctionRef<void(const FIndexSet Instances)> ObservationFunction,
 			const TFunctionRef<void(const FIndexSet Instances)> PolicyFunction,
@@ -288,137 +278,90 @@ namespace UE::Learning
 			const TFunctionRef<void(const FIndexSet Instances)> CompletionFunction,
 			const FIndexSet Instances)
 		{
-			UE_LEARNING_TRACE_CPUPROFILER_EVENT_SCOPE(Learning::Experience::GatherExperienceIteration);
+			UE_LEARNING_TRACE_CPUPROFILER_EVENT_SCOPE(Learning::Experience::GatherExperienceUntilReplayBufferFull);
+
+			// Reset Everything
 
-			if (ResetBuffer.GetResetInstanceNum() > 0)
+			ReplayBuffer.Reset();
+			EpisodeBuffer.Reset(Instances);
+			ResetFunction(Instances);
+
+			while (true)
 			{
-				// Reset Environment
+				// Encode Observations
 
-				ResetFunction(ResetBuffer.GetResetInstances());
-			}
+				ObservationFunction(Instances);
 
-			// Encode Observations
+				// Evaluate Policy
 
-			ObservationFunction(Instances);
-			Array::Check(ObservationVectorBuffer, Instances);
+				PolicyFunction(Instances);
 
-			// Evaluate Policy
+				// Decode Actions
 
-			PolicyFunction(Instances);
+				ActionFunction(Instances);
 
-			// Decode Actions
+				// Update Environment
 
-			Array::Check(ActionVectorBuffer, Instances);
-			ActionFunction(Instances);
+				UpdateFunction(Instances);
 
-			// Update Environment
+				// Compute Rewards
 
-			UpdateFunction(Instances);
+				RewardFunction(Instances);
 
-			// Compute Rewards
+				// Push to Experience Buffer
 
-			RewardFunction(Instances);
-			Array::Check(RewardBuffer, Instances);
+				EpisodeBuffer.Push(
+					ObservationVectorBuffer,
+					ActionVectorBuffer,
+					PreEvaluationMemoryStateVectorBuffer,
+					RewardBuffer,
+					Instances);
 
-			// Push to Experience Buffer
+				// Evaluate Completions
 
-			EpisodeBuffer.Push(
-				ObservationVectorBuffer,
-				ActionVectorBuffer,
-				PreEvaluationMemoryStateVectorBuffer,
-				RewardBuffer,
-				Instances);
+				CompletionFunction(Instances);
 
-			// Evaluate Completions
+				Completion::EvaluateEndOfEpisodeCompletions(
+					EpisodeCompletionBuffer,
+					EpisodeBuffer.GetEpisodeStepNums(),
+					EpisodeBuffer.GetMaxStepNum(),
+					Instances);
 
-			CompletionFunction(Instances);
+				for (const int32 Instance : Instances)
+				{
+					AllCompletionBuffer[Instance] = Completion::Or(CompletionBuffer[Instance], EpisodeCompletionBuffer[Instance]);
+				}
 
-			Completion::EvaluateEndOfEpisodeCompletions(
-				CompletionBuffer,
-				EpisodeBuffer.GetEpisodeStepNums(),
-				EpisodeBuffer.GetMaxStepNum(),
-				EpisodeEndCompletionMode,
-				Instances);
+				ResetBuffer.SetResetInstancesFromCompletions(AllCompletionBuffer, Instances);
 
-			ResetBuffer.SetResetInstancesFromCompletions(CompletionBuffer, Instances);
+				if (ResetBuffer.GetResetInstanceNum() == 0)
+				{
+					continue;
+				}
 
-			// Evaluate Observations again for instances that are completed
+				// Evaluate Observations again for instances that are completed
 
-			if (ResetBuffer.GetResetInstanceNum() > 0)
-			{
 				ObservationFunction(ResetBuffer.GetResetInstances());
 
-				// Push to Replay Buffer
+				// Push completed instances to Replay Buffer and return if full
 
-				const bool bReplayBufferFull = ReplayBuffer.AddEpisodes(
-					CompletionBuffer,
+				if (ReplayBuffer.AddEpisodes(
+					AllCompletionBuffer,
 					ObservationVectorBuffer,
 					MemoryStateVectorBuffer,
 					EpisodeBuffer,
-					ResetBuffer.GetResetInstances());
+					ResetBuffer.GetResetInstances()))
+				{
+					return;
+				}
 
+				// Just reset Episode Buffer for instances who reached the maximum episode length
+				ResetBuffer.SetResetInstancesFromCompletions(EpisodeCompletionBuffer, Instances);
 				EpisodeBuffer.Reset(ResetBuffer.GetResetInstances());
 
-				return bReplayBufferFull;
-			}
-			else
-			{
-				return false;
-			}
-		}
-
-		void GatherExperienceUntilReplayBufferFull(
-			FReplayBuffer& ReplayBuffer,
-			FEpisodeBuffer& EpisodeBuffer,
-			FResetInstanceBuffer& ResetBuffer,
-			TLearningArrayView<2, float> ObservationVectorBuffer,
-			TLearningArrayView<2, float> ActionVectorBuffer,
-			TLearningArrayView<2, float> PreEvaluationMemoryStateVectorBuffer,
-			TLearningArrayView<2, float> MemoryStateVectorBuffer,
-			TLearningArrayView<1, float> RewardBuffer,
-			TLearningArrayView<1, ECompletionMode> CompletionBuffer,
-			const ECompletionMode EpisodeEndCompletionMode,
-			const TFunctionRef<void(const FIndexSet Instances)> ResetFunction,
-			const TFunctionRef<void(const FIndexSet Instances)> ObservationFunction,
-			const TFunctionRef<void(const FIndexSet Instances)> PolicyFunction,
-			const TFunctionRef<void(const FIndexSet Instances)> ActionFunction,
-			const TFunctionRef<void(const FIndexSet Instances)> UpdateFunction,
-			const TFunctionRef<void(const FIndexSet Instances)> RewardFunction,
-			const TFunctionRef<void(const FIndexSet Instances)> CompletionFunction,
-			const FIndexSet Instances)
-		{
-			UE_LEARNING_TRACE_CPUPROFILER_EVENT_SCOPE(Learning::Experience::GatherExperienceUntilReplayBufferFull);
-
-			bool bReplayBufferFull = false;
-
-			GatherExperienceReset(ReplayBuffer, EpisodeBuffer, ResetBuffer, Instances);
-
-			while (true)
-			{
-				bReplayBufferFull = GatherExperienceIteration(
-					ReplayBuffer,
-					EpisodeBuffer,
-					ResetBuffer,
-					ObservationVectorBuffer,
-					ActionVectorBuffer,
-					PreEvaluationMemoryStateVectorBuffer,
-					MemoryStateVectorBuffer,
-					RewardBuffer,
-					CompletionBuffer,
-					EpisodeEndCompletionMode,
-					ResetFunction,
-					ObservationFunction,
-					PolicyFunction,
-					ActionFunction,
-					UpdateFunction,
-					RewardFunction,
-					CompletionFunction,
-					Instances);
-
-				if (bReplayBufferFull)
-				{
-					break;
-				}
+				// Call Reset Function for instances which signaled a completion
+				ResetBuffer.SetResetInstancesFromCompletions(CompletionBuffer, Instances);
+				ResetFunction(ResetBuffer.GetResetInstances());
 			}
 		}
 
diff --git a/Engine/Plugins/Experimental/LearningAgents/Source/LearningTraining/Private/LearningImitationTrainer.cpp b/Engine/Plugins/Experimental/LearningAgents/Source/LearningTraining/Private/LearningImitationTrainer.cpp
--- a/Engine/Plugins/Experimental/LearningAgents/Source/LearningTraining/Private/LearningImitationTrainer.cpp
+++ b/Engine/Plugins/Experimental/LearningAgents/Source/LearningTraining/Private/LearningImitationTrainer.cpp
@@ -115,10 +115,14 @@ namespace UE::Learning
 		const int32 ObservationDimNum,
 		const int32 ActionDimNum,
 		const int32 MemoryStateDimNum,
-		const INeuralNetwork& PolicyNetwork,
+		const ULearningNeuralNetworkData& PolicyNetwork,
+		const ULearningNeuralNetworkData& EncoderNetwork,
+		const ULearningNeuralNetworkData& DecoderNetwork,
+		const Observation::FSchema& ObservationSchema,
+		const Observation::FSchemaElement& ObservationSchemaElement,
+		const Action::FSchema& ActionSchema,
+		const Action::FSchemaElement& ActionSchemaElement,
 		const FImitationTrainerTrainingSettings& TrainingSettings,
-		const FImitationTrainerNetworkSettings& NetworkSettings,
-		const EImitationTrainerFlags TrainerFlags,
 		const ESubprocessFlags TrainingProcessFlags,
 		const ELogSetting LogSettings)
 	{
@@ -128,7 +132,9 @@ namespace UE::Learning
 
 		// Allocate Shared Memory
 
-		Policy = SharedMemory::Allocate<1, uint8>({ PolicyNetwork.GetSerializationByteNum() });
+		Policy = SharedMemory::Allocate<1, uint8>({ PolicyNetwork.GetSnapshotByteNum() });
+		Encoder = SharedMemory::Allocate<1, uint8>({ EncoderNetwork.GetSnapshotByteNum() });
+		Decoder = SharedMemory::Allocate<1, uint8>({ DecoderNetwork.GetSnapshotByteNum() });
 		Controls = SharedMemory::Allocate<1, volatile int32>({ SharedMemoryTraining::GetControlNum() });
 		EpisodeStarts = SharedMemory::Allocate<1, int32>({ MaxEpisodeNum });
 		EpisodeLengths = SharedMemory::Allocate<1, int32>({ MaxEpisodeNum });
@@ -158,38 +164,39 @@ namespace UE::Learning
 		ConfigObject->SetStringField(TEXT("SitePackagesPath"), *FileManager.ConvertToAbsolutePathForExternalAppForRead(*SitePackagesPath));
 		ConfigObject->SetStringField(TEXT("IntermediatePath"), *FileManager.ConvertToAbsolutePathForExternalAppForRead(*IntermediatePath));
 
-		ConfigObject->SetStringField(TEXT("PolicyNetworkClass"), PolicyNetwork.GetPythonClassName());
-
 		ConfigObject->SetStringField(TEXT("PolicyGuid"), *Policy.Guid.ToString(EGuidFormats::DigitsWithHyphensInBraces));
+		ConfigObject->SetStringField(TEXT("EncoderGuid"), *Encoder.Guid.ToString(EGuidFormats::DigitsWithHyphensInBraces));
+		ConfigObject->SetStringField(TEXT("DecoderGuid"), *Decoder.Guid.ToString(EGuidFormats::DigitsWithHyphensInBraces));
 		ConfigObject->SetStringField(TEXT("ControlsGuid"), *Controls.Guid.ToString(EGuidFormats::DigitsWithHyphensInBraces));
 		ConfigObject->SetStringField(TEXT("EpisodeStartsGuid"), *EpisodeStarts.Guid.ToString(EGuidFormats::DigitsWithHyphensInBraces));
 		ConfigObject->SetStringField(TEXT("EpisodeLengthsGuid"), *EpisodeLengths.Guid.ToString(EGuidFormats::DigitsWithHyphensInBraces));
 		ConfigObject->SetStringField(TEXT("ObservationsGuid"), *Observations.Guid.ToString(EGuidFormats::DigitsWithHyphensInBraces));
 		ConfigObject->SetStringField(TEXT("ActionsGuid"), *Actions.Guid.ToString(EGuidFormats::DigitsWithHyphensInBraces));
 
+		ConfigObject->SetObjectField(TEXT("ObservationSchema"), Trainer::ConvertObservationSchemaToJSON(ObservationSchema, ObservationSchemaElement));
+		ConfigObject->SetObjectField(TEXT("ActionSchema"), Trainer::ConvertActionSchemaToJSON(ActionSchema, ActionSchemaElement));
 		ConfigObject->SetNumberField(TEXT("ObservationVectorDimensionNum"), ObservationDimNum);
 		ConfigObject->SetNumberField(TEXT("ActionVectorDimensionNum"), ActionDimNum);
 		ConfigObject->SetNumberField(TEXT("MemoryStateVectorDimensionNum"), MemoryStateDimNum);
 		ConfigObject->SetNumberField(TEXT("MaxEpisodeNum"), MaxEpisodeNum);
 		ConfigObject->SetNumberField(TEXT("MaxStepNum"), MaxStepNum);
 
-		ConfigObject->SetNumberField(TEXT("PolicyNetworkByteNum"), PolicyNetwork.GetSerializationByteNum());
-		ConfigObject->SetNumberField(TEXT("PolicyActionNoiseMin"), NetworkSettings.PolicyActionNoiseMin);
-		ConfigObject->SetNumberField(TEXT("PolicyActionNoiseMax"), NetworkSettings.PolicyActionNoiseMax);
+		ConfigObject->SetNumberField(TEXT("PolicyNetworkByteNum"), PolicyNetwork.GetSnapshotByteNum());
+		ConfigObject->SetNumberField(TEXT("EncoderNetworkByteNum"), EncoderNetwork.GetSnapshotByteNum());
+		ConfigObject->SetNumberField(TEXT("DecoderNetworkByteNum"), DecoderNetwork.GetSnapshotByteNum());
 
 		ConfigObject->SetNumberField(TEXT("IterationNum"), TrainingSettings.IterationNum);
-		ConfigObject->SetNumberField(TEXT("LearningRatePolicy"), TrainingSettings.LearningRatePolicy);
+		ConfigObject->SetNumberField(TEXT("LearningRate"), TrainingSettings.LearningRate);
 		ConfigObject->SetNumberField(TEXT("LearningRateDecay"), TrainingSettings.LearningRateDecay);
 		ConfigObject->SetNumberField(TEXT("WeightDecay"), TrainingSettings.WeightDecay);
-		ConfigObject->SetNumberField(TEXT("InitialActionScale"), TrainingSettings.InitialActionScale);
-		ConfigObject->SetNumberField(TEXT("InitialMemoryScale"), TrainingSettings.InitialMemoryScale);
-		ConfigObject->SetNumberField(TEXT("PolicyBatchSize"), TrainingSettings.PolicyBatchSize);
-		ConfigObject->SetNumberField(TEXT("PolicyWindow"), TrainingSettings.PolicyWindow);
+		ConfigObject->SetNumberField(TEXT("BatchSize"), TrainingSettings.BatchSize);
+		ConfigObject->SetNumberField(TEXT("Window"), TrainingSettings.Window);
+		ConfigObject->SetNumberField(TEXT("ActionRegularizationWeight"), TrainingSettings.ActionRegularizationWeight);
+		ConfigObject->SetNumberField(TEXT("ActionEntropyWeight"), TrainingSettings.ActionEntropyWeight);
 		ConfigObject->SetNumberField(TEXT("Seed"), TrainingSettings.Seed);
 		ConfigObject->SetStringField(TEXT("Device"), Trainer::GetDeviceString(TrainingSettings.Device));
 		ConfigObject->SetBoolField(TEXT("UseTensorBoard"), TrainingSettings.bUseTensorboard);
-
-		ConfigObject->SetBoolField(TEXT("UseInitialPolicyNetwork"), (bool)(TrainerFlags & EImitationTrainerFlags::UseInitialPolicyNetwork));
+		ConfigObject->SetBoolField(TEXT("SaveSnapshots"), TrainingSettings.bSaveSnapshots);
 
 		ConfigObject->SetBoolField(TEXT("LoggingEnabled"), LogSettings == ELogSetting::Silent ? false : true);
 
@@ -237,29 +244,95 @@ namespace UE::Learning
 	}
 
 	ETrainerResponse FSharedMemoryImitationTrainer::RecvPolicy(
-		INeuralNetwork& OutNetwork,
+		ULearningNeuralNetworkData& OutNetwork,
 		const float Timeout,
 		FRWLock* NetworkLock,
 		const ELogSetting LogSettings)
 	{
-		return SharedMemoryTraining::RecvPolicy(
+		return SharedMemoryTraining::RecvNetwork(
 			Controls.View,
 			OutNetwork,
+			SharedMemoryTraining::EControls::PolicySignal,
 			Policy.View,
 			Timeout,
 			NetworkLock,
 			LogSettings);
 	}
 
+	ETrainerResponse FSharedMemoryImitationTrainer::RecvEncoder(
+		ULearningNeuralNetworkData& OutNetwork,
+		const float Timeout,
+		FRWLock* NetworkLock,
+		const ELogSetting LogSettings)
+	{
+		return SharedMemoryTraining::RecvNetwork(
+			Controls.View,
+			OutNetwork,
+			SharedMemoryTraining::EControls::EncoderSignal,
+			Encoder.View,
+			Timeout,
+			NetworkLock,
+			LogSettings);
+	}
+
+	ETrainerResponse FSharedMemoryImitationTrainer::RecvDecoder(
+		ULearningNeuralNetworkData& OutNetwork,
+		const float Timeout,
+		FRWLock* NetworkLock,
+		const ELogSetting LogSettings)
+	{
+		return SharedMemoryTraining::RecvNetwork(
+			Controls.View,
+			OutNetwork,
+			SharedMemoryTraining::EControls::DecoderSignal,
+			Decoder.View,
+			Timeout,
+			NetworkLock,
+			LogSettings);
+	}
+
 	ETrainerResponse FSharedMemoryImitationTrainer::SendPolicy(
-		const INeuralNetwork& Network,
+		const ULearningNeuralNetworkData& Network,
 		const float Timeout,
 		FRWLock* NetworkLock,
 		const ELogSetting LogSettings)
 	{
-		return SharedMemoryTraining::SendPolicy(
+		return SharedMemoryTraining::SendNetwork(
 			Controls.View,
 			Policy.View,
+			SharedMemoryTraining::EControls::PolicySignal,
+			Network,
+			Timeout,
+			NetworkLock,
+			LogSettings);
+	}
+
+	ETrainerResponse FSharedMemoryImitationTrainer::SendEncoder(
+		const ULearningNeuralNetworkData& Network,
+		const float Timeout,
+		FRWLock* NetworkLock,
+		const ELogSetting LogSettings)
+	{
+		return SharedMemoryTraining::SendNetwork(
+			Controls.View,
+			Encoder.View,
+			SharedMemoryTraining::EControls::EncoderSignal,
+			Network,
+			Timeout,
+			NetworkLock,
+			LogSettings);
+	}
+
+	ETrainerResponse FSharedMemoryImitationTrainer::SendDecoder(
+		const ULearningNeuralNetworkData& Network,
+		const float Timeout,
+		FRWLock* NetworkLock,
+		const ELogSetting LogSettings)
+	{
+		return SharedMemoryTraining::SendNetwork(
+			Controls.View,
+			Decoder.View,
+			SharedMemoryTraining::EControls::DecoderSignal,
 			Network,
 			Timeout,
 			NetworkLock,
@@ -291,6 +364,8 @@ namespace UE::Learning
 	void FSharedMemoryImitationTrainer::Deallocate()
 	{
 		SharedMemory::Deallocate(Policy);
+		SharedMemory::Deallocate(Encoder);
+		SharedMemory::Deallocate(Decoder);
 		SharedMemory::Deallocate(Controls);
 		SharedMemory::Deallocate(Observations);
 		SharedMemory::Deallocate(Actions);
@@ -467,13 +542,17 @@ namespace UE::Learning
 		const int32 ObservationDimNum,
 		const int32 ActionDimNum,
 		const int32 MemoryStateDimNum,
-		const INeuralNetwork& PolicyNetwork,
+		const ULearningNeuralNetworkData& PolicyNetwork,
+		const ULearningNeuralNetworkData& EncoderNetwork,
+		const ULearningNeuralNetworkData& DecoderNetwork,
+		const Observation::FSchema& ObservationSchema,
+		const Observation::FSchemaElement& ObservationSchemaElement,
+		const Action::FSchema& ActionSchema,
+		const Action::FSchemaElement& ActionSchemaElement,
 		const TCHAR* IpAddress,
 		const uint32 Port,
 		const float Timeout,
-		const FImitationTrainerTrainingSettings& TrainingSettings,
-		const FImitationTrainerNetworkSettings& NetworkSettings,
-		const EImitationTrainerFlags TrainerFlags)
+		const FImitationTrainerTrainingSettings& TrainingSettings)
 	{
 		// Write Config
 
@@ -487,31 +566,30 @@ namespace UE::Learning
 		ConfigObject->SetStringField(TEXT("TrainerType"), TrainerType);
 		ConfigObject->SetStringField(TEXT("TimeStamp"), *TimeStamp);
 
-		ConfigObject->SetStringField(TEXT("PolicyNetworkClass"), PolicyNetwork.GetPythonClassName());
-
+		ConfigObject->SetObjectField(TEXT("ObservationSchema"), Trainer::ConvertObservationSchemaToJSON(ObservationSchema, ObservationSchemaElement));
+		ConfigObject->SetObjectField(TEXT("ActionSchema"), Trainer::ConvertActionSchemaToJSON(ActionSchema, ActionSchemaElement));
 		ConfigObject->SetNumberField(TEXT("ObservationVectorDimensionNum"), ObservationDimNum);
 		ConfigObject->SetNumberField(TEXT("ActionVectorDimensionNum"), ActionDimNum);
 		ConfigObject->SetNumberField(TEXT("MemoryStateVectorDimensionNum"), MemoryStateDimNum);
 		ConfigObject->SetNumberField(TEXT("MaxEpisodeNum"), MaxEpisodeNum);
 		ConfigObject->SetNumberField(TEXT("MaxStepNum"), MaxStepNum);
 
-		ConfigObject->SetNumberField(TEXT("PolicyNetworkByteNum"), PolicyNetwork.GetSerializationByteNum());
-		ConfigObject->SetNumberField(TEXT("PolicyActionNoiseMin"), NetworkSettings.PolicyActionNoiseMin);
-		ConfigObject->SetNumberField(TEXT("PolicyActionNoiseMax"), NetworkSettings.PolicyActionNoiseMax);
+		ConfigObject->SetNumberField(TEXT("PolicyNetworkByteNum"), PolicyNetwork.GetSnapshotByteNum());
+		ConfigObject->SetNumberField(TEXT("EncoderNetworkByteNum"), EncoderNetwork.GetSnapshotByteNum());
+		ConfigObject->SetNumberField(TEXT("DecoderNetworkByteNum"), DecoderNetwork.GetSnapshotByteNum());
 
 		ConfigObject->SetNumberField(TEXT("IterationNum"), TrainingSettings.IterationNum);
-		ConfigObject->SetNumberField(TEXT("LearningRatePolicy"), TrainingSettings.LearningRatePolicy);
+		ConfigObject->SetNumberField(TEXT("LearningRate"), TrainingSettings.LearningRate);
 		ConfigObject->SetNumberField(TEXT("LearningRateDecay"), TrainingSettings.LearningRateDecay);
 		ConfigObject->SetNumberField(TEXT("WeightDecay"), TrainingSettings.WeightDecay);
-		ConfigObject->SetNumberField(TEXT("InitialActionScale"), TrainingSettings.InitialActionScale);
-		ConfigObject->SetNumberField(TEXT("InitialMemoryScale"), TrainingSettings.InitialMemoryScale);
-		ConfigObject->SetNumberField(TEXT("PolicyBatchSize"), TrainingSettings.PolicyBatchSize);
-		ConfigObject->SetNumberField(TEXT("PolicyWindow"), TrainingSettings.PolicyWindow);
+		ConfigObject->SetNumberField(TEXT("BatchSize"), TrainingSettings.BatchSize);
+		ConfigObject->SetNumberField(TEXT("Window"), TrainingSettings.Window);
+		ConfigObject->SetNumberField(TEXT("ActionRegularizationWeight"), TrainingSettings.ActionRegularizationWeight);
+		ConfigObject->SetNumberField(TEXT("ActionEntropyWeight"), TrainingSettings.ActionEntropyWeight);
 		ConfigObject->SetNumberField(TEXT("Seed"), TrainingSettings.Seed);
 		ConfigObject->SetStringField(TEXT("Device"), Trainer::GetDeviceString(TrainingSettings.Device));
 		ConfigObject->SetBoolField(TEXT("UseTensorBoard"), TrainingSettings.bUseTensorboard);
-
-		ConfigObject->SetBoolField(TEXT("UseInitialPolicyNetwork"), (bool)(TrainerFlags & EImitationTrainerFlags::UseInitialPolicyNetwork));
+		ConfigObject->SetBoolField(TEXT("SaveSnapshots"), TrainingSettings.bSaveSnapshots);
 
 		FString JsonString;
 		TSharedRef<TJsonWriter<>> JsonWriter = TJsonWriterFactory<>::Create(&JsonString, 0);
@@ -519,7 +597,9 @@ namespace UE::Learning
 
 		// Allocate buffer to receive network data in
 
-		NetworkBuffer.SetNumUninitialized({ PolicyNetwork.GetSerializationByteNum() });
+		PolicyBuffer.SetNumUninitialized({ PolicyNetwork.GetSnapshotByteNum() });
+		EncoderBuffer.SetNumUninitialized({ EncoderNetwork.GetSnapshotByteNum() });
+		DecoderBuffer.SetNumUninitialized({ DecoderNetwork.GetSnapshotByteNum() });
 
 		// Create Socket
 
@@ -588,21 +668,57 @@ namespace UE::Learning
 	}
 
 	ETrainerResponse FSocketImitationTrainer::RecvPolicy(
-		INeuralNetwork& OutNetwork,
+		ULearningNeuralNetworkData& OutNetwork,
+		const float Timeout,
+		FRWLock* NetworkLock,
+		const ELogSetting LogSettings)
+	{
+		return SocketTraining::RecvNetwork(*Socket, OutNetwork, PolicyBuffer, SocketTraining::ESignal::RecvPolicy, Timeout, NetworkLock, LogSettings);
+	}
+
+	ETrainerResponse FSocketImitationTrainer::RecvEncoder(
+		ULearningNeuralNetworkData& OutNetwork,
 		const float Timeout,
 		FRWLock* NetworkLock,
 		const ELogSetting LogSettings)
 	{
-		return SocketTraining::RecvPolicy(*Socket, OutNetwork, NetworkBuffer, Timeout, NetworkLock, LogSettings);
+		return SocketTraining::RecvNetwork(*Socket, OutNetwork, EncoderBuffer, SocketTraining::ESignal::RecvEncoder, Timeout, NetworkLock, LogSettings);
+	}
+
+	ETrainerResponse FSocketImitationTrainer::RecvDecoder(
+		ULearningNeuralNetworkData& OutNetwork,
+		const float Timeout,
+		FRWLock* NetworkLock,
+		const ELogSetting LogSettings)
+	{
+		return SocketTraining::RecvNetwork(*Socket, OutNetwork, DecoderBuffer, SocketTraining::ESignal::RecvDecoder, Timeout, NetworkLock, LogSettings);
 	}
 
 	ETrainerResponse FSocketImitationTrainer::SendPolicy(
-		const INeuralNetwork& Network,
+		const ULearningNeuralNetworkData& Network,
+		const float Timeout,
+		FRWLock* NetworkLock,
+		const ELogSetting LogSettings)
+	{
+		return SocketTraining::SendNetwork(*Socket, PolicyBuffer, SocketTraining::ESignal::SendPolicy, Network, Timeout, NetworkLock, LogSettings);
+	}
+
+	ETrainerResponse FSocketImitationTrainer::SendEncoder(
+		const ULearningNeuralNetworkData& Network,
+		const float Timeout,
+		FRWLock* NetworkLock,
+		const ELogSetting LogSettings)
+	{
+		return SocketTraining::SendNetwork(*Socket, EncoderBuffer, SocketTraining::ESignal::SendEncoder, Network, Timeout, NetworkLock, LogSettings);
+	}
+
+	ETrainerResponse FSocketImitationTrainer::SendDecoder(
+		const ULearningNeuralNetworkData& Network,
 		const float Timeout,
 		FRWLock* NetworkLock,
 		const ELogSetting LogSettings)
 	{
-		return SocketTraining::SendPolicy(*Socket, NetworkBuffer, Network, Timeout, NetworkLock, LogSettings);
+		return SocketTraining::SendNetwork(*Socket, DecoderBuffer, SocketTraining::ESignal::SendDecoder, Network, Timeout, NetworkLock, LogSettings);
 	}
 
 	ETrainerResponse FSocketImitationTrainer::SendExperience(
@@ -627,15 +743,20 @@ namespace UE::Learning
 	{
 		ETrainerResponse Train(
 			IImitationTrainer& Trainer,
-			INeuralNetwork& Network,
+			ULearningNeuralNetworkData& PolicyNetwork,
+			ULearningNeuralNetworkData& EncoderNetwork,
+			ULearningNeuralNetworkData& DecoderNetwork,
 			const TLearningArrayView<1, const int32> EpisodeStartsExperience,
 			const TLearningArrayView<1, const int32> EpisodeLengthsExperience,
 			const TLearningArrayView<2, const float> ObservationsExperience,
 			const TLearningArrayView<2, const float> ActionsExperience,
-			const EImitationTrainerFlags TrainerFlags, 
 			TAtomic<bool>* bRequestTrainingStopSignal,
-			FRWLock* NetworkLock,
-			TAtomic<bool>* bNetworkUpdatedSignal,
+			FRWLock* PolicyNetworkLock,
+			FRWLock* EncoderNetworkLock,
+			FRWLock* DecoderNetworkLock,
+			TAtomic<bool>* bPolicyNetworkUpdatedSignal,
+			TAtomic<bool>* bEncoderNetworkUpdatedSignal,
+			TAtomic<bool>* bDecoderNetworkUpdatedSignal,
 			const ELogSetting LogSettings)
 		{
 			ETrainerResponse Response = ETrainerResponse::Success;
@@ -647,7 +768,7 @@ namespace UE::Learning
 				UE_LOG(LogLearning, Display, TEXT("Sending initial Policy..."));
 			}
 
-			Response = Trainer.SendPolicy(Network, 20.0f, NetworkLock);
+			Response = Trainer.SendPolicy(PolicyNetwork, Trainer::DefaultTimeout, PolicyNetworkLock);
 
 			if (Response != ETrainerResponse::Success)
 			{
@@ -660,27 +781,44 @@ namespace UE::Learning
 				return Response;
 			}
 
-			if (!(bool)(TrainerFlags & EImitationTrainerFlags::UseInitialPolicyNetwork))
+			// Send initial Encoder
+
+			if (LogSettings != ELogSetting::Silent)
 			{
-				// Receive initial Policy
+				UE_LOG(LogLearning, Display, TEXT("Sending initial Encoder..."));
+			}
+
+			Response = Trainer.SendEncoder(EncoderNetwork, Trainer::DefaultTimeout, EncoderNetworkLock);
 
+			if (Response != ETrainerResponse::Success)
+			{
 				if (LogSettings != ELogSetting::Silent)
 				{
-					UE_LOG(LogLearning, Display, TEXT("Receiving initial Policy..."));
+					UE_LOG(LogLearning, Error, TEXT("Error sending initial encoder from trainer: %s. Check log for errors."), Trainer::GetResponseString(Response));
 				}
 
-				Response = Trainer.RecvPolicy(Network, 20.0f, NetworkLock);
+				Trainer.Terminate();
+				return Response;
+			}
 
-				if (Response != ETrainerResponse::Success)
-				{
-					if (LogSettings != ELogSetting::Silent)
-					{
-						UE_LOG(LogLearning, Error, TEXT("Error receiving initial policy from trainer: %s. Check log for errors."), Trainer::GetResponseString(Response));
-					}
+			// Send initial Decoder
+
+			if (LogSettings != ELogSetting::Silent)
+			{
+				UE_LOG(LogLearning, Display, TEXT("Sending initial Decoder..."));
+			}
+
+			Response = Trainer.SendDecoder(DecoderNetwork, Trainer::DefaultTimeout, DecoderNetworkLock);
 
-					Trainer.Terminate();
-					return Response;
+			if (Response != ETrainerResponse::Success)
+			{
+				if (LogSettings != ELogSetting::Silent)
+				{
+					UE_LOG(LogLearning, Error, TEXT("Error sending initial decoder from trainer: %s. Check log for errors."), Trainer::GetResponseString(Response));
 				}
+
+				Trainer.Terminate();
+				return Response;
 			}
 
 			// Send Experience
@@ -690,7 +828,7 @@ namespace UE::Learning
 				EpisodeLengthsExperience,
 				ObservationsExperience, 
 				ActionsExperience, 
-				10.0f);
+				Trainer::DefaultTimeout);
 
 			if (Response != ETrainerResponse::Success)
 			{
@@ -734,7 +872,7 @@ namespace UE::Learning
 				
 				if (Trainer.HasPolicyOrCompleted())
 				{
-					Response = Trainer.RecvPolicy(Network, 10.0f, NetworkLock);
+					Response = Trainer.RecvPolicy(PolicyNetwork, Trainer::DefaultTimeout, PolicyNetworkLock);
 
 					if (Response == ETrainerResponse::Completed)
 					{
@@ -752,17 +890,53 @@ namespace UE::Learning
 						}
 						break;
 					}
-				}
 
-				if (bNetworkUpdatedSignal)
+					if (bPolicyNetworkUpdatedSignal)
+					{
+						*bPolicyNetworkUpdatedSignal = true;
+					}
+
+					Response = Trainer.RecvEncoder(EncoderNetwork, Trainer::DefaultTimeout, EncoderNetworkLock);
+
+					if (Response != ETrainerResponse::Success)
+					{
+						if (LogSettings != ELogSetting::Silent)
+						{
+							UE_LOG(LogLearning, Error, TEXT("Error receiving encoder from trainer: %s. Check log for errors."), Trainer::GetResponseString(Response));
+						}
+						break;
+					}
+
+					if (bEncoderNetworkUpdatedSignal)
+					{
+						*bEncoderNetworkUpdatedSignal = true;
+					}
+
+					Response = Trainer.RecvDecoder(DecoderNetwork, Trainer::DefaultTimeout, DecoderNetworkLock);
+
+					if (Response != ETrainerResponse::Success)
+					{
+						if (LogSettings != ELogSetting::Silent)
+						{
+							UE_LOG(LogLearning, Error, TEXT("Error receiving decoder from trainer: %s. Check log for errors."), Trainer::GetResponseString(Response));
+						}
+						break;
+					}
+
+					if (bDecoderNetworkUpdatedSignal)
+					{
+						*bDecoderNetworkUpdatedSignal = true;
+					}
+				}
+				else
 				{
-					*bNetworkUpdatedSignal = true;
+					FPlatformProcess::Sleep(0.001f);
 				}
 			}
 
 			// Allow some time for trainer to shut down gracefully before we kill it...
 
-			Response = Trainer.Wait(5.0f);
+			Response = Trainer.Wait(Trainer::DefaultTimeout);
 
 			if (Response != ETrainerResponse::Success)
 			{
