LearningAgents: Large restructuring to allow for structured actions and observations.

diff --git a/Engine/Plugins/Experimental/LearningAgents/Source/LearningTraining/Private/LearningPPOTrainer.cpp b/Engine/Plugins/Experimental/LearningAgents/Source/LearningTraining/Private/LearningPPOTrainer.cpp
--- a/Engine/Plugins/Experimental/LearningAgents/Source/LearningTraining/Private/LearningPPOTrainer.cpp
+++ b/Engine/Plugins/Experimental/LearningAgents/Source/LearningTraining/Private/LearningPPOTrainer.cpp
@@ -10,6 +10,8 @@
 #include "LearningSharedMemory.h"
 #include "LearningSharedMemoryTraining.h"
 #include "LearningSocketTraining.h"
+#include "LearningObservation.h"
+#include "LearningAction.h"
 
 #include "Misc/MonitoredProcess.h"
 #include "Misc/Guid.h"
@@ -115,11 +117,15 @@ namespace UE::Learning
 		const FString& PythonContentPath,
 		const FString& IntermediatePath,
 		const FReplayBuffer& ReplayBuffer,
-		const INeuralNetwork& PolicyNetwork,
-		const INeuralNetwork& CriticNetwork,
+		const ULearningNeuralNetworkData& PolicyNetwork,
+		const ULearningNeuralNetworkData& CriticNetwork,
+		const ULearningNeuralNetworkData& EncoderNetwork,
+		const ULearningNeuralNetworkData& DecoderNetwork,
+		const Observation::FSchema& ObservationSchema,
+		const Observation::FSchemaElement& ObservationSchemaElement,
+		const Action::FSchema& ActionSchema,
+		const Action::FSchemaElement& ActionSchemaElement,
 		const FPPOTrainerTrainingSettings& TrainingSettings,
-		const FPPOTrainerNetworkSettings& NetworkSettings,
-		const EPPOTrainerFlags TrainerFlags,
 		const ELogSetting LogSettings,
 		const ESubprocessFlags TrainingProcessFlags,
 		const uint16 ProcessNum,
@@ -148,8 +154,10 @@ namespace UE::Learning
 		{
 			// Allocate Shared Memory
 
-			Policy = SharedMemory::Allocate<1, uint8>({ PolicyNetwork.GetSerializationByteNum() });
-			Critic = SharedMemory::Allocate<1, uint8>({ CriticNetwork.GetSerializationByteNum() });
+			Policy = SharedMemory::Allocate<1, uint8>({ PolicyNetwork.GetSnapshotByteNum() });
+			Critic = SharedMemory::Allocate<1, uint8>({ CriticNetwork.GetSnapshotByteNum() });
+			Encoder = SharedMemory::Allocate<1, uint8>({ EncoderNetwork.GetSnapshotByteNum() });
+			Decoder = SharedMemory::Allocate<1, uint8>({ DecoderNetwork.GetSnapshotByteNum() });
 			Controls = SharedMemory::Allocate<2, volatile int32>({ ProcessNum, SharedMemoryTraining::GetControlNum() });
 			EpisodeStarts = SharedMemory::Allocate<2, int32>({ ProcessNum, ReplayBuffer.GetMaxEpisodeNum() });
 			EpisodeLengths = SharedMemory::Allocate<2, int32>({ ProcessNum, ReplayBuffer.GetMaxEpisodeNum() });
@@ -178,6 +186,8 @@ namespace UE::Learning
 				SubprocessCommandLine += FString::Printf(TEXT(" -LearningProcessIdx %i"), SubprocessIdx);
 				SubprocessCommandLine += FString(TEXT(" -LearningPolicyGuid ")) + Policy.Guid.ToString();
 				SubprocessCommandLine += FString(TEXT(" -LearningCriticGuid ")) + Critic.Guid.ToString();
+				SubprocessCommandLine += FString(TEXT(" -LearningEncoderGuid ")) + Encoder.Guid.ToString();
+				SubprocessCommandLine += FString(TEXT(" -LearningDecoderGuid ")) + Decoder.Guid.ToString();
 				SubprocessCommandLine += FString(TEXT(" -LearningControlsGuid ")) + Controls.Guid.ToString();
 				SubprocessCommandLine += FString(TEXT(" -LearningEpisodeStartsGuid ")) + EpisodeStarts.Guid.ToString();
 				SubprocessCommandLine += FString(TEXT(" -LearningEpisodeLengthsGuid ")) + EpisodeLengths.Guid.ToString();
@@ -226,11 +236,10 @@ namespace UE::Learning
 			ConfigObject->SetStringField(TEXT("SitePackagesPath"), *FileManager.ConvertToAbsolutePathForExternalAppForRead(*SitePackagesPath));
 			ConfigObject->SetStringField(TEXT("IntermediatePath"), *FileManager.ConvertToAbsolutePathForExternalAppForRead(*IntermediatePath));
 
-			ConfigObject->SetStringField(TEXT("PolicyNetworkClass"), PolicyNetwork.GetPythonClassName());
-			ConfigObject->SetStringField(TEXT("CriticNetworkClass"), CriticNetwork.GetPythonClassName());
-
 			ConfigObject->SetStringField(TEXT("PolicyGuid"), *Policy.Guid.ToString(EGuidFormats::DigitsWithHyphensInBraces));
 			ConfigObject->SetStringField(TEXT("CriticGuid"), *Critic.Guid.ToString(EGuidFormats::DigitsWithHyphensInBraces));
+			ConfigObject->SetStringField(TEXT("EncoderGuid"), *Encoder.Guid.ToString(EGuidFormats::DigitsWithHyphensInBraces));
+			ConfigObject->SetStringField(TEXT("DecoderGuid"), *Decoder.Guid.ToString(EGuidFormats::DigitsWithHyphensInBraces));
 			ConfigObject->SetStringField(TEXT("ControlsGuid"), *Controls.Guid.ToString(EGuidFormats::DigitsWithHyphensInBraces));
 			ConfigObject->SetStringField(TEXT("EpisodeStartsGuid"), *EpisodeStarts.Guid.ToString(EGuidFormats::DigitsWithHyphensInBraces));
 			ConfigObject->SetStringField(TEXT("EpisodeLengthsGuid"), *EpisodeLengths.Guid.ToString(EGuidFormats::DigitsWithHyphensInBraces));
@@ -242,17 +251,18 @@ namespace UE::Learning
 			ConfigObject->SetStringField(TEXT("MemoryStatesGuid"), *MemoryStates.Guid.ToString(EGuidFormats::DigitsWithHyphensInBraces));
 			ConfigObject->SetStringField(TEXT("RewardsGuid"), *Rewards.Guid.ToString(EGuidFormats::DigitsWithHyphensInBraces));
 
+			ConfigObject->SetObjectField(TEXT("ObservationSchema"), Trainer::ConvertObservationSchemaToJSON(ObservationSchema, ObservationSchemaElement));
+			ConfigObject->SetObjectField(TEXT("ActionSchema"), Trainer::ConvertActionSchemaToJSON(ActionSchema, ActionSchemaElement));
 			ConfigObject->SetNumberField(TEXT("ObservationVectorDimensionNum"), ObservationVectorDimensionNum);
 			ConfigObject->SetNumberField(TEXT("ActionVectorDimensionNum"), ActionVectorDimensionNum);
 			ConfigObject->SetNumberField(TEXT("MemoryStateVectorDimensionNum"), MemoryStateVectorDimensionNum);
 			ConfigObject->SetNumberField(TEXT("MaxEpisodeNum"), ReplayBuffer.GetMaxEpisodeNum());
 			ConfigObject->SetNumberField(TEXT("MaxStepNum"), ReplayBuffer.GetMaxStepNum());
 
-			ConfigObject->SetNumberField(TEXT("PolicyNetworkByteNum"), PolicyNetwork.GetSerializationByteNum());
-			ConfigObject->SetNumberField(TEXT("PolicyActionNoiseMin"), NetworkSettings.PolicyActionNoiseMin);
-			ConfigObject->SetNumberField(TEXT("PolicyActionNoiseMax"), NetworkSettings.PolicyActionNoiseMax);
-
-			ConfigObject->SetNumberField(TEXT("CriticNetworkByteNum"), CriticNetwork.GetSerializationByteNum());
+			ConfigObject->SetNumberField(TEXT("PolicyNetworkByteNum"), PolicyNetwork.GetSnapshotByteNum());
+			ConfigObject->SetNumberField(TEXT("CriticNetworkByteNum"), CriticNetwork.GetSnapshotByteNum());
+			ConfigObject->SetNumberField(TEXT("EncoderNetworkByteNum"), EncoderNetwork.GetSnapshotByteNum());
+			ConfigObject->SetNumberField(TEXT("DecoderNetworkByteNum"), DecoderNetwork.GetSnapshotByteNum());
 
 			ConfigObject->SetNumberField(TEXT("ProcessNum"), ProcessNum);
 
@@ -261,14 +271,13 @@ namespace UE::Learning
 			ConfigObject->SetNumberField(TEXT("LearningRateCritic"), TrainingSettings.LearningRateCritic);
 			ConfigObject->SetNumberField(TEXT("LearningRateDecay"), TrainingSettings.LearningRateDecay);
 			ConfigObject->SetNumberField(TEXT("WeightDecay"), TrainingSettings.WeightDecay);
-			ConfigObject->SetNumberField(TEXT("InitialActionScale"), TrainingSettings.InitialActionScale);
-			ConfigObject->SetNumberField(TEXT("InitialMemoryScale"), TrainingSettings.InitialMemoryScale);
-			ConfigObject->SetNumberField(TEXT("InitialCriticScale"), TrainingSettings.InitialCriticScale);
 			ConfigObject->SetNumberField(TEXT("PolicyBatchSize"), TrainingSettings.PolicyBatchSize);
 			ConfigObject->SetNumberField(TEXT("CriticBatchSize"), TrainingSettings.CriticBatchSize);
 			ConfigObject->SetNumberField(TEXT("PolicyWindow"), TrainingSettings.PolicyWindow);
 			ConfigObject->SetNumberField(TEXT("IterationsPerGather"), TrainingSettings.IterationsPerGather);
+			ConfigObject->SetNumberField(TEXT("CriticWarmupIterations"), TrainingSettings.CriticWarmupIterations);
 			ConfigObject->SetNumberField(TEXT("EpsilonClip"), TrainingSettings.EpsilonClip);
+			ConfigObject->SetNumberField(TEXT("ActionSurrogateWeight"), TrainingSettings.ActionSurrogateWeight);
 			ConfigObject->SetNumberField(TEXT("ActionRegularizationWeight"), TrainingSettings.ActionRegularizationWeight);
 			ConfigObject->SetNumberField(TEXT("ActionEntropyWeight"), TrainingSettings.ActionEntropyWeight);
 			ConfigObject->SetNumberField(TEXT("ReturnRegularizationWeight"), TrainingSettings.ReturnRegularizationWeight);
@@ -282,9 +291,7 @@ namespace UE::Learning
 			ConfigObject->SetNumberField(TEXT("DiscountFactor"), TrainingSettings.DiscountFactor);
 			ConfigObject->SetStringField(TEXT("Device"), Trainer::GetDeviceString(TrainingSettings.Device));
 			ConfigObject->SetBoolField(TEXT("UseTensorBoard"), TrainingSettings.bUseTensorboard);
-
-			ConfigObject->SetBoolField(TEXT("UseInitialPolicyNetwork"), (bool)(TrainerFlags& EPPOTrainerFlags::UseInitialPolicyNetwork));
-			ConfigObject->SetBoolField(TEXT("UseInitialCriticNetwork"), (bool)(TrainerFlags& EPPOTrainerFlags::UseInitialCriticNetwork));
+			ConfigObject->SetBoolField(TEXT("SaveSnapshots"), TrainingSettings.bSaveSnapshots);
 
 			ConfigObject->SetBoolField(TEXT("LoggingEnabled"), LogSettings == ELogSetting::Silent ? false : true);
 
@@ -321,6 +328,8 @@ namespace UE::Learning
 
 			FGuid PolicyGuid; ensure(FParse::Value(FCommandLine::Get(), TEXT("LearningPolicyGuid"), PolicyGuid));
 			FGuid CriticGuid; ensure(FParse::Value(FCommandLine::Get(), TEXT("LearningCriticGuid"), CriticGuid));
+			FGuid EncoderGuid; ensure(FParse::Value(FCommandLine::Get(), TEXT("LearningEncoderGuid"), EncoderGuid));
+			FGuid DecoderGuid; ensure(FParse::Value(FCommandLine::Get(), TEXT("LearningDecoderGuid"), DecoderGuid));
 			FGuid ControlsGuid; ensure(FParse::Value(FCommandLine::Get(), TEXT("LearningControlsGuid"), ControlsGuid));
 			FGuid EpisodeStartsGuid; ensure(FParse::Value(FCommandLine::Get(), TEXT("LearningEpisodeStartsGuid"), EpisodeStartsGuid));
 			FGuid EpisodeLengthsGuid; ensure(FParse::Value(FCommandLine::Get(), TEXT("LearningEpisodeLengthsGuid"), EpisodeLengthsGuid));
@@ -334,8 +343,10 @@ namespace UE::Learning
 
 			// Map shared memory
 
-			Policy = SharedMemory::Map<1, uint8>(PolicyGuid, { PolicyNetwork.GetSerializationByteNum() });
-			Critic = SharedMemory::Map<1, uint8>(CriticGuid, { CriticNetwork.GetSerializationByteNum() });
+			Policy = SharedMemory::Map<1, uint8>(PolicyGuid, { PolicyNetwork.GetSnapshotByteNum() });
+			Critic = SharedMemory::Map<1, uint8>(CriticGuid, { CriticNetwork.GetSnapshotByteNum() });
+			Encoder = SharedMemory::Map<1, uint8>(EncoderGuid, { EncoderNetwork.GetSnapshotByteNum() });
+			Decoder = SharedMemory::Map<1, uint8>(DecoderGuid, { DecoderNetwork.GetSnapshotByteNum() });
 			Controls = SharedMemory::Map<2, volatile int32>(ControlsGuid, { ProcessNum, SharedMemoryTraining::GetControlNum() });
 			EpisodeStarts = SharedMemory::Map<2, int32>(EpisodeStartsGuid, { ProcessNum, ReplayBuffer.GetMaxEpisodeNum() });
 			EpisodeLengths = SharedMemory::Map<2, int32>(EpisodeLengthsGuid, { ProcessNum, ReplayBuffer.GetMaxEpisodeNum() });
@@ -394,14 +405,15 @@ namespace UE::Learning
 	}
 
 	ETrainerResponse FSharedMemoryPPOTrainer::RecvPolicy(
-		INeuralNetwork& OutNetwork,
+		ULearningNeuralNetworkData& OutNetwork,
 		const float Timeout,
 		FRWLock* NetworkLock,
 		const ELogSetting LogSettings)
 	{
-		return SharedMemoryTraining::RecvPolicy(
+		return SharedMemoryTraining::RecvNetwork(
 			Controls.View[ProcessIdx],
 			OutNetwork,
+			SharedMemoryTraining::EControls::PolicySignal,
 			Policy.View,
 			Timeout,
 			NetworkLock,
@@ -409,32 +421,63 @@ namespace UE::Learning
 	}
 
 	ETrainerResponse FSharedMemoryPPOTrainer::RecvCritic(
-		INeuralNetwork& OutNetwork,
+		ULearningNeuralNetworkData& OutNetwork,
 		const float Timeout,
 		FRWLock* NetworkLock,
 		const ELogSetting LogSettings)
 	{
-		UE_LEARNING_CHECKF(Critic.View.GetData(), 
-			TEXT("Buffer for critic was not allocated. Were the correct EPPOTrainerFlags provided to construct the FSharedMemoryPPOTrainer?"));
-
-		return SharedMemoryTraining::RecvCritic(
+		return SharedMemoryTraining::RecvNetwork(
 			Controls.View[ProcessIdx],
 			OutNetwork,
+			SharedMemoryTraining::EControls::CriticSignal,
 			Critic.View,
 			Timeout,
 			NetworkLock,
 			LogSettings);
 	}
 
+	ETrainerResponse FSharedMemoryPPOTrainer::RecvEncoder(
+		ULearningNeuralNetworkData& OutNetwork,
+		const float Timeout,
+		FRWLock* NetworkLock,
+		const ELogSetting LogSettings)
+	{
+		return SharedMemoryTraining::RecvNetwork(
+			Controls.View[ProcessIdx],
+			OutNetwork,
+			SharedMemoryTraining::EControls::EncoderSignal,
+			Encoder.View,
+			Timeout,
+			NetworkLock,
+			LogSettings);
+	}
+
+	ETrainerResponse FSharedMemoryPPOTrainer::RecvDecoder(
+		ULearningNeuralNetworkData& OutNetwork,
+		const float Timeout,
+		FRWLock* NetworkLock,
+		const ELogSetting LogSettings)
+	{
+		return SharedMemoryTraining::RecvNetwork(
+			Controls.View[ProcessIdx],
+			OutNetwork,
+			SharedMemoryTraining::EControls::DecoderSignal,
+			Decoder.View,
+			Timeout,
+			NetworkLock,
+			LogSettings);
+	}
+
 	ETrainerResponse FSharedMemoryPPOTrainer::SendPolicy(
-		const INeuralNetwork& Network,
+		const ULearningNeuralNetworkData& Network,
 		const float Timeout,
 		FRWLock* NetworkLock,
 		const ELogSetting LogSettings)
 	{
-		return SharedMemoryTraining::SendPolicy(
+		return SharedMemoryTraining::SendNetwork(
 			Controls.View[ProcessIdx],
 			Policy.View,
+			SharedMemoryTraining::EControls::PolicySignal,
 			Network,
 			Timeout,
 			NetworkLock,
@@ -442,17 +485,47 @@ namespace UE::Learning
 	}
 
 	ETrainerResponse FSharedMemoryPPOTrainer::SendCritic(
-		const INeuralNetwork& Network,
+		const ULearningNeuralNetworkData& Network,
 		const float Timeout,
 		FRWLock* NetworkLock,
 		const ELogSetting LogSettings)
 	{
-		UE_LEARNING_CHECKF(Critic.View.GetData(),
-			TEXT("Buffer for critic was not allocated. Were the correct EPPOTrainerFlags provided to construct the FSharedMemoryPPOTrainer?"));
-
-		return SharedMemoryTraining::SendCritic(
+		return SharedMemoryTraining::SendNetwork(
 			Controls.View[ProcessIdx],
 			Critic.View,
+			SharedMemoryTraining::EControls::CriticSignal,
+			Network,
+			Timeout,
+			NetworkLock,
+			LogSettings);
+	}
+
+	ETrainerResponse FSharedMemoryPPOTrainer::SendEncoder(
+		const ULearningNeuralNetworkData& Network,
+		const float Timeout,
+		FRWLock* NetworkLock,
+		const ELogSetting LogSettings)
+	{
+		return SharedMemoryTraining::SendNetwork(
+			Controls.View[ProcessIdx],
+			Encoder.View,
+			SharedMemoryTraining::EControls::EncoderSignal,
+			Network,
+			Timeout,
+			NetworkLock,
+			LogSettings);
+	}
+
+	ETrainerResponse FSharedMemoryPPOTrainer::SendDecoder(
+		const ULearningNeuralNetworkData& Network,
+		const float Timeout,
+		FRWLock* NetworkLock,
+		const ELogSetting LogSettings)
+	{
+		return SharedMemoryTraining::SendNetwork(
+			Controls.View[ProcessIdx],
+			Decoder.View,
+			SharedMemoryTraining::EControls::DecoderSignal,
 			Network,
 			Timeout,
 			NetworkLock,
@@ -485,6 +558,9 @@ namespace UE::Learning
 		if (Policy.Region != nullptr)
 		{
 			SharedMemory::Deallocate(Policy);
+			SharedMemory::Deallocate(Critic);
+			SharedMemory::Deallocate(Encoder);
+			SharedMemory::Deallocate(Decoder);
 			SharedMemory::Deallocate(Controls);
 			SharedMemory::Deallocate(EpisodeStarts);
 			SharedMemory::Deallocate(EpisodeLengths);
@@ -652,14 +728,18 @@ namespace UE::Learning
 		ETrainerResponse& OutResponse,
 		const FString& TaskName,
 		const FReplayBuffer& ReplayBuffer,
-		const INeuralNetwork& PolicyNetwork,
-		const INeuralNetwork& CriticNetwork,
+		const ULearningNeuralNetworkData& PolicyNetwork,
+		const ULearningNeuralNetworkData& CriticNetwork,
+		const ULearningNeuralNetworkData& EncoderNetwork,
+		const ULearningNeuralNetworkData& DecoderNetwork,
+		const Observation::FSchema& ObservationSchema,
+		const Observation::FSchemaElement& ObservationSchemaElement,
+		const Action::FSchema& ActionSchema,
+		const Action::FSchemaElement& ActionSchemaElement,
 		const TCHAR* IpAddress,
 		const uint32 Port,
 		const float Timeout,
-		const FPPOTrainerTrainingSettings& TrainingSettings,
-		const FPPOTrainerNetworkSettings& NetworkSettings,
-		const EPPOTrainerFlags TrainerFlags)
+		const FPPOTrainerTrainingSettings& TrainingSettings)
 	{
 		const int32 ObservationVectorDimensionNum = ReplayBuffer.GetObservations().Num<1>();
 		const int32 ActionVectorDimensionNum = ReplayBuffer.GetActions().Num<1>();
@@ -673,34 +753,31 @@ namespace UE::Learning
 		ConfigObject->SetStringField(TEXT("TrainerType"), TEXT("Network"));
 		ConfigObject->SetStringField(TEXT("TimeStamp"), *FDateTime::Now().ToFormattedString(TEXT("%Y-%m-%d_%H-%M-%S")));
 
-		ConfigObject->SetStringField(TEXT("PolicyNetworkClass"), PolicyNetwork.GetPythonClassName());
-		ConfigObject->SetStringField(TEXT("CriticNetworkClass"), CriticNetwork.GetPythonClassName());
-
+		ConfigObject->SetObjectField(TEXT("ObservationSchema"), Trainer::ConvertObservationSchemaToJSON(ObservationSchema, ObservationSchemaElement));
+		ConfigObject->SetObjectField(TEXT("ActionSchema"), Trainer::ConvertActionSchemaToJSON(ActionSchema, ActionSchemaElement));
 		ConfigObject->SetNumberField(TEXT("ObservationVectorDimensionNum"), ObservationVectorDimensionNum);
 		ConfigObject->SetNumberField(TEXT("ActionVectorDimensionNum"), ActionVectorDimensionNum);
 		ConfigObject->SetNumberField(TEXT("MemoryStateVectorDimensionNum"), MemoryStateVectorDimensionNum);
 		ConfigObject->SetNumberField(TEXT("MaxEpisodeNum"), ReplayBuffer.GetMaxEpisodeNum());
 		ConfigObject->SetNumberField(TEXT("MaxStepNum"), ReplayBuffer.GetMaxStepNum());
 
-		ConfigObject->SetNumberField(TEXT("PolicyNetworkByteNum"), PolicyNetwork.GetSerializationByteNum());
-		ConfigObject->SetNumberField(TEXT("PolicyActionNoiseMin"), NetworkSettings.PolicyActionNoiseMin);
-		ConfigObject->SetNumberField(TEXT("PolicyActionNoiseMax"), NetworkSettings.PolicyActionNoiseMax);
-
-		ConfigObject->SetNumberField(TEXT("CriticNetworkByteNum"), CriticNetwork.GetSerializationByteNum());
+		ConfigObject->SetNumberField(TEXT("PolicyNetworkByteNum"), PolicyNetwork.GetSnapshotByteNum());
+		ConfigObject->SetNumberField(TEXT("CriticNetworkByteNum"), CriticNetwork.GetSnapshotByteNum());
+		ConfigObject->SetNumberField(TEXT("EncoderNetworkByteNum"), EncoderNetwork.GetSnapshotByteNum());
+		ConfigObject->SetNumberField(TEXT("DecoderNetworkByteNum"), DecoderNetwork.GetSnapshotByteNum());
 
 		ConfigObject->SetNumberField(TEXT("IterationNum"), TrainingSettings.IterationNum);
 		ConfigObject->SetNumberField(TEXT("LearningRatePolicy"), TrainingSettings.LearningRatePolicy);
 		ConfigObject->SetNumberField(TEXT("LearningRateCritic"), TrainingSettings.LearningRateCritic);
 		ConfigObject->SetNumberField(TEXT("LearningRateDecay"), TrainingSettings.LearningRateDecay);
 		ConfigObject->SetNumberField(TEXT("WeightDecay"), TrainingSettings.WeightDecay);
-		ConfigObject->SetNumberField(TEXT("InitialActionScale"), TrainingSettings.InitialActionScale);
-		ConfigObject->SetNumberField(TEXT("InitialMemoryScale"), TrainingSettings.InitialMemoryScale);
-		ConfigObject->SetNumberField(TEXT("InitialCriticScale"), TrainingSettings.InitialCriticScale);
 		ConfigObject->SetNumberField(TEXT("PolicyBatchSize"), TrainingSettings.PolicyBatchSize);
 		ConfigObject->SetNumberField(TEXT("CriticBatchSize"), TrainingSettings.CriticBatchSize);
 		ConfigObject->SetNumberField(TEXT("PolicyWindow"), TrainingSettings.PolicyWindow);
 		ConfigObject->SetNumberField(TEXT("IterationsPerGather"), TrainingSettings.IterationsPerGather);
+		ConfigObject->SetNumberField(TEXT("CriticWarmupIterations"), TrainingSettings.CriticWarmupIterations);
 		ConfigObject->SetNumberField(TEXT("EpsilonClip"), TrainingSettings.EpsilonClip);
+		ConfigObject->SetNumberField(TEXT("ActionSurrogateWeight"), TrainingSettings.ActionSurrogateWeight);
 		ConfigObject->SetNumberField(TEXT("ActionRegularizationWeight"), TrainingSettings.ActionRegularizationWeight);
 		ConfigObject->SetNumberField(TEXT("ActionEntropyWeight"), TrainingSettings.ActionEntropyWeight);
 		ConfigObject->SetNumberField(TEXT("ReturnRegularizationWeight"), TrainingSettings.ReturnRegularizationWeight);
@@ -714,9 +791,7 @@ namespace UE::Learning
 		ConfigObject->SetNumberField(TEXT("DiscountFactor"), TrainingSettings.DiscountFactor);
 		ConfigObject->SetStringField(TEXT("Device"), Trainer::GetDeviceString(TrainingSettings.Device));
 		ConfigObject->SetBoolField(TEXT("UseTensorBoard"), TrainingSettings.bUseTensorboard);
-
-		ConfigObject->SetBoolField(TEXT("UseInitialPolicyNetwork"), (bool)(TrainerFlags & EPPOTrainerFlags::UseInitialPolicyNetwork));
-		ConfigObject->SetBoolField(TEXT("UseInitialCriticNetwork"), (bool)(TrainerFlags & EPPOTrainerFlags::UseInitialCriticNetwork));
+		ConfigObject->SetBoolField(TEXT("SaveSnapshots"), TrainingSettings.bSaveSnapshots);
 
 		FString JsonString;
 		TSharedRef<TJsonWriter<>> JsonWriter = TJsonWriterFactory<>::Create(&JsonString, 0);
@@ -724,8 +799,10 @@ namespace UE::Learning
 
 		// Allocate buffer to receive network data in
 
-		PolicyNetworkBuffer.SetNumUninitialized({ PolicyNetwork.GetSerializationByteNum() });
-		CriticNetworkBuffer.SetNumUninitialized({ CriticNetwork.GetSerializationByteNum() });
+		PolicyNetworkBuffer.SetNumUninitialized({ PolicyNetwork.GetSnapshotByteNum() });
+		CriticNetworkBuffer.SetNumUninitialized({ CriticNetwork.GetSnapshotByteNum() });
+		EncoderNetworkBuffer.SetNumUninitialized({ EncoderNetwork.GetSnapshotByteNum() });
+		DecoderNetworkBuffer.SetNumUninitialized({ DecoderNetwork.GetSnapshotByteNum() });
 
 		// Create Socket
 
@@ -788,45 +865,75 @@ namespace UE::Learning
 	}
 
 	ETrainerResponse FSocketPPOTrainer::RecvPolicy(
-		INeuralNetwork& OutNetwork,
+		ULearningNeuralNetworkData& OutNetwork,
 		const float Timeout,
 		FRWLock* NetworkLock,
 		const ELogSetting LogSettings)
 	{
-		return SocketTraining::RecvPolicy(*Socket, OutNetwork, PolicyNetworkBuffer, Timeout, NetworkLock, LogSettings);
+		return SocketTraining::RecvNetwork(*Socket, OutNetwork, PolicyNetworkBuffer, SocketTraining::ESignal::RecvPolicy, Timeout, NetworkLock, LogSettings);
 	}
 
 	ETrainerResponse FSocketPPOTrainer::RecvCritic(
-		INeuralNetwork& OutNetwork,
+		ULearningNeuralNetworkData& OutNetwork,
+		const float Timeout,
+		FRWLock* NetworkLock,
+		const ELogSetting LogSettings)
+	{
+		return SocketTraining::RecvNetwork(*Socket, OutNetwork, CriticNetworkBuffer, SocketTraining::ESignal::RecvCritic, Timeout, NetworkLock, LogSettings);
+	}
+
+	ETrainerResponse FSocketPPOTrainer::RecvEncoder(
+		ULearningNeuralNetworkData& OutNetwork,
 		const float Timeout,
 		FRWLock* NetworkLock,
 		const ELogSetting LogSettings)
 	{
-		UE_LEARNING_CHECKF(CriticNetworkBuffer.GetData(),
-			TEXT("Buffer for critic was not allocated. Were the correct EPPOTrainerFlags provided to construct the FSocketPPOTrainer?"));
+		return SocketTraining::RecvNetwork(*Socket, OutNetwork, EncoderNetworkBuffer, SocketTraining::ESignal::RecvEncoder, Timeout, NetworkLock, LogSettings);
+	}
 
-		return SocketTraining::RecvCritic(*Socket, OutNetwork, CriticNetworkBuffer, Timeout, NetworkLock, LogSettings);
+	ETrainerResponse FSocketPPOTrainer::RecvDecoder(
+		ULearningNeuralNetworkData& OutNetwork,
+		const float Timeout,
+		FRWLock* NetworkLock,
+		const ELogSetting LogSettings)
+	{
+		return SocketTraining::RecvNetwork(*Socket, OutNetwork, DecoderNetworkBuffer, SocketTraining::ESignal::RecvDecoder, Timeout, NetworkLock, LogSettings);
 	}
 
 	ETrainerResponse FSocketPPOTrainer::SendPolicy(
-		const INeuralNetwork& Network,
+		const ULearningNeuralNetworkData& Network,
 		const float Timeout,
 		FRWLock* NetworkLock,
 		const ELogSetting LogSettings)
 	{
-		return SocketTraining::SendPolicy(*Socket, PolicyNetworkBuffer, Network, Timeout, NetworkLock, LogSettings);
+		return SocketTraining::SendNetwork(*Socket, PolicyNetworkBuffer, SocketTraining::ESignal::SendPolicy, Network, Timeout, NetworkLock, LogSettings);
 	}
 
 	ETrainerResponse FSocketPPOTrainer::SendCritic(
-		const INeuralNetwork& Network,
+		const ULearningNeuralNetworkData& Network,
 		const float Timeout,
 		FRWLock* NetworkLock,
 		const ELogSetting LogSettings)
 	{
-		UE_LEARNING_CHECKF(CriticNetworkBuffer.GetData(),
-			TEXT("Buffer for critic was not allocated. Were the correct EPPOTrainerFlags provided to construct the FSocketPPOTrainer?"));
+		return SocketTraining::SendNetwork(*Socket, CriticNetworkBuffer, SocketTraining::ESignal::SendCritic, Network, Timeout, NetworkLock, LogSettings);
+	}
 
-		return SocketTraining::SendCritic(*Socket, CriticNetworkBuffer, Network, Timeout, NetworkLock, LogSettings);
+	ETrainerResponse FSocketPPOTrainer::SendEncoder(
+		const ULearningNeuralNetworkData& Network,
+		const float Timeout,
+		FRWLock* NetworkLock,
+		const ELogSetting LogSettings)
+	{
+		return SocketTraining::SendNetwork(*Socket, EncoderNetworkBuffer, SocketTraining::ESignal::SendEncoder, Network, Timeout, NetworkLock, LogSettings);
+	}
+
+	ETrainerResponse FSocketPPOTrainer::SendDecoder(
+		const ULearningNeuralNetworkData& Network,
+		const float Timeout,
+		FRWLock* NetworkLock,
+		const ELogSetting LogSettings)
+	{
+		return SocketTraining::SendNetwork(*Socket, DecoderNetworkBuffer, SocketTraining::ESignal::SendDecoder, Network, Timeout, NetworkLock, LogSettings);
 	}
 
 	ETrainerResponse FSocketPPOTrainer::SendExperience(
@@ -844,15 +951,18 @@ namespace UE::Learning
 			FReplayBuffer& ReplayBuffer,
 			FEpisodeBuffer& EpisodeBuffer,
 			FResetInstanceBuffer& ResetBuffer,
-			INeuralNetwork& PolicyNetwork,
-			INeuralNetwork& CriticNetwork,
+			ULearningNeuralNetworkData& PolicyNetwork,
+			ULearningNeuralNetworkData& CriticNetwork,
+			ULearningNeuralNetworkData& EncoderNetwork,
+			ULearningNeuralNetworkData& DecoderNetwork,
 			TLearningArrayView<2, float> ObservationVectorBuffer,
 			TLearningArrayView<2, float> ActionVectorBuffer,
 			TLearningArrayView<2, float> PreEvaluationMemoryStateVectorBuffer,
 			TLearningArrayView<2, float> MemoryStateVectorBuffer,
 			TLearningArrayView<1, float> RewardBuffer,
 			TLearningArrayView<1, ECompletionMode> CompletionBuffer,
-			const ECompletionMode EpisodeEndCompletionMode,
+			TLearningArrayView<1, ECompletionMode> EpisodeCompletionBuffer,
+			TLearningArrayView<1, ECompletionMode> AllCompletionBuffer,
 			const TFunctionRef<void(const FIndexSet Instances)> ResetFunction,
 			const TFunctionRef<void(const FIndexSet Instances)> ObservationFunction,
 			const TFunctionRef<void(const FIndexSet Instances)> PolicyFunction,
@@ -861,12 +971,15 @@ namespace UE::Learning
 			const TFunctionRef<void(const FIndexSet Instances)> RewardFunction,
 			const TFunctionRef<void(const FIndexSet Instances)> CompletionFunction,
 			const FIndexSet Instances,
-			const EPPOTrainerFlags TrainerFlags,
 			TAtomic<bool>* bRequestTrainingStopSignal,
 			FRWLock* PolicyNetworkLock,
 			FRWLock* CriticNetworkLock,
+			FRWLock* EncoderNetworkLock,
+			FRWLock* DecoderNetworkLock,
 			TAtomic<bool>* bPolicyNetworkUpdatedSignal,
 			TAtomic<bool>* bCriticNetworkUpdatedSignal,
+			TAtomic<bool>* bEncoderNetworkUpdatedSignal,
+			TAtomic<bool>* bDecoderNetworkUpdatedSignal,
 			const ELogSetting LogSettings)
 		{
 			TRACE_CPUPROFILER_EVENT_SCOPE(Learning::PPOTrainer::Train);
@@ -880,7 +993,7 @@ namespace UE::Learning
 				UE_LOG(LogLearning, Display, TEXT("Sending initial Policy..."));
 			}
 
-			Response = Trainer.SendPolicy(PolicyNetwork, UE::Learning::Trainer::DefaultTimeout, PolicyNetworkLock);
+			Response = Trainer.SendPolicy(PolicyNetwork, Trainer::DefaultTimeout, PolicyNetworkLock);
 
 			if (Response != ETrainerResponse::Success)
 			{
@@ -900,7 +1013,7 @@ namespace UE::Learning
 				UE_LOG(LogLearning, Display, TEXT("Sending initial Critic..."));
 			}
 
-			Response = Trainer.SendCritic(CriticNetwork, UE::Learning::Trainer::DefaultTimeout, CriticNetworkLock);
+			Response = Trainer.SendCritic(CriticNetwork, Trainer::DefaultTimeout, CriticNetworkLock);
 
 			if (Response != ETrainerResponse::Success)
 			{
@@ -913,60 +1026,44 @@ namespace UE::Learning
 				return Response;
 			}
 
-			if (!(bool)(TrainerFlags & EPPOTrainerFlags::UseInitialPolicyNetwork))
+			// Send initial Encoder
+
+			if (LogSettings != ELogSetting::Silent)
 			{
-				// Receive initial Policy
+				UE_LOG(LogLearning, Display, TEXT("Sending initial Encoder..."));
+			}
+
+			Response = Trainer.SendEncoder(EncoderNetwork, Trainer::DefaultTimeout, EncoderNetworkLock);
 
+			if (Response != ETrainerResponse::Success)
+			{
 				if (LogSettings != ELogSetting::Silent)
 				{
-					UE_LOG(LogLearning, Display, TEXT("Receiving initial Policy..."));
+					UE_LOG(LogLearning, Error, TEXT("Error sending initial encoder to trainer: %s. Check log for errors."), Trainer::GetResponseString(Response));
 				}
 
-				Response = Trainer.RecvPolicy(PolicyNetwork, UE::Learning::Trainer::DefaultTimeout, PolicyNetworkLock);
-
-				if (Response != ETrainerResponse::Success)
-				{
-					if (LogSettings != ELogSetting::Silent)
-					{
-						UE_LOG(LogLearning, Error, TEXT("Error receiving initial policy from trainer: %s. Check log for errors."), Trainer::GetResponseString(Response));
-					}
+				Trainer.Terminate();
+				return Response;
+			}
 
-					Trainer.Terminate();
-					return Response;
-				}
+			// Send initial Decoder
 
-				if (bPolicyNetworkUpdatedSignal)
-				{
-					*bPolicyNetworkUpdatedSignal = true;
-				}
+			if (LogSettings != ELogSetting::Silent)
+			{
+				UE_LOG(LogLearning, Display, TEXT("Sending initial Decoder..."));
 			}
 
-			if (!(bool)(TrainerFlags & EPPOTrainerFlags::UseInitialCriticNetwork))
-			{
-				// Receive initial Critic
+			Response = Trainer.SendDecoder(DecoderNetwork, Trainer::DefaultTimeout, DecoderNetworkLock);
 
+			if (Response != ETrainerResponse::Success)
+			{
 				if (LogSettings != ELogSetting::Silent)
 				{
-					UE_LOG(LogLearning, Display, TEXT("Receiving initial Critic..."));
+					UE_LOG(LogLearning, Error, TEXT("Error sending initial decoder to trainer: %s. Check log for errors."), Trainer::GetResponseString(Response));
 				}
 
-				Response = Trainer.RecvCritic(CriticNetwork, UE::Learning::Trainer::DefaultTimeout, CriticNetworkLock);
-
-				if (Response != ETrainerResponse::Success)
-				{
-					if (LogSettings != ELogSetting::Silent)
-					{
-						UE_LOG(LogLearning, Error, TEXT("Error receiving initial critic from trainer: %s. Check log for errors."), Trainer::GetResponseString(Response));
-					}
-
-					Trainer.Terminate();
-					return Response;
-				}
-
-				if (bCriticNetworkUpdatedSignal)
-				{
-					*bCriticNetworkUpdatedSignal = true;
-				}
+				Trainer.Terminate();
+				return Response;
 			}
 
 			// Start Training Loop
@@ -1009,7 +1106,8 @@ namespace UE::Learning
 						MemoryStateVectorBuffer,
 						RewardBuffer,
 						CompletionBuffer,
-						EpisodeEndCompletionMode,
+						EpisodeCompletionBuffer,
+						AllCompletionBuffer,
 						ResetFunction,
 						ObservationFunction,
 						PolicyFunction,
@@ -1019,7 +1117,7 @@ namespace UE::Learning
 						CompletionFunction,
 						Instances);
 
-					Response = Trainer.SendExperience(ReplayBuffer, UE::Learning::Trainer::DefaultTimeout);
+					Response = Trainer.SendExperience(ReplayBuffer, Trainer::DefaultTimeout);
 
 					if (Response != ETrainerResponse::Success)
 					{
@@ -1033,7 +1131,9 @@ namespace UE::Learning
 					}
 				}
 
-				Response = Trainer.RecvPolicy(PolicyNetwork, UE::Learning::Trainer::DefaultTimeout, PolicyNetworkLock);
+				// Update Policy
+
+				Response = Trainer.RecvPolicy(PolicyNetwork, Trainer::DefaultTimeout, PolicyNetworkLock);
 
 				if (Response == ETrainerResponse::Completed)
 				{
@@ -1057,7 +1157,9 @@ namespace UE::Learning
 					*bPolicyNetworkUpdatedSignal = true;
 				}
 
-				Response = Trainer.RecvCritic(CriticNetwork, UE::Learning::Trainer::DefaultTimeout, CriticNetworkLock);
+				// Update Critic
+
+				Response = Trainer.RecvCritic(CriticNetwork, Trainer::DefaultTimeout, CriticNetworkLock);
 
 				if (Response != ETrainerResponse::Success)
 				{
@@ -1072,11 +1174,47 @@ namespace UE::Learning
 				{
 					*bCriticNetworkUpdatedSignal = true;
 				}
+
+				// Update Encoder
+
+				Response = Trainer.RecvEncoder(EncoderNetwork, Trainer::DefaultTimeout, EncoderNetworkLock);
+
+				if (Response != ETrainerResponse::Success)
+				{
+					if (LogSettings != ELogSetting::Silent)
+					{
+						UE_LOG(LogLearning, Error, TEXT("Error receiving encoder from trainer: %s. Check log for errors."), Trainer::GetResponseString(Response));
+					}
+					break;
+				}
+
+				if (bEncoderNetworkUpdatedSignal)
+				{
+					*bEncoderNetworkUpdatedSignal = true;
+				}
+
+				// Update Decoder
+
+				Response = Trainer.RecvDecoder(DecoderNetwork, Trainer::DefaultTimeout, DecoderNetworkLock);
+
+				if (Response != ETrainerResponse::Success)
+				{
+					if (LogSettings != ELogSetting::Silent)
+					{
+						UE_LOG(LogLearning, Error, TEXT("Error receiving decoder from trainer: %s. Check log for errors."), Trainer::GetResponseString(Response));
+					}
+					break;
+				}
+
+				if (bDecoderNetworkUpdatedSignal)
+				{
+					*bDecoderNetworkUpdatedSignal = true;
+				}
 			}
 
 			// Allow some time for trainer to shut down gracefully before we kill it...
 			
-			Response = Trainer.Wait(UE::Learning::Trainer::DefaultTimeout);
+			Response = Trainer.Wait(Trainer::DefaultTimeout);
 
 			if (Response != ETrainerResponse::Success)
 			{
diff --git a/Engine/Plugins/Experimental/LearningAgents/Source/LearningTraining/Private/LearningSharedMemoryTraining.cpp b/Engine/Plugins/Experimental/LearningAgents/Source/LearningTraining/Private/LearningSharedMemoryTraining.cpp
--- a/Engine/Plugins/Experimental/LearningAgents/Source/LearningTraining/Private/LearningSharedMemoryTraining.cpp
+++ b/Engine/Plugins/Experimental/LearningAgents/Source/LearningTraining/Private/LearningSharedMemoryTraining.cpp
@@ -11,19 +11,6 @@
 
 namespace UE::Learning::SharedMemoryTraining
 {
-	enum class EControls : uint8
-	{
-		ExperienceEpisodeNum = 0,
-		ExperienceStepNum = 1,
-		ExperienceSignal = 2,
-		PolicySignal = 3,
-		CriticSignal = 4,
-		CompleteSignal = 5,
-		StopSignal = 6,
-
-		ControlNum = 8,
-	};
-
 	uint8 GetControlNum()
 	{
 		return (uint8)EControls::ControlNum;
@@ -40,10 +27,11 @@ namespace UE::Learning::SharedMemoryTraining
 		return Controls[(uint8)EControls::PolicySignal] || Controls[(uint8)EControls::CompleteSignal];
 	}
 
-	ETrainerResponse RecvPolicy(
+	ETrainerResponse RecvNetwork(
 		TLearningArrayView<1, volatile int32> Controls,
-		INeuralNetwork& OutNetwork,
-		const TLearningArrayView<1, const uint8> Policy,
+		ULearningNeuralNetworkData& OutNetwork,
+		const EControls Signal,
+		const TLearningArrayView<1, const uint8> NetworkData,
 		const float Timeout,
 		FRWLock* NetworkLock,
 		const ELogSetting LogSettings)
@@ -51,8 +39,8 @@ namespace UE::Learning::SharedMemoryTraining
 		const float SleepTime = 0.001f;
 		float WaitTime = 0.0f;
 
-		// Wait until the policy is done being written by the sub-process
-		while (!Controls[(uint8)EControls::PolicySignal])
+		// Wait until the network is done being written by the sub-process
+		while (!Controls[(uint8)Signal])
 		{
 			// Check if Completed Signal has been raised
 			if (Controls[(uint8)EControls::CompleteSignal])
@@ -73,79 +61,44 @@ namespace UE::Learning::SharedMemoryTraining
 
 		if (LogSettings != ELogSetting::Silent)
 		{
-			UE_LOG(LogLearning, Display, TEXT("Pulling Policy..."));
+			UE_LOG(LogLearning, Display, TEXT("Pulling network..."));
 		}
 
-		// Read the policy
+		// Read the network
 		bool bSuccess = false;
 		{
 			FScopeNullableWriteLock ScopeLock(NetworkLock);
-			int32 Offset = 0;
-			bSuccess = OutNetwork.DeserializeFromBytes(Offset, Policy);
-
-			if (!bSuccess)
-			{
-				UE_LOG(LogLearning, Error, TEXT("Error receiving Policy network. Format invalid."));
-			}
-		}
-
-		// Confirm we have read the policy
-		Controls[(uint8)EControls::PolicySignal] = false;
-
-		return bSuccess ? ETrainerResponse::Success : ETrainerResponse::Unexpected;
-	}
-
-	ETrainerResponse RecvCritic(
-		TLearningArrayView<1, volatile int32> Controls,
-		INeuralNetwork& OutNetwork,
-		const TLearningArrayView<1, const uint8> Critic,
-		const float Timeout,
-		FRWLock* NetworkLock,
-		const ELogSetting LogSettings)
-	{
-		const float SleepTime = 0.001f;
-		float WaitTime = 0.0f;
 
-		// Wait until the critic is done being written by the sub-process
-		while (!Controls[(uint8)EControls::CriticSignal])
-		{
-			FPlatformProcess::Sleep(SleepTime);
-			WaitTime += SleepTime;
-
-			if (WaitTime > Timeout)
+			if (NetworkData.Num() != OutNetwork.GetSnapshotByteNum())
 			{
-				return ETrainerResponse::Timeout;
+				UE_LOG(LogLearning, Error, TEXT("Error receiving network. Incorrect buffer size. Buffer is %i bytes, expected %i."), NetworkData.Num(), OutNetwork.GetSnapshotByteNum());
+				bSuccess = false;
 			}
-		}
-
-		if (LogSettings != ELogSetting::Silent)
-		{
-			UE_LOG(LogLearning, Display, TEXT("Pulling Critic..."));
-		}
-
-		// Read the critic
-		bool bSuccess = false;
-		{
-			FScopeNullableWriteLock ScopeLock(NetworkLock);
-			int32 Offset = 0;
-			bSuccess = OutNetwork.DeserializeFromBytes(Offset, Critic);
-
-			if (!bSuccess)
+			else
 			{
-				UE_LOG(LogLearning, Error, TEXT("Error receiving Critic network. Format invalid."));
+				if (!OutNetwork.LoadFromSnapshot(MakeArrayView(NetworkData.GetData(), NetworkData.Num())))
+				{
+					UE_LOG(LogLearning, Error, TEXT("Error receiving network. Invalid Format."));
+					bSuccess = false;
+				}
+				else
+				{
+					bSuccess = true;
+				}
 			}
 		}
 
-		// Confirm we have read the critic
-		Controls[(uint8)EControls::CriticSignal] = false;
+		// Confirm we have read the network
+		Controls[(uint8)Signal] = false;
 
 		return bSuccess ? ETrainerResponse::Success : ETrainerResponse::Unexpected;
 	}
 
-	ETrainerResponse SendPolicy(
+	ETrainerResponse SendNetwork(
 		TLearningArrayView<1, volatile int32> Controls,
-		TLearningArrayView<1, uint8> Policy,
-		const INeuralNetwork& Network,
+		TLearningArrayView<1, uint8> NetworkData,
+		const EControls Signal,
+		const ULearningNeuralNetworkData& Network,
 		const float Timeout,
 		FRWLock* NetworkLock,
 		const ELogSetting LogSettings)
@@ -154,7 +107,7 @@ namespace UE::Learning::SharedMemoryTraining
 		float WaitTime = 0.0f;
 
 		// Wait until the policy is requested by the sub-process
-		while (!Controls[(uint8)EControls::PolicySignal])
+		while (!Controls[(uint8)Signal])
 		{
 			FPlatformProcess::Sleep(SleepTime);
 			WaitTime += SleepTime;
@@ -167,61 +120,29 @@ namespace UE::Learning::SharedMemoryTraining
 
 		if (LogSettings != ELogSetting::Silent)
 		{
-			UE_LOG(LogLearning, Display, TEXT("Pushing Policy..."));
+			UE_LOG(LogLearning, Display, TEXT("Pushing network..."));
 		}
 
-		// Write the policy
+		// Write the network
+		bool bSuccess = false;
 		{
 			FScopeNullableReadLock ScopeLock(NetworkLock);
-			int32 Offset = 0;
-			Network.SerializeToBytes(Offset, Policy);
-		}
-
-		// Confirm we have written the policy
-		Controls[(uint8)EControls::PolicySignal] = false;
-
-		return ETrainerResponse::Success;
-	}
-
-	ETrainerResponse SendCritic(
-		TLearningArrayView<1, volatile int32> Controls,
-		TLearningArrayView<1, uint8> Critic,
-		const INeuralNetwork& Network,
-		const float Timeout,
-		FRWLock* NetworkLock,
-		const ELogSetting LogSettings)
-	{
-		const float SleepTime = 0.001f;
-		float WaitTime = 0.0f;
-
-		// Wait until the critic is requested by the sub-process
-		while (!Controls[(uint8)EControls::CriticSignal])
-		{
-			FPlatformProcess::Sleep(SleepTime);
-			WaitTime += SleepTime;
-
-			if (WaitTime > Timeout)
+			if (NetworkData.Num() != Network.GetSnapshotByteNum())
 			{
-				return ETrainerResponse::Timeout;
+				UE_LOG(LogLearning, Error, TEXT("Error sending network. Incorrect buffer size. Buffer is %i bytes, expected %i."), NetworkData.Num(), Network.GetSnapshotByteNum());
+				bSuccess = false;
+			}
+			else
+			{
+				Network.SaveToSnapshot(MakeArrayView(NetworkData.GetData(), NetworkData.Num()));
+				bSuccess = true;
 			}
 		}
 
-		if (LogSettings != ELogSetting::Silent)
-		{
-			UE_LOG(LogLearning, Display, TEXT("Pushing Critic..."));
-		}
-
-		// Write the critic
-		{
-			FScopeNullableReadLock ScopeLock(NetworkLock);
-			int32 Offset = 0;
-			Network.SerializeToBytes(Offset, Critic);
-		}
-
-		// Confirm we have written the critic
-		Controls[(uint8)EControls::CriticSignal] = false;
+		// Confirm we have written the network
+		Controls[(uint8)Signal] = false;
 
-		return ETrainerResponse::Success;
+		return bSuccess ? ETrainerResponse::Success : ETrainerResponse::Unexpected;
 	}
 
 	ETrainerResponse SendExperience(
diff --git a/Engine/Plugins/Experimental/LearningAgents/Source/LearningTraining/Private/LearningSharedMemoryTraining.h b/Engine/Plugins/Experimental/LearningAgents/Source/LearningTraining/Private/LearningSharedMemoryTraining.h
--- a/Engine/Plugins/Experimental/LearningAgents/Source/LearningTraining/Private/LearningSharedMemoryTraining.h
+++ b/Engine/Plugins/Experimental/LearningAgents/Source/LearningTraining/Private/LearningSharedMemoryTraining.h
@@ -4,28 +4,37 @@
 
 #include "LearningTrainer.h"
 
+class ULearningNeuralNetworkData;
+
 namespace UE::Learning
 {
 	enum class ECompletionMode : uint8;
-	struct INeuralNetwork;
 	struct FReplayBuffer;
 
 	namespace SharedMemoryTraining
 	{
-		LEARNINGTRAINING_API uint8 GetControlNum();
+		enum class EControls : uint8
+		{
+			ExperienceEpisodeNum	= 0,
+			ExperienceStepNum		= 1,
+			ExperienceSignal		= 2,
+			PolicySignal			= 3,
+			CriticSignal			= 4,
+			EncoderSignal			= 5,
+			DecoderSignal			= 6,
+			CompleteSignal			= 7,
+			StopSignal				= 8,
 
-		LEARNINGTRAINING_API ETrainerResponse RecvPolicy(
-			TLearningArrayView<1, volatile int32> Controls,
-			INeuralNetwork& OutNetwork,
-			const TLearningArrayView<1, const uint8> Policy,
-			const float Timeout = Trainer::DefaultTimeout,
-			FRWLock* NetworkLock = nullptr,
-			const ELogSetting LogSettings = Trainer::DefaultLogSettings);
+			ControlNum				= 9,
+		};
 
-		LEARNINGTRAINING_API ETrainerResponse RecvCritic(
+		LEARNINGTRAINING_API uint8 GetControlNum();
+
+		LEARNINGTRAINING_API ETrainerResponse RecvNetwork(
 			TLearningArrayView<1, volatile int32> Controls,
-			INeuralNetwork& OutNetwork,
-			const TLearningArrayView<1, const uint8> Critic,
+			ULearningNeuralNetworkData& OutNetwork,
+			const EControls Signal,
+			const TLearningArrayView<1, const uint8> NetworkData,
 			const float Timeout = Trainer::DefaultTimeout,
 			FRWLock* NetworkLock = nullptr,
 			const ELogSetting LogSettings = Trainer::DefaultLogSettings);
@@ -35,18 +44,11 @@ namespace UE::Learning
 
 		LEARNINGTRAINING_API bool HasPolicyOrCompleted(TLearningArrayView<1, volatile int32> Controls);
 
-		LEARNINGTRAINING_API ETrainerResponse SendPolicy(
-			TLearningArrayView<1, volatile int32> Controls,
-			TLearningArrayView<1, uint8> Policy,
-			const INeuralNetwork& Network,
-			const float Timeout = Trainer::DefaultTimeout,
-			FRWLock* NetworkLock = nullptr,
-			const ELogSetting LogSettings = Trainer::DefaultLogSettings);
-
-		LEARNINGTRAINING_API ETrainerResponse SendCritic(
+		LEARNINGTRAINING_API ETrainerResponse SendNetwork(
 			TLearningArrayView<1, volatile int32> Controls,
-			TLearningArrayView<1, uint8> Critic,
-			const INeuralNetwork& Network,
+			TLearningArrayView<1, uint8> NetworkData,
+			const EControls Signal,
+			const ULearningNeuralNetworkData& Network,
 			const float Timeout = Trainer::DefaultTimeout,
 			FRWLock* NetworkLock = nullptr,
 			const ELogSetting LogSettings = Trainer::DefaultLogSettings);
diff --git a/Engine/Plugins/Experimental/LearningAgents/Source/LearningTraining/Private/LearningSocketTraining.cpp b/Engine/Plugins/Experimental/LearningAgents/Source/LearningTraining/Private/LearningSocketTraining.cpp
--- a/Engine/Plugins/Experimental/LearningAgents/Source/LearningTraining/Private/LearningSocketTraining.cpp
+++ b/Engine/Plugins/Experimental/LearningAgents/Source/LearningTraining/Private/LearningSocketTraining.cpp
@@ -11,19 +11,6 @@
 
 namespace UE::Learning::SocketTraining
 {
-	enum class ESignal : uint8
-	{
-		Invalid = 0,
-		SendConfig = 1,
-		SendExperience = 2,
-		RecvPolicy = 3,
-		SendPolicy = 4,
-		RecvCritic = 5,
-		SendCritic = 6,
-		RecvComplete = 7,
-		SendStop = 8,
-	};
-
 	ETrainerResponse WaitForConnection(FSocket& Socket, const float Timeout)
 	{
 		float WaitTime = 0.0f;
@@ -86,17 +73,18 @@ namespace UE::Learning::SocketTraining
 		}
 	}
 
-	ETrainerResponse RecvPolicy(
+	ETrainerResponse RecvNetwork(
 		FSocket& Socket,
-		INeuralNetwork& OutNetwork,
+		ULearningNeuralNetworkData& OutNetwork,
 		TLearningArrayView<1, uint8> OutNetworkBuffer,
+		const ESignal NetworkSignal,
 		const float Timeout,
 		FRWLock* NetworkLock,
 		const ELogSetting LogSettings)
 	{
 		if (LogSettings != ELogSetting::Silent)
 		{
-			UE_LOG(LogLearning, Display, TEXT("Pulling Policy..."));
+			UE_LOG(LogLearning, Display, TEXT("Pulling Network..."));
 		}
 
 		uint8 Signal = (uint8)ESignal::Invalid;
@@ -108,7 +96,7 @@ namespace UE::Learning::SocketTraining
 			return ETrainerResponse::Completed;
 		}
 
-		if (Signal != (uint8)ESignal::RecvPolicy)
+		if (Signal != (uint8)NetworkSignal)
 		{
 			return ETrainerResponse::Unexpected;
 		}
@@ -119,52 +107,23 @@ namespace UE::Learning::SocketTraining
 		bool bSuccess = false;
 		{
 			FScopeNullableWriteLock ScopeLock(NetworkLock);
-			int32 Offset = 0;
-			bSuccess = OutNetwork.DeserializeFromBytes(Offset, OutNetworkBuffer);
 
-			if (!bSuccess)
+			if (OutNetworkBuffer.Num() != OutNetwork.GetSnapshotByteNum())
 			{
-				UE_LOG(LogLearning, Error, TEXT("Error receiving Policy network. Format invalid."));
+				UE_LOG(LogLearning, Error, TEXT("Error receiving network. Incorrect buffer size. Buffer is %i bytes, expected %i."), OutNetworkBuffer.Num(), OutNetwork.GetSnapshotByteNum());
+				bSuccess = false;
 			}
-		}
-
-		return bSuccess ? ETrainerResponse::Success : ETrainerResponse::Unexpected;
-	}
-
-	ETrainerResponse RecvCritic(
-		FSocket& Socket,
-		INeuralNetwork& OutNetwork,
-		TLearningArrayView<1, uint8> OutNetworkBuffer,
-		const float Timeout,
-		FRWLock* NetworkLock,
-		const ELogSetting LogSettings)
-	{
-		if (LogSettings != ELogSetting::Silent)
-		{
-			UE_LOG(LogLearning, Display, TEXT("Pulling Critic..."));
-		}
-
-		uint8 Signal = (uint8)ESignal::Invalid;
-		ETrainerResponse Response = RecvWithTimeout(Socket, &Signal, 1, Timeout);
-		if (Response != ETrainerResponse::Success) { return Response; }
-
-		if (Signal != (uint8)ESignal::RecvCritic)
-		{
-			return ETrainerResponse::Unexpected;
-		}
-
-		Response = RecvWithTimeout(Socket, OutNetworkBuffer.GetData(), OutNetworkBuffer.Num(), Timeout);
-		if (Response != ETrainerResponse::Success) { return Response; }
-
-		bool bSuccess = false;
-		{
-			FScopeNullableWriteLock ScopeLock(NetworkLock);
-			int32 Offset = 0;
-			bSuccess = OutNetwork.DeserializeFromBytes(Offset, OutNetworkBuffer);
-
-			if (!bSuccess)
+			else
 			{
-				UE_LOG(LogLearning, Error, TEXT("Error receiving Critic network. Format invalid."));
+				if (!OutNetwork.LoadFromSnapshot(MakeArrayView(OutNetworkBuffer.GetData(), OutNetworkBuffer.Num())))
+				{
+					UE_LOG(LogLearning, Error, TEXT("Error receiving network. Invalid Format."));
+					bSuccess = false;
+				}
+				else
+				{
+					bSuccess = true;
+				}
 			}
 		}
 
@@ -231,62 +190,43 @@ namespace UE::Learning::SocketTraining
 		return Socket.HasPendingData(PendingDataSize);
 	}
 
-	ETrainerResponse SendPolicy(
+	ETrainerResponse SendNetwork(
 		FSocket& Socket,
 		TLearningArrayView<1, uint8> NetworkBuffer,
-		const INeuralNetwork& Network,
+		const ESignal NetworkSignal,
+		const ULearningNeuralNetworkData& Network,
 		const float Timeout,
 		FRWLock* NetworkLock,
 		const ELogSetting LogSettings)
 	{
 		if (LogSettings != ELogSetting::Silent)
 		{
-			UE_LOG(LogLearning, Display, TEXT("Pushing Policy..."));
-		}
-
-		{
-			FScopeNullableReadLock ScopeLock(NetworkLock);
-			int32 Offset = 0;
-			Network.SerializeToBytes(Offset, NetworkBuffer);
-		}
-
-		const uint8 Signal = (uint8)ESignal::SendPolicy;
-		ETrainerResponse Response = SendWithTimeout(Socket, &Signal, 1, Timeout);
-		if (Response != ETrainerResponse::Success) { return Response; }
-
-		Response = SendWithTimeout(Socket, NetworkBuffer.GetData(), NetworkBuffer.Num(), Timeout);
-		if (Response != ETrainerResponse::Success) { return Response; }
-
-		return ETrainerResponse::Success;
-	}
-
-	ETrainerResponse SendCritic(
-		FSocket& Socket,
-		TLearningArrayView<1, uint8> NetworkBuffer,
-		const INeuralNetwork& Network,
-		const float Timeout,
-		FRWLock* NetworkLock,
-		const ELogSetting LogSettings)
-	{
-		if (LogSettings != ELogSetting::Silent)
-		{
-			UE_LOG(LogLearning, Display, TEXT("Pushing Critic..."));
+			UE_LOG(LogLearning, Display, TEXT("Pushing Network..."));
 		}
 
+		bool bSuccess = false;
 		{
 			FScopeNullableReadLock ScopeLock(NetworkLock);
-			int32 Offset = 0;
-			Network.SerializeToBytes(Offset, NetworkBuffer);
+			if (NetworkBuffer.Num() != Network.GetSnapshotByteNum())
+			{
+				UE_LOG(LogLearning, Error, TEXT("Error sending network. Incorrect buffer size. Buffer is %i bytes, expected %i."), NetworkBuffer.Num(), Network.GetSnapshotByteNum());
+				bSuccess = false;
+			}
+			else
+			{
+				Network.SaveToSnapshot(MakeArrayView(NetworkBuffer.GetData(), NetworkBuffer.Num()));
+				bSuccess = true;
+			}
 		}
 
-		const uint8 Signal = (uint8)ESignal::SendCritic;
+		const uint8 Signal = (uint8)NetworkSignal;
 		ETrainerResponse Response = SendWithTimeout(Socket, &Signal, 1, Timeout);
 		if (Response != ETrainerResponse::Success) { return Response; }
 
 		Response = SendWithTimeout(Socket, NetworkBuffer.GetData(), NetworkBuffer.Num(), Timeout);
 		if (Response != ETrainerResponse::Success) { return Response; }
 
-		return ETrainerResponse::Success;
+		return bSuccess ? Response : ETrainerResponse::Unexpected;
 	}
 
 	ETrainerResponse SendExperience(
diff --git a/Engine/Plugins/Experimental/LearningAgents/Source/LearningTraining/Private/LearningSocketTraining.h b/Engine/Plugins/Experimental/LearningAgents/Source/LearningTraining/Private/LearningSocketTraining.h
--- a/Engine/Plugins/Experimental/LearningAgents/Source/LearningTraining/Private/LearningSocketTraining.h
+++ b/Engine/Plugins/Experimental/LearningAgents/Source/LearningTraining/Private/LearningSocketTraining.h
@@ -5,14 +5,31 @@
 #include "LearningTrainer.h"
 
 class FSocket;
+class ULearningNeuralNetworkData;
 
 namespace UE::Learning
 {
-	struct INeuralNetwork;
 	struct FReplayBuffer;
 
 	namespace SocketTraining
 	{
+		enum class ESignal : uint8
+		{
+			Invalid				= 0,
+			SendConfig			= 1,
+			SendExperience		= 2,
+			RecvPolicy			= 3,
+			SendPolicy			= 4,
+			RecvCritic			= 5,
+			SendCritic			= 6,
+			RecvEncoder			= 7,
+			SendEncoder			= 8,
+			RecvDecoder			= 9,
+			SendDecoder			= 10,
+			RecvComplete		= 11,
+			SendStop			= 12,
+		};
+
 		LEARNINGTRAINING_API ETrainerResponse WaitForConnection(
 			FSocket& Socket, 
 			const float Timeout = Trainer::DefaultTimeout);
@@ -23,18 +40,11 @@ namespace UE::Learning
 			const int32 ByteNum, 
 			const float Timeout = Trainer::DefaultTimeout);
 
-		LEARNINGTRAINING_API ETrainerResponse RecvPolicy(
-			FSocket& Socket,
-			INeuralNetwork& OutNetwork,
-			TLearningArrayView<1, uint8> OutNetworkBuffer,
-			const float Timeout = Trainer::DefaultTimeout,
-			FRWLock* NetworkLock = nullptr,
-			const ELogSetting LogSettings = Trainer::DefaultLogSettings);
-
-		LEARNINGTRAINING_API ETrainerResponse RecvCritic(
+		LEARNINGTRAINING_API ETrainerResponse RecvNetwork(
 			FSocket& Socket,
-			INeuralNetwork& OutNetwork,
+			ULearningNeuralNetworkData& OutNetwork,
 			TLearningArrayView<1, uint8> OutNetworkBuffer,
+			const ESignal NetworkSignal,
 			const float Timeout = Trainer::DefaultTimeout,
 			FRWLock* NetworkLock = nullptr,
 			const ELogSetting LogSettings = Trainer::DefaultLogSettings);
@@ -56,18 +66,11 @@ namespace UE::Learning
 
 		LEARNINGTRAINING_API bool HasPolicyOrCompleted(FSocket& Socket);
 
-		LEARNINGTRAINING_API ETrainerResponse SendPolicy(
-			FSocket& Socket,
-			TLearningArrayView<1, uint8> NetworkBuffer,
-			const INeuralNetwork& Network,
-			const float Timeout = Trainer::DefaultTimeout,
-			FRWLock* NetworkLock = nullptr,
-			const ELogSetting LogSettings = Trainer::DefaultLogSettings);
-
-		LEARNINGTRAINING_API ETrainerResponse SendCritic(
+		LEARNINGTRAINING_API ETrainerResponse SendNetwork(
 			FSocket& Socket,
 			TLearningArrayView<1, uint8> NetworkBuffer,
-			const INeuralNetwork& Network,
+			const ESignal NetworkSignal,
+			const ULearningNeuralNetworkData& Network,
 			const float Timeout = Trainer::DefaultTimeout,
 			FRWLock* NetworkLock = nullptr,
 			const ELogSetting LogSettings = Trainer::DefaultLogSettings);
diff --git a/Engine/Plugins/Experimental/LearningAgents/Source/LearningTraining/Private/LearningTrainer.cpp b/Engine/Plugins/Experimental/LearningAgents/Source/LearningTraining/Private/LearningTrainer.cpp
--- a/Engine/Plugins/Experimental/LearningAgents/Source/LearningTraining/Private/LearningTrainer.cpp
+++ b/Engine/Plugins/Experimental/LearningAgents/Source/LearningTraining/Private/LearningTrainer.cpp
@@ -2,10 +2,261 @@
 
 #include "LearningTrainer.h"
 
+#include "LearningObservation.h"
+#include "LearningAction.h"
+
 #include "HAL/Platform.h"
+#include "Dom/JsonObject.h"
 
 namespace UE::Learning::Trainer
 {
+	TSharedPtr<FJsonObject> ConvertObservationSchemaToJSON(
+		const Observation::FSchema& ObservationSchema,
+		const Observation::FSchemaElement& ObservationSchemaElement)
+	{
+		TSharedPtr<FJsonObject> Object = MakeShared<FJsonObject>();
+		Object->SetNumberField(TEXT("VectorSize"), ObservationSchema.GetObservationVectorSize(ObservationSchemaElement));
+		Object->SetNumberField(TEXT("EncodedSize"), ObservationSchema.GetEncodedVectorSize(ObservationSchemaElement));
+
+		switch (ObservationSchema.GetType(ObservationSchemaElement))
+		{
+		case Observation::EType::Null:
+		{
+			Object->SetStringField(TEXT("Type"), TEXT("Null"));
+			break;
+		}
+
+		case Observation::EType::Continuous:
+		{
+			const Observation::FSchemaContinuousParameters Parameters = ObservationSchema.GetContinuous(ObservationSchemaElement);
+			Object->SetStringField(TEXT("Type"), TEXT("Continuous"));
+			Object->SetNumberField(TEXT("Num"), Parameters.Num);
+			break;
+		}
+
+		case Observation::EType::And:
+		{
+			const Observation::FSchemaAndParameters Parameters = ObservationSchema.GetAnd(ObservationSchemaElement);
+			
+			Object->SetStringField(TEXT("Type"), TEXT("And"));
+
+			TSharedPtr<FJsonObject> SubObject = MakeShared<FJsonObject>();
+			for (int32 SubElementIdx = 0; SubElementIdx < Parameters.Elements.Num(); SubElementIdx++)
+			{
+				TSharedPtr<FJsonObject> SubElement = ConvertObservationSchemaToJSON(ObservationSchema, Parameters.Elements[SubElementIdx]);
+				SubElement->SetNumberField(TEXT("Index"), SubElementIdx);
+				SubObject->SetObjectField(Parameters.ElementNames[SubElementIdx].ToString(), SubElement);
+			}
+
+			Object->SetObjectField(TEXT("Elements"), SubObject);
+			break;
+		}
+
+		case Observation::EType::OrExclusive:
+		{
+			const Observation::FSchemaOrExclusiveParameters Parameters = ObservationSchema.GetOrExclusive(ObservationSchemaElement);
+
+			Object->SetStringField(TEXT("Type"), TEXT("OrExclusive"));
+			Object->SetNumberField(TEXT("EncodingSize"), Parameters.EncodingSize);
+
+			TSharedPtr<FJsonObject> SubObject = MakeShared<FJsonObject>();
+			for (int32 SubElementIdx = 0; SubElementIdx < Parameters.Elements.Num(); SubElementIdx++)
+			{
+				TSharedPtr<FJsonObject> SubElement = ConvertObservationSchemaToJSON(ObservationSchema, Parameters.Elements[SubElementIdx]);
+				SubElement->SetNumberField(TEXT("Index"), SubElementIdx);
+				SubObject->SetObjectField(Parameters.ElementNames[SubElementIdx].ToString(), SubElement);
+			}
+
+			Object->SetObjectField(TEXT("Elements"), SubObject);
+			break;
+		}
+
+		case Observation::EType::OrInclusive:
+		{
+			const Observation::FSchemaOrInclusiveParameters Parameters = ObservationSchema.GetOrInclusive(ObservationSchemaElement);
+
+			Object->SetStringField(TEXT("Type"), TEXT("OrInclusive"));
+			Object->SetNumberField(TEXT("AttentionEncodingSize"), Parameters.AttentionEncodingSize);
+			Object->SetNumberField(TEXT("AttentionHeadNum"), Parameters.AttentionHeadNum);
+			Object->SetNumberField(TEXT("ValueEncodingSize"), Parameters.ValueEncodingSize);
+
+			TSharedPtr<FJsonObject> SubObject = MakeShared<FJsonObject>();
+			for (int32 SubElementIdx = 0; SubElementIdx < Parameters.Elements.Num(); SubElementIdx++)
+			{
+				TSharedPtr<FJsonObject> SubElement = ConvertObservationSchemaToJSON(ObservationSchema, Parameters.Elements[SubElementIdx]);
+				SubElement->SetNumberField(TEXT("Index"), SubElementIdx);
+				SubObject->SetObjectField(Parameters.ElementNames[SubElementIdx].ToString(), SubElement);
+			}
+
+			Object->SetObjectField(TEXT("Elements"), SubObject);
+			break;
+		}
+
+		case Observation::EType::Array:
+		{
+			const Observation::FSchemaArrayParameters Parameters = ObservationSchema.GetArray(ObservationSchemaElement);
+
+			Object->SetStringField(TEXT("Type"), TEXT("Array"));
+			Object->SetNumberField(TEXT("Num"), Parameters.Num);
+			Object->SetObjectField(TEXT("Element"), ConvertObservationSchemaToJSON(ObservationSchema, Parameters.Element));
+			break;
+		}
+
+		case Observation::EType::Set:
+		{
+			const Observation::FSchemaSetParameters Parameters = ObservationSchema.GetSet(ObservationSchemaElement);
+
+			Object->SetStringField(TEXT("Type"), TEXT("Set"));
+			Object->SetNumberField(TEXT("MaxNum"), Parameters.MaxNum);
+			Object->SetNumberField(TEXT("AttentionEncodingSize"), Parameters.AttentionEncodingSize);
+			Object->SetNumberField(TEXT("AttentionHeadNum"), Parameters.AttentionHeadNum);
+			Object->SetNumberField(TEXT("ValueEncodingSize"), Parameters.ValueEncodingSize);
+			Object->SetObjectField(TEXT("Element"), ConvertObservationSchemaToJSON(ObservationSchema, Parameters.Element));
+			break;
+		}
+
+		case Observation::EType::Encoding:
+		{
+			const Observation::FSchemaEncodingParameters Parameters = ObservationSchema.GetEncoding(ObservationSchemaElement);
+
+			Object->SetStringField(TEXT("Type"), TEXT("Encoding"));
+			Object->SetNumberField(TEXT("EncodingSize"), Parameters.EncodingSize);
+			Object->SetObjectField(TEXT("Element"), ConvertObservationSchemaToJSON(ObservationSchema, Parameters.Element));
+			break;
+		}
+
+		default:
+			UE_LEARNING_NOT_IMPLEMENTED();
+		}
+
+		return Object;
+	}
+
+	TSharedPtr<FJsonObject> ConvertActionSchemaToJSON(
+		const Action::FSchema& ActionSchema,
+		const Action::FSchemaElement& ActionSchemaElement)
+	{
+		TSharedPtr<FJsonObject> Object = MakeShared<FJsonObject>();
+		Object->SetNumberField(TEXT("VectorSize"), ActionSchema.GetActionVectorSize(ActionSchemaElement));
+		Object->SetNumberField(TEXT("DistributionSize"), ActionSchema.GetActionDistributionVectorSize(ActionSchemaElement));
+		Object->SetNumberField(TEXT("EncodedSize"), ActionSchema.GetEncodedVectorSize(ActionSchemaElement));
+
+		switch (ActionSchema.GetType(ActionSchemaElement))
+		{
+		case Action::EType::Null:
+		{
+			Object->SetStringField(TEXT("Type"), TEXT("Null"));
+			break;
+		}
+
+		case Action::EType::Continuous:
+		{
+			const Action::FSchemaContinuousParameters Parameters = ActionSchema.GetContinuous(ActionSchemaElement);
+
+			Object->SetStringField(TEXT("Type"), TEXT("Continuous"));
+			Object->SetNumberField(TEXT("Num"), Parameters.Num);
+			break;
+		}
+
+		case Action::EType::DiscreteExclusive:
+		{
+			const Action::FSchemaDiscreteExclusiveParameters Parameters = ActionSchema.GetDiscreteExclusive(ActionSchemaElement);
+
+			Object->SetStringField(TEXT("Type"), TEXT("DiscreteExclusive"));
+			Object->SetNumberField(TEXT("Num"), Parameters.Num);
+			break;
+		}
+
+		case Action::EType::DiscreteInclusive:
+		{
+			const Action::FSchemaDiscreteInclusiveParameters Parameters = ActionSchema.GetDiscreteInclusive(ActionSchemaElement);
+
+			Object->SetStringField(TEXT("Type"), TEXT("DiscreteInclusive"));
+			Object->SetNumberField(TEXT("Num"), Parameters.Num);
+			break;
+		}
+
+		case Action::EType::And:
+		{
+			const Action::FSchemaAndParameters Parameters = ActionSchema.GetAnd(ActionSchemaElement);
+
+			Object->SetStringField(TEXT("Type"), TEXT("And"));
+
+			TSharedPtr<FJsonObject> SubObject = MakeShared<FJsonObject>();
+			for (int32 SubElementIdx = 0; SubElementIdx < Parameters.Elements.Num(); SubElementIdx++)
+			{
+				TSharedPtr<FJsonObject> SubElement = ConvertActionSchemaToJSON(ActionSchema, Parameters.Elements[SubElementIdx]);
+				SubElement->SetNumberField(TEXT("Index"), SubElementIdx);
+				SubObject->SetObjectField(Parameters.ElementNames[SubElementIdx].ToString(), SubElement);
+			}
+
+			Object->SetObjectField(TEXT("Elements"), SubObject);
+			break;
+		}
+
+		case Action::EType::OrExclusive:
+		{
+			const Action::FSchemaOrExclusiveParameters Parameters = ActionSchema.GetOrExclusive(ActionSchemaElement);
+
+			Object->SetStringField(TEXT("Type"), TEXT("OrExclusive"));
+
+			TSharedPtr<FJsonObject> SubObject = MakeShared<FJsonObject>();
+			for (int32 SubElementIdx = 0; SubElementIdx < Parameters.Elements.Num(); SubElementIdx++)
+			{
+				TSharedPtr<FJsonObject> SubElement = ConvertActionSchemaToJSON(ActionSchema, Parameters.Elements[SubElementIdx]);
+				SubElement->SetNumberField(TEXT("Index"), SubElementIdx);
+				SubObject->SetObjectField(Parameters.ElementNames[SubElementIdx].ToString(), SubElement);
+			}
+
+			Object->SetObjectField(TEXT("Elements"), SubObject);
+			break;
+		}
+
+		case Action::EType::OrInclusive:
+		{
+			const Action::FSchemaOrInclusiveParameters Parameters = ActionSchema.GetOrInclusive(ActionSchemaElement);
+
+			Object->SetStringField(TEXT("Type"), TEXT("OrInclusive"));
+
+			TSharedPtr<FJsonObject> SubObject = MakeShared<FJsonObject>();
+			for (int32 SubElementIdx = 0; SubElementIdx < Parameters.Elements.Num(); SubElementIdx++)
+			{
+				TSharedPtr<FJsonObject> SubElement = ConvertActionSchemaToJSON(ActionSchema, Parameters.Elements[SubElementIdx]);
+				SubElement->SetNumberField(TEXT("Index"), SubElementIdx);
+				SubObject->SetObjectField(Parameters.ElementNames[SubElementIdx].ToString(), SubElement);
+			}
+
+			Object->SetObjectField(TEXT("Elements"), SubObject);
+			break;
+		}
+
+		case Action::EType::Array:
+		{
+			const Action::FSchemaArrayParameters Parameters = ActionSchema.GetArray(ActionSchemaElement);
+
+			Object->SetStringField(TEXT("Type"), TEXT("Array"));
+			Object->SetNumberField(TEXT("Num"), Parameters.Num);
+			Object->SetObjectField(TEXT("Element"), ConvertActionSchemaToJSON(ActionSchema, Parameters.Element));
+			break;
+		}
+
+		case Action::EType::Encoding:
+		{
+			const Action::FSchemaEncodingParameters Parameters = ActionSchema.GetEncoding(ActionSchemaElement);
+
+			Object->SetStringField(TEXT("Type"), TEXT("Encoding"));
+			Object->SetNumberField(TEXT("EncodingSize"), Parameters.EncodingSize);
+			Object->SetObjectField(TEXT("Element"), ConvertActionSchemaToJSON(ActionSchema, Parameters.Element));
+			break;
+		}
+
+		default:
+			UE_LEARNING_NOT_IMPLEMENTED();
+		}
+
+		return Object;
+	}
+
 	const TCHAR* GetDeviceString(const ETrainerDevice Device)
 	{
 		switch (Device)
diff --git a/Engine/Plugins/Experimental/LearningAgents/Source/LearningTraining/Public/LearningExperience.h b/Engine/Plugins/Experimental/LearningAgents/Source/LearningTraining/Public/LearningExperience.h
--- a/Engine/Plugins/Experimental/LearningAgents/Source/LearningTraining/Public/LearningExperience.h
+++ b/Engine/Plugins/Experimental/LearningAgents/Source/LearningTraining/Public/LearningExperience.h
@@ -153,57 +153,6 @@ namespace UE::Learning
 
 	namespace Experience
 	{
-		/**
-		* Resets an environment, replay buffer, and experience buffer
-		*/
-		LEARNINGTRAINING_API void GatherExperienceReset(
-			FReplayBuffer& ReplayBuffer,
-			FEpisodeBuffer& EpisodeBuffer,
-			FResetInstanceBuffer& ResetBuffer,
-			const FIndexSet Instances);
-
-		/**
-		* Runs a single step of experience gathering
-		*
-		* @param ReplayBuffer							Replay Buffer
-		* @param EpisodeBuffer							Episode Buffer
-		* @param ResetBuffer							Reset Buffer
-		* @param ObservationVectorBuffer				Buffer to read/write observation vectors into
-		* @param ActionVectorBuffer						Buffer to read/write action vectors into
-		* @param PreEvaluationMemoryStateVectorBuffer	Buffer to read/write pre-evaluation memory state vectors into
-		* @param MemoryStateVectorBuffer				Buffer to read/write (post-evaluation) memory state vectors into
-		* @param RewardBuffer							Buffer to read/write rewards into
-		* @param CompletionBuffer						Buffer to read/write completions into
-		* @param EpisodeEndCompletionMode				Completion mode to use for episodes that reach the max length
-		* @param ResetFunction							Function to run for resetting the environment
-		* @param ObservationFunction					Function to run for evaluating observations
-		* @param PolicyFunction							Function to run generating actions from observations
-		* @param ActionFunction							Function to run for evaluating actions
-		* @param UpdateFunction							Function to run for updating the environment
-		* @param RewardFunction							Function to run for evaluating rewards
-		* @param CompletionFunction						Function to run for evaluating completions
-		* @param Instances								Set of instances to gather experience for
-		*/
-		LEARNINGTRAINING_API bool GatherExperienceIteration(
-			FReplayBuffer& ReplayBuffer,
-			FEpisodeBuffer& EpisodeBuffer,
-			FResetInstanceBuffer& ResetBuffer,
-			TLearningArrayView<2, float> ObservationVectorBuffer,
-			TLearningArrayView<2, float> ActionVectorBuffer,
-			TLearningArrayView<2, float> PreEvaluationMemoryStateVectorBuffer,
-			TLearningArrayView<2, float> MemoryStateVectorBuffer,
-			TLearningArrayView<1, float> RewardBuffer,
-			TLearningArrayView<1, ECompletionMode> CompletionBuffer,
-			const ECompletionMode EpisodeEndCompletionMode,
-			const TFunctionRef<void(const FIndexSet Instances)> ResetFunction,
-			const TFunctionRef<void(const FIndexSet Instances)> ObservationFunction,
-			const TFunctionRef<void(const FIndexSet Instances)> PolicyFunction,
-			const TFunctionRef<void(const FIndexSet Instances)> ActionFunction,
-			const TFunctionRef<void(const FIndexSet Instances)> UpdateFunction,
-			const TFunctionRef<void(const FIndexSet Instances)> RewardFunction,
-			const TFunctionRef<void(const FIndexSet Instances)> CompletionFunction,
-			const FIndexSet Instances);
-
 		/**
 		* Resets, and then runs experience gathering until the provided replay buffer is full
 		*
@@ -216,7 +165,8 @@ namespace UE::Learning
 		* @param MemoryStateVectorBuffer				Buffer to read/write (post-evaluation) memory state vectors into
 		* @param RewardBuffer							Buffer to read/write rewards into
 		* @param CompletionBuffer						Buffer to read/write completions into
-		* @param EpisodeEndCompletionMode				Completion mode to use for episodes that reach the max length
+		* @param EpisodeCompletionBuffer				Additional buffer to record completions from full episode buffers
+		* @param AllCompletionBuffer					Additional buffer to record all completions from full episodes and normal completions
 		* @param ResetFunction							Function to run for resetting the environment
 		* @param ObservationFunction					Function to run for evaluating observations
 		* @param PolicyFunction							Function to run generating actions from observations
@@ -236,7 +186,8 @@ namespace UE::Learning
 			TLearningArrayView<2, float> MemoryStateVectorBuffer,
 			TLearningArrayView<1, float> RewardBuffer,
 			TLearningArrayView<1, ECompletionMode> CompletionBuffer,
-			const ECompletionMode EpisodeEndCompletionMode,
+			TLearningArrayView<1, ECompletionMode> EpisodeCompletionBuffer,
+			TLearningArrayView<1, ECompletionMode> AllCompletionBuffer,
 			const TFunctionRef<void(const FIndexSet Instances)> ResetFunction,
 			const TFunctionRef<void(const FIndexSet Instances)> ObservationFunction,
 			const TFunctionRef<void(const FIndexSet Instances)> PolicyFunction,
@@ -245,6 +196,5 @@ namespace UE::Learning
 			const TFunctionRef<void(const FIndexSet Instances)> RewardFunction,
 			const TFunctionRef<void(const FIndexSet Instances)> CompletionFunction,
 			const FIndexSet Instances);
-
 	}
 };
\ No newline at end of file
diff --git a/Engine/Plugins/Experimental/LearningAgents/Source/LearningTraining/Public/LearningImitationTrainer.h b/Engine/Plugins/Experimental/LearningAgents/Source/LearningTraining/Public/LearningImitationTrainer.h
--- a/Engine/Plugins/Experimental/LearningAgents/Source/LearningTraining/Public/LearningImitationTrainer.h
+++ b/Engine/Plugins/Experimental/LearningAgents/Source/LearningTraining/Public/LearningImitationTrainer.h
@@ -14,6 +14,7 @@
 
 class FSocket;
 class FMonitoredProcess;
+class ULearningNeuralNetworkData;
 
 UCLASS()
 class LEARNINGTRAINING_API ULearningSocketImitationTrainerServerCommandlet : public UCommandlet
@@ -28,21 +29,6 @@ class LEARNINGTRAINING_API ULearningSocketImitationTrainerServerCommandlet : pub
 
 namespace UE::Learning
 {
-	struct INeuralNetwork;
-
-	/**
-	* Settings for the network used for training. These settings must match the Neural Network
-	* objects passed to ImitationTrainer::Train.
-	*/
-	struct FImitationTrainerNetworkSettings
-	{
-		/** Minimum action noise used by the policy */
-		float PolicyActionNoiseMin = 0.25f;
-
-		/** Maximum action noise used by the policy */
-		float PolicyActionNoiseMax = 0.25f;
-	};
-
 	struct FImitationTrainerTrainingSettings
 	{
 		// Number of iterations to train the network for. Controls the overall training time.
@@ -50,8 +36,8 @@ namespace UE::Learning
 		// closer to 1000000 iterations or more is required for an exhaustively trained network.
 		uint32 IterationNum = 1000000;
 
-		// Learning rate of the policy network. Typical values are between 0.001f and 0.0001f
-		float LearningRatePolicy = 0.0001f;
+		// Learning rate of the training network. Typical values are between 0.001f and 0.0001f
+		float LearningRate = 0.001f;
 
 		// Ratio by which to decay the learning rate every 1000 iterations.
 		float LearningRateDecay = 0.99f;
@@ -60,24 +46,23 @@ namespace UE::Learning
 		// weights to be smaller.
 		float WeightDecay = 0.001f;
 
-		// Initial scale to apply to actions before noise is added to them. The smaller this is, 
-		// the less likely you are to have spurious correlations at the beginning of training which 
-		// can make things slow or unstable. Too small and the network may become too slow to train.
-		float InitialActionScale = 0.1f;
-
-		// Initial scale to apply to memory. The smaller this is, the more stable training will be. Too small and the 
-		// network may not use its memory while training.
-		float InitialMemoryScale = 0.1f;
-
-		// Batch size to use for training the policy. Smaller values tend to produce better results 
+		// Batch size to use for training. Smaller values tend to produce better results 
 		// at the cost of slowing down training.
-		uint32 PolicyBatchSize = 128;
+		uint32 BatchSize = 128;
 
-		// The window of observations and actions over which to do the training of the policy. Increasing this value 
+		// The window of observations and actions over which to do the training. Increasing this value 
 		// will encourage the policy to use its memory effectively. Too large and training can become unstable. Given
 		// we don't know the memory state during imitation learning it is better this is slightly larger than when we 
 		// are doing reinforcement learning.
-		uint32 PolicyWindow = 64;
+		uint32 Window = 64;
+
+		// Weight used to regularize actions. Larger values will encourage smaller actions but too large
+		// will cause actions to become always zero.
+		float ActionRegularizationWeight = 0.001f;
+
+		// Weighting used for the entropy bonus. Larger values encourage larger action 
+		// noise and therefore greater exploration but can make actions very noisy.
+		float ActionEntropyWeight = 0.0f;
 
 		// Random seed to use for training
 		uint32 Seed = 1234;
@@ -92,20 +77,10 @@ namespace UE::Learning
 		// for this version of Python by going to your Unreal Editor Python Binaries directory 
 		// (e.g. "\Engine\Binaries\ThirdParty\Python3\Win64") and running `python -m pip install tensorboard`. 
 		bool bUseTensorboard = false;
-	};
-
-	/**
-	* ImitationTrainer flags controlling some aspects of the process of communication with the trainer
-	*/
-	enum class EImitationTrainerFlags : uint8
-	{
-		None = 0,
 
-		// If to send over the initial provided policy network rather than reinitialize it from random weights at 
-		// the start of training. Use this if you want to start from a network which has already been trained.
-		UseInitialPolicyNetwork = 1 << 0,
+		// If to save snapshots of the trained networks every 1000 iterations
+		bool bSaveSnapshots = false;
 	};
-	ENUM_CLASS_FLAGS(EImitationTrainerFlags)
 
 	/**
 	* Interface for an object which can train a policy from experience using imitation learning.
@@ -152,7 +127,37 @@ namespace UE::Learning
 		* @returns				Trainer response
 		*/
 		virtual ETrainerResponse RecvPolicy(
-			INeuralNetwork& OutNetwork,
+			ULearningNeuralNetworkData& OutNetwork,
+			const float Timeout = Trainer::DefaultTimeout,
+			FRWLock* NetworkLock = nullptr,
+			const ELogSetting LogSettings = Trainer::DefaultLogSettings) = 0;
+
+		/**
+		* Wait for the trainer and pull an updated encoder.
+		*
+		* @param OutNetwork		Network to update
+		* @param Timeout		Timeout to wait in seconds
+		* @param NetworkLock	Lock to use when updating network
+		* @param LogSettings	Log settings
+		* @returns				Trainer response
+		*/
+		virtual ETrainerResponse RecvEncoder(
+			ULearningNeuralNetworkData& OutNetwork,
+			const float Timeout = Trainer::DefaultTimeout,
+			FRWLock* NetworkLock = nullptr,
+			const ELogSetting LogSettings = Trainer::DefaultLogSettings) = 0;
+
+		/**
+		* Wait for the trainer and pull an updated decoder.
+		*
+		* @param OutNetwork		Network to update
+		* @param Timeout		Timeout to wait in seconds
+		* @param NetworkLock	Lock to use when updating network
+		* @param LogSettings	Log settings
+		* @returns				Trainer response
+		*/
+		virtual ETrainerResponse RecvDecoder(
+			ULearningNeuralNetworkData& OutNetwork,
 			const float Timeout = Trainer::DefaultTimeout,
 			FRWLock* NetworkLock = nullptr,
 			const ELogSetting LogSettings = Trainer::DefaultLogSettings) = 0;
@@ -167,7 +172,37 @@ namespace UE::Learning
 		* @returns				Trainer response
 		*/
 		virtual ETrainerResponse SendPolicy(
-			const INeuralNetwork& Network,
+			const ULearningNeuralNetworkData& Network,
+			const float Timeout = Trainer::DefaultTimeout,
+			FRWLock* NetworkLock = nullptr,
+			const ELogSetting LogSettings = Trainer::DefaultLogSettings) = 0;
+
+		/**
+		* Wait for the training process to be ready and push an updated encoder to the shared memory.
+		*
+		* @param Network		Network to push
+		* @param Timeout		Timeout to wait in seconds
+		* @param NetworkLock	Lock to use when pushing network
+		* @param LogSettings	Log settings
+		* @returns				Trainer response
+		*/
+		virtual ETrainerResponse SendEncoder(
+			const ULearningNeuralNetworkData& Network,
+			const float Timeout = Trainer::DefaultTimeout,
+			FRWLock* NetworkLock = nullptr,
+			const ELogSetting LogSettings = Trainer::DefaultLogSettings) = 0;
+
+		/**
+		* Wait for the training process to be ready and push an updated decoder to the shared memory.
+		*
+		* @param Network		Network to push
+		* @param Timeout		Timeout to wait in seconds
+		* @param NetworkLock	Lock to use when pushing network
+		* @param LogSettings	Log settings
+		* @returns				Trainer response
+		*/
+		virtual ETrainerResponse SendDecoder(
+			const ULearningNeuralNetworkData& Network,
 			const float Timeout = Trainer::DefaultTimeout,
 			FRWLock* NetworkLock = nullptr,
 			const ELogSetting LogSettings = Trainer::DefaultLogSettings) = 0;
@@ -200,21 +235,26 @@ namespace UE::Learning
 		/**
 		* Create a new imitation trainer sub-process
 		*
-		* @param TaskName				Name of the training task - used to help identify the logs, snapshots, and other files generated by training
-		* @param PythonExecutablePath	Path to the python executable used for training
-		* @param SitePackagesPath		Path to the site-packages shipped with the PythonFoundationPackages plugin
-		* @param PythonContentPath		Path to the Python Content folder provided by the Learning plugin
-		* @param IntermediatePath		Path to the intermediate folder to write temporary files, logs, and snapshots to
-		* @param MaxEpisodeNum			Maximum number of episodes in the training data
-		* @param MaxStepNum				Maximum number of steps in the training data
-		* @param ObservationDimNum		Number of dimensions in the observation vector
-		* @param ActionDimNum			Number of dimensions in the action vector
-		* @param MemoryStateDimNum		Number of dimensions in the memory state vector
-		* @param PolicyNetwork			Policy Network to use
-		* @param TrainingSettings		Trainer Training settings
-		* @param NetworkSettings		Trainer Network settings
-		* @param TrainingProcessFlags	Training subprocess flags
-		* @param LogSettings			Logging settings to use
+		* @param TaskName					Name of the training task - used to help identify the logs, snapshots, and other files generated by training
+		* @param PythonExecutablePath		Path to the python executable used for training
+		* @param SitePackagesPath			Path to the site-packages shipped with the PythonFoundationPackages plugin
+		* @param PythonContentPath			Path to the Python Content folder provided by the Learning plugin
+		* @param IntermediatePath			Path to the intermediate folder to write temporary files, logs, and snapshots to
+		* @param MaxEpisodeNum				Maximum number of episodes in the training data
+		* @param MaxStepNum					Maximum number of steps in the training data
+		* @param ObservationDimNum			Number of dimensions in the observation vector
+		* @param ActionDimNum				Number of dimensions in the action vector
+		* @param MemoryStateDimNum			Number of dimensions in the memory state vector
+		* @param PolicyNetwork				Policy Network to use
+		* @param EncoderNetwork				Encoder Network to use
+		* @param DecoderNetwork				Decoder Network to use
+		* @param ObservationSchema			Schema used for Observations
+		* @param ObservationSchemaElement	Schema Observation Element
+		* @param ActionSchema				Schema used for Actions
+		* @param ActionSchemaElement		Schema Action Element
+		* @param TrainingSettings			Trainer Training settings
+		* @param TrainingProcessFlags		Training subprocess flags
+		* @param LogSettings				Logging settings to use
 		*/
 		FSharedMemoryImitationTrainer(
 			const FString& TaskName,
@@ -227,10 +267,14 @@ namespace UE::Learning
 			const int32 ObservationDimNum,
 			const int32 ActionDimNum,
 			const int32 MemoryStateDimNum,
-			const INeuralNetwork& PolicyNetwork,
+			const ULearningNeuralNetworkData& PolicyNetwork,
+			const ULearningNeuralNetworkData& EncoderNetwork,
+			const ULearningNeuralNetworkData& DecoderNetwork,
+			const Observation::FSchema& ObservationSchema,
+			const Observation::FSchemaElement& ObservationSchemaElement,
+			const Action::FSchema& ActionSchema,
+			const Action::FSchemaElement& ActionSchemaElement,
 			const FImitationTrainerTrainingSettings& TrainingSettings = FImitationTrainerTrainingSettings(),
-			const FImitationTrainerNetworkSettings& NetworkSettings = FImitationTrainerNetworkSettings(),
-			const EImitationTrainerFlags TrainerFlags = EImitationTrainerFlags::None,
 			const ESubprocessFlags TrainingProcessFlags = ESubprocessFlags::None,
 			const ELogSetting LogSettings = ELogSetting::Normal);
 
@@ -245,13 +289,37 @@ namespace UE::Learning
 		virtual bool HasPolicyOrCompleted() override final;
 
 		virtual ETrainerResponse RecvPolicy(
-			INeuralNetwork& OutNetwork,
+			ULearningNeuralNetworkData& OutNetwork,
+			const float Timeout = Trainer::DefaultTimeout,
+			FRWLock* NetworkLock = nullptr,
+			const ELogSetting LogSettings = Trainer::DefaultLogSettings) override final;
+
+		virtual ETrainerResponse RecvEncoder(
+			ULearningNeuralNetworkData& OutNetwork,
+			const float Timeout = Trainer::DefaultTimeout,
+			FRWLock* NetworkLock = nullptr,
+			const ELogSetting LogSettings = Trainer::DefaultLogSettings) override final;
+
+		virtual ETrainerResponse RecvDecoder(
+			ULearningNeuralNetworkData& OutNetwork,
 			const float Timeout = Trainer::DefaultTimeout,
 			FRWLock* NetworkLock = nullptr,
 			const ELogSetting LogSettings = Trainer::DefaultLogSettings) override final;
 
 		virtual ETrainerResponse SendPolicy(
-			const INeuralNetwork& Network,
+			const ULearningNeuralNetworkData& Network,
+			const float Timeout = Trainer::DefaultTimeout,
+			FRWLock* NetworkLock = nullptr,
+			const ELogSetting LogSettings = Trainer::DefaultLogSettings) override final;
+
+		virtual ETrainerResponse SendEncoder(
+			const ULearningNeuralNetworkData& Network,
+			const float Timeout = Trainer::DefaultTimeout,
+			FRWLock* NetworkLock = nullptr,
+			const ELogSetting LogSettings = Trainer::DefaultLogSettings) override final;
+
+		virtual ETrainerResponse SendDecoder(
+			const ULearningNeuralNetworkData& Network,
 			const float Timeout = Trainer::DefaultTimeout,
 			FRWLock* NetworkLock = nullptr,
 			const ELogSetting LogSettings = Trainer::DefaultLogSettings) override final;
@@ -274,6 +342,8 @@ namespace UE::Learning
 		// Shared memory
 
 		UE::Learning::TSharedMemoryArrayView<1, uint8> Policy;
+		UE::Learning::TSharedMemoryArrayView<1, uint8> Encoder;
+		UE::Learning::TSharedMemoryArrayView<1, uint8> Decoder;
 		UE::Learning::TSharedMemoryArrayView<1, volatile int32> Controls; // Mark as volatile to avoid compiler optimizing away reads without writes etc.
 		UE::Learning::TSharedMemoryArrayView<1, int32> EpisodeStarts;
 		UE::Learning::TSharedMemoryArrayView<1, int32> EpisodeLengths;
@@ -359,12 +429,16 @@ namespace UE::Learning
 		* @param ActionDimNum				Number of dimensions in the action vector
 		* @param MemoryStateDimNum			Number of dimensions in the memory state vector
 		* @param PolicyNetwork				Policy Network to use
+		* @param EncoderNetwork				Encoder Network to use
+		* @param DecoderNetwork				Decoder Network to use
+		* @param ObservationSchema			Schema used for Observations
+		* @param ObservationSchemaElement	Schema Observation Element
+		* @param ActionSchema				Schema used for Actions
+		* @param ActionSchemaElement		Schema Action Element
 		* @param IpAddress					Server Ip address
 		* @param Port						Server Port
 		* @param Timeout					Timeout to wait in seconds for connection and initial data transfer
 		* @param TrainingSettings			Trainer Training settings
-		* @param NetworkSettings			Trainer Network settings
-		* @param TrainerFlags				Flags for the trainer
 		*/
 		FSocketImitationTrainer(
 			ETrainerResponse& OutResponse,
@@ -374,13 +448,17 @@ namespace UE::Learning
 			const int32 ObservationDimNum,
 			const int32 ActionDimNum,
 			const int32 MemoryStateDimNum,
-			const INeuralNetwork& PolicyNetwork,
+			const ULearningNeuralNetworkData& PolicyNetwork,
+			const ULearningNeuralNetworkData& EncoderNetwork,
+			const ULearningNeuralNetworkData& DecoderNetwork,
+			const Observation::FSchema& ObservationSchema,
+			const Observation::FSchemaElement& ObservationSchemaElement,
+			const Action::FSchema& ActionSchema,
+			const Action::FSchemaElement& ActionSchemaElement,
 			const TCHAR* IpAddress = Trainer::DefaultIp,
 			const uint32 Port = Trainer::DefaultPort,
 			const float Timeout = Trainer::DefaultTimeout,
-			const FImitationTrainerTrainingSettings& TrainingSettings = FImitationTrainerTrainingSettings(),
-			const FImitationTrainerNetworkSettings& NetworkSettings = FImitationTrainerNetworkSettings(),
-			const EImitationTrainerFlags TrainerFlags = EImitationTrainerFlags::None);
+			const FImitationTrainerTrainingSettings& TrainingSettings = FImitationTrainerTrainingSettings());
 
 		~FSocketImitationTrainer();
 
@@ -393,13 +471,37 @@ namespace UE::Learning
 		virtual bool HasPolicyOrCompleted() override final;
 
 		virtual ETrainerResponse RecvPolicy(
-			INeuralNetwork& OutNetwork,
+			ULearningNeuralNetworkData& OutNetwork,
+			const float Timeout = Trainer::DefaultTimeout,
+			FRWLock* NetworkLock = nullptr,
+			const ELogSetting LogSettings = Trainer::DefaultLogSettings) override final;
+
+		virtual ETrainerResponse RecvEncoder(
+			ULearningNeuralNetworkData& OutNetwork,
+			const float Timeout = Trainer::DefaultTimeout,
+			FRWLock* NetworkLock = nullptr,
+			const ELogSetting LogSettings = Trainer::DefaultLogSettings) override final;
+
+		virtual ETrainerResponse RecvDecoder(
+			ULearningNeuralNetworkData& OutNetwork,
 			const float Timeout = Trainer::DefaultTimeout,
 			FRWLock* NetworkLock = nullptr,
 			const ELogSetting LogSettings = Trainer::DefaultLogSettings) override final;
 
 		virtual ETrainerResponse SendPolicy(
-			const INeuralNetwork& Network,
+			const ULearningNeuralNetworkData& Network,
+			const float Timeout = Trainer::DefaultTimeout,
+			FRWLock* NetworkLock = nullptr,
+			const ELogSetting LogSettings = Trainer::DefaultLogSettings) override final;
+
+		virtual ETrainerResponse SendEncoder(
+			const ULearningNeuralNetworkData& Network,
+			const float Timeout = Trainer::DefaultTimeout,
+			FRWLock* NetworkLock = nullptr,
+			const ELogSetting LogSettings = Trainer::DefaultLogSettings) override final;
+
+		virtual ETrainerResponse SendDecoder(
+			const ULearningNeuralNetworkData& Network,
 			const float Timeout = Trainer::DefaultTimeout,
 			FRWLock* NetworkLock = nullptr,
 			const ELogSetting LogSettings = Trainer::DefaultLogSettings) override final;
@@ -414,7 +516,9 @@ namespace UE::Learning
 
 	private:
 
-		TLearningArray<1, uint8> NetworkBuffer;
+		TLearningArray<1, uint8> PolicyBuffer;
+		TLearningArray<1, uint8> EncoderBuffer;
+		TLearningArray<1, uint8> DecoderBuffer;
 		FSocket* Socket = nullptr;
 	};
 
@@ -424,29 +528,39 @@ namespace UE::Learning
 		* Train a policy using experience already gathered from example episodes
 		*
 		* @param Trainer							Trainer
-		* @param Network							Policy network
+		* @param PolicyNetwork						Policy network
+		* @param EncoderNetwork						Encoder network
+		* @param DecoderNetwork						Decoder network
 		* @param EpisodeStartsExperience			Array of offsets where episodes start
 		* @param EpisodeLengthsExperience			Array of episode lengths
 		* @param ObservationsExperience				Set of observation vectors
 		* @param ActionsExperience					Set of action vectors
-		* @param TrainerFlags						Flags for the trainer, should match what was used to initialize the Trainer object.
 		* @param bRequestTrainingStopSignal			Optional signal that can be raised to indicate training should be stopped
-		* @param NetworkLock						Optional Lock to use when updating the policy network
-		* @param bNetworkUpdatedSignal				Optional signal that will be raised when the policy network is updated
+		* @param PolicyNetworkLock					Optional Lock to use when updating the policy network
+		* @param EncoderNetworkLock					Optional Lock to use when updating the encoder network
+		* @param DecoderNetworkLock					Optional Lock to use when updating the decoder network
+		* @param bPolicyNetworkUpdatedSignal		Optional signal that will be raised when the policy network is updated
+		* @param bEncoderNetworkUpdatedSignal		Optional signal that will be raised when the encoder network is updated
+		* @param bDecoderNetworkUpdatedSignal		Optional signal that will be raised when the decoder network is updated
 		* @param LogSettings						Logging settings
 		* @returns									Trainer response in case of errors during communication otherwise Success
 		*/
 		LEARNINGTRAINING_API ETrainerResponse Train(
 			IImitationTrainer& Trainer,
-			INeuralNetwork& Network,
+			ULearningNeuralNetworkData& PolicyNetwork,
+			ULearningNeuralNetworkData& EncoderNetwork,
+			ULearningNeuralNetworkData& DecoderNetwork,
 			const TLearningArrayView<1, const int32> EpisodeStartsExperience,
 			const TLearningArrayView<1, const int32> EpisodeLengthsExperience,
 			const TLearningArrayView<2, const float> ObservationsExperience,
 			const TLearningArrayView<2, const float> ActionsExperience,
-			const EImitationTrainerFlags TrainerFlags = EImitationTrainerFlags::None,
 			TAtomic<bool>* bRequestTrainingStopSignal = nullptr,
-			FRWLock* NetworkLock = nullptr,
-			TAtomic<bool>* bNetworkUpdatedSignal = nullptr,
+			FRWLock* PolicyNetworkLock = nullptr,
+			FRWLock* EncoderNetworkLock = nullptr,
+			FRWLock* DecoderNetworkLock = nullptr,
+			TAtomic<bool>* bPolicyNetworkUpdatedSignal = nullptr,
+			TAtomic<bool>* bEncoderNetworkUpdatedSignal = nullptr,
+			TAtomic<bool>* bDecoderNetworkUpdatedSignal = nullptr,
 			const ELogSetting LogSettings = ELogSetting::Normal);
 	}
 
