LearningAgents: Large restructuring to allow for structured actions and observations.

diff --git a/Engine/Plugins/Experimental/LearningAgents/Source/LearningTraining/Public/LearningPPOTrainer.h b/Engine/Plugins/Experimental/LearningAgents/Source/LearningTraining/Public/LearningPPOTrainer.h
--- a/Engine/Plugins/Experimental/LearningAgents/Source/LearningTraining/Public/LearningPPOTrainer.h
+++ b/Engine/Plugins/Experimental/LearningAgents/Source/LearningTraining/Public/LearningPPOTrainer.h
@@ -14,6 +14,7 @@
 
 class FSocket;
 class FMonitoredProcess;
+class ULearningNeuralNetworkData;
 
 UCLASS()
 class LEARNINGTRAINING_API ULearningSocketPPOTrainerServerCommandlet : public UCommandlet
@@ -28,26 +29,11 @@ class LEARNINGTRAINING_API ULearningSocketPPOTrainerServerCommandlet : public UC
 
 namespace UE::Learning
 {
-	struct INeuralNetwork;
 	struct FReplayBuffer;
 	struct FResetInstanceBuffer;
 	struct FEpisodeBuffer;
 	enum class ECompletionMode : uint8;
 
-	/**
-	* Settings for the networks used for training PPO. These settings must match the Neural Network 
-	* objects passed to PPOTrainer::Train and the action noise must match the action noise used by 
-	* the policy while gathering experience. 
-	*/
-	struct FPPOTrainerNetworkSettings
-	{
-		/** Minimum action noise used by the policy */
-		float PolicyActionNoiseMin = 0.25f;
-
-		/** Maximum action noise used by the policy */
-		float PolicyActionNoiseMax = 0.25f;
-	};
-
 	/**
 	* Settings used for training with PPO
 	*/
@@ -72,19 +58,6 @@ namespace UE::Learning
 		// weights to be smaller but too large a value can cause the network weights to collapse to all zeros.
 		float WeightDecay = 0.001f;
 
-		// Initial scale to apply to actions before noise is added to them. The smaller this is, 
-		// the less likely you are to have spurious correlations at the beginning of training which 
-		// can make things slow or unstable. Too small and the network may become difficult to train.
-		float InitialActionScale = 0.1f;
-
-		// Initial scale to apply to memory. The smaller this is, the more stable training will be. Too small and the 
-		// network may not use its memory while training.
-		float InitialMemoryScale = 0.1f;
-
-		// Initial scale to apply to the critic outputs. The smaller this is, the more stable training will be. Too 
-		// small and the network may become too slow to train.
-		float InitialCriticScale = 0.1f;
-
 		// Batch size to use for training the policy. Smaller values tend to produce better results 
 		// at the cost of slowing down training. Large batch sizes are much more computationally efficient 
 		// when training on the GPU.
@@ -103,6 +76,10 @@ namespace UE::Learning
 		// the critic and policy to be effectively updated, but too large and it will simply slow down training.
 		uint32 IterationsPerGather = 250;
 
+		// Number of iterations of training to perform to warm-up the Critic. This helps speed up and stabilize training
+		// at the beginning when the Critic may be producing predictions at the wrong order of magnitude.
+		uint32 CriticWarmupIterations = 10;
+
 		// Clipping ratio to apply to policy updates. Keeps the training "on-policy". 
 		// Larger values may speed up training at the cost of stability. Conversely, too small 
 		// values will keep the policy from being able to learn an optimal policy.
@@ -111,13 +88,16 @@ namespace UE::Learning
 		// Weight used to regularize predicted returns. Encourages the critic not to over or under estimate returns.
 		float ReturnRegularizationWeight = 0.0001f;
 
+		// Weight for the loss used to train the policy via the PPO surrogate objective.
+		float ActionSurrogateWeight = 1.0f;
+
 		// Weight used to regularize actions. Larger values will encourage smaller actions but too large
 		// will cause actions to become always zero.
 		float ActionRegularizationWeight = 0.001f;
 
 		// Weighting used for the entropy bonus. Larger values encourage larger action 
 		// noise and therefore greater exploration but can make actions very noisy.
-		float ActionEntropyWeight = 0.01f;
+		float ActionEntropyWeight = 0.0f;
 
 		// This is used in the Generalized Advantage Estimation as what is essentially 
 		// an exponential smoothing/decay. Typical values should be between 0.9 and 1.0.
@@ -163,26 +143,10 @@ namespace UE::Learning
 		// for this version of Python by going to your Unreal Editor Python Binaries directory 
 		// (e.g. "\Engine\Binaries\ThirdParty\Python3\Win64") and running `python -m pip install tensorboard`. 
 		bool bUseTensorboard = false;
-	};
 
-	/**
-	* PPOTrainer flags controlling some aspects of the process of communication with the trainer
-	*/
-	enum class EPPOTrainerFlags : uint8
-	{
-		None = 0,
-
-		// If to send over the initial provided policy network rather than reinitialize it from random weights at 
-		// the start of training. Use this if you want to start from a network which has already been trained 
-		// such as via Imitation Training.
-		UseInitialPolicyNetwork = 1 << 0,
-
-		// If to send over the initial provided critic network rather than reinitialize it from random weights at 
-		// the start of training. Use this if you want to start from a network which has already been trained 
-		// such as via Imitation Training.
-		UseInitialCriticNetwork = 1 << 1,
+		// If to save snapshots of the trained networks every 1000 iterations
+		bool bSaveSnapshots = false;
 	};
-	ENUM_CLASS_FLAGS(EPPOTrainerFlags)
 
 	/**
 	* Interface for an object which can train a policy using PPO.
@@ -214,7 +178,7 @@ namespace UE::Learning
 		* @returns				Trainer response
 		*/
 		virtual ETrainerResponse RecvPolicy(
-			INeuralNetwork& OutNetwork,
+			ULearningNeuralNetworkData& OutNetwork,
 			const float Timeout = Trainer::DefaultTimeout,
 			FRWLock* NetworkLock = nullptr,
 			const ELogSetting LogSettings = Trainer::DefaultLogSettings) = 0;
@@ -229,7 +193,37 @@ namespace UE::Learning
 		* @returns				Trainer response
 		*/
 		virtual ETrainerResponse RecvCritic(
-			INeuralNetwork& OutNetwork,
+			ULearningNeuralNetworkData& OutNetwork,
+			const float Timeout = Trainer::DefaultTimeout,
+			FRWLock* NetworkLock = nullptr,
+			const ELogSetting LogSettings = Trainer::DefaultLogSettings) = 0;
+
+		/**
+		* Wait for the trainer to push an updated encoder network.
+		*
+		* @param OutNetwork		Network to update
+		* @param Timeout		Timeout to wait in seconds
+		* @param NetworkLock	Lock to use when updating network
+		* @param LogSettings	Log settings
+		* @returns				Trainer response
+		*/
+		virtual ETrainerResponse RecvEncoder(
+			ULearningNeuralNetworkData& OutNetwork,
+			const float Timeout = Trainer::DefaultTimeout,
+			FRWLock* NetworkLock = nullptr,
+			const ELogSetting LogSettings = Trainer::DefaultLogSettings) = 0;
+
+		/**
+		* Wait for the trainer to push an updated decoder network.
+		*
+		* @param OutNetwork		Network to update
+		* @param Timeout		Timeout to wait in seconds
+		* @param NetworkLock	Lock to use when updating network
+		* @param LogSettings	Log settings
+		* @returns				Trainer response
+		*/
+		virtual ETrainerResponse RecvDecoder(
+			ULearningNeuralNetworkData& OutNetwork,
 			const float Timeout = Trainer::DefaultTimeout,
 			FRWLock* NetworkLock = nullptr,
 			const ELogSetting LogSettings = Trainer::DefaultLogSettings) = 0;
@@ -249,7 +243,7 @@ namespace UE::Learning
 		* @returns				Trainer response
 		*/
 		virtual ETrainerResponse SendPolicy(
-			const INeuralNetwork& Network,
+			const ULearningNeuralNetworkData& Network,
 			const float Timeout = Trainer::DefaultTimeout,
 			FRWLock* NetworkLock = nullptr,
 			const ELogSetting LogSettings = Trainer::DefaultLogSettings) = 0;
@@ -264,7 +258,37 @@ namespace UE::Learning
 		* @returns				Trainer response
 		*/
 		virtual ETrainerResponse SendCritic(
-			const INeuralNetwork& Network,
+			const ULearningNeuralNetworkData& Network,
+			const float Timeout = Trainer::DefaultTimeout,
+			FRWLock* NetworkLock = nullptr,
+			const ELogSetting LogSettings = Trainer::DefaultLogSettings) = 0;
+
+		/**
+		* Wait for the trainer to be ready and push the current encoder network.
+		*
+		* @param Network		Network to push
+		* @param Timeout		Timeout to wait in seconds
+		* @param NetworkLock	Lock to use when pushing network
+		* @param LogSettings	Log settings
+		* @returns				Trainer response
+		*/
+		virtual ETrainerResponse SendEncoder(
+			const ULearningNeuralNetworkData& Network,
+			const float Timeout = Trainer::DefaultTimeout,
+			FRWLock* NetworkLock = nullptr,
+			const ELogSetting LogSettings = Trainer::DefaultLogSettings) = 0;
+
+		/**
+		* Wait for the trainer to be ready and push the current decoder network.
+		*
+		* @param Network		Network to push
+		* @param Timeout		Timeout to wait in seconds
+		* @param NetworkLock	Lock to use when pushing network
+		* @param LogSettings	Log settings
+		* @returns				Trainer response
+		*/
+		virtual ETrainerResponse SendDecoder(
+			const ULearningNeuralNetworkData& Network,
 			const float Timeout = Trainer::DefaultTimeout,
 			FRWLock* NetworkLock = nullptr,
 			const ELogSetting LogSettings = Trainer::DefaultLogSettings) = 0;
@@ -302,9 +326,13 @@ namespace UE::Learning
 		* @param ReplayBuffer				Replay buffer used to collect experience
 		* @param PolicyNetwork				Policy Network to use
 		* @param CriticNetwork				Critic Network to use
+		* @param EncoderNetwork				Encoder Network to use
+		* @param DecoderNetwork				Decoder Network to use
+		* @param ObservationSchema			Schema used for Observations
+		* @param ObservationSchemaElement	Schema Observation Element
+		* @param ActionSchema				Schema used for Actions
+		* @param ActionSchemaElement		Schema Action Element
 		* @param TrainingSettings			Trainer Training settings
-		* @param NetworkSettings			Trainer Network settings
-		* @param TrainerFlags				Flags for the trainer
 		* @param LogSettings				Logging settings to use
 		* @param TrainingProcessFlags		Training subprocess flags
 
@@ -336,11 +364,15 @@ namespace UE::Learning
 			const FString& PythonContentPath,
 			const FString& IntermediatePath,
 			const FReplayBuffer& ReplayBuffer,
-			const INeuralNetwork& PolicyNetwork,
-			const INeuralNetwork& CriticNetwork,
+			const ULearningNeuralNetworkData& PolicyNetwork,
+			const ULearningNeuralNetworkData& CriticNetwork,
+			const ULearningNeuralNetworkData& EncoderNetwork,
+			const ULearningNeuralNetworkData& DecoderNetwork,
+			const Observation::FSchema& ObservationSchema,
+			const Observation::FSchemaElement& ObservationSchemaElement,
+			const Action::FSchema& ActionSchema,
+			const Action::FSchemaElement& ActionSchemaElement,
 			const FPPOTrainerTrainingSettings& TrainingSettings = FPPOTrainerTrainingSettings(),
-			const FPPOTrainerNetworkSettings& NetworkSettings = FPPOTrainerNetworkSettings(),
-			const EPPOTrainerFlags TrainerFlags = EPPOTrainerFlags::None,
 			const ELogSetting LogSettings = ELogSetting::Normal,
 			const ESubprocessFlags TrainingProcessFlags = ESubprocessFlags::None,
 			const uint16 ProcessNum = 1,
@@ -353,13 +385,25 @@ namespace UE::Learning
 		virtual ETrainerResponse Wait(const float Timeout = Trainer::DefaultTimeout) override final;
 
 		virtual ETrainerResponse RecvPolicy(
-			INeuralNetwork& OutNetwork,
+			ULearningNeuralNetworkData& OutNetwork,
 			const float Timeout = Trainer::DefaultTimeout,
 			FRWLock* NetworkLock = nullptr,
 			const ELogSetting LogSettings = Trainer::DefaultLogSettings) override final;
 
 		virtual ETrainerResponse RecvCritic(
-			INeuralNetwork& OutNetwork,
+			ULearningNeuralNetworkData& OutNetwork,
+			const float Timeout = Trainer::DefaultTimeout,
+			FRWLock* NetworkLock = nullptr,
+			const ELogSetting LogSettings = Trainer::DefaultLogSettings) override final;
+
+		virtual ETrainerResponse RecvEncoder(
+			ULearningNeuralNetworkData& OutNetwork,
+			const float Timeout = Trainer::DefaultTimeout,
+			FRWLock* NetworkLock = nullptr,
+			const ELogSetting LogSettings = Trainer::DefaultLogSettings) override final;
+
+		virtual ETrainerResponse RecvDecoder(
+			ULearningNeuralNetworkData& OutNetwork,
 			const float Timeout = Trainer::DefaultTimeout,
 			FRWLock* NetworkLock = nullptr,
 			const ELogSetting LogSettings = Trainer::DefaultLogSettings) override final;
@@ -367,13 +411,25 @@ namespace UE::Learning
 		virtual ETrainerResponse SendStop(const float Timeout = Trainer::DefaultTimeout) override final;
 
 		virtual ETrainerResponse SendPolicy(
-			const INeuralNetwork& Network,
+			const ULearningNeuralNetworkData& Network,
 			const float Timeout = Trainer::DefaultTimeout,
 			FRWLock* NetworkLock = nullptr,
 			const ELogSetting LogSettings = Trainer::DefaultLogSettings) override final;
 
 		virtual ETrainerResponse SendCritic(
-			const INeuralNetwork& Network,
+			const ULearningNeuralNetworkData& Network,
+			const float Timeout = Trainer::DefaultTimeout,
+			FRWLock* NetworkLock = nullptr,
+			const ELogSetting LogSettings = Trainer::DefaultLogSettings) override final;
+
+		virtual ETrainerResponse SendEncoder(
+			const ULearningNeuralNetworkData& Network,
+			const float Timeout = Trainer::DefaultTimeout,
+			FRWLock* NetworkLock = nullptr,
+			const ELogSetting LogSettings = Trainer::DefaultLogSettings) override final;
+
+		virtual ETrainerResponse SendDecoder(
+			const ULearningNeuralNetworkData& Network,
 			const float Timeout = Trainer::DefaultTimeout,
 			FRWLock* NetworkLock = nullptr,
 			const ELogSetting LogSettings = Trainer::DefaultLogSettings) override final;
@@ -394,6 +450,8 @@ namespace UE::Learning
 
 		TSharedMemoryArrayView<1, uint8> Policy;
 		TSharedMemoryArrayView<1, uint8> Critic;
+		TSharedMemoryArrayView<1, uint8> Encoder;
+		TSharedMemoryArrayView<1, uint8> Decoder;
 		TSharedMemoryArrayView<2, volatile int32> Controls; // Mark as volatile to avoid compiler optimizing away reads without writes etc.
 		TSharedMemoryArrayView<2, int32> EpisodeStarts;
 		TSharedMemoryArrayView<2, int32> EpisodeLengths;
@@ -493,25 +551,33 @@ namespace UE::Learning
 		* @param ReplayBuffer				Replay buffer used to collect experience
 		* @param PolicyNetwork				Policy Network to use
 		* @param CriticNetwork				Critic Network to use
+		* @param EncoderNetwork				Encoder Network to use
+		* @param DecoderNetwork				Decoder Network to use
+		* @param ObservationSchema			Schema used for Observations
+		* @param ObservationSchemaElement	Schema Observation Element
+		* @param ActionSchema				Schema used for Actions
+		* @param ActionSchemaElement		Schema Action Element
 		* @param IpAddress					Server Ip address
 		* @param Port						Server Port
 		* @param Timeout					Timeout to wait in seconds for connection and initial data transfer
 		* @param TrainingSettings			Trainer Training settings
-		* @param NetworkSettings			Trainer Network settings
-		* @param TrainerFlags				Flags for the trainer
 		*/
 		FSocketPPOTrainer(
 			ETrainerResponse& OutResponse,
 			const FString& TaskName,
 			const FReplayBuffer& ReplayBuffer,
-			const INeuralNetwork& PolicyNetwork,
-			const INeuralNetwork& CriticNetwork,
+			const ULearningNeuralNetworkData& PolicyNetwork,
+			const ULearningNeuralNetworkData& CriticNetwork,
+			const ULearningNeuralNetworkData& EncoderNetwork,
+			const ULearningNeuralNetworkData& DecoderNetwork,
+			const Observation::FSchema& ObservationSchema,
+			const Observation::FSchemaElement& ObservationSchemaElement,
+			const Action::FSchema& ActionSchema,
+			const Action::FSchemaElement& ActionSchemaElement,
 			const TCHAR* IpAddress = Trainer::DefaultIp,
 			const uint32 Port = Trainer::DefaultPort,
 			const float Timeout = Trainer::DefaultTimeout,
-			const FPPOTrainerTrainingSettings& TrainingSettings = FPPOTrainerTrainingSettings(),
-			const FPPOTrainerNetworkSettings& NetworkSettings = FPPOTrainerNetworkSettings(),
-			const EPPOTrainerFlags TrainerFlags = EPPOTrainerFlags::None);
+			const FPPOTrainerTrainingSettings& TrainingSettings = FPPOTrainerTrainingSettings());
 
 		~FSocketPPOTrainer();
 
@@ -520,13 +586,25 @@ namespace UE::Learning
 		virtual ETrainerResponse Wait(const float Timeout = Trainer::DefaultTimeout) override final;
 
 		virtual ETrainerResponse RecvPolicy(
-			INeuralNetwork& OutNetwork,
+			ULearningNeuralNetworkData& OutNetwork,
 			const float Timeout = Trainer::DefaultTimeout,
 			FRWLock* NetworkLock = nullptr,
 			const ELogSetting LogSettings = Trainer::DefaultLogSettings) override final;
 
 		virtual ETrainerResponse RecvCritic(
-			INeuralNetwork& OutNetwork,
+			ULearningNeuralNetworkData& OutNetwork,
+			const float Timeout = Trainer::DefaultTimeout,
+			FRWLock* NetworkLock = nullptr,
+			const ELogSetting LogSettings = Trainer::DefaultLogSettings) override final;
+
+		virtual ETrainerResponse RecvEncoder(
+			ULearningNeuralNetworkData& OutNetwork,
+			const float Timeout = Trainer::DefaultTimeout,
+			FRWLock* NetworkLock = nullptr,
+			const ELogSetting LogSettings = Trainer::DefaultLogSettings) override final;
+
+		virtual ETrainerResponse RecvDecoder(
+			ULearningNeuralNetworkData& OutNetwork,
 			const float Timeout = Trainer::DefaultTimeout,
 			FRWLock* NetworkLock = nullptr,
 			const ELogSetting LogSettings = Trainer::DefaultLogSettings) override final;
@@ -534,13 +612,25 @@ namespace UE::Learning
 		virtual ETrainerResponse SendStop(const float Timeout = Trainer::DefaultTimeout) override final;
 
 		virtual ETrainerResponse SendPolicy(
-			const INeuralNetwork& Network,
+			const ULearningNeuralNetworkData& Network,
 			const float Timeout = Trainer::DefaultTimeout,
 			FRWLock* NetworkLock = nullptr,
 			const ELogSetting LogSettings = Trainer::DefaultLogSettings) override final;
 
 		virtual ETrainerResponse SendCritic(
-			const INeuralNetwork& Network,
+			const ULearningNeuralNetworkData& Network,
+			const float Timeout = Trainer::DefaultTimeout,
+			FRWLock* NetworkLock = nullptr,
+			const ELogSetting LogSettings = Trainer::DefaultLogSettings) override final;
+
+		virtual ETrainerResponse SendEncoder(
+			const ULearningNeuralNetworkData& Network,
+			const float Timeout = Trainer::DefaultTimeout,
+			FRWLock* NetworkLock = nullptr,
+			const ELogSetting LogSettings = Trainer::DefaultLogSettings) override final;
+
+		virtual ETrainerResponse SendDecoder(
+			const ULearningNeuralNetworkData& Network,
 			const float Timeout = Trainer::DefaultTimeout,
 			FRWLock* NetworkLock = nullptr,
 			const ELogSetting LogSettings = Trainer::DefaultLogSettings) override final;
@@ -554,6 +644,8 @@ namespace UE::Learning
 
 		TLearningArray<1, uint8> PolicyNetworkBuffer;
 		TLearningArray<1, uint8> CriticNetworkBuffer;
+		TLearningArray<1, uint8> EncoderNetworkBuffer;
+		TLearningArray<1, uint8> DecoderNetworkBuffer;
 		FSocket* Socket = nullptr;
 	};
 
@@ -568,14 +660,17 @@ namespace UE::Learning
 		* @param EpisodeBuffer							Episode Buffer
 		* @param ResetBuffer							Reset Buffer
 		* @param PolicyNetwork							Policy Network to use
-		* @param CriticNetwork							Optional Critic Network to use
+		* @param CriticNetwork							Critic Network to use
+		* @param EncoderNetwork							Encoder Network to use
+		* @param DecoderNetwork							Decoder Network to use
 		* @param ObservationVectorBuffer				Buffer to read/write observation vectors into
 		* @param ActionVectorBuffer						Buffer to read/write action vectors into
 		* @param PreEvaluationMemoryStateVectorBuffer	Buffer to read/write pre-evaluation memory state vectors into
 		* @param MemoryStateVectorBuffer				Buffer to read/write (post-evaluation) memory state vectors into
 		* @param RewardBuffer							Buffer to read/write rewards into
 		* @param CompletionBuffer						Buffer to read/write completions into
-		* @param EpisodeEndCompletionMode				Completion mode to use for episodes that reach the max length
+		* @param EpisodeCompletionBuffer				Additional buffer to record completions from full episode buffers
+		* @param AllCompletionBuffer					Additional buffer to record all completions from full episodes and normal completions
 		* @param ResetFunction							Function to run for resetting the environment
 		* @param ObservationFunction					Function to run for evaluating observations
 		* @param PolicyFunction							Function to run for evaluating the policy
@@ -584,12 +679,15 @@ namespace UE::Learning
 		* @param RewardFunction							Function to run for evaluating rewards
 		* @param CompletionFunction						Function to run for evaluating completions
 		* @param Instances								Set of instances to run training for
-		* @param TrainerFlags							Flags for the trainer, should match what was used to initialize the Trainer object.
 		* @param bRequestTrainingStopSignal				Optional signal that can be set to indicate training should be stopped
 		* @param PolicyNetworkLock						Optional Lock to use when updating the policy network
 		* @param CriticNetworkLock						Optional Lock to use when updating the critic network
+		* @param EncoderNetworkLock						Optional Lock to use when updating the encoder network
+		* @param DecoderNetworkLock						Optional Lock to use when updating the decoder network
 		* @param bPolicyNetworkUpdatedSignal			Optional signal that will be set when the policy network is updated
 		* @param bCriticNetworkUpdatedSignal			Optional signal that will be set when the critic network is updated
+		* @param bEncoderNetworkUpdatedSignal			Optional signal that will be set when the encoder network is updated
+		* @param bDecoderNetworkUpdatedSignal			Optional signal that will be set when the decoder network is updated
 		* @param LogSettings							Logging settings
 		* @returns										Trainer response in case of errors during communication otherwise Success
 		*/
@@ -598,15 +696,18 @@ namespace UE::Learning
 			FReplayBuffer& ReplayBuffer,
 			FEpisodeBuffer& EpisodeBuffer,
 			FResetInstanceBuffer& ResetBuffer,
-			INeuralNetwork& PolicyNetwork,
-			INeuralNetwork& CriticNetwork,
+			ULearningNeuralNetworkData& PolicyNetwork,
+			ULearningNeuralNetworkData& CriticNetwork,
+			ULearningNeuralNetworkData& EncoderNetwork,
+			ULearningNeuralNetworkData& DecoderNetwork,
 			TLearningArrayView<2, float> ObservationVectorBuffer,
 			TLearningArrayView<2, float> ActionVectorBuffer,
 			TLearningArrayView<2, float> PreEvaluationMemoryStateVectorBuffer,
 			TLearningArrayView<2, float> MemoryStateVectorBuffer,
 			TLearningArrayView<1, float> RewardBuffer,
 			TLearningArrayView<1, ECompletionMode> CompletionBuffer,
-			const ECompletionMode EpisodeEndCompletionMode,
+			TLearningArrayView<1, ECompletionMode> EpisodeCompletionBuffer,
+			TLearningArrayView<1, ECompletionMode> AllCompletionBuffer,
 			const TFunctionRef<void(const FIndexSet Instances)> ResetFunction,
 			const TFunctionRef<void(const FIndexSet Instances)> ObservationFunction,
 			const TFunctionRef<void(const FIndexSet Instances)> PolicyFunction,
@@ -615,12 +716,15 @@ namespace UE::Learning
 			const TFunctionRef<void(const FIndexSet Instances)> RewardFunction,
 			const TFunctionRef<void(const FIndexSet Instances)> CompletionFunction,
 			const FIndexSet Instances,
-			const EPPOTrainerFlags TrainerFlags = EPPOTrainerFlags::None,
 			TAtomic<bool>* bRequestTrainingStopSignal = nullptr,
 			FRWLock* PolicyNetworkLock = nullptr,
 			FRWLock* CriticNetworkLock = nullptr,
+			FRWLock* EncoderNetworkLock = nullptr,
+			FRWLock* DecoderNetworkLock = nullptr,
 			TAtomic<bool>* bPolicyNetworkUpdatedSignal = nullptr,
 			TAtomic<bool>* bCriticNetworkUpdatedSignal = nullptr,
+			TAtomic<bool>* bEncoderNetworkUpdatedSignal = nullptr,
+			TAtomic<bool>* bDecoderNetworkUpdatedSignal = nullptr,
 			const ELogSetting LogSettings = ELogSetting::Normal);
 	}
 
diff --git a/Engine/Plugins/Experimental/LearningAgents/Source/LearningTraining/Public/LearningTrainer.h b/Engine/Plugins/Experimental/LearningAgents/Source/LearningTraining/Public/LearningTrainer.h
--- a/Engine/Plugins/Experimental/LearningAgents/Source/LearningTraining/Public/LearningTrainer.h
+++ b/Engine/Plugins/Experimental/LearningAgents/Source/LearningTraining/Public/LearningTrainer.h
@@ -5,8 +5,24 @@
 #include "LearningArray.h"
 #include "LearningLog.h"
 
+#include "Templates/SharedPointer.h"
+
+class FJsonObject;
+
 namespace UE::Learning
 {
+	namespace Observation
+	{
+		struct FSchema;
+		struct FSchemaElement;
+	}
+
+	namespace Action
+	{
+		struct FSchema;
+		struct FSchemaElement;
+	}
+
 	/**
 	* Training Device
 	*/
@@ -120,6 +136,20 @@ namespace UE::Learning
 		* Gets the LearningAgents Intermediate path from the intermediate directory.
 		*/
 		LEARNINGTRAINING_API FString GetIntermediatePath(const FString& IntermediateDir);
+
+		/**
+		* Converts an observation schema element into a JSON representation.
+		*/
+		LEARNINGTRAINING_API TSharedPtr<FJsonObject> ConvertObservationSchemaToJSON(
+			const Observation::FSchema& ObservationSchema,
+			const Observation::FSchemaElement& ObservationSchemaElement);
+
+		/**
+		* Converts an action schema element into a JSON representation.
+		*/
+		LEARNINGTRAINING_API TSharedPtr<FJsonObject> ConvertActionSchemaToJSON(
+			const Action::FSchema& ActionSchema,
+			const Action::FSchemaElement& ActionSchemaElement);
 	}
 
 }
\ No newline at end of file
diff --git a/Engine/Plugins/Experimental/NNERuntimeBasicCpu/Source/NNERuntimeBasicCpu/Private/NNERuntimeBasicCpu.ispc b/Engine/Plugins/Experimental/NNERuntimeBasicCpu/Source/NNERuntimeBasicCpu/Private/NNERuntimeBasicCpu.ispc
--- a/Engine/Plugins/Experimental/NNERuntimeBasicCpu/Source/NNERuntimeBasicCpu/Private/NNERuntimeBasicCpu.ispc
+++ b/Engine/Plugins/Experimental/NNERuntimeBasicCpu/Source/NNERuntimeBasicCpu/Private/NNERuntimeBasicCpu.ispc
@@ -1,6 +1,17 @@
 // Copyright Epic Games, Inc. All Rights Reserved.
 
-export void NNERuntimeBasicCPUNormalizeLayer(
+inline float NNERuntimeBasicCPUTanH(const float X, const float Min = -10.0f, const float Max = 10.0f)
+{
+	const float Exp2X = exp(2.0f * clamp(X, Min, Max));
+	return (Exp2X - 1) / (Exp2X + 1);
+}
+
+inline float NNERuntimeBasicCPUSigmoid(const float X)
+{
+	return 1.0f / (1.0f + exp(-X));
+}
+
+export void NNERuntimeBasicCPUOperatorNormalize(
 	uniform float Output[],
 	uniform const float Input[],
 	uniform const float Mean[],
@@ -16,7 +27,7 @@ export void NNERuntimeBasicCPUNormalizeLayer(
 	}
 }
 
-export void NNERuntimeBasicCPUDenormalizeLayer(
+export void NNERuntimeBasicCPUOperatorDenormalize(
 	uniform float Output[],
 	uniform const float Input[],
 	uniform const float Mean[],
@@ -32,14 +43,14 @@ export void NNERuntimeBasicCPUDenormalizeLayer(
 	}
 }
 
-export void NNERuntimeBasicCPULinearLayer(
+export void NNERuntimeBasicCPUOperatorLinear(
 	uniform float Output[],
 	uniform const float Input[],
 	uniform const float Weights[],
 	uniform const float Biases[],
 	uniform const int BatchNum,
-	uniform const int RowNum,
 	uniform const int ColNum,
+	uniform const int RowNum,
 	uniform const int OutputStride,
 	uniform const int InputStride)
 {
@@ -65,7 +76,7 @@ export void NNERuntimeBasicCPULinearLayer(
 	}
 }
 
-export void NNERuntimeBasicCPUCompressedLinearLayer(
+export void NNERuntimeBasicCPUOperatorCompressedLinear(
 	uniform float Output[],
 	uniform const float Input[],
 	uniform const uint16 Weights[],
@@ -73,8 +84,8 @@ export void NNERuntimeBasicCPUCompressedLinearLayer(
 	uniform const float WeightScales[],
 	uniform const float Biases[],
 	uniform const int BatchNum,
-	uniform const int RowNum,
 	uniform const int ColNum,
+	uniform const int RowNum,
 	uniform const int OutputStride,
 	uniform const int InputStride)
 {
@@ -103,15 +114,15 @@ export void NNERuntimeBasicCPUCompressedLinearLayer(
 	}
 }
 
-export void NNERuntimeBasicCPUMultiLinearLayer(
+export void NNERuntimeBasicCPUOperatorMultiLinear(
 	uniform float Output[],
 	uniform const float Input[],
 	uniform const float Weights[],
 	uniform const float Biases[],
 	uniform const int BatchNum,
 	uniform const int BlockNum,
-	uniform const int RowNum,
 	uniform const int ColNum,
+	uniform const int RowNum,
 	uniform const int OutputStride,
 	uniform const int InputStride)
 {
@@ -140,7 +151,7 @@ export void NNERuntimeBasicCPUMultiLinearLayer(
 	}
 }
 
-export void NNERuntimeBasicCPUIdentityLayer(
+export void NNERuntimeBasicCPUOperatorCopy(
 	uniform float Output[],
 	uniform const float Input[],
 	uniform const int BatchNum,
@@ -154,7 +165,7 @@ export void NNERuntimeBasicCPUIdentityLayer(
 	}
 }
 
-export void NNERuntimeBasicCPUActivationReLU(
+export void NNERuntimeBasicCPUOperatorReLU(
 	uniform float Output[],
 	uniform const float Input[],
 	uniform const int BatchNum,
@@ -168,7 +179,7 @@ export void NNERuntimeBasicCPUActivationReLU(
 	}
 }
 
-export void NNERuntimeBasicCPUActivationELU(
+export void NNERuntimeBasicCPUOperatorELU(
 	uniform float Output[],
 	uniform const float Input[],
 	uniform const int BatchNum,
@@ -183,13 +194,7 @@ export void NNERuntimeBasicCPUActivationELU(
 	}
 }
 
-inline float NNERuntimeBasicCPUTanH(const float X, const float Min = -10.0f, const float Max = 10.0f)
-{
-	const float Exp2X = exp(2.0f * clamp(X, Min, Max));
-	return (Exp2X - 1) / (Exp2X + 1);
-}
-
-export void NNERuntimeBasicCPUActivationTanH(
+export void NNERuntimeBasicCPUOperatorTanH(
 	uniform float Output[],
 	uniform const float Input[],
 	uniform const int BatchNum,
@@ -203,7 +208,7 @@ export void NNERuntimeBasicCPUActivationTanH(
 	}
 }
 
-export void NNERuntimeBasicCPUActivationPReLU(
+export void NNERuntimeBasicCPUOperatorPReLU(
 	uniform float Output[],
 	uniform const float Input[],
 	uniform const float Alpha[],
@@ -218,3 +223,301 @@ export void NNERuntimeBasicCPUActivationPReLU(
 		Output[BatchIdx * OutputStride + Idx] = select(Value > 0.0f, Value, Alpha[Idx] * Value);
 	}
 }
+
+export void NNERuntimeBasicCPUOperatorMemoryCellUpdateMemory(
+	uniform float Output[],
+	uniform const float RememberGate[],
+	uniform const float Memory[],
+	uniform const float Update[],
+	uniform const int BatchSize,
+	uniform const int MemorySize,
+	uniform const int OutputStride,
+	uniform const int RememberGateStride,
+	uniform const int MemoryStride,
+	uniform const int UpdateStride)
+{
+	foreach(BatchIdx = 0 ... BatchSize, Idx = 0 ... MemorySize)
+	{
+		const float Gate = NNERuntimeBasicCPUSigmoid(RememberGate[BatchIdx * RememberGateStride + Idx]);
+		const float Prev = Memory[BatchIdx * MemoryStride + Idx];
+		const float Targ = NNERuntimeBasicCPUTanH(Update[BatchIdx * UpdateStride + Idx]);
+
+		Output[BatchIdx * OutputStride + Idx] = (1.0f - Gate) * Prev + Gate * Targ;
+	}
+}
+
+export void NNERuntimeBasicCPUOperatorMemoryCellUpdateOutput(
+	uniform float Output[],
+	uniform const float PassthroughGate[],
+	uniform const float MemoryUpdate[],
+	uniform const float InputUpdate[],
+	uniform const int BatchSize,
+	uniform const int OutputSize,
+	uniform const int OutputStride,
+	uniform const int PassthroughGateStride,
+	uniform const int MemoryUpdateStride,
+	uniform const int InputUpdateStride)
+{
+	foreach(BatchIdx = 0 ... BatchSize, Idx = 0 ... OutputSize)
+	{
+		const float Gate = NNERuntimeBasicCPUSigmoid(PassthroughGate[BatchIdx * PassthroughGateStride + Idx]);
+		const float MemTarg = NNERuntimeBasicCPUTanH(MemoryUpdate[BatchIdx * MemoryUpdateStride + Idx]);
+		const float InTarg = NNERuntimeBasicCPUTanH(InputUpdate[BatchIdx * InputUpdateStride + Idx]);
+
+		Output[BatchIdx * OutputStride + Idx] = (1.0f - Gate) * MemTarg + Gate * InTarg;
+	}
+}
+
+export void NNERuntimeBasicCPUOperatorAggregateGatherElements(
+	uniform float OutputBuffer[],
+	uniform const float InputBuffer[],
+	uniform const uint32 ElementNums[],
+	uniform const uint32 ElementOffsets[],
+	uniform const int BatchSize,
+	uniform const int ElementSize,
+	uniform const int InputStride)
+{
+	for (uniform int BatchIdx = 0; BatchIdx < BatchSize; BatchIdx++)
+	{
+		uniform const uint32 ElementNum = ElementNums[BatchIdx];
+		uniform const uint32 ElementOffset = ElementOffsets[BatchIdx];
+
+		foreach (ElementIdx = 0 ... ElementNum, Idx = 0 ... ElementSize)
+		{
+			OutputBuffer[(ElementOffset + ElementIdx) * ElementSize + Idx] = InputBuffer[BatchIdx * InputStride + ElementIdx * ElementSize + Idx];
+		}
+	}
+}
+
+export void NNERuntimeBasicCPUOperatorAggregateInsertOneHot(
+	uniform float QueryBuffer[],
+	uniform const int Index,
+	uniform const int BatchSize,
+	uniform const int MaskSize,
+	uniform const int QueryBufferStride)
+{
+	for (uniform int BatchIdx = 0; BatchIdx < BatchSize; BatchIdx++)
+	{
+		foreach (MaskIdx = 0 ... MaskSize)
+		{
+			QueryBuffer[BatchIdx * QueryBufferStride + MaskIdx] = 0.0f;
+		}
+
+		QueryBuffer[BatchIdx * QueryBufferStride + Index] = 1.0f;
+	}
+}
+
+export void NNERuntimeBasicCPUOperatorAggregateCountElementNum(
+	uniform uint32& TotalElementNum,
+	uniform uint32 ElementNums[],
+	uniform uint32 ElementOffsets[],
+	uniform const float MaskBuffer[],
+	uniform const int BatchSize,
+	uniform const int MaskSize,
+	uniform const int MaskBufferStride)
+{
+	TotalElementNum = 0;
+
+	for (uniform int BatchIdx = 0; BatchIdx < BatchSize; BatchIdx++)
+	{
+		uint32 ElementSumAccum = 0;
+
+		foreach (MaskIdx = 0 ... MaskSize)
+		{
+			if (MaskBuffer[BatchIdx * MaskBufferStride + MaskIdx]) { ElementSumAccum++; }
+		}
+
+		uniform uint32 ElementSum = reduce_add(ElementSumAccum);
+
+		ElementOffsets[BatchIdx] = TotalElementNum;
+		ElementNums[BatchIdx] = ElementSum;
+		TotalElementNum += ElementSum;
+	}
+}
+
+export void NNERuntimeBasicCPUOperatorAggregateDotProductAttention(
+	uniform float Attention[],
+	uniform const float Queries[],
+	uniform const float Keys[],
+	uniform const int ElementNum,
+	uniform const int AttentionEncodingSize,
+	uniform const int AttentionHeadNum)
+{
+	for (uniform int ElementIdx = 0; ElementIdx < ElementNum; ElementIdx++)
+	{
+		for (uniform int HeadIdx = 0; HeadIdx < AttentionHeadNum; HeadIdx++)
+		{
+			float AttentionAccum = 0.0f;
+
+			foreach (Idx = 0 ... AttentionEncodingSize)
+			{
+				AttentionAccum += (
+					Keys[ElementIdx * AttentionHeadNum * AttentionEncodingSize + HeadIdx * AttentionEncodingSize + Idx] *
+					Queries[ElementIdx * AttentionHeadNum * AttentionEncodingSize + HeadIdx * AttentionEncodingSize + Idx]);
+			}
+
+			Attention[ElementIdx * AttentionHeadNum + HeadIdx] = reduce_add(AttentionAccum) / sqrt((uniform float)AttentionEncodingSize);
+		}
+	}
+}
+
+export void NNERuntimeBasicCPUOperatorAggregateSoftmaxPlusOneInplace(
+	uniform float AttentionMaxs[],
+	uniform float AttentionDenoms[],
+	uniform float Attention[],
+	uniform const uint32 ElementNums[],
+	uniform const uint32 ElementOffsets[],
+	uniform const int BatchSize,
+	uniform const int AttentionHeadNum)
+{
+	for (uniform int BatchIdx = 0; BatchIdx < BatchSize; BatchIdx++)
+	{
+		uniform const uint32 ElementNum = ElementNums[BatchIdx];
+		uniform const uint32 ElementOffset = ElementOffsets[BatchIdx];
+
+		foreach (HeadIdx = 0 ... AttentionHeadNum)
+		{
+			AttentionMaxs[HeadIdx] = 0.0f;
+			AttentionDenoms[HeadIdx] = 0.0f;
+		}
+
+		for (uniform uint32 ElementIdx = ElementOffset; ElementIdx < ElementOffset + ElementNum; ElementIdx++)
+		{
+			foreach (HeadIdx = 0 ... AttentionHeadNum)
+			{
+				AttentionMaxs[HeadIdx] = max(AttentionMaxs[HeadIdx], Attention[ElementIdx * AttentionHeadNum + HeadIdx]);
+			}
+		}
+
+		for (uniform uint32 ElementIdx = ElementOffset; ElementIdx < ElementOffset + ElementNum; ElementIdx++)
+		{
+			foreach (HeadIdx = 0 ... AttentionHeadNum)
+			{
+				AttentionDenoms[HeadIdx] += exp(Attention[ElementIdx * AttentionHeadNum + HeadIdx] - AttentionMaxs[HeadIdx]);
+			}
+		}
+
+		foreach (ElementIdx = ElementOffset ... ElementOffset + ElementNum, HeadIdx = 0 ... AttentionHeadNum)
+		{
+			Attention[ElementIdx * AttentionHeadNum + HeadIdx] = exp(Attention[ElementIdx * AttentionHeadNum + HeadIdx] - AttentionMaxs[HeadIdx]) / (AttentionDenoms[HeadIdx] + 1.0f);
+		}
+	}
+}
+
+export void NNERuntimeBasicCPUOperatorAggregateAttentionSum(
+	uniform float Output[],
+	uniform const float Attention[],
+	uniform const float Values[],
+	uniform const uint32 ElementNums[],
+	uniform const uint32 ElementOffsets[],
+	uniform const int BatchSize,
+	uniform const int EncodingSize,
+	uniform const int AttentionHeadNum,
+	uniform const int OutputStride)
+{
+	for (uniform int BatchIdx = 0; BatchIdx < BatchSize; BatchIdx++)
+	{
+		foreach (HeadIdx = 0 ... AttentionHeadNum, Idx = 0 ... EncodingSize)
+		{
+			Output[BatchIdx * OutputStride + HeadIdx * EncodingSize + Idx] = 0.0f;
+		}
+
+		uniform const uint32 ElementNum = ElementNums[BatchIdx];
+		uniform const uint32 ElementOffset = ElementOffsets[BatchIdx];
+
+		for (uniform uint32 ElementIdx = ElementOffset; ElementIdx < ElementOffset + ElementNum; ElementIdx++)
+		{
+			for (uniform int HeadIdx = 0; HeadIdx < AttentionHeadNum; HeadIdx++)
+			{
+				uniform const float Scale = Attention[ElementIdx * AttentionHeadNum + HeadIdx];
+
+				if (Scale != 0.0f)
+				{
+					foreach(Idx = 0 ... EncodingSize)
+					{
+						Output[BatchIdx * OutputStride + HeadIdx * EncodingSize + Idx] += Scale * Values[ElementIdx * AttentionHeadNum * EncodingSize + HeadIdx * EncodingSize + Idx];
+					}
+				}
+			}
+		}
+	}
+}
+
+export void NNERuntimeBasicCPUOperatorAggregateGatherQueryValueFromSubLayers(
+	uniform float QueryBuffer[],
+	uniform float KeyBuffer[],
+	uniform float ValueBuffer[],
+	uniform uint32 ElementAccum[],
+	uniform const uint32 ElementOffsets[],
+	uniform const float SubLayerQueryBuffer[],
+	uniform const float SubLayerKeyBuffer[],
+	uniform const float SubLayerValueBuffer[],
+	uniform const uint32 SubLayerBatchIndicesBuffer[],
+	uniform const int SubLayerBatchIndexNum,
+	uniform const int QuerySize,
+	uniform const int KeySize,
+	uniform const int ValueSize)
+{
+	for (uniform int ElementIdx = 0; ElementIdx < SubLayerBatchIndexNum; ElementIdx++)
+	{
+		uniform const uint32 BatchIdx = SubLayerBatchIndicesBuffer[ElementIdx];
+		uniform const uint32 ElementOffset = ElementOffsets[BatchIdx] + ElementAccum[BatchIdx];
+
+		foreach (QueryIdx = 0 ... QuerySize)
+		{
+			QueryBuffer[ElementOffset * QuerySize + QueryIdx] = SubLayerQueryBuffer[ElementIdx * QuerySize + QueryIdx];
+		}
+
+		foreach(KeyIdx = 0 ... KeySize)
+		{
+			KeyBuffer[ElementOffset * KeySize + KeyIdx] = SubLayerKeyBuffer[ElementIdx * KeySize + KeyIdx];
+		}
+
+		foreach(ValueIdx = 0 ... ValueSize)
+		{
+			ValueBuffer[ElementOffset * ValueSize + ValueIdx] = SubLayerValueBuffer[ElementIdx * ValueSize + ValueIdx];
+		}
+
+		ElementAccum[BatchIdx]++;
+	}
+}
+
+export void NNERuntimeBasicCPUOperatorGather(
+	uniform float OutputBuffer[],
+	uniform const float InputBuffer[],
+	uniform const uint32 BatchIndices[],
+	uniform const int BatchIndexNum,
+	uniform const int InputOutputSize,
+	uniform const int OutputStride,
+	uniform const int InputStride)
+{
+	for (uniform int BatchIndexIdx = 0; BatchIndexIdx < BatchIndexNum; BatchIndexIdx++)
+	{
+		uniform const uint32 SrcBatchIdx = BatchIndices[BatchIndexIdx];
+
+		foreach (Idx = 0 ... InputOutputSize)
+		{
+			OutputBuffer[BatchIndexIdx * OutputStride + Idx] = InputBuffer[SrcBatchIdx * InputStride + Idx];
+		}
+	}
+}
+
+export void NNERuntimeBasicCPUOperatorScatter(
+	uniform float OutputBuffer[],
+	uniform const float InputBuffer[],
+	uniform const uint32 BatchIndices[],
+	uniform const int BatchIndexNum,
+	uniform const int InputOutputSize,
+	uniform const int OutputStride,
+	uniform const int InputStride)
+{
+	for (uniform int BatchIndexIdx = 0; BatchIndexIdx < BatchIndexNum; BatchIndexIdx++)
+	{
+		uniform const uint32 DstBatchIdx = BatchIndices[BatchIndexIdx];
+
+		foreach (Idx = 0 ... InputOutputSize)
+		{
+			OutputBuffer[DstBatchIdx * OutputStride + Idx] = InputBuffer[BatchIndexIdx * InputStride + Idx];
+		}
+	}
+}
\ No newline at end of file
