[Backout] - CL30154408
Original CL Desc
-----------------------------------------------------------------
LearningAgents: Large restructuring to allow for structured actions and observations.

#virtualized

diff --git a/Engine/Plugins/Animation/MLDeformer/NearestNeighborModel/Source/NearestNeighborModel/Private/NearestNeighborOptimizedNetwork.cpp b/Engine/Plugins/Animation/MLDeformer/NearestNeighborModel/Source/NearestNeighborModel/Private/NearestNeighborOptimizedNetwork.cpp
--- a/Engine/Plugins/Animation/MLDeformer/NearestNeighborModel/Source/NearestNeighborModel/Private/NearestNeighborOptimizedNetwork.cpp
+++ b/Engine/Plugins/Animation/MLDeformer/NearestNeighborModel/Source/NearestNeighborModel/Private/NearestNeighborOptimizedNetwork.cpp
@@ -9,7 +9,7 @@
 #include "NNERuntime.h"
 #include "NNERuntimeCPU.h"
 #include "NNEModelData.h"
-#include "NNERuntimeBasicCpuBuilder.h"
+#include "NNERuntimeBasicCpu.h"
 
 //--------------------------------------------------------------------------
 // UNearestNeighborOptimizedNetwork
@@ -22,30 +22,33 @@ namespace UE::NearestNeighborModel::Private
 	/** Creates the FileData from the legacy network format and clears it. */
 	static inline void CreateFileDataAndClearLayers(TArray<uint8>& OutFileData, TArray<TObjectPtr<UNearestNeighborNetworkLayer>>& Layers)
 	{
-		UE::NNE::RuntimeBasic::FModelBuilder Builder;
+		UE::NNE::RuntimeBasic::FSequentialModelBuilder Builder;
 
-		TArray<UE::NNE::RuntimeBasic::FModelBuilderElement, TInlineAllocator<32>> LayerElements;
-		LayerElements.Reserve(2 * Layers.Num());
+		// Allocate data for the PReLuAlpha values as these need to be kept around until we run "Builder.WriteAndReset"
+		TArray<TArray<float>> PreluAlphaData;
 
 		for (UNearestNeighborNetworkLayer* Layer : Layers)
 		{
 			if (UNearestNeighborNetworkLayer_Gemm_Prelu* GemmPreluLayer = Cast<UNearestNeighborNetworkLayer_Gemm_Prelu>(Layer))
 			{
-				LayerElements.Add(Builder.MakeLinear(
+				Builder.AddLinear(
 					GemmPreluLayer->NumInputs,
 					GemmPreluLayer->NumOutputs,
 					GemmPreluLayer->Parameters[0].Values,
-					GemmPreluLayer->Parameters[1].Values));
+					GemmPreluLayer->Parameters[1].Values);
 
-				LayerElements.Add(Builder.MakePReLU(GemmPreluLayer->NumOutputs, Builder.MakeWeightsConstant(GemmPreluLayer->NumOutputs, GemmPreluLayer->Parameters[2].Values[0])));
+				PreluAlphaData.AddDefaulted();
+				PreluAlphaData.Last().Init(GemmPreluLayer->Parameters[2].Values[0], GemmPreluLayer->NumOutputs);
+
+				Builder.AddPReLU(PreluAlphaData.Last());
 			}
 			else if (UNearestNeighborNetworkLayer_Gemm* GemmLayer = Cast<UNearestNeighborNetworkLayer_Gemm>(Layer))
 			{
-				LayerElements.Add(Builder.MakeLinear(
+				Builder.AddLinear(
 					GemmLayer->NumInputs,
 					GemmLayer->NumOutputs,
 					GemmLayer->Parameters[0].Values,
-					GemmLayer->Parameters[1].Values));
+					GemmLayer->Parameters[1].Values);
 			}
 			else
 			{
@@ -53,10 +56,10 @@ namespace UE::NearestNeighborModel::Private
 			}
 		}
 
-		uint32 InputSize, OutputSize;
-		Builder.WriteFileDataAndReset(OutFileData, InputSize, OutputSize, Builder.MakeSequence(LayerElements));
-
 		Layers.Empty();
+
+		OutFileData.SetNumUninitialized(Builder.GetWriteByteNum());
+		Builder.WriteAndReset(OutFileData);
 	}
 
 	PRAGMA_ENABLE_DEPRECATION_WARNINGS
diff --git a/Engine/Plugins/Animation/MLDeformer/NeuralMorphModel/Source/NeuralMorphModel/Private/NeuralMorphNetwork.cpp b/Engine/Plugins/Animation/MLDeformer/NeuralMorphModel/Source/NeuralMorphModel/Private/NeuralMorphNetwork.cpp
--- a/Engine/Plugins/Animation/MLDeformer/NeuralMorphModel/Source/NeuralMorphModel/Private/NeuralMorphNetwork.cpp
+++ b/Engine/Plugins/Animation/MLDeformer/NeuralMorphModel/Source/NeuralMorphModel/Private/NeuralMorphNetwork.cpp
@@ -11,7 +11,7 @@
 #include "NNERuntime.h"
 #include "NNERuntimeCPU.h"
 #include "NNEModelData.h"
-#include "NNERuntimeBasicCpuBuilder.h"
+#include "NNERuntimeBasicCpu.h"
 
 //--------------------------------------------------------------------------
 // UNeuralMorphMLP
@@ -275,13 +275,10 @@ namespace UE::NeuralMorphModel::Private
 
 	static inline void ConvertMLPToFileData(TArray<uint8>& OutBytes, TObjectPtr<UNeuralMorphMLP>& MLP)
 	{
-		UE::NNE::RuntimeBasic::FModelBuilder Builder;
+		UE::NNE::RuntimeBasic::FSequentialModelBuilder Builder;
 
 		const uint32 LayerNum = MLP->GetNumLayers();
 
-		TArray<UE::NNE::RuntimeBasic::FModelBuilderElement, TInlineAllocator<32>> LayerElements;
-		LayerElements.Reserve(2 * LayerNum);
-
 		for (uint32 LayerIdx = 0; LayerIdx < LayerNum; LayerIdx++)
 		{
 			const uint32 InputSize = MLP->GetLayer(LayerIdx).NumInputs;
@@ -292,20 +289,20 @@ namespace UE::NeuralMorphModel::Private
 
 			if (BlockNum == 1)
 			{
-				LayerElements.Add(Builder.MakeLinear(InputSize, OutputSize, Weights, Biases));
+				Builder.AddLinear(InputSize, OutputSize, Weights, Biases);
 			}
 			else
 			{
-				LayerElements.Add(Builder.MakeMultiLinear(InputSize, OutputSize, BlockNum, Weights, Biases));
+				Builder.AddMultiLinear(InputSize, OutputSize, BlockNum, Weights, Biases);
 			}
 
-			LayerElements.Add(Builder.MakeELU(OutputSize * BlockNum));
+			Builder.AddELU();
 		}
 
-		uint32 InputSize, OutputSize;
-		Builder.WriteFileDataAndReset(OutBytes, InputSize, OutputSize, Builder.MakeSequence(LayerElements));
-
 		MLP = nullptr;
+
+		OutBytes.SetNumUninitialized(Builder.GetWriteByteNum());
+		Builder.WriteAndReset(OutBytes);
 	}
 
 	PRAGMA_ENABLE_DEPRECATION_WARNINGS
diff --git a/Engine/Plugins/Experimental/LearningAgents/LearningAgents.uplugin b/Engine/Plugins/Experimental/LearningAgents/LearningAgents.uplugin
--- a/Engine/Plugins/Experimental/LearningAgents/LearningAgents.uplugin
+++ b/Engine/Plugins/Experimental/LearningAgents/LearningAgents.uplugin
@@ -31,6 +31,11 @@
 			"Type": "Runtime",
 			"LoadingPhase": "Default"
 		},
+		{
+			"Name": "LearningAgentsEditor",
+			"Type": "Editor",
+			"LoadingPhase": "PostEngineInit"
+		},
 		{
 			"Name": "LearningAgentsTraining",
 			"Type": "Runtime",
diff --git a/Engine/Plugins/Experimental/LearningAgents/Source/Learning/Learning.Build.cs b/Engine/Plugins/Experimental/LearningAgents/Source/Learning/Learning.Build.cs
--- a/Engine/Plugins/Experimental/LearningAgents/Source/Learning/Learning.Build.cs
+++ b/Engine/Plugins/Experimental/LearningAgents/Source/Learning/Learning.Build.cs
@@ -38,10 +38,7 @@ public class Learning : ModuleRules
 		PrivateDependencyModuleNames.AddRange(
 			new string[]
 			{
-				"CoreUObject",
 				"Engine",
-				"NNE",
-				"NNERuntimeBasicCpu",
 				// ... add private dependencies that you statically link with here ...	
 			}
 			);
diff --git a/Engine/Plugins/Experimental/LearningAgents/Source/Learning/Private/Learning.ispc b/Engine/Plugins/Experimental/LearningAgents/Source/Learning/Private/Learning.ispc
--- a/Engine/Plugins/Experimental/LearningAgents/Source/Learning/Private/Learning.ispc
+++ b/Engine/Plugins/Experimental/LearningAgents/Source/Learning/Private/Learning.ispc
@@ -5,18 +5,9 @@
 #endif
 
 #ifndef INDEX_NONE
-#define INDEX_NONE		(-1)
+#define INDEX_NONE		-1
 #endif
 
-#ifndef UE_SMALL_NUMBER
-#define UE_SMALL_NUMBER (1.e-8f)
-#endif
-
-inline float LearningSigmoid(const float X)
-{
-	return 1.0f / (1.0f + exp(-X));
-}
-
 /////////////////////////////////////////////////////////
 // 
 // Functions matching those in LearningRandom.cpp
@@ -68,16 +59,16 @@ inline uniform int32 LearningRandomIntInRange(uniform const uint32 State, unifor
 
 inline float LearningRandomUniform(
 	const uint32 State,
-	const float Min = 0.0f,
-	const float Max = 1.0f)
+	const float Min,
+	const float Max)
 {
 	return (Max - Min) * LearningRandomFloat(State ^ 0x72404541) + Min;
 }
 
 inline float LearningRandomGaussian(
 	const uint32 State,
-	const float Mean = 0.0f,
-	const float Std = 1.0f)
+	const float Mean,
+	const float Std)
 {
 	return Std * LearningRandomUniformToGaussian(
 		LearningRandomFloat(State ^ 0x4855e88f),
@@ -132,50 +123,154 @@ export void LearningRandomGaussianArray(
 	}
 }
 
-export void LearningRandomDistributionIndependantNormal(
+export void LearningResampleStateArray(
+	uniform uint32 States[],
+	uniform const int DimNum)
+{
+	foreach (DimIdx = 0 ... DimNum)
+	{
+		States[DimIdx] = LearningRandomInt(States[DimIdx] ^ 0x3616dcc8 ^ LearningRandomInt(DimIdx ^ 0x6a837ffd));
+	}
+}
+
+/////////////////////////////////////////////////////////
+// 
+// Some functions related to evaluation of Neural Networks
+// 
+/////////////////////////////////////////////////////////
+
+export void LearningLayerMatMulPlusBias(
 	uniform float Output[],
-	uniform const float Mean[],
-	uniform const float LogStd[],
-	uniform const int DimNum,
-	uniform const uint32 State,
-	uniform const float Scale)
+	const uniform float Input[],
+	const uniform float Weights[],
+	const uniform float Biases[],
+	const uniform int BatchNum,
+	const uniform int BufferSize,
+	const uniform int RowNum,
+	const uniform int ColNum)
 {
-	if (Scale < UE_SMALL_NUMBER)
+	foreach (BatchIdx = 0 ... BatchNum, ColIdx = 0 ... ColNum)
 	{
-		foreach(DimIdx = 0 ... DimNum)
-		{
-			Output[DimIdx] = Mean[DimIdx];
-		}
+		Output[BatchIdx * BufferSize + ColIdx] = Biases[ColIdx];
 	}
-	else
+	
+	for (uniform int BatchIdx = 0; BatchIdx < BatchNum; BatchIdx++)
 	{
-		foreach(DimIdx = 0 ... DimNum)
+		for (uniform int RowIdx = 0; RowIdx < RowNum; RowIdx++)
 		{
-			Output[DimIdx] = LearningRandomGaussian(State ^ 0x7db0e4e9 ^ LearningRandomInt(DimIdx ^ 0xf4976a00), Mean[DimIdx], Scale * exp(min(LogStd[DimIdx], 10.0f)));
+			uniform const float Value = Input[BatchIdx * BufferSize + RowIdx];
+
+			if (Value != 0.0)
+			{
+				foreach(ColIdx = 0 ... ColNum)
+				{
+					Output[BatchIdx * BufferSize + ColIdx] += Value * Weights[RowIdx * ColNum + ColIdx];
+				}
+			}
 		}
 	}
 }
 
-export void LearningRandomDistributionBernoulli(
+export void LearningLayerReLU(
+	uniform float InputOutput[],
+	const uniform int BatchNum,
+	const uniform int BufferSize,
+	const uniform int ColNum)
+{
+	foreach(BatchIdx = 0 ... BatchNum, ColIdx = 0 ... ColNum)
+	{
+		InputOutput[BatchIdx * BufferSize + ColIdx] = max(InputOutput[BatchIdx * BufferSize + ColIdx], 0.0f);
+	}
+}
+
+inline float LearningELU(const float X)
+{
+	return X > 0.0f ? X : exp(X) - 1.0f;
+}
+
+export void LearningLayerELU(
+	uniform float InputOutput[],
+	const uniform int BatchNum,
+	const uniform int BufferSize,
+	const uniform int ColNum)
+{
+	foreach (BatchIdx = 0 ... BatchNum, ColIdx = 0 ... ColNum)
+	{
+		InputOutput[BatchIdx * BufferSize + ColIdx] = LearningELU(InputOutput[BatchIdx * BufferSize + ColIdx]);
+	}
+}
+
+inline float LearningTanH(const float X, const float Min=-10.0f, const float Max=10.0f)
+{
+	const float Exp2X = exp(2.0f * clamp(X, Min, Max));
+	return (Exp2X - 1) / (Exp2X + 1);
+}
+
+export void LearningLayerTanH(
+	uniform float InputOutput[],
+	const uniform int BatchNum,
+	const uniform int BufferSize,
+	const uniform int ColNum)
+{
+	foreach(BatchIdx = 0 ... BatchNum, ColIdx = 0 ... ColNum)
+	{
+		InputOutput[BatchIdx * BufferSize + ColIdx] = LearningTanH(InputOutput[BatchIdx * BufferSize + ColIdx]);
+	}
+}
+
+inline float LearningSigmoid(const float X)
+{
+	return 1.0f / (1.0f + exp(-X));
+}
+
+export void LearningLayerActionNoise(
 	uniform float Output[],
-	uniform const float Logits[],
-	uniform const int DimNum,
-	uniform const uint32 State,
-	uniform const float Scale)
+	uniform float OutputMean[],
+	uniform float OutputStd[],
+	const uniform float OutputNetwork[],
+	const uniform uint32 Seed[],
+	const uniform float ActionNoiseScale[],
+	const uniform int BatchNum,
+	const uniform int ActionNum,
+	const uniform int MemoryStateNum,
+	const uniform float LogActionNoiseMin,
+	const uniform float LogActionNoiseMax)
 {
-	if (Scale < UE_SMALL_NUMBER)
+	foreach (BatchIdx = 0 ... BatchNum, ActionIdx = 0 ... ActionNum)
 	{
-		foreach(DimIdx = 0 ... DimNum)
-		{
-			Output[DimIdx] = select(Logits[DimIdx] > 0.0f, 1.0f, 0.0f);
-		}
+		OutputMean[BatchIdx * ActionNum + ActionIdx] = OutputNetwork[BatchIdx * (2 * ActionNum + MemoryStateNum) + ActionIdx];
+		OutputStd[BatchIdx * ActionNum + ActionIdx] = ActionNoiseScale[BatchIdx] *
+			exp(LearningSigmoid(OutputNetwork[BatchIdx * (2 * ActionNum + MemoryStateNum) + ActionNum + ActionIdx]) * (LogActionNoiseMax - LogActionNoiseMin) + LogActionNoiseMin);
+
+		Output[BatchIdx * ActionNum + ActionIdx] = LearningRandomGaussian(
+			Seed[BatchIdx] ^ 0xab744615 ^ LearningRandomInt(ActionIdx ^ 0xf8a88a27),
+			OutputMean[BatchIdx * ActionNum + ActionIdx],
+			OutputStd[BatchIdx * ActionNum + ActionIdx]);
 	}
-	else
+}
+
+export void LearningLayerActionNoiseSingleBatch(
+	uniform float Output[],
+	uniform float OutputMean[],
+	uniform float OutputStd[],
+	const uniform float OutputNetwork[],
+	const uniform uint32 Seed,
+	const uniform float ActionNoiseScale,
+	const uniform int ActionNum,
+	const uniform int MemoryStateNum,
+	const uniform float LogActionNoiseMin,
+	const uniform float LogActionNoiseMax)
+{
+	foreach(ActionIdx = 0 ... ActionNum)
 	{
-		foreach(DimIdx = 0 ... DimNum)
-		{
-			Output[DimIdx] = select(LearningSigmoid(Logits[DimIdx] / Scale) > LearningRandomUniform(State ^ 0xf4021a46 ^ LearningRandomInt(DimIdx ^ 0x7a8cc64e)), 1.0f, 0.0f);
-		}
+		OutputMean[ActionIdx] = OutputNetwork[ActionIdx];
+		OutputStd[ActionIdx] = ActionNoiseScale *
+			exp(LearningSigmoid(OutputNetwork[ActionNum + ActionIdx]) * (LogActionNoiseMax - LogActionNoiseMin) + LogActionNoiseMin);
+
+		Output[ActionIdx] = LearningRandomGaussian(
+			Seed ^ 0xab744615 ^ LearningRandomInt(ActionIdx ^ 0xf8a88a27),
+			OutputMean[ActionIdx],
+			OutputStd[ActionIdx]);
 	}
 }
 
@@ -420,6 +515,70 @@ export void LearningUpdatePSOBest(
 	}
 }
 
+/////////////////////////////////////////////////////////
+// 
+// Functions related to constructing feature vectors
+// 
+/////////////////////////////////////////////////////////
+
+export void LearningCombineFeature(
+	uniform float Output[],
+	uniform const float Buffer[],
+	uniform const int DimOffset,
+	uniform const int DimNum,
+	uniform const int OutputDimNum,
+	uniform const int InstanceNum,
+	uniform const float Scale)
+{
+	foreach (InstanceIdx = 0 ... InstanceNum, DimIdx = 0 ... DimNum)
+	{
+		Output[InstanceIdx * OutputDimNum + DimOffset + DimIdx] = Scale * Buffer[InstanceIdx * DimNum + DimIdx];
+	}
+}
+
+export void LearningCombineFeatureSingleInstance(
+	uniform float Output[],
+	uniform const float Buffer[],
+	uniform const int DimOffset,
+	uniform const int DimNum,
+	uniform const float Scale)
+{
+	foreach(DimIdx = 0 ... DimNum)
+	{
+		Output[DimOffset + DimIdx] = Scale * Buffer[DimIdx];
+	}
+}
+
+export void LearningSeparateFeature(
+	uniform float Buffer[],
+	uniform const float Input[],
+	uniform const int DimOffset,
+	uniform const int DimNum,
+	uniform const int OutputDimNum,
+	uniform const int InstanceNum,
+	uniform const float Scale,
+	uniform const float Epsilon)
+{
+	foreach (InstanceIdx = 0 ... InstanceNum, DimIdx = 0 ... DimNum)
+	{
+		Buffer[InstanceIdx * DimNum + DimIdx] = Input[InstanceIdx * OutputDimNum + DimOffset + DimIdx] / max(Scale, Epsilon);
+	}
+}
+
+export void LearningSeparateFeatureSingleInstance(
+	uniform float Buffer[],
+	uniform const float Input[],
+	uniform const int DimOffset,
+	uniform const int DimNum,
+	uniform const float Scale,
+	uniform const float Epsilon)
+{
+	foreach(DimIdx = 0 ... DimNum)
+	{
+		Buffer[DimIdx] = Input[DimOffset + DimIdx] / max(Scale, Epsilon);
+	}
+}
+
 /////////////////////////////////////////////////////////
 // 
 // Functions related to K-Means Clustering
diff --git a/Engine/Plugins/Experimental/LearningAgents/Source/Learning/Private/LearningArrayMap.cpp b/Engine/Plugins/Experimental/LearningAgents/Source/Learning/Private/LearningArrayMap.cpp
new file mode 100644
--- /dev/null
+++ b/Engine/Plugins/Experimental/LearningAgents/Source/Learning/Private/LearningArrayMap.cpp
@@ -0,0 +1,301 @@
+// Copyright Epic Games, Inc. All Rights Reserved.
+
+#include "LearningArrayMap.h"
+
+namespace UE::Learning
+{
+	FArrayMap::FArrayMap() {}
+	FArrayMap::~FArrayMap() { Empty(); }
+
+	void FArrayMap::Empty()
+	{
+		Handles.Empty();
+
+		const int32 VariableNum = Arrays.Num();
+
+		for (int32 VariableIdx = 0; VariableIdx < VariableNum; VariableIdx++)
+		{
+			if (Flags[VariableIdx] & FlagConstructed)
+			{
+				Destructors[VariableIdx](
+					Arrays[VariableIdx].Data,
+					DimNums[VariableIdx],
+					Arrays[VariableIdx].Shape);
+			}
+
+			if (Flags[VariableIdx] & FlagAllocated)
+			{
+				FMemory::Free(Arrays[VariableIdx].Data);
+			}
+		}
+
+		Arrays.Empty();
+		Destructors.Empty();
+		DimNums.Empty();
+		Flags.Empty();
+		TypeIds.Empty();
+		TypeNames.Empty();
+	}
+
+	FArrayMapHandle FArrayMap::Lookup(const FArrayMapKey Key) const
+	{
+		return Handles[Key];
+	}
+
+	const FArrayMapHandle* FArrayMap::Find(const FArrayMapKey Key) const
+	{
+		return Handles.Find(Key);
+	}
+
+	int16 FArrayMap::TypeId(const FArrayMapKey Key) const
+	{
+		return TypeIds[Lookup(Key).Index];
+	}
+
+	int16 FArrayMap::TypeId(const FArrayMapHandle Variable) const
+	{
+		return TypeIds[Variable.Index];
+	}
+
+	const TCHAR* FArrayMap::TypeName(const FArrayMapKey Key) const
+	{
+		return TypeNames[Lookup(Key).Index];
+	}
+
+	const TCHAR* FArrayMap::TypeName(const FArrayMapHandle Variable) const
+	{
+		return TypeNames[Variable.Index];
+	}
+
+	uint8 FArrayMap::DimNum(const FArrayMapKey Key) const
+	{
+		return DimNums[Lookup(Key).Index];
+	}
+
+	uint8 FArrayMap::DimNum(const FArrayMapHandle Handle) const
+	{
+		return DimNums[Handle.Index];
+	}
+
+	bool FArrayMap::Contains(const FArrayMapKey Key) const
+	{
+		return Handles.Contains(Key);
+	}
+
+	void FArrayMap::LinkKeys(const FArrayMapKey Src, const FArrayMapKey Dst)
+	{
+		const FArrayMapHandle* SrcVariable = Handles.Find(Src);
+		const FArrayMapHandle* DstVariable = Handles.Find(Dst);
+
+		UE_LEARNING_CHECKF(SrcVariable, TEXT("Does not contain an array called { \"%s\", \"%s\" }"),
+			*Src.Namespace.ToString(), *Src.Variable.ToString());
+
+		UE_LEARNING_CHECKF(DstVariable, TEXT("Does not contain an array called { \"%s\", \"%s\" }"),
+			*Dst.Namespace.ToString(), *Dst.Variable.ToString());
+
+		const int16 DstIndex = DstVariable->Index;
+		const int16 SrcIndex = SrcVariable->Index;
+
+		UE_LEARNING_CHECKF(!(Flags[DstIndex] & FlagLinked), TEXT("Destination array is already linked!"));
+
+		UE_LEARNING_CHECKF(TypeIds[SrcIndex] == TypeIds[DstIndex],
+			TEXT("Source array { \"%s\", \"%s\" } is of type %s, while destination array { \"%s\", \"%s\" } is of type %s"),
+			*Src.Namespace.ToString(), *Src.Variable.ToString(), TypeNames[SrcIndex],
+			*Dst.Namespace.ToString(), *Dst.Variable.ToString(), TypeNames[DstIndex]);
+
+		UE_LEARNING_CHECKF(DimNums[SrcIndex] == DimNums[DstIndex],
+			TEXT("Source array { \"%s\", \"%s\" } has %i dimensions, while destination array { \"%s\", \"%s\" } has %i dimensions"),
+			*Src.Namespace.ToString(), *Src.Variable.ToString(), DimNums[SrcIndex],
+			*Dst.Namespace.ToString(), *Dst.Variable.ToString(), DimNums[DstIndex]);
+
+		const uint8 DimNum = DimNums[SrcIndex];
+
+		for (uint8 DimIdx = 0; DimIdx < DimNum; DimIdx++)
+		{
+			UE_LEARNING_CHECKF(Arrays[SrcIndex].Shape[DimIdx] == Arrays[DstIndex].Shape[DimIdx],
+				TEXT("Array Shapes don't match on dimension %i of %i ({ \"%s\", \"%s\" }: %i, { \"%s\", \"%s\" }: %i)"),
+				DimIdx + 1, DimNum,
+				*Src.Namespace.ToString(), *Src.Variable.ToString(), Arrays[SrcIndex].Shape[DimIdx],
+				*Dst.Namespace.ToString(), *Dst.Variable.ToString(), Arrays[DstIndex].Shape[DimIdx]);
+		}
+
+		// Free the Destination array data
+		if (Flags[DstIndex] & FlagConstructed)
+		{
+			Destructors[DstIndex](
+				Arrays[DstIndex].Data,
+				DimNums[DstIndex],
+				Arrays[DstIndex].Shape);
+		}
+
+		if (Flags[DstIndex] & FlagAllocated)
+		{
+			FMemory::Free(Arrays[DstIndex].Data);
+		}
+
+		// Re-direct existing Links
+		for (int16 ArrayIdx = 0; ArrayIdx < Arrays.Num(); ArrayIdx++)
+		{
+			if (Flags[ArrayIdx] == FlagLinked && Arrays[ArrayIdx].Data == Arrays[DstIndex].Data)
+			{
+				Arrays[ArrayIdx].Data = Arrays[SrcIndex].Data;
+			}
+		}
+
+		// Insert the new Link
+		Arrays[DstIndex].Data = Arrays[SrcIndex].Data;
+		Flags[DstIndex] = FlagLinked;
+	}
+
+	void FArrayMap::LinkHandles(const FArrayMapHandle Src, const FArrayMapHandle Dst)
+	{
+		const int16 SrcIndex = Src.Index;
+		const int16 DstIndex = Dst.Index;
+
+		UE_LEARNING_CHECKF(!(Flags[DstIndex] & FlagLinked), TEXT("Destination array is already linked!"));
+
+		UE_LEARNING_CHECKF(TypeIds[SrcIndex] == TypeIds[DstIndex],
+			TEXT("Source array is of type %s, while destination array is of type %s"),
+			TypeNames[SrcIndex], TypeNames[DstIndex]);
+
+		UE_LEARNING_CHECKF(DimNums[SrcIndex] == DimNums[DstIndex],
+			TEXT("Source array has %i dimensions, while destination array has %i dimensions"),
+			DimNums[SrcIndex], DimNums[DstIndex]);
+
+		const uint8 DimNum = DimNums[SrcIndex];
+
+		for (uint8 DimIdx = 0; DimIdx < DimNum; DimIdx++)
+		{
+			UE_LEARNING_CHECKF(Arrays[SrcIndex].Shape[DimIdx] == Arrays[DstIndex].Shape[DimIdx],
+				TEXT("Array Shapes don't match on dimension %i of %i (%i vs %i)"),
+				DimIdx + 1, DimNum, Arrays[SrcIndex].Shape[DimIdx], Arrays[DstIndex].Shape[DimIdx]);
+		}
+
+		// Free the Destination array data
+		if (Flags[DstIndex] & FlagConstructed)
+		{
+			Destructors[DstIndex](
+				Arrays[DstIndex].Data,
+				DimNums[DstIndex],
+				Arrays[DstIndex].Shape);
+		}
+
+		if (Flags[DstIndex] & FlagAllocated)
+		{
+			FMemory::Free(Arrays[DstIndex].Data);
+		}
+
+		// Re-direct existing Links
+		for (int16 ArrayIdx = 0; ArrayIdx < Arrays.Num(); ArrayIdx++)
+		{
+			if (Flags[ArrayIdx] == FlagLinked && Arrays[ArrayIdx].Data == Arrays[DstIndex].Data)
+			{
+				Arrays[ArrayIdx].Data = Arrays[SrcIndex].Data;
+			}
+		}
+
+		// Insert the new Link
+		Arrays[DstIndex].Data = Arrays[SrcIndex].Data;
+		Flags[DstIndex] = FlagLinked;
+	}
+
+	void FArrayMap::LinkHandlesFlat(const FArrayMapHandle Src, const FArrayMapHandle Dst)
+	{
+		const int16 SrcIndex = Src.Index;
+		const int16 DstIndex = Dst.Index;
+
+		UE_LEARNING_CHECKF(!(Flags[DstIndex] & FlagLinked), TEXT("Destination array is already linked!"));
+
+		UE_LEARNING_CHECKF(TypeIds[SrcIndex] == TypeIds[DstIndex],
+			TEXT("Source array is of type %s, while destination array is of type %s"),
+			TypeNames[SrcIndex], TypeNames[DstIndex]);
+
+		int32 SrcTotalSize = Arrays[SrcIndex].Shape[0];
+		for (uint8 SrcDimIdx = 1; SrcDimIdx < DimNums[SrcIndex]; SrcDimIdx++)
+		{
+			SrcTotalSize *= Arrays[SrcIndex].Shape[SrcDimIdx];
+		}
+
+		int32 DstTotalSize = Arrays[DstIndex].Shape[0];
+		for (uint8 DstDimIdx = 1; DstDimIdx < DimNums[DstIndex]; DstDimIdx++)
+		{
+			DstTotalSize *= Arrays[DstIndex].Shape[DstDimIdx];
+		}
+
+		UE_LEARNING_CHECKF(SrcTotalSize == DstTotalSize,
+			TEXT("Array total sizes don't match (%i vs %i)"),
+			SrcTotalSize, DstTotalSize);
+
+		// Free the Destination array data
+		if (Flags[DstIndex] & FlagConstructed)
+		{
+			Destructors[DstIndex](
+				Arrays[DstIndex].Data,
+				DimNums[DstIndex],
+				Arrays[DstIndex].Shape);
+		}
+
+		if (Flags[DstIndex] & FlagAllocated)
+		{
+			FMemory::Free(Arrays[DstIndex].Data);
+		}
+
+		// Re-direct existing Links
+		for (int16 ArrayIdx = 0; ArrayIdx < Arrays.Num(); ArrayIdx++)
+		{
+			if (Flags[ArrayIdx] == FlagLinked && Arrays[ArrayIdx].Data == Arrays[DstIndex].Data)
+			{
+				Arrays[ArrayIdx].Data = Arrays[SrcIndex].Data;
+			}
+		}
+
+		// Insert the new Link
+		Arrays[DstIndex].Data = Arrays[SrcIndex].Data;
+		Flags[DstIndex] = FlagLinked;
+	}
+
+	bool FArrayMap::HasLink(const FArrayMapKey Dst) const
+	{
+		const FArrayMapHandle* DstVariable = Handles.Find(Dst);
+
+		UE_LEARNING_CHECKF(DstVariable, TEXT("Does not contain an array called { \"%s\", \"%s\" }"),
+			*Dst.Namespace.ToString(), *Dst.Variable.ToString());
+
+		return HasLink(*DstVariable);
+	}
+
+	bool FArrayMap::HasLink(const FArrayMapHandle Dst) const
+	{
+		return Flags[Dst.Index] & FlagLinked;
+	}
+
+	bool FArrayMap::IsLinkedTo(const FArrayMapKey Src) const
+	{
+		const FArrayMapHandle* SrcVariable = Handles.Find(Src);
+
+		UE_LEARNING_CHECKF(SrcVariable, TEXT("Does not contain an array called { \"%s\", \"%s\" }"),
+			*Src.Namespace.ToString(), *Src.Variable.ToString());
+
+		return IsLinkedTo(*SrcVariable);
+	}
+
+	bool FArrayMap::IsLinkedTo(const FArrayMapHandle Src) const
+	{
+		const int32 VariableNum = Arrays.Num();
+
+		for (int32 VariableIdx = 0; VariableIdx < VariableNum; VariableIdx++)
+		{
+			if (Arrays[VariableIdx].Data == Arrays[Src.Index].Data && (Flags[VariableIdx] & FlagLinked))
+			{
+				return true;
+			}
+		}
+
+		return false;
+	}
+
+	const FArrayMapKey* FArrayMap::FindKey(FArrayMapHandle Handle) const
+	{
+		return Handles.FindKey(Handle);
+	}
+}
diff --git a/Engine/Plugins/Experimental/LearningAgents/Source/Learning/Private/LearningCompletion.cpp b/Engine/Plugins/Experimental/LearningAgents/Source/Learning/Private/LearningCompletion.cpp
--- a/Engine/Plugins/Experimental/LearningAgents/Source/Learning/Private/LearningCompletion.cpp
+++ b/Engine/Plugins/Experimental/LearningAgents/Source/Learning/Private/LearningCompletion.cpp
@@ -4,41 +4,45 @@
 
 namespace UE::Learning
 {
-	void FResetInstanceBuffer::Reserve(const int32 InMaxInstanceNum)
+	void FResetInstanceBuffer::Resize(const int32 InMaxInstanceNum)
 	{
-		ResetInstances.Reserve(InMaxInstanceNum);
+		MaxInstanceNum = InMaxInstanceNum;
+		ResetInstanceNum = 0;
+		ResetInstances.SetNumUninitialized({ MaxInstanceNum });
 	}
 
 	void FResetInstanceBuffer::SetResetInstances(const FIndexSet Instances)
 	{
-		ResetInstances.Reset();
+		ResetInstanceNum = 0;
+
 		for (const int32 InstanceIdx : Instances)
 		{
-			ResetInstances.Add(InstanceIdx);
+			ResetInstances[ResetInstanceNum] = InstanceIdx;
+			ResetInstanceNum++;
 		}
 
-		ResetInstancesSet = ResetInstances;
+		ResetInstancesSet = ResetInstances.Slice(0, ResetInstanceNum);
 		ResetInstancesSet.TryMakeSlice();
 	}
 
 	void FResetInstanceBuffer::SetResetInstancesFromCompletions(const TLearningArrayView<1, const ECompletionMode> Completions, const FIndexSet Instances)
 	{
-		ResetInstances.Reset();
+		ResetInstanceNum = 0;
+
 		for (const int32 InstanceIdx : Instances)
 		{
 			if (Completions[InstanceIdx] != ECompletionMode::Running)
 			{
-				ResetInstances.Add(InstanceIdx);
+				ResetInstances[ResetInstanceNum] = InstanceIdx;
+				ResetInstanceNum++;
 			}
 		}
 
-		ResetInstancesSet = ResetInstances;
+		ResetInstancesSet = ResetInstances.Slice(0, ResetInstanceNum);
 		ResetInstancesSet.TryMakeSlice();
 	}
 
-	const int32 FResetInstanceBuffer::GetResetInstanceNum() const { return ResetInstances.Num(); }
-
-	const TArray<int32>& FResetInstanceBuffer::GetResetInstancesArray() const { return ResetInstances; }
+	const int32 FResetInstanceBuffer::GetResetInstanceNum() const { return ResetInstanceNum; }
 
 	const FIndexSet FResetInstanceBuffer::GetResetInstances() const { return ResetInstancesSet; }
 
@@ -72,11 +76,17 @@ namespace UE::Learning
 			TLearningArrayView<1, ECompletionMode> InOutCompletions,
 			const TLearningArrayView<1, const int32> EpisodeStepNums,
 			const int32 EpisodeMaxStepNum,
+			const ECompletionMode EpisodeEndCompletionMode,
 			const FIndexSet Instances)
 		{
+			UE_LEARNING_TRACE_CPUPROFILER_EVENT_SCOPE(Learning::Completion::EvaluateEndOfEpisodeCompletions);
+
 			for (const int32 InstanceIdx : Instances)
 			{
-				InOutCompletions[InstanceIdx] = EpisodeStepNums[InstanceIdx] == EpisodeMaxStepNum ? ECompletionMode::Truncated : ECompletionMode::Running;
+				if (InOutCompletions[InstanceIdx] == ECompletionMode::Running && EpisodeStepNums[InstanceIdx] == EpisodeMaxStepNum)
+				{
+					InOutCompletions[InstanceIdx] = EpisodeEndCompletionMode;
+				}
 			}
 		}
 	}
diff --git a/Engine/Plugins/Experimental/LearningAgents/Source/Learning/Private/LearningCompletionObject.cpp b/Engine/Plugins/Experimental/LearningAgents/Source/Learning/Private/LearningCompletionObject.cpp
new file mode 100644
--- /dev/null
+++ b/Engine/Plugins/Experimental/LearningAgents/Source/Learning/Private/LearningCompletionObject.cpp
@@ -0,0 +1,376 @@
+// Copyright Epic Games, Inc. All Rights Reserved.
+
+#include "LearningCompletionObject.h"
+#include "LearningLog.h"
+
+namespace UE::Learning
+{
+	FCompletionObject::FCompletionObject(
+		const FName& InIdentifier,
+		const TSharedRef<FArrayMap>& InInstanceData,
+		const int32 InMaxInstanceNum,
+		const ECompletionMode InCompletionMode)
+		: InstanceData(InInstanceData)
+		, CompletionMode(InCompletionMode)
+	{
+		CompletionHandle = InstanceData->Add<1, ECompletionMode>({ InIdentifier, TEXT("Completion") }, { InMaxInstanceNum }, ECompletionMode::Running);
+	}
+
+	TLearningArrayView<1, ECompletionMode> FCompletionObject::CompletionBuffer()
+	{
+		return InstanceData->View(CompletionHandle);
+	}
+
+	//------------------------------------------------------------------
+
+	FAnyCompletion::FAnyCompletion(
+		const FName& InIdentifier,
+		const TLearningArrayView<1, const TSharedRef<FCompletionObject>> InCompletions,
+		const TSharedRef<FArrayMap>& InInstanceData,
+		const int32 InMaxInstanceNum)
+		: FCompletionObject(InIdentifier, InInstanceData, InMaxInstanceNum, ECompletionMode::Terminated)
+		, Completions(InCompletions) {}
+
+	void FAnyCompletion::Evaluate(const FIndexSet Instances)
+	{
+		UE_LEARNING_TRACE_CPUPROFILER_EVENT_SCOPE(FAnyCompletion::Evaluate);
+
+		TLearningArrayView<1, ECompletionMode> Completion = InstanceData->View(CompletionHandle);
+
+		const int32 CompletionNum = Completions.Num();
+
+		for (int32 CompletionIdx = 0; CompletionIdx < CompletionNum; CompletionIdx++)
+		{
+			Completions[CompletionIdx]->Evaluate(Instances);
+		}
+
+		Array::Set(Completion, ECompletionMode::Running, Instances);
+
+		for (int32 CompletionIdx = 0; CompletionIdx < CompletionNum; CompletionIdx++)
+		{
+			const TLearningArrayView<1, const ECompletionMode> Input = InstanceData->ConstView(Completions[CompletionIdx]->CompletionHandle);
+
+			for (const int32 InstanceIdx : Instances)
+			{
+				Completion[InstanceIdx] = Completion::Or(Completion[InstanceIdx], Input[InstanceIdx]);
+			}
+		}
+	}
+
+	//------------------------------------------------------------------
+
+	FConditionalCompletion::FConditionalCompletion(
+		const FName& InIdentifier,
+		const TSharedRef<FArrayMap>& InInstanceData,
+		const int32 InMaxInstanceNum,
+		const ECompletionMode InCompletionMode)
+		: FCompletionObject(InIdentifier, InInstanceData, InMaxInstanceNum, InCompletionMode)
+	{
+		ConditionHandle = InstanceData->Add<1, bool>({ InIdentifier, TEXT("Condition") }, { InMaxInstanceNum }, 0.0f);
+	}
+
+	void FConditionalCompletion::Evaluate(const FIndexSet Instances)
+	{
+		UE_LEARNING_TRACE_CPUPROFILER_EVENT_SCOPE(FConditionalCompletion::Evaluate);
+
+		const TLearningArrayView<1, const bool> Condition = InstanceData->ConstView(ConditionHandle);
+		TLearningArrayView<1, ECompletionMode> Completion = InstanceData->View(CompletionHandle);
+
+		for (const int32 InstanceIdx : Instances)
+		{
+			Completion[InstanceIdx] = Condition[InstanceIdx] ? CompletionMode : ECompletionMode::Running;
+		}
+	}
+
+	//------------------------------------------------------------------
+
+	FTimeElapsedCompletion::FTimeElapsedCompletion(
+		const FName& InIdentifier,
+		const TSharedRef<FArrayMap>& InInstanceData,
+		const int32 InMaxInstanceNum,
+		const float InThreshold,
+		const ECompletionMode InCompletionMode)
+		: FCompletionObject(InIdentifier, InInstanceData, InMaxInstanceNum, InCompletionMode)
+	{
+		TimeHandle = InstanceData->Add<1, float>({ InIdentifier, TEXT("Time") }, { InMaxInstanceNum }, 0.0f);
+		ThresholdHandle = InstanceData->Add<1, float>({ InIdentifier, TEXT("Threshold") }, { InMaxInstanceNum }, InThreshold);
+	}
+
+	void FTimeElapsedCompletion::Evaluate(const FIndexSet Instances)
+	{
+		UE_LEARNING_TRACE_CPUPROFILER_EVENT_SCOPE(FTimeElapsedCompletion::Evaluate);
+
+		const TLearningArrayView<1, const float> Time = InstanceData->ConstView(TimeHandle);
+		const TLearningArrayView<1, const float> Threshold = InstanceData->ConstView(ThresholdHandle);
+		TLearningArrayView<1, ECompletionMode> Completion = InstanceData->View(CompletionHandle);
+
+		for (const int32 InstanceIdx : Instances)
+		{
+			Completion[InstanceIdx] = Time[InstanceIdx] > Threshold[InstanceIdx] ? CompletionMode : ECompletionMode::Running;
+		}
+	}
+
+	//------------------------------------------------------------------
+
+	FScalarPositionDifferenceCompletion::FScalarPositionDifferenceCompletion(
+		const FName& InIdentifier,
+		const TSharedRef<FArrayMap>& InInstanceData,
+		const int32 InMaxInstanceNum,
+		const int32 InPositionNum,
+		const float InThreshold,
+		const ECompletionMode InCompletionMode)
+		: FCompletionObject(InIdentifier, InInstanceData, InMaxInstanceNum, InCompletionMode)
+	{
+		Position0Handle = InstanceData->Add<2, float>({ InIdentifier, TEXT("Position0") }, { InMaxInstanceNum, InPositionNum }, 0.0f);
+		Position1Handle = InstanceData->Add<2, float>({ InIdentifier, TEXT("Position1") }, { InMaxInstanceNum, InPositionNum }, 0.0f);
+		ThresholdHandle = InstanceData->Add<1, float>({ InIdentifier, TEXT("Threshold") }, { InMaxInstanceNum }, InThreshold);
+	}
+
+	void FScalarPositionDifferenceCompletion::Evaluate(const FIndexSet Instances)
+	{
+		UE_LEARNING_TRACE_CPUPROFILER_EVENT_SCOPE(FScalarPositionDifferenceCompletion::Evaluate);
+
+		const TLearningArrayView<2, const float> Position0 = InstanceData->ConstView(Position0Handle);
+		const TLearningArrayView<2, const float> Position1 = InstanceData->ConstView(Position1Handle);
+		const TLearningArrayView<1, const float> Threshold = InstanceData->ConstView(ThresholdHandle);
+		TLearningArrayView<1, ECompletionMode> Completion = InstanceData->View(CompletionHandle);
+
+		const int32 PositionNum = Position0.Num<1>();
+
+		for (const int32 InstanceIdx : Instances)
+		{
+			Completion[InstanceIdx] = ECompletionMode::Running;
+
+			for (int32 PositionIdx = 0; PositionIdx < PositionNum; PositionIdx++)
+			{
+				if (FMath::Abs(Position0[InstanceIdx][PositionIdx] - Position1[InstanceIdx][PositionIdx]) > Threshold[InstanceIdx])
+				{
+					Completion[InstanceIdx] = CompletionMode;
+					break;
+				}
+			}
+		}
+	}
+
+	FPlanarPositionDifferenceCompletion::FPlanarPositionDifferenceCompletion(
+		const FName& InIdentifier,
+		const TSharedRef<FArrayMap>& InInstanceData,
+		const int32 InMaxInstanceNum,
+		const int32 InPositionNum,
+		const float InThreshold,
+		const ECompletionMode InCompletionMode,
+		const FVector InAxis0,
+		const FVector InAxis1)
+		: FCompletionObject(InIdentifier, InInstanceData, InMaxInstanceNum, InCompletionMode)
+		, Axis0(InAxis0)
+		, Axis1(InAxis1)
+	{
+		Position0Handle = InstanceData->Add<2, FVector>({ InIdentifier, TEXT("Position0") }, { InMaxInstanceNum, InPositionNum }, FVector::ZeroVector);
+		Position1Handle = InstanceData->Add<2, FVector>({ InIdentifier, TEXT("Position1") }, { InMaxInstanceNum, InPositionNum }, FVector::ZeroVector);
+		ThresholdHandle = InstanceData->Add<1, float>({ InIdentifier, TEXT("Threshold") }, { InMaxInstanceNum }, InThreshold);
+	}
+
+	void FPlanarPositionDifferenceCompletion::Evaluate(const FIndexSet Instances)
+	{
+		UE_LEARNING_TRACE_CPUPROFILER_EVENT_SCOPE(FPlanarPositionDifferenceCompletion::Evaluate);
+
+		const TLearningArrayView<2, const FVector> Position0 = InstanceData->ConstView(Position0Handle);
+		const TLearningArrayView<2, const FVector> Position1 = InstanceData->ConstView(Position1Handle);
+		const TLearningArrayView<1, const float> Threshold = InstanceData->ConstView(ThresholdHandle);
+		TLearningArrayView<1, ECompletionMode> Completion = InstanceData->View(CompletionHandle);
+
+		const int32 PositionNum = Position0.Num<1>();
+
+		for (const int32 InstanceIdx : Instances)
+		{
+			Completion[InstanceIdx] = ECompletionMode::Running;
+
+			for (int32 PositionIdx = 0; PositionIdx < PositionNum; PositionIdx++)
+			{
+				const FVector ProjectedPosition0 = FVector(Axis0.Dot(Position0[InstanceIdx][PositionIdx]), Axis1.Dot(Position0[InstanceIdx][PositionIdx]), 0.0f);
+				const FVector ProjectedPosition1 = FVector(Axis0.Dot(Position1[InstanceIdx][PositionIdx]), Axis1.Dot(Position1[InstanceIdx][PositionIdx]), 0.0f);
+
+				if (FVector::Distance(ProjectedPosition0, ProjectedPosition1) > Threshold[InstanceIdx])
+				{
+					Completion[InstanceIdx] = CompletionMode;
+					break;
+				}
+			}
+		}
+	}
+
+	FPlanarPositionSimilarityCompletion::FPlanarPositionSimilarityCompletion(
+		const FName& InIdentifier,
+		const TSharedRef<FArrayMap>& InInstanceData,
+		const int32 InMaxInstanceNum,
+		const int32 InPositionNum,
+		const float InThreshold,
+		const ECompletionMode InCompletionMode,
+		const FVector InAxis0,
+		const FVector InAxis1)
+		: FCompletionObject(InIdentifier, InInstanceData, InMaxInstanceNum, InCompletionMode)
+		, Axis0(InAxis0)
+		, Axis1(InAxis1)
+	{
+		Position0Handle = InstanceData->Add<2, FVector>({ InIdentifier, TEXT("Position0") }, { InMaxInstanceNum, InPositionNum }, FVector::ZeroVector);
+		Position1Handle = InstanceData->Add<2, FVector>({ InIdentifier, TEXT("Position1") }, { InMaxInstanceNum, InPositionNum }, FVector::ZeroVector);
+		ThresholdHandle = InstanceData->Add<1, float>({ InIdentifier, TEXT("Threshold") }, { InMaxInstanceNum }, InThreshold);
+	}
+
+	void FPlanarPositionSimilarityCompletion::Evaluate(const FIndexSet Instances)
+	{
+		UE_LEARNING_TRACE_CPUPROFILER_EVENT_SCOPE(FPlanarPositionSimilarityCompletion::Evaluate);
+
+		const TLearningArrayView<2, const FVector> Position0 = InstanceData->ConstView(Position0Handle);
+		const TLearningArrayView<2, const FVector> Position1 = InstanceData->ConstView(Position1Handle);
+		const TLearningArrayView<1, const float> Threshold = InstanceData->ConstView(ThresholdHandle);
+		TLearningArrayView<1, ECompletionMode> Completion = InstanceData->View(CompletionHandle);
+
+		const int32 PositionNum = Position0.Num<1>();
+
+		for (const int32 InstanceIdx : Instances)
+		{
+			Completion[InstanceIdx] = ECompletionMode::Running;
+
+			for (int32 PositionIdx = 0; PositionIdx < PositionNum; PositionIdx++)
+			{
+				const FVector ProjectedPosition0 = FVector(Axis0.Dot(Position0[InstanceIdx][PositionIdx]), Axis1.Dot(Position0[InstanceIdx][PositionIdx]), 0.0f);
+				const FVector ProjectedPosition1 = FVector(Axis0.Dot(Position1[InstanceIdx][PositionIdx]), Axis1.Dot(Position1[InstanceIdx][PositionIdx]), 0.0f);
+
+				if (FVector::Distance(ProjectedPosition0, ProjectedPosition1) < Threshold[InstanceIdx])
+				{
+					Completion[InstanceIdx] = CompletionMode;
+					break;
+				}
+			}
+		}
+	}
+
+	//------------------------------------------------------------------
+
+	FScalarVelocityDifferenceCompletion::FScalarVelocityDifferenceCompletion(
+		const FName& InIdentifier,
+		const TSharedRef<FArrayMap>& InInstanceData,
+		const int32 InMaxInstanceNum,
+		const int32 InVelocityNum,
+		const float InThreshold,
+		const ECompletionMode InCompletionMode)
+		: FCompletionObject(InIdentifier, InInstanceData, InMaxInstanceNum, InCompletionMode)
+	{
+		Velocity0Handle = InstanceData->Add<2, float>({ InIdentifier, TEXT("Velocity0") }, { InMaxInstanceNum, InVelocityNum }, 0.0f);
+		Velocity1Handle = InstanceData->Add<2, float>({ InIdentifier, TEXT("Velocity1") }, { InMaxInstanceNum, InVelocityNum }, 0.0f);
+		ThresholdHandle = InstanceData->Add<1, float>({ InIdentifier, TEXT("Threshold") }, { InMaxInstanceNum }, InThreshold);
+	}
+
+	void FScalarVelocityDifferenceCompletion::Evaluate(const FIndexSet Instances)
+	{
+		UE_LEARNING_TRACE_CPUPROFILER_EVENT_SCOPE(FScalarVelocityDifferenceCompletion::Evaluate);
+
+		const TLearningArrayView<2, const float> Velocity0 = InstanceData->ConstView(Velocity0Handle);
+		const TLearningArrayView<2, const float> Velocity1 = InstanceData->ConstView(Velocity1Handle);
+		const TLearningArrayView<1, const float> Threshold = InstanceData->ConstView(ThresholdHandle);
+		TLearningArrayView<1, ECompletionMode> Completion = InstanceData->View(CompletionHandle);
+
+		const int32 VelocityNum = Velocity0.Num<1>();
+
+		for (const int32 InstanceIdx : Instances)
+		{
+			Completion[InstanceIdx] = ECompletionMode::Running;
+
+			for (int32 VelocityIdx = 0; VelocityIdx < VelocityNum; VelocityIdx++)
+			{
+				if (FMath::Abs(Velocity0[InstanceIdx][VelocityIdx] - Velocity1[InstanceIdx][VelocityIdx]) > Threshold[InstanceIdx])
+				{
+					Completion[InstanceIdx] = CompletionMode;
+					break;
+				}
+			}
+		}
+	}
+
+	//------------------------------------------------------------------
+
+	FScalarRotationDifferenceCompletion::FScalarRotationDifferenceCompletion(
+		const FName& InIdentifier,
+		const TSharedRef<FArrayMap>& InInstanceData,
+		const int32 InMaxInstanceNum,
+		const int32 InRotationNum,
+		const float InThreshold,
+		const ECompletionMode InCompletionMode)
+		: FCompletionObject(InIdentifier, InInstanceData, InMaxInstanceNum, InCompletionMode)
+	{
+		Rotation0Handle = InstanceData->Add<2, float>({ InIdentifier, TEXT("Rotation0") }, { InMaxInstanceNum, InRotationNum }, 0.0f);
+		Rotation1Handle = InstanceData->Add<2, float>({ InIdentifier, TEXT("Rotation1") }, { InMaxInstanceNum, InRotationNum }, 0.0f);
+		ThresholdHandle = InstanceData->Add<1, float>({ InIdentifier, TEXT("Threshold") }, { InMaxInstanceNum }, InThreshold);
+	}
+
+	void FScalarRotationDifferenceCompletion::Evaluate(const FIndexSet Instances)
+	{
+		UE_LEARNING_TRACE_CPUPROFILER_EVENT_SCOPE(FScalarRotationDifferenceCompletion::Evaluate);
+
+		const TLearningArrayView<2, const float> Rotation0 = InstanceData->ConstView(Rotation0Handle);
+		const TLearningArrayView<2, const float> Rotation1 = InstanceData->ConstView(Rotation1Handle);
+		const TLearningArrayView<1, const float> Threshold = InstanceData->ConstView(ThresholdHandle);
+		TLearningArrayView<1, ECompletionMode> Completion = InstanceData->View(CompletionHandle);
+
+		const int32 RotationNum = Rotation0.Num<1>();
+
+		for (const int32 InstanceIdx : Instances)
+		{
+			Completion[InstanceIdx] = ECompletionMode::Running;
+
+			for (int32 RotationIdx = 0; RotationIdx < RotationNum; RotationIdx++)
+			{
+				if (FMath::Abs(FMath::FindDeltaAngleRadians(Rotation1[InstanceIdx][RotationIdx], Rotation0[InstanceIdx][RotationIdx])) > Threshold[InstanceIdx])
+				{
+					Completion[InstanceIdx] = CompletionMode;
+					break;
+				}
+			}
+		}
+	}
+
+	//------------------------------------------------------------------
+
+	FScalarAngularVelocityDifferenceCompletion::FScalarAngularVelocityDifferenceCompletion(
+		const FName& InIdentifier,
+		const TSharedRef<FArrayMap>& InInstanceData,
+		const int32 InMaxInstanceNum,
+		const int32 InAngularVelocityNum,
+		const float InThreshold,
+		const ECompletionMode InCompletionMode)
+		: FCompletionObject(InIdentifier, InInstanceData, InMaxInstanceNum, InCompletionMode)
+	{
+		AngularVelocity0Handle = InstanceData->Add<2, float>({ InIdentifier, TEXT("AngularVelocity0") }, { InMaxInstanceNum, InAngularVelocityNum }, 0.0f);
+		AngularVelocity1Handle = InstanceData->Add<2, float>({ InIdentifier, TEXT("AngularVelocity1") }, { InMaxInstanceNum, InAngularVelocityNum }, 0.0f);
+		ThresholdHandle = InstanceData->Add<1, float>({ InIdentifier, TEXT("Threshold") }, { InMaxInstanceNum }, InThreshold);
+	}
+
+	void FScalarAngularVelocityDifferenceCompletion::Evaluate(const FIndexSet Instances)
+	{
+		UE_LEARNING_TRACE_CPUPROFILER_EVENT_SCOPE(FScalarAngularVelocityDifferenceCompletion::Evaluate);
+
+		const TLearningArrayView<2, const float> AngularVelocity0 = InstanceData->ConstView(AngularVelocity0Handle);
+		const TLearningArrayView<2, const float> AngularVelocity1 = InstanceData->ConstView(AngularVelocity1Handle);
+		const TLearningArrayView<1, const float> Threshold = InstanceData->ConstView(ThresholdHandle);
+		TLearningArrayView<1, ECompletionMode> Completion = InstanceData->View(CompletionHandle);
+
+		const int32 AngularVelocityNum = AngularVelocity0.Num<1>();
+
+		for (const int32 InstanceIdx : Instances)
+		{
+			Completion[InstanceIdx] = ECompletionMode::Running;
+
+			for (int32 AngularVelocityIdx = 0; AngularVelocityIdx < AngularVelocityNum; AngularVelocityIdx++)
+			{
+				if (FMath::Abs(AngularVelocity0[InstanceIdx][AngularVelocityIdx] - AngularVelocity1[InstanceIdx][AngularVelocityIdx]) > Threshold[InstanceIdx])
+				{
+					Completion[InstanceIdx] = CompletionMode;
+					break;
+				}
+			}
+		}
+	}
+
+	//------------------------------------------------------------------
+}
