[Backout] - CL30154408

diff --git a/Engine/Plugins/Experimental/LearningAgents/Source/Learning/Private/LearningNeuralNetwork.cpp b/Engine/Plugins/Experimental/LearningAgents/Source/Learning/Private/LearningNeuralNetwork.cpp
--- a/Engine/Plugins/Experimental/LearningAgents/Source/Learning/Private/LearningNeuralNetwork.cpp
+++ b/Engine/Plugins/Experimental/LearningAgents/Source/Learning/Private/LearningNeuralNetwork.cpp
@@ -2,377 +2,368 @@
 
 #include "LearningNeuralNetwork.h"
 
-#include "LearningProgress.h"
+#if UE_LEARNING_ISPC
+#include "Learning.ispc.generated.h"
+#endif
 
-#include "Async/ParallelFor.h"
-#include "Modules/ModuleManager.h"
+namespace UE::Learning
+{
+	namespace NeuralNetwork
+	{
+		static inline void MatMulPlusBias(
+			TLearningArrayView<2, float> Output,
+			const TLearningArrayView<2, const float> Input,
+			const TLearningArrayView<2, const float> Weights,
+			const TLearningArrayView<1, const float> Biases,
+			const int32 BatchOffset,
+			const int32 BatchNum)
+		{
+			UE_LEARNING_TRACE_CPUPROFILER_EVENT_SCOPE(Learning::NeuralNetwork::MatMulPlusBias);
+
+			const int32 RowNum = Weights.Num<0>();
+			const int32 ColNum = Weights.Num<1>();
+			const int32 BufferSize = Input.Num<1>();
+
+#if UE_LEARNING_ISPC
+			ispc::LearningLayerMatMulPlusBias(
+				Output.Slice(BatchOffset, BatchNum).GetData(),
+				Input.Slice(BatchOffset, BatchNum).GetData(),
+				Weights.GetData(),
+				Biases.GetData(),
+				BatchNum,
+				BufferSize,
+				RowNum,
+				ColNum);
+#else
+			for (int32 BatchIdx = BatchOffset; BatchIdx < BatchOffset + BatchNum; BatchIdx++)
+			{
+				for (int32 ColIdx = 0; ColIdx < ColNum; ColIdx++)
+				{
+					Output[BatchIdx][ColIdx] = Biases[ColIdx];
+				}
 
-#include "NNE.h"
-#include "NNEModelData.h"
-#include "NNERuntimeCPU.h"
+				for (int32 RowIdx = 0; RowIdx < RowNum; RowIdx++)
+				{
+					const float Value = Input[BatchIdx][RowIdx];
+
+					if (Value != 0.0)
+					{
+						for (int32 ColIdx = 0; ColIdx < ColNum; ColIdx++)
+						{
+							Output[BatchIdx][ColIdx] += Value * Weights[RowIdx][ColIdx];
+						}
+					}
+				}
+			}
+#endif
+		}
 
-void ULearningNeuralNetworkData::PostLoad()
-{
-	Super::PostLoad();
+		static inline float ReLU(const float X)
+		{
+			return FMath::Max(X, 0.0f);
+		}
 
-	UpdateNetwork();
-}
+		static inline float ELU(const float X)
+		{
+			return X > 0.0f ? X : FMath::Exp(X) - 1.0f;
+		}
 
-void ULearningNeuralNetworkData::Init(const int32 InInputSize, const int32 InOutputSize, const int32 InCompatibilityHash, const TArrayView<const uint8> InFileData)
-{
-	InputSize = InInputSize;
-	OutputSize = InOutputSize;
-	CompatibilityHash = InCompatibilityHash;
-	FileData = InFileData;
-	UpdateNetwork();
-}
+		static inline float TanH(const float X)
+		{
+			return FMath::Tanh(X);
+		}
 
-void ULearningNeuralNetworkData::InitFrom(const ULearningNeuralNetworkData* OtherNetworkData)
-{
-	InputSize = OtherNetworkData->InputSize;
-	OutputSize = OtherNetworkData->OutputSize;
-	CompatibilityHash = OtherNetworkData->CompatibilityHash;
-	FileData = OtherNetworkData->FileData;
-	UpdateNetwork();
-}
+		static inline void ActivationReLU(TLearningArrayView<2, float> InputOutput, const int32 BatchOffset, const int32 BatchNum, const int32 InputOutputNum)
+		{
+			UE_LEARNING_TRACE_CPUPROFILER_EVENT_SCOPE(Learning::NeuralNetwork::ActivationReLU);
+
+			const int32 BufferSize = InputOutput.Num<1>();
+
+#if UE_LEARNING_ISPC
+			ispc::LearningLayerReLU(
+				InputOutput.Slice(BatchOffset, BatchNum).GetData(),
+				BatchNum,
+				BufferSize, 
+				InputOutputNum);
+#else
+			for (int32 BatchIdx = BatchOffset; BatchIdx < BatchOffset + BatchNum; BatchIdx++)
+			{
+				for (int32 ColIdx = 0; ColIdx < InputOutputNum; ColIdx++)
+				{
+					InputOutput[BatchIdx][ColIdx] = ReLU(InputOutput[BatchIdx][ColIdx]);
+				}
+			}
+#endif
+		}
 
-namespace UE::Learning::NeuralNetwork::Private
-{
-	static constexpr int32 MagicNumber = 0x1e9b0c80;
-	static constexpr int32 VersionNumber = 1;
-}
+		static inline void ActivationELU(TLearningArrayView<2, float> InputOutput, const int32 BatchOffset, const int32 BatchNum, const int32 InputOutputNum)
+		{
+			UE_LEARNING_TRACE_CPUPROFILER_EVENT_SCOPE(Learning::NeuralNetwork::ActivationELU);
+
+			const int32 BufferSize = InputOutput.Num<1>();
+
+#if UE_LEARNING_ISPC
+			ispc::LearningLayerELU(
+				InputOutput.Slice(BatchOffset, BatchNum).GetData(),
+				BatchNum,
+				BufferSize,
+				InputOutputNum);
+#else
+			for (int32 BatchIdx = BatchOffset; BatchIdx < BatchOffset + BatchNum; BatchIdx++)
+			{
+				for (int32 ColIdx = 0; ColIdx < InputOutputNum; ColIdx++)
+				{
+					InputOutput[BatchIdx][ColIdx] = ELU(InputOutput[BatchIdx][ColIdx]);
+				}
+			}
+#endif
+		}
 
-bool ULearningNeuralNetworkData::LoadFromSnapshot(const TArrayView<const uint8> InBytes)
-{
-	if (!ensureMsgf(InBytes.Num() > sizeof(int32) * 6, TEXT("Snapshot format invalid. File not large enough.")))
-	{
-		return false;
-	}
+		static inline void ActivationTanH(TLearningArrayView<2, float> InputOutput, const int32 BatchOffset, const int32 BatchNum, const int32 InputOutputNum)
+		{
+			UE_LEARNING_TRACE_CPUPROFILER_EVENT_SCOPE(Learning::NeuralNetwork::ActivationTanH);
+
+			const int32 BufferSize = InputOutput.Num<1>();
+
+#if UE_LEARNING_ISPC
+			ispc::LearningLayerTanH(
+				InputOutput.Slice(BatchOffset, BatchNum).GetData(),
+				BatchNum,
+				BufferSize,
+				InputOutputNum);
+#else
+			for (int32 BatchIdx = BatchOffset; BatchIdx < BatchOffset + BatchNum; BatchIdx++)
+			{
+				for (int32 ValueIdx = 0; ValueIdx < InputOutputNum; ValueIdx++)
+				{
+					InputOutput[BatchIdx][ValueIdx] = TanH(InputOutput[BatchIdx][ValueIdx]);
+				}
+			}
+#endif
+		}
 
-	int32 Offset = 0;
-	
-	int32 MagicNumber = 0;
-	UE::Learning::DeserializeFromBytes(Offset, InBytes, MagicNumber);
-	if (!ensureMsgf(MagicNumber == UE::Learning::NeuralNetwork::Private::MagicNumber, TEXT("Snapshot format invalid. Incorrect Magic Number. Got %i, expected %i."), MagicNumber, UE::Learning::NeuralNetwork::Private::MagicNumber))
-	{
-		return false;
+		static constexpr int32 MagicNumber = 0x0cd353cf;
+		static constexpr int32 VersionNumber = 1;
 	}
 
-	int32 VersionNumber = 0;
-	UE::Learning::DeserializeFromBytes(Offset, InBytes, VersionNumber);
-	if (!ensureMsgf(VersionNumber == UE::Learning::NeuralNetwork::Private::VersionNumber, TEXT("Snapshot format invalid. Unsupported version number. Got %i, expected %i."), VersionNumber, UE::Learning::NeuralNetwork::Private::VersionNumber))
+	bool FNeuralNetworkMLP::DeserializeFromBytes(int32& InOutOffset, const TLearningArrayView<1, const uint8> RawBytes)
 	{
-		return false;
-	}
-
-	UE::Learning::DeserializeFromBytes(Offset, InBytes, InputSize);
-	UE::Learning::DeserializeFromBytes(Offset, InBytes, OutputSize);
-	UE::Learning::DeserializeFromBytes(Offset, InBytes, CompatibilityHash);
+		UE_LEARNING_TRACE_CPUPROFILER_EVENT_SCOPE(Learning::FNeuralNetworkMLP::DeserializeFromBytes);
 
-	int32 FileDataSize = 0;
-	UE::Learning::DeserializeFromBytes(Offset, InBytes, FileDataSize);
-	FileData = InBytes.Slice(Offset, FileDataSize);
+		// At least Magic Number and Version Number Required
+		if (RawBytes.Num() < sizeof(int32) * 2)
+		{
+			return false;
+		}
 
-	UpdateNetwork();
-	return true;
-}
+		int32 Magic;
+		UE::Learning::DeserializeFromBytes(InOutOffset, RawBytes, Magic);
+		if (Magic != NeuralNetwork::MagicNumber) { return false; }
 
-void ULearningNeuralNetworkData::SaveToSnapshot(TArrayView<uint8> OutBytes) const
-{
-	int32 Offset = 0;
-	UE::Learning::SerializeToBytes(Offset, OutBytes, UE::Learning::NeuralNetwork::Private::MagicNumber);
-	UE::Learning::SerializeToBytes(Offset, OutBytes, UE::Learning::NeuralNetwork::Private::VersionNumber);
-	UE::Learning::SerializeToBytes(Offset, OutBytes, InputSize);
-	UE::Learning::SerializeToBytes(Offset, OutBytes, OutputSize);
-	UE::Learning::SerializeToBytes(Offset, OutBytes, CompatibilityHash);
-	UE::Learning::SerializeToBytes(Offset, OutBytes, FileData.Num());
-	FMemory::Memcpy(OutBytes.Slice(Offset, FileData.Num()).GetData(), FileData.GetData(), FileData.Num());
-}
+		int32 Version;
+		UE::Learning::DeserializeFromBytes(InOutOffset, RawBytes, Version);
+		if (Version != NeuralNetwork::VersionNumber) { return false; }
 
-int32 ULearningNeuralNetworkData::GetSnapshotByteNum() const
-{
-	return
-		sizeof(int32) + // Magic Number
-		sizeof(int32) + // Version Number
-		sizeof(int32) + // Input Size
-		sizeof(int32) + // Output Size
-		sizeof(int32) + // Compatibility Hash
-		sizeof(int32) + // FileData size
-		FileData.Num(); // FileData
-}
+		int32 ActivationFunctionInt;
+		UE::Learning::DeserializeFromBytes(InOutOffset, RawBytes, ActivationFunctionInt);
+		ActivationFunction = (EActivationFunction)ActivationFunctionInt;
 
-TSharedPtr<UE::Learning::FNeuralNetwork>& ULearningNeuralNetworkData::GetNetwork()
-{
-	return Network;
-}
+		int32 LayerNum;
+		UE::Learning::DeserializeFromBytes(InOutOffset, RawBytes, LayerNum);
 
-bool ULearningNeuralNetworkData::IsEmpty() const
-{
-	return Network->IsEmpty();
-}
-
-int32 ULearningNeuralNetworkData::GetInputSize() const
-{
-	return InputSize;
-}
+		Weights.SetNum(LayerNum);
+		Biases.SetNum(LayerNum);
 
-int32 ULearningNeuralNetworkData::GetOutputSize() const
-{
-	return OutputSize;
-}
+		for (int32 LayerIdx = 0; LayerIdx < LayerNum; LayerIdx++)
+		{
+			UE::Learning::Array::DeserializeFromBytes(InOutOffset, RawBytes, Weights[LayerIdx]);
+			UE::Learning::Array::DeserializeFromBytes(InOutOffset, RawBytes, Biases[LayerIdx]);
+		}
 
-int32 ULearningNeuralNetworkData::GetCompatibilityHash() const
-{
-	return CompatibilityHash;
-}
+		return true;
+	}
 
-void ULearningNeuralNetworkData::UpdateNetwork()
-{
-	if (FileData.Num() > 0)
+	void FNeuralNetworkMLP::SerializeToBytes(int32& InOutOffset, TLearningArrayView<1, uint8> OutRawBytes) const
 	{
-		if (!ModelData)
-		{
-			ModelData = NewObject<UNNEModelData>(this);
-		}
+		UE_LEARNING_TRACE_CPUPROFILER_EVENT_SCOPE(Learning::FNeuralNetworkMLP::SerializeToBytes);
 
-		ModelData->Init(TEXT("ubnne"), FileData);
+		const int32 LayerNum = Weights.Num();
 
-		if (!Network)
+		UE::Learning::SerializeToBytes(InOutOffset, OutRawBytes, NeuralNetwork::MagicNumber);
+		UE::Learning::SerializeToBytes(InOutOffset, OutRawBytes, NeuralNetwork::VersionNumber);
+		UE::Learning::SerializeToBytes(InOutOffset, OutRawBytes, (int32)ActivationFunction);
+		UE::Learning::SerializeToBytes(InOutOffset, OutRawBytes, LayerNum);
+
+		for (int32 LayerIdx = 0; LayerIdx < LayerNum; LayerIdx++)
 		{
-			Network = MakeShared<UE::Learning::FNeuralNetwork>();
+			UE::Learning::Array::SerializeToBytes(InOutOffset, OutRawBytes, Weights[LayerIdx]);
+			UE::Learning::Array::SerializeToBytes(InOutOffset, OutRawBytes, Biases[LayerIdx]);
 		}
+	}
 
-		ensureMsgf(FModuleManager::Get().LoadModule(TEXT("NNERuntimeBasicCpu")), TEXT("Unable to load module for NNE runtime NNERuntimeBasicCpu."));
+	int32 FNeuralNetworkMLP::GetSerializationByteNum() const
+	{
+		const int32 LayerNum = GetLayerNum();
 
-		TWeakInterfacePtr<INNERuntimeCPU> RuntimeCPU = UE::NNE::GetRuntime<INNERuntimeCPU>(TEXT("NNERuntimeBasicCpu"));
+		int32 Total = 0;
 
-		TSharedPtr<UE::NNE::IModelCPU> UpdatedModel = nullptr;
+		Total += sizeof(int32);  // Magic
+		Total += sizeof(int32);  // Version
+		Total += sizeof(int32);  // Activation
+		Total += sizeof(int32);  // Layer Num
 
-		if (ensureMsgf(RuntimeCPU.IsValid(), TEXT("Could not find requested NNE Runtime")))
+		for (int32 LayerIdx = 0; LayerIdx < LayerNum; LayerIdx++)
 		{
-			UpdatedModel = RuntimeCPU->CreateModelCPU(ModelData);
+			Total += UE::Learning::Array::SerializationByteNum<2, float>(Weights[LayerIdx].Shape());
+			Total += UE::Learning::Array::SerializationByteNum<1, float>(Biases[LayerIdx].Shape());
 		}
 
-		// If we are not in the editor we can now clear the FileData and FileType since these will be
-		// using additional memory and we are not going to save this asset and so don't require them.
-
-#if !WITH_EDITOR
-		ModelData->ClearFileDataAndFileType();
-#endif
-
-		Network->UpdateModel(UpdatedModel, InputSize, OutputSize);
+		return Total;
 	}
-}
 
-namespace UE::Learning
-{
-	TSharedRef<FNeuralNetworkInference> FNeuralNetwork::CreateInferenceObject(
+	TSharedRef<INeuralNetworkInference> FNeuralNetworkMLP::CreateInferenceObject(
 		const int32 MaxBatchSize,
 		const FNeuralNetworkInferenceSettings& Settings)
 	{
-		TSharedRef<FNeuralNetworkInference> InferenceObject = MakeShared<FNeuralNetworkInference>(
-			*Model,
-			MaxBatchSize,
-			InputSize,
-			OutputSize,
-			Settings);
-
-		InferenceObjects.Emplace(InferenceObject.ToWeakPtr());
-
-		return InferenceObject;
+		return MakeShared<FNeuralNetworkInferenceMLP>(*this, MaxBatchSize, Settings);
 	}
 
-	bool FNeuralNetwork::IsEmpty() const
+	int32 FNeuralNetworkMLP::GetInputNum() const
 	{
-		return Model != nullptr;
+		UE_LEARNING_CHECK(Weights.Num() > 0);
+		return Weights[0].Num<0>();
 	}
-	
-	int32 FNeuralNetwork::GetInputSize() const
+
+	int32 FNeuralNetworkMLP::GetOutputNum() const
 	{
-		return InputSize;
+		UE_LEARNING_CHECK(Weights.Num() > 0);
+		return Weights.Last().Num<1>();
 	}
 
-	int32 FNeuralNetwork::GetOutputSize() const
+	int32 FNeuralNetworkMLP::GetHiddenNum() const
 	{
-		return OutputSize;
+		UE_LEARNING_CHECK(Weights.Num() > 0);
+		return Weights[0].Num<1>();
 	}
 
-	void FNeuralNetwork::UpdateModel(const TSharedPtr<NNE::IModelCPU>& InModel, const int32 InInputSize, const int32 InOutputSize)
+	int32 FNeuralNetworkMLP::GetLayerNum() const
 	{
-		// Update Model
-
-		InputSize = InInputSize;
-		OutputSize = InOutputSize;
-		Model = InModel;
+		return Weights.Num();
+	}
 
-		if (Model)
-		{
-			// We need to tell all the created inference objects to update their instances since the Model has changed
-			// and so the instance objects themselves are either no longer valid or tied to the old model.
+	void FNeuralNetworkMLP::Resize(
+		const int32 InputNum,
+		const int32 OutputNum,
+		const int32 HiddenNum,
+		const int32 LayerNum)
+	{
+		UE_LEARNING_TRACE_CPUPROFILER_EVENT_SCOPE(Learning::FNeuralNetworkMLP::Resize);
+		UE_LEARNING_CHECKF(LayerNum >= 2, TEXT("At least two layers required (input and output layers)"));
 
-			for (const TWeakPtr<UE::Learning::FNeuralNetworkInference>& InferenceObject : InferenceObjects)
-			{
-				if (TSharedPtr<UE::Learning::FNeuralNetworkInference> InferenceObjectPtr = InferenceObject.Pin())
-				{
-					InferenceObjectPtr->ReloadModelInstances(*Model);
-				}
-			}
+		Weights.SetNum(LayerNum);
+		Biases.SetNum(LayerNum);
 
-			// Remove all invalid weak pointers
+		Weights[0].SetNumUninitialized({ InputNum, HiddenNum });
+		Biases[0].SetNumUninitialized({ HiddenNum });
 
-			InferenceObjects.RemoveAll(
-				[](const TWeakPtr<UE::Learning::FNeuralNetworkInference>& InferenceObject) { return !InferenceObject.IsValid(); });
-		}
-		else
+		for (int32 LayerIdx = 0; LayerIdx < LayerNum - 2; LayerIdx++)
 		{
-			InferenceObjects.Empty();
+			Weights[LayerIdx + 1].SetNumUninitialized({ HiddenNum, HiddenNum });
+			Biases[LayerIdx + 1].SetNumUninitialized({ HiddenNum });
 		}
-	}
 
-	FNeuralNetworkInference::FNeuralNetworkInference(
-		UE::NNE::IModelCPU& InModel,
-		const int32 MaxBatchSize,
-		const int32 InInputSize,
-		const int32 InOutputSize,
-		const FNeuralNetworkInferenceSettings& InSettings)
-		: Settings(InSettings)
-	{
-		InputBuffer.SetNumUninitialized({ MaxBatchSize, InInputSize });
-		OutputBuffer.SetNumUninitialized({ MaxBatchSize, InOutputSize });
+		Weights[LayerNum - 1].SetNumUninitialized({ HiddenNum, OutputNum });
+		Biases[LayerNum - 1].SetNumUninitialized({ OutputNum });
 
-		ReloadModelInstances(InModel);
+		for (int32 LayerIdx = 0; LayerIdx < LayerNum; LayerIdx++)
+		{
+			Array::Zero(Weights[LayerIdx]);
+			Array::Zero(Biases[LayerIdx]);
+		}
 	}
 
-	void FNeuralNetworkInference::ReloadModelInstances(NNE::IModelCPU& InModel)
+	FNeuralNetworkInferenceMLP::FNeuralNetworkInferenceMLP(const FNeuralNetworkMLP& InNetwork, const int32 MaxBatchSize, const FNeuralNetworkInferenceSettings& InSettings)
+		: Network(InNetwork)
+		, Settings(InSettings)
 	{
-		UE_LEARNING_TRACE_CPUPROFILER_EVENT_SCOPE(Learning::FNeuralNetworkInference::ReloadModelInstances);
-
-		const int32 MaxBatchSize = InputBuffer.Num<0>();
-		const int32 InputSize = InputBuffer.Num<1>();
-		const int32 OutputSize = OutputBuffer.Num<1>();
+		const int32 BufferSize = FMath::Max(FMath::Max(Network.GetInputNum(), Network.GetOutputNum()), Network.GetHiddenNum());
 
-		if (Settings.bParallelEvaluation)
-		{
-			const int32 IdealInferenceInstanceNum = ParallelForImpl::GetNumberOfThreadTasks(MaxBatchSize, Settings.MinParallelBatchSize, EParallelForFlags::None);
-			const int32 InferenceInstanceSliceLength = FMath::DivideAndRoundUp((int32)MaxBatchSize, IdealInferenceInstanceNum);
-			const int32 InferenceInstanceNum = FMath::DivideAndRoundUp((int32)MaxBatchSize, InferenceInstanceSliceLength);
-
-			ModelInstances.Empty(InferenceInstanceNum);
-
-			for (int32 InferenceInstanceIdx = 0; InferenceInstanceIdx < InferenceInstanceNum; InferenceInstanceIdx++)
-			{
-				ModelInstances.Emplace(InModel.CreateModelInstanceCPU());
-				ModelInstances.Last()->SetInputTensorShapes({ NNE::FTensorShape::Make({(uint32)InferenceInstanceSliceLength, (uint32)InputSize}) });
-			}
-		}
-		else
-		{
-			ModelInstances.Empty(1);
-			ModelInstances.Emplace(InModel.CreateModelInstanceCPU());
-			ModelInstances.Last()->SetInputTensorShapes({ NNE::FTensorShape::Make({(uint32)MaxBatchSize, (uint32)InputSize}) });
-		}
+		FrontBuffer.SetNumUninitialized({ MaxBatchSize, BufferSize });
+		BackBuffer.SetNumUninitialized({ MaxBatchSize, BufferSize });
 	}
 
-	void FNeuralNetworkInference::Evaluate(
+	void FNeuralNetworkInferenceMLP::Evaluate(
 		TLearningArrayView<2, float> Output,
 		const TLearningArrayView<2, const float> Input,
-		const FIndexSet Instances)
+		const FIndexSet Instances) 
 	{
-		UE_LEARNING_TRACE_CPUPROFILER_EVENT_SCOPE(Learning::FNeuralNetworkInference::Evaluate);
+		const int32 InputNum = Network.GetInputNum();
+		const int32 OutputNum = Network.GetOutputNum();
+		const int32 LayerNum = Network.GetLayerNum();
 
-		const int32 InstanceNum = Instances.Num();
-		const int32 InputNum = Input.Num<1>();
-		const int32 OutputNum = Output.Num<1>();
+		UE_LEARNING_CHECK(Input.Num<1>() == InputNum);
+		UE_LEARNING_CHECK(Output.Num<1>() == OutputNum);
 
-		if (InstanceNum == 0) { return; }
-
-		// Evaluation is not thread-safe due to the scatter and gather we do into the InputBuffer and OutputBuffer
-		// so we take a lock here since the interface for this function makes it appear like this might be thread safe as 
-		// long as the Instances provided are not an overlapping set.
-		FScopeNullableWriteLock ScopeLock(&EvaluationLock);
-
-		// Gather Inputs
+		// Gather Instance Inputs
 
 		Array::Check(Input, Instances);
 
-		for (int32 InstanceIdx = 0; InstanceIdx < InstanceNum; InstanceIdx++)
+		const TLearningArrayView<2, float> InputBuffer = FrontBuffer;
+
+		for (int32 InstanceIdx = 0; InstanceIdx < Instances.Num(); InstanceIdx++)
 		{
-			Array::Copy(InputBuffer[InstanceIdx], Input[Instances[InstanceIdx]]);
+			Array::Copy(InputBuffer[InstanceIdx].Slice(0, InputNum), Input[Instances[InstanceIdx]]);
 		}
 
-		// Run Inference
+		// Evaluate Network
 
-		Array::Check(InputBuffer.Slice(0, InstanceNum));
-
-		if (Settings.bParallelEvaluation && InstanceNum > Settings.MinParallelBatchSize)
+		auto EvaluationFunction = [this, LayerNum](const int32 BatchOffset, const int32 BatchSize)
 		{
-			const int32 IdealInferenceInstanceNum = ParallelForImpl::GetNumberOfThreadTasks(InstanceNum, Settings.MinParallelBatchSize, EParallelForFlags::None);
-			const int32 InferenceInstanceSliceLength = FMath::DivideAndRoundUp(InstanceNum, IdealInferenceInstanceNum);
-			const int32 InferenceInstanceNum = FMath::DivideAndRoundUp(InstanceNum, InferenceInstanceSliceLength);
-
-			UE_LEARNING_CHECK(InferenceInstanceNum <= ModelInstances.Num());
-
-			ParallelFor(InferenceInstanceNum, [this, InstanceNum, InputNum, InferenceInstanceSliceLength](int32 InferenceInstanceIdx)
+			for (int32 LayerIdx = 0; LayerIdx < LayerNum; LayerIdx++)
 			{
-				const int32 StartIndex = InferenceInstanceIdx * InferenceInstanceSliceLength;
-				const int32 StopIndex = FMath::Min((InferenceInstanceIdx + 1) * InferenceInstanceSliceLength, InstanceNum);
-
-				TLearningArrayView<2, float> InputBufferSlice = InputBuffer.Slice(StartIndex, StopIndex - StartIndex);
-				TLearningArrayView<2, float> OutputBufferSlice = OutputBuffer.Slice(StartIndex, StopIndex - StartIndex);
+				const TLearningArrayView<2, float> LayerInputBuffer = LayerIdx % 2 == 0 ? FrontBuffer : BackBuffer;
+				const TLearningArrayView<2, float> LayerOutputBuffer = LayerIdx % 2 == 0 ? BackBuffer : FrontBuffer;
+				const int32 LayerOutputNum = Network.Weights[LayerIdx].Num<1>();
+
+				NeuralNetwork::MatMulPlusBias(
+					LayerOutputBuffer,
+					LayerInputBuffer,
+					Network.Weights[LayerIdx],
+					Network.Biases[LayerIdx],
+					BatchOffset,
+					BatchSize);
+
+				if (LayerIdx != LayerNum - 1)
+				{
+					switch (Network.ActivationFunction)
+					{
+					case EActivationFunction::ReLU: NeuralNetwork::ActivationReLU(LayerOutputBuffer, BatchOffset, BatchSize, LayerOutputNum); break;
+					case EActivationFunction::ELU: NeuralNetwork::ActivationELU(LayerOutputBuffer, BatchOffset, BatchSize, LayerOutputNum); break;
+					case EActivationFunction::TanH: NeuralNetwork::ActivationTanH(LayerOutputBuffer, BatchOffset, BatchSize, LayerOutputNum); break;
+					default: UE_LEARNING_CHECKF(false, TEXT("Unimplemented"));
+					}
+				}
+			}
+		};
 
-				ModelInstances[InferenceInstanceIdx]->SetInputTensorShapes({ NNE::FTensorShape::Make({(uint32)(StopIndex - StartIndex), (uint32)InputNum}) });
-				ModelInstances[InferenceInstanceIdx]->RunSync(
-					{ { (void*)InputBufferSlice.GetData(), InputBufferSlice.Num() * sizeof(float) } },
-					{ { (void*)OutputBufferSlice.GetData(), OutputBufferSlice.Num() * sizeof(float) } });
-			});
+		if (Settings.bParallelEvaluation && Instances.Num() > Settings.MinParallelBatchSize)
+		{
+			Learning::SlicedParallelFor(Instances.Num(), Settings.MinParallelBatchSize, EvaluationFunction);
 		}
 		else
 		{
-			ModelInstances[0]->SetInputTensorShapes({ NNE::FTensorShape::Make({(uint32)InstanceNum, (uint32)InputNum}) });
-			ModelInstances[0]->RunSync(
-				{ { (void*)InputBuffer.GetData(), InputBuffer.Slice(0, InstanceNum).Num() * sizeof(float) } },
-				{ { (void*)OutputBuffer.GetData(), OutputBuffer.Slice(0, InstanceNum).Num() * sizeof(float) } });
+			EvaluationFunction(0, Instances.Num());
 		}
 
-		Array::Check(OutputBuffer.Slice(0, InstanceNum));
+		// Scatter Instance Outputs
 
-		// Scatter Outputs
+		const TLearningArrayView<2, float> OutputBuffer = LayerNum % 2 == 0 ? FrontBuffer : BackBuffer;
 
-		for (int32 InstanceIdx = 0; InstanceIdx < InstanceNum; InstanceIdx++)
+		for (int32 InstanceIdx = 0; InstanceIdx < Instances.Num(); InstanceIdx++)
 		{
-			Array::Copy(Output[Instances[InstanceIdx]], OutputBuffer[InstanceIdx]);
+			Array::Copy(Output[Instances[InstanceIdx]], OutputBuffer[InstanceIdx].Slice(0, OutputNum));
 		}
 
 		Array::Check(Output, Instances);
 	}
-
-
-	FNeuralNetworkFunction::FNeuralNetworkFunction(
-		const int32 InMaxInstanceNum,
-		const TSharedPtr<FNeuralNetwork>& InNeuralNetwork,
-		const FNeuralNetworkInferenceSettings& InInferenceSettings) 
-		: MaxInstanceNum(InMaxInstanceNum)
-		, NeuralNetwork(InNeuralNetwork)
-		, InferenceSettings(InInferenceSettings)
-	{
-		NeuralNetworkInference = NeuralNetwork->CreateInferenceObject(InMaxInstanceNum, InInferenceSettings);
-	}
-
-	void FNeuralNetworkFunction::Evaluate(
-		TLearningArrayView<2, float> Output,
-		const TLearningArrayView<2, const float> Input,
-		const FIndexSet Instances)
-	{
-		NeuralNetworkInference->Evaluate(Output, Input, Instances);
-	}
-
-	void FNeuralNetworkFunction::UpdateNeuralNetwork(const TSharedPtr<FNeuralNetwork>& NewNeuralNetwork)
-	{
-		if (NeuralNetwork != NewNeuralNetwork)
-		{
-			NeuralNetwork = NewNeuralNetwork;
-			NeuralNetworkInference = NeuralNetwork->CreateInferenceObject(MaxInstanceNum, InferenceSettings);
-		}
-	}
-
 }
diff --git a/Engine/Plugins/Experimental/LearningAgents/Source/Learning/Private/LearningNeuralNetworkObject.cpp b/Engine/Plugins/Experimental/LearningAgents/Source/Learning/Private/LearningNeuralNetworkObject.cpp
new file mode 100644
--- /dev/null
+++ b/Engine/Plugins/Experimental/LearningAgents/Source/Learning/Private/LearningNeuralNetworkObject.cpp
@@ -0,0 +1,241 @@
+// Copyright Epic Games, Inc. All Rights Reserved.
+
+#include "LearningNeuralNetworkObject.h"
+
+#include "LearningRandom.h"
+
+#if UE_LEARNING_ISPC
+#include "Learning.ispc.generated.h"
+#endif
+
+namespace UE::Learning
+{
+	namespace NeuralNetworkPolicyFunction::Private
+	{
+		static inline float Sigmoid(const float X)
+		{
+			return 1.0f / (1.0f + FMath::Exp(-X));
+		}
+	}
+
+	FNeuralNetworkPolicyFunction::FNeuralNetworkPolicyFunction(
+		const FName& InIdentifier,
+		const TSharedRef<FArrayMap>& InInstanceData,
+		const int32 InMaxInstanceNum,
+		const int32 InObservationNum,
+		const int32 InActionNum,
+		const int32 InMemoryStateNum,
+		const TSharedPtr<INeuralNetwork>& InNeuralNetwork,
+		const uint32 InSeed,
+		const FNeuralNetworkInferenceSettings& InInferenceSettings,
+		const FNeuralNetworkPolicyFunctionSettings& InSettings)
+		: FFunctionObject(InInstanceData)
+		, MaxInstanceNum(InMaxInstanceNum)
+		, ObservationNum(InObservationNum)
+		, ActionNum(InActionNum)
+		, MemoryStateNum(InMemoryStateNum)
+		, NeuralNetwork(InNeuralNetwork)
+		, InferenceSettings(InInferenceSettings)
+		, Settings(InSettings)
+	{
+		UE_LEARNING_CHECK(NeuralNetwork->GetInputNum() == ObservationNum + MemoryStateNum);
+		UE_LEARNING_CHECK(NeuralNetwork->GetOutputNum() == 2 * ActionNum + MemoryStateNum);
+
+		SeedHandle = InstanceData->Add<1, uint32>({ InIdentifier, TEXT("Seed") }, { InMaxInstanceNum });
+		
+		InputObservationHandle = InstanceData->Add<2, float>({ InIdentifier, TEXT("InputObservation") }, { InMaxInstanceNum, ObservationNum }, 0.0f);
+		InputMemoryStateHandle = InstanceData->Add<2, float>({ InIdentifier, TEXT("InputMemoryState") }, { InMaxInstanceNum, MemoryStateNum }, 0.0f);
+		InputNetworkHandle = InstanceData->Add<2, float>({ InIdentifier, TEXT("InputNetwork") }, { InMaxInstanceNum, ObservationNum + MemoryStateNum }, 0.0f);
+
+		OutputNetworkHandle = InstanceData->Add<2, float>({ InIdentifier, TEXT("OutputNetwork") }, { InMaxInstanceNum, 2 * ActionNum + MemoryStateNum }, 0.0f);
+		OutputMemoryStateHandle = InstanceData->Add<2, float>({ InIdentifier, TEXT("OutputMemoryState") }, { InMaxInstanceNum, MemoryStateNum }, 0.0f);
+		OutputActionHandle = InstanceData->Add<2, float>({ InIdentifier, TEXT("OutputAction") }, { InMaxInstanceNum, ActionNum }, 0.0f);
+		OutputActionMeanHandle = InstanceData->Add<2, float>({ InIdentifier, TEXT("OutputActionMean") }, { InMaxInstanceNum, ActionNum }, 0.0f);
+		OutputActionStdHandle = InstanceData->Add<2, float>({ InIdentifier, TEXT("OutputActionStd") }, { InMaxInstanceNum, ActionNum }, 0.0f);
+
+		ActionNoiseScaleHandle = InstanceData->Add<1, float>({ InIdentifier, TEXT("ActionNoiseScale") }, { InMaxInstanceNum }, Settings.ActionNoiseScale);
+
+		NeuralNetworkInference = NeuralNetwork->CreateInferenceObject(InMaxInstanceNum, InInferenceSettings);
+		
+		Random::IntArray(InstanceData->View(SeedHandle), InSeed);
+	}
+
+	void FNeuralNetworkPolicyFunction::Evaluate(const FIndexSet Instances)
+	{
+		UE_LEARNING_TRACE_CPUPROFILER_EVENT_SCOPE(Learning::FNeuralNetworkPolicyFunction::Evaluate);
+
+		TLearningArrayView<1, uint32> Seed = InstanceData->View(SeedHandle);
+
+		const TLearningArrayView<2, const float> InputObservation = InstanceData->ConstView(InputObservationHandle);
+		const TLearningArrayView<2, const float> InputMemoryState = InstanceData->ConstView(InputMemoryStateHandle);
+		TLearningArrayView<2, float> InputNetwork = InstanceData->View(InputNetworkHandle);
+
+		TLearningArrayView<2, float> OutputNetwork = InstanceData->View(OutputNetworkHandle);
+		TLearningArrayView<2, float> OutputMemoryState = InstanceData->View(OutputMemoryStateHandle);
+		TLearningArrayView<2, float> OutputAction = InstanceData->View(OutputActionHandle);
+		TLearningArrayView<2, float> OutputActionMean = InstanceData->View(OutputActionMeanHandle);
+		TLearningArrayView<2, float> OutputActionStd = InstanceData->View(OutputActionStdHandle);
+
+		const TLearningArrayView<1, const float> ActionNoiseScale = InstanceData->ConstView(ActionNoiseScaleHandle);
+
+		const int32 InstanceNum = Instances.Num();
+
+		Array::Check(InputObservation, Instances);
+		Array::Check(InputMemoryState, Instances);
+
+		// Copy in Observation and Memory State into network input
+
+		for (const int32 InstanceIdx : Instances)
+		{
+			Array::Copy(InputNetwork[InstanceIdx].Slice(0, ObservationNum), InputObservation[InstanceIdx]);
+			Array::Copy(InputNetwork[InstanceIdx].Slice(ObservationNum, MemoryStateNum), InputMemoryState[InstanceIdx]);
+		}
+
+		NeuralNetworkInference->Evaluate(OutputNetwork, InputNetwork, Instances);
+
+		Array::Check(OutputNetwork, Instances);
+
+		// Copy Out Memory State
+
+		for (const int32 InstanceIdx : Instances)
+		{
+			Array::Copy(
+				OutputMemoryState[InstanceIdx],
+				OutputNetwork[InstanceIdx].Slice(2 * ActionNum, MemoryStateNum));
+		}
+
+		Array::Check(OutputMemoryState, Instances);
+
+		// Apply Action Noise
+
+		UE_LEARNING_ARRAY_VALUE_CHECK(Settings.ActionNoiseMin >= 0.0f && Settings.ActionNoiseMax >= 0.0f);
+		const float LogActionNoiseMin = FMath::Loge(Settings.ActionNoiseMin + UE_KINDA_SMALL_NUMBER);
+		const float LogActionNoiseMax = FMath::Loge(Settings.ActionNoiseMax + UE_KINDA_SMALL_NUMBER);
+
+#if UE_LEARNING_ISPC
+		if (Instances.IsSlice())
+		{
+			ispc::LearningLayerActionNoise(
+				OutputAction.Slice(Instances.GetSliceStart(), Instances.GetSliceNum()).GetData(),
+				OutputActionMean.Slice(Instances.GetSliceStart(), Instances.GetSliceNum()).GetData(),
+				OutputActionStd.Slice(Instances.GetSliceStart(), Instances.GetSliceNum()).GetData(),
+				OutputNetwork.Slice(Instances.GetSliceStart(), Instances.GetSliceNum()).GetData(),
+				Seed.Slice(Instances.GetSliceStart(), Instances.GetSliceNum()).GetData(),
+				ActionNoiseScale.Slice(Instances.GetSliceStart(), Instances.GetSliceNum()).GetData(),
+				Instances.GetSliceNum(),
+				ActionNum,
+				MemoryStateNum,
+				LogActionNoiseMin,
+				LogActionNoiseMax);
+		}
+		else
+		{
+			for (const int32 InstanceIdx : Instances)
+			{
+				ispc::LearningLayerActionNoiseSingleBatch(
+					OutputAction[InstanceIdx].GetData(),
+					OutputActionMean[InstanceIdx].GetData(),
+					OutputActionStd[InstanceIdx].GetData(),
+					OutputNetwork[InstanceIdx].GetData(),
+					Seed[InstanceIdx],
+					ActionNoiseScale[InstanceIdx],
+					ActionNum,
+					MemoryStateNum,
+					LogActionNoiseMin,
+					LogActionNoiseMax);
+			}
+		}
+#else
+		for (const int32 InstanceIdx : Instances)
+		{
+			for (int32 ActionIdx = 0; ActionIdx < ActionNum; ActionIdx++)
+			{
+				OutputActionMean[InstanceIdx][ActionIdx] = OutputNetwork[InstanceIdx][ActionIdx];
+				OutputActionStd[InstanceIdx][ActionIdx] = ActionNoiseScale[InstanceIdx] *
+					FMath::Exp(NeuralNetworkPolicyFunction::Private::Sigmoid(
+						OutputNetwork[InstanceIdx][ActionNum + ActionIdx]) * (LogActionNoiseMax - LogActionNoiseMin) + LogActionNoiseMin);
+
+				OutputAction[InstanceIdx][ActionIdx] = Random::Gaussian(
+					Seed[InstanceIdx] ^ 0xab744615 ^ Random::Int(ActionIdx ^ 0xf8a88a27),
+					OutputActionMean[InstanceIdx][ActionIdx],
+					OutputActionStd[InstanceIdx][ActionIdx]);
+			}
+		}
+#endif
+
+		Random::ResampleStateArray(Seed, Instances);
+
+		Array::Check(OutputActionMean, Instances);
+		Array::Check(OutputActionStd, Instances);
+		Array::Check(OutputAction, Instances);
+	}
+
+	void FNeuralNetworkPolicyFunction::UpdateNeuralNetwork(const TSharedPtr<INeuralNetwork>& NewNeuralNetwork)
+	{
+		if (NeuralNetwork != NewNeuralNetwork)
+		{
+			NeuralNetwork = NewNeuralNetwork;
+			NeuralNetworkInference = NeuralNetwork->CreateInferenceObject(MaxInstanceNum, InferenceSettings);
+		}
+	}
+
+	FNeuralNetworkCriticFunction::FNeuralNetworkCriticFunction(
+		const FName& InIdentifier,
+		const TSharedRef<FArrayMap>& InInstanceData,
+		const int32 InMaxInstanceNum,
+		const int32 InObservationNum,
+		const int32 InMemoryStateNum,
+		const TSharedPtr<INeuralNetwork>& InNeuralNetwork,
+		const FNeuralNetworkInferenceSettings& InInferenceSettings)
+		: FFunctionObject(InInstanceData)
+		, MaxInstanceNum(InMaxInstanceNum)
+		, ObservationNum(InObservationNum)
+		, MemoryStateNum(InMemoryStateNum)
+		, NeuralNetwork(InNeuralNetwork)
+		, InferenceSettings(InInferenceSettings)
+	{
+		UE_LEARNING_CHECK(NeuralNetwork->GetInputNum() == ObservationNum + MemoryStateNum);
+		UE_LEARNING_CHECK(NeuralNetwork->GetOutputNum() == 1);
+
+		InputObservationHandle = InstanceData->Add<2, float>({ InIdentifier, TEXT("InputObservation") }, { InMaxInstanceNum, ObservationNum }, 0.0f);
+		InputMemoryStateHandle = InstanceData->Add<2, float>({ InIdentifier, TEXT("InputMemoryState") }, { InMaxInstanceNum, MemoryStateNum }, 0.0f);
+		InputNetworkHandle = InstanceData->Add<2, float>({ InIdentifier, TEXT("InputNetwork") }, { InMaxInstanceNum, ObservationNum + MemoryStateNum }, 0.0f);
+		OutputHandle = InstanceData->Add<1, float>({ InIdentifier, TEXT("Output") }, { InMaxInstanceNum }, 0.0f);
+
+		NeuralNetworkInference = NeuralNetwork->CreateInferenceObject(InMaxInstanceNum, InInferenceSettings);
+	}
+
+	void FNeuralNetworkCriticFunction::Evaluate(const FIndexSet Instances)
+	{
+		UE_LEARNING_TRACE_CPUPROFILER_EVENT_SCOPE(Learning::FNeuralNetworkCriticFunction::Evaluate);
+
+		const TLearningArrayView<2, const float> InputObservation = InstanceData->ConstView(InputObservationHandle);
+		const TLearningArrayView<2, const float> InputMemoryState = InstanceData->ConstView(InputMemoryStateHandle);
+		TLearningArrayView<2, float> InputNetwork = InstanceData->View(InputNetworkHandle);
+		TLearningArrayView<1, float> Output = InstanceData->View(OutputHandle);
+
+		// Copy in Observation and Memory State into network input
+
+		for (const int32 InstanceIdx : Instances)
+		{
+			Array::Copy(InputNetwork[InstanceIdx].Slice(0, ObservationNum), InputObservation[InstanceIdx]);
+			Array::Copy(InputNetwork[InstanceIdx].Slice(ObservationNum, MemoryStateNum), InputMemoryState[InstanceIdx]);
+		}
+
+		NeuralNetworkInference->Evaluate(
+			TLearningArrayView<2, float>(Output.GetData(), { Output.Num(), 1 }),
+			InputNetwork, 
+			Instances);
+
+		Array::Check(Output, Instances);
+	}
+
+	void FNeuralNetworkCriticFunction::UpdateNeuralNetwork(const TSharedPtr<INeuralNetwork>& NewNeuralNetwork)
+	{
+		if (NeuralNetwork != NewNeuralNetwork)
+		{
+			NeuralNetwork = NewNeuralNetwork;
+			NeuralNetworkInference = NeuralNetwork->CreateInferenceObject(MaxInstanceNum, InferenceSettings);
+		}
+	}
+}
diff --git a/Engine/Plugins/Experimental/LearningAgents/Source/Learning/Private/LearningRandom.cpp b/Engine/Plugins/Experimental/LearningAgents/Source/Learning/Private/LearningRandom.cpp
--- a/Engine/Plugins/Experimental/LearningAgents/Source/Learning/Private/LearningRandom.cpp
+++ b/Engine/Plugins/Experimental/LearningAgents/Source/Learning/Private/LearningRandom.cpp
@@ -13,11 +13,6 @@ namespace UE::Learning::Random
 		return FMath::Sqrt(-2.0f * FMath::Loge(FMath::Max(R1, UE_SMALL_NUMBER))) * FMath::Cos(R2 * UE_TWO_PI);
 	}
 
-	static inline float Sigmoid(const float X)
-	{
-		return 1.0f / (1.0f + FMath::Exp(-X));
-	}
-
 	uint32 Int(const uint32 State)
 	{
 		// Here we use a simple xor and shift based 
@@ -148,28 +143,6 @@ namespace UE::Learning::Random
 #endif
 	}
 
-	LEARNING_API void IntArray(
-		TLearningArrayView<1, uint32> Output,
-		const uint32 State,
-		const FIndexSet Indices)
-	{
-		UE_LEARNING_TRACE_CPUPROFILER_EVENT_SCOPE(Learning::Random::IntArray);
-
-		if (Indices.IsSlice() && Indices.Num() > 4)
-		{
-			IntArray(
-				Output.Slice(Indices.GetSliceStart(), Indices.GetSliceNum()),
-				State);
-		}
-		else
-		{
-			for (const int32 ElementIdx : Indices)
-			{
-				Output[ElementIdx] = Int(State ^ 0xbed25b30 ^ Int(ElementIdx ^ 0xb521a8d9));
-			}
-		}
-	}
-
 	void FloatArray(
 		TLearningArrayView<1, float> Output,
 		const uint32 State)
@@ -276,126 +249,6 @@ namespace UE::Learning::Random
 		}
 	}
 
-	void DistributionIndependantNormal(
-		TLearningArrayView<1, float> Output,
-		const uint32 State,
-		const TLearningArrayView<1, const float> Mean,
-		const TLearningArrayView<1, const float> LogStd,
-		const float Scale)
-	{
-		UE_LEARNING_TRACE_CPUPROFILER_EVENT_SCOPE(Learning::Random::DistributionIndependantNormal);
-
-		const int32 ElementNum = Output.Num();
-
-#if UE_LEARNING_ISPC
-		ispc::LearningRandomDistributionIndependantNormal(
-			Output.GetData(),
-			Mean.GetData(),
-			LogStd.GetData(),
-			ElementNum,
-			State,
-			Scale);
-#else
-		if (Scale < UE_SMALL_NUMBER)
-		{
-			for (int32 ElementIdx = 0; ElementIdx < ElementNum; ElementIdx++)
-			{
-				Output[ElementIdx] = Mean[ElementIdx];
-			}
-		}
-		else
-		{
-			for (int32 ElementIdx = 0; ElementIdx < ElementNum; ElementIdx++)
-			{
-				Output[ElementIdx] = Gaussian(State ^ 0x7db0e4e9 ^ Int(ElementIdx ^ 0xf4976a00), Mean[ElementIdx], Scale * FMath::Exp(FMath::Min(LogStd[ElementIdx], 10.0f)));
-			}
-		}
-#endif
-	}
-
-	void DistributionMultinoulli(
-		TLearningArrayView<1, float> Output,
-		const uint32 State,
-		const TLearningArrayView<1, const float> Logits,
-		const float Scale)
-	{
-		UE_LEARNING_TRACE_CPUPROFILER_EVENT_SCOPE(Learning::Random::DistributionMultinoulli);
-
-		const int32 ElementNum = Output.Num();
-		if (ElementNum == 0)
-		{
-			return;
-		}
-
-		float MaxValue = -FLT_MAX;
-		uint32 MaxIdx = INDEX_NONE;
-
-		if (Scale < UE_SMALL_NUMBER)
-		{
-			for (int32 ElementIdx = 0; ElementIdx < ElementNum; ElementIdx++)
-			{
-				if (Logits[ElementIdx] > MaxValue)
-				{
-					MaxValue = Logits[ElementIdx];
-					MaxIdx = ElementIdx;
-				}
-			}
-		}
-		else
-		{
-			for (int32 ElementIdx = 0; ElementIdx < ElementNum; ElementIdx++)
-			{
-				const float ElementValue = Logits[ElementIdx] / Scale - FMath::Loge(-FMath::Loge(Uniform(State ^ 0x7a156b37 ^ Int(ElementIdx ^ 0x0c71bebb))));
-
-				if (ElementValue > MaxValue)
-				{
-					MaxValue = ElementValue;
-					MaxIdx = ElementIdx;
-				}
-			}
-		}
-
-		UE_LEARNING_CHECK(MaxIdx != INDEX_NONE);
-
-		Array::Zero(Output);
-		Output[MaxIdx] = 1.0f;
-	}
-
-	void DistributionBernoulli(
-		TLearningArrayView<1, float> Output,
-		const uint32 State,
-		const TLearningArrayView<1, const float> Logits,
-		const float Scale)
-	{
-		UE_LEARNING_TRACE_CPUPROFILER_EVENT_SCOPE(Learning::Random::DistributionBernoulli);
-
-		const int32 ElementNum = Output.Num();
-
-#if UE_LEARNING_ISPC
-		ispc::LearningRandomDistributionBernoulli(
-			Output.GetData(),
-			Logits.GetData(),
-			ElementNum,
-			State,
-			Scale);
-#else
-		if (Scale < UE_SMALL_NUMBER)
-		{
-			for (int32 ElementIdx = 0; ElementIdx < ElementNum; ElementIdx++)
-			{
-				Output[ElementIdx] = Logits[ElementIdx] > 0.0f ? 1.0f : 0.0f;
-			}
-		}
-		else
-		{
-			for (int32 ElementIdx = 0; ElementIdx < ElementNum; ElementIdx++)
-			{
-				Output[ElementIdx] = Sigmoid(Logits[ElementIdx] / Scale) > Uniform(State ^ 0xf4021a46 ^ Int(ElementIdx ^ 0x7a8cc64e)) ? 1.0f : 0.0f;
-			}
-		}
-#endif
-	}
-
 	////////////
 
 	uint32 SampleInt(uint32& State)
@@ -501,15 +354,6 @@ namespace UE::Learning::Random
 		IntArray(Output, State ^ 0x6adcdb41);
 	}
 
-	void SampleIntArray(
-		TLearningArrayView<1, uint32> Output,
-		uint32& State,
-		const FIndexSet Indices)
-	{
-		State = Int(State ^ 0xab9c9ee3);
-		IntArray(Output, State ^ 0x6adcdb41, Indices);
-	}
-
 	void SampleFloatArray(
 		TLearningArrayView<1, float> Output,
 		uint32& State)
@@ -538,6 +382,29 @@ namespace UE::Learning::Random
 		GaussianArray(Output, State ^ 0xf2381309, Mean, Std);
 	}
 
+	void ResampleStateArray(TLearningArrayView<1, uint32> InOutStates, const FIndexSet Indices)
+	{
+		UE_LEARNING_TRACE_CPUPROFILER_EVENT_SCOPE(Learning::Random::ResampleStateArray);
+
+		const int32 ElementNum = Indices.Num();
+
+		if (UE_LEARNING_ISPC && Indices.IsSlice() && Indices.Num() > 4)
+		{
+#if UE_LEARNING_ISPC
+			ispc::LearningResampleStateArray(
+				InOutStates.Slice(Indices.GetSliceStart(), Indices.GetSliceNum()).GetData(),
+				Indices.GetSliceNum());
+#endif
+		}
+		else
+		{
+			for (const int32 ElementIdx : Indices)
+			{
+				InOutStates[ElementIdx] = Int(InOutStates[ElementIdx] ^ 0x3616dcc8 ^ Int(ElementIdx ^ 0x6a837ffd));
+			}
+		}
+	}
+
 	void SamplePlanarClippedGaussianArray(
 		TLearningArrayView<1, FVector> Output,
 		uint32& State,
@@ -561,48 +428,4 @@ namespace UE::Learning::Random
 		PlanarDirectionArray(Output, State ^ 0x6bfd3e6a, Axis0, Axis1);
 	}
 
-	void SampleDistributionIndependantNormal(
-		TLearningArrayView<1, float> Output,
-		uint32& State,
-		const TLearningArrayView<1, const float> Mean,
-		const TLearningArrayView<1, const float> LogStd,
-		const float Scale)
-	{
-		State = Int(State ^ 0x984c4d5f);
-		DistributionIndependantNormal(
-			Output,
-			State,
-			Mean,
-			LogStd,
-			Scale);
-	}
-
-	void SampleDistributionMultinoulli(
-		TLearningArrayView<1, float> Output,
-		uint32& State,
-		const TLearningArrayView<1, const float> Logits,
-		const float Scale)
-	{
-		State = Int(State ^ 0xc9261d09);
-		DistributionMultinoulli(
-			Output,
-			State,
-			Logits,
-			Scale);
-	}
-
-	void SampleDistributionBernoulli(
-		TLearningArrayView<1, float> Output,
-		uint32& State,
-		const TLearningArrayView<1, const float> Logits,
-		const float Scale)
-	{
-		State = Int(State ^ 0xfeba3d63);
-		DistributionBernoulli(
-			Output,
-			State,
-			Logits,
-			Scale);
-	}
-
 }
\ No newline at end of file
diff --git a/Engine/Plugins/Experimental/LearningAgents/Source/Learning/Private/LearningRecorderObject.cpp b/Engine/Plugins/Experimental/LearningAgents/Source/Learning/Private/LearningRecorderObject.cpp
new file mode 100644
--- /dev/null
+++ b/Engine/Plugins/Experimental/LearningAgents/Source/Learning/Private/LearningRecorderObject.cpp
@@ -0,0 +1,46 @@
+// Copyright Epic Games, Inc. All Rights Reserved.
+
+#include "LearningRecorderObject.h"
+
+namespace UE::Learning
+{
+	FRecorderObject::FRecorderObject(
+		const int32 InMaxInstanceNum,
+		const int32 InMaxStepNum,
+		const TSharedRef<FArrayMap>& InSrcArrayMap,
+		const TSharedRef<FArrayMap>& InLogArrayMap)
+		: MaxInstanceNum(InMaxInstanceNum)
+		, MaxStepNum(InMaxStepNum)
+		, SrcArrayMap(InSrcArrayMap)
+		, LogArrayMap(InLogArrayMap)
+	{
+		StepNums.SetNumUninitialized({ MaxInstanceNum });
+		Array::Zero(StepNums);
+	}
+
+	FRecorderObject::~FRecorderObject() { }
+
+	void FRecorderObject::Reset(const FIndexSet Instances)
+	{
+		Array::Zero(StepNums, Instances);
+	}
+
+	void FRecorderObject::Update(const FIndexSet Instances)
+	{
+		const int32 VariableNum = SrcHandles.Num();
+
+		for (int32 VariableIdx = 0; VariableIdx < VariableNum; VariableIdx++)
+		{
+			CopyFunctions[VariableIdx](
+				StepNums,
+				*SrcArrayMap,
+				*LogArrayMap,
+				SrcHandles[VariableIdx],
+				LogHandles[VariableIdx],
+				MaxStepNum,
+				Instances);
+		}
+	}
+
+}
+
diff --git a/Engine/Plugins/Experimental/LearningAgents/Source/Learning/Private/LearningRewardObject.cpp b/Engine/Plugins/Experimental/LearningAgents/Source/Learning/Private/LearningRewardObject.cpp
new file mode 100644
--- /dev/null
+++ b/Engine/Plugins/Experimental/LearningAgents/Source/Learning/Private/LearningRewardObject.cpp
@@ -0,0 +1,1015 @@
+// Copyright Epic Games, Inc. All Rights Reserved.
+
+#include "LearningRewardObject.h"
+#include "LearningLog.h"
+
+namespace UE::Learning
+{
+	namespace Reward
+	{
+		// Misc
+
+		static inline FVector VectorExp(const FVector X)
+		{
+			return FVector(FMath::Exp(X[0]), FMath::Exp(X[1]), FMath::Exp(X[2]));
+		}
+
+		static inline FVector VectorMax(const FVector X, const float Y)
+		{
+			return FVector(FMath::Max(X[0], Y), FMath::Max(X[1], Y), FMath::Max(X[2], Y));
+		}
+
+		static inline float SpringInertializeHalfLifeToDamping(const float HalfLife, const float Epsilon = UE_SMALL_NUMBER)
+		{
+			UE_LEARNING_ARRAY_VALUE_CHECK(HalfLife > 0.0f);
+
+			const float Ln2 = 0.69314718056f;
+			return (4.0f * Ln2) / FMath::Max(HalfLife, Epsilon);
+		}
+
+		static inline FVector SpringInertializeDecayIntersectionTime(const FVector Position, const FVector Velocity, const float HalfLife, const float Epsilon = UE_SMALL_NUMBER)
+		{
+			UE_LEARNING_ARRAY_VALUE_CHECK(HalfLife > 0.0f);
+
+			const float Y = SpringInertializeHalfLifeToDamping(HalfLife, Epsilon) / 2.0f;
+			return -Position / VectorMax(Velocity + Position * Y, Epsilon);
+		}
+
+		static inline FVector SpringInertializeDecayTotalDisplacement(const FVector PositionOffset, const FVector VelocityOffset, const float HalfLife, const float Epsilon = UE_SMALL_NUMBER)
+		{
+			const FVector X = PositionOffset;
+			const FVector V = VelocityOffset;
+			const float Y = SpringInertializeHalfLifeToDamping(HalfLife, Epsilon) / 2.0f;
+			const FVector T = SpringInertializeDecayIntersectionTime(X, V, HalfLife, Epsilon);
+
+			const FVector Integral0 = (X * Y * 2 + V) / (Y * Y);
+			const FVector IntegralT = (VectorExp(-Y * T) / (Y * Y)) * (T * V * Y + X * Y * (T * Y + 2) + V);
+
+			return FVector(
+				T[0] > 0.0f ? FMath::Abs(Integral0[0] - IntegralT[0]) + FMath::Abs(IntegralT[0]) : FMath::Abs(Integral0[0]),
+				T[1] > 0.0f ? FMath::Abs(Integral0[1] - IntegralT[1]) + FMath::Abs(IntegralT[1]) : FMath::Abs(Integral0[1]),
+				T[2] > 0.0f ? FMath::Abs(Integral0[2] - IntegralT[2]) + FMath::Abs(IntegralT[2]) : FMath::Abs(Integral0[2]));
+		}
+
+		// Helpers
+
+		float DistanceToReward(const float Distance)
+		{
+			return SquaredDistanceToReward(FMath::Square(Distance));
+		}
+
+		float DistanceToPenalty(const float Distance)
+		{
+			return -Distance;
+		}
+
+		float SquaredDistanceToReward(const float SquaredDistance)
+		{
+			return FMath::InvExpApprox(SquaredDistance);
+		}
+
+		float SquaredDistanceToPenalty(const float SquaredDistance)
+		{
+			return DistanceToPenalty(FMath::Sqrt(SquaredDistance));
+		}
+
+		// Rewards
+
+		float ScalarSimilarityReward(
+			const float Scalar0,
+			const float Scalar1,
+			const float Scale,
+			const float Threshold,
+			const float Weight,
+			const float Epsilon)
+		{
+			UE_LEARNING_ARRAY_VALUE_CHECK(Scale > 0.0f);
+
+			return Weight * DistanceToReward(FMath::Max(FMath::Abs(Scalar0 - Scalar1) - Threshold, 0.0f) / FMath::Max(Scale, Epsilon));
+		}
+
+		float ScalarPositionSimilarityReward(
+			const float Position0,
+			const float Position1,
+			const float Scale,
+			const float Threshold,
+			const float Weight,
+			const float Epsilon)
+		{
+			UE_LEARNING_ARRAY_VALUE_CHECK(Scale > 0.0f);
+
+			return Weight * DistanceToReward(FMath::Max(FMath::Abs(Position0 - Position1) - Threshold, 0.0f) / FMath::Max(Scale, Epsilon));
+		}
+
+		float PositionSimilarityReward(
+			const FVector Position0,
+			const FVector Position1,
+			const float Scale,
+			const float Threshold,
+			const float Weight,
+			const float Epsilon)
+		{
+			UE_LEARNING_ARRAY_VALUE_CHECK(Scale > 0.0f);
+
+			return Weight * DistanceToReward(FMath::Max(FVector::Distance(Position0, Position1) - Threshold, 0.0f) / FMath::Max(Scale, Epsilon));
+		}
+
+		float ScalarRotationSimilarityReward(
+			const float Rotation0,
+			const float Rotation1,
+			const float Scale,
+			const float Threshold,
+			const float Weight,
+			const float Epsilon)
+		{
+			UE_LEARNING_ARRAY_VALUE_CHECK(Scale > 0.0f);
+
+			return Weight * DistanceToReward(FMath::Max(
+				FMath::Abs(FMath::FindDeltaAngleRadians(Rotation1, Rotation0)) - Threshold, 0.0f) / FMath::Max(Scale, Epsilon));
+		}
+
+		float ScalarAngularVelocitySimilarityReward(
+			const float AngularVelocity0,
+			const float AngularVelocity1,
+			const float Scale,
+			const float Threshold,
+			const float Weight,
+			const float Epsilon)
+		{
+			UE_LEARNING_ARRAY_VALUE_CHECK(Scale > 0.0f);
+
+			return Weight * DistanceToReward(FMath::Max(FMath::Abs(AngularVelocity0 - AngularVelocity1) - Threshold, 0.0f) / FMath::Max(Scale, Epsilon));
+		}
+
+		LEARNING_API float ScalarVelocityReward(
+			const float Velocity,
+			const float Scale,
+			const float Weight,
+			const float Epsilon)
+		{
+			return Weight * (Velocity / FMath::Max(Scale, Epsilon));
+		}
+
+		LEARNING_API float LocalDirectionalVelocityReward(
+			const FVector Velocity,
+			const FQuat RelativeRotation,
+			const float Scale,
+			const float Weight,
+			const FVector Axis,
+			const float Epsilon)
+		{
+			UE_LEARNING_ARRAY_VALUE_CHECK(Scale > 0.0f);
+
+			return Weight * (RelativeRotation.UnrotateVector(Velocity).Dot(Axis) / FMath::Max(Scale, Epsilon));
+		}
+
+		// Penalties
+
+		float ScalarDifferencePenalty(
+			const float Scalar0,
+			const float Scalar1,
+			const float Scale,
+			const float Threshold,
+			const float Weight,
+			const float Epsilon)
+		{
+			UE_LEARNING_ARRAY_VALUE_CHECK(Scale > 0.0f);
+
+			return Weight * DistanceToPenalty(FMath::Max(FMath::Abs(Scalar0 - Scalar1) - Threshold, 0.0f) / FMath::Max(Scale, Epsilon));
+		}
+
+		float ScalarMagnitudePenalty(
+			const float Scalar,
+			const float Scale,
+			const float Threshold,
+			const float Weight,
+			const float Epsilon)
+		{
+			UE_LEARNING_ARRAY_VALUE_CHECK(Scale > 0.0f);
+
+			return Weight * DistanceToPenalty(FMath::Max(FMath::Abs(Scalar) - Threshold, 0.0f) / FMath::Max(Scale, Epsilon));
+		}
+
+		float DirectionAngleDifferencePenalty(
+			const FVector Direction0,
+			const FVector Direction1,
+			const float Scale,
+			const float Threshold,
+			const float Weight,
+			const float Epsilon)
+		{
+			UE_LEARNING_ARRAY_VALUE_CHECK(Scale > 0.0f);
+
+			return Weight * DistanceToPenalty(FMath::Max(FMath::Acos(Direction0.Dot(Direction1)) - Threshold, 0.0f) / FMath::Max(Scale, Epsilon));
+		}
+
+		float PlanarPositionDifferencePenalty(
+			const FVector Position0,
+			const FVector Position1,
+			const float Scale,
+			const float Threshold,
+			const float Weight,
+			const FVector Axis0,
+			const FVector Axis1,
+			const float Epsilon)
+		{
+			UE_LEARNING_ARRAY_VALUE_CHECK(Scale > 0.0f);
+
+			const FVector ProjectedPosition0 = FVector(Axis0.Dot(Position0), Axis1.Dot(Position0), 0.0f);
+			const FVector ProjectedPosition1 = FVector(Axis0.Dot(Position1), Axis1.Dot(Position1), 0.0f);
+
+			return Weight * DistanceToPenalty(FMath::Max(FVector::Distance(ProjectedPosition0, ProjectedPosition1) - Threshold, 0.0f) / FMath::Max(Scale, Epsilon));
+		}
+
+		float PositionDifferencePenalty(
+			const FVector Position0,
+			const FVector Position1,
+			const float Scale,
+			const float Threshold,
+			const float Weight,
+			const float Epsilon)
+		{
+			UE_LEARNING_ARRAY_VALUE_CHECK(Scale > 0.0f);
+
+			return Weight * DistanceToPenalty(FMath::Max(FVector::Dist(Position0, Position1) - Threshold, 0.0f) / FMath::Max(Scale, Epsilon));
+		}
+
+		float VelocityDifferencePenalty(
+			const FVector Velocity0,
+			const FVector Velocity1,
+			const float Scale,
+			const float Threshold,
+			const float Weight,
+			const float Epsilon)
+		{
+			UE_LEARNING_ARRAY_VALUE_CHECK(Scale > 0.0f);
+
+			return Weight * DistanceToPenalty(FMath::Max(FVector::Dist(Velocity0, Velocity1) - Threshold, 0.0f) / FMath::Max(Scale, Epsilon));
+		}
+
+		float SpringInertializationPenalty(
+			const FVector Position0,
+			const FVector Position1,
+			const FVector Velocity0,
+			const FVector Velocity1,
+			const float HalfLife,
+			const float Scale,
+			const float Threshold,
+			const float Weight,
+			const float Epsilon)
+		{
+			UE_LEARNING_ARRAY_VALUE_CHECK(Scale > 0.0f);
+			UE_LEARNING_ARRAY_VALUE_CHECK(HalfLife > 0.0f);
+
+			return Weight * DistanceToPenalty(FMath::Max(SpringInertializeDecayTotalDisplacement(
+				Position0 - Position1,
+				Velocity0 - Velocity1,
+				HalfLife,
+				Epsilon).Length() / FMath::Max(HalfLife, Epsilon) - Threshold, 0.0f) / FMath::Max(Scale, Epsilon));
+		}
+	}
+
+	//------------------------------------------------------------------
+
+	FRewardObject::FRewardObject(
+		const FName& InIdentifier,
+		const TSharedRef<FArrayMap>& InInstanceData,
+		const int32 InMaxInstanceNum)
+		: InstanceData(InInstanceData)
+	{
+		RewardHandle = InstanceData->Add<1, float>({ InIdentifier, TEXT("Reward") }, { InMaxInstanceNum }, 0.0f);
+	}
+
+	TLearningArrayView<1, float> FRewardObject::RewardBuffer()
+	{
+		return  InstanceData->View(RewardHandle);
+	}
+
+	//------------------------------------------------------------------
+
+	FSumReward::FSumReward(
+		const FName& InIdentifier,
+		const TLearningArrayView<1, const TSharedRef<FRewardObject>> InRewards,
+		const TSharedRef<FArrayMap>& InInstanceData,
+		const int32 InMaxInstanceNum)
+		: FRewardObject(InIdentifier, InInstanceData, InMaxInstanceNum)
+		, Rewards(InRewards) {}
+
+	void FSumReward::Evaluate(const FIndexSet Instances)
+	{
+		UE_LEARNING_TRACE_CPUPROFILER_EVENT_SCOPE(FSumReward::Evaluate);
+
+		TLearningArrayView<1, float> Reward = InstanceData->View(RewardHandle);
+
+		const int32 RewardNum = Rewards.Num();
+
+		for (int32 RewardIdx = 0; RewardIdx < RewardNum; RewardIdx++)
+		{
+			Rewards[RewardIdx]->Evaluate(Instances);
+		}
+
+		Array::Zero(Reward, Instances);
+
+		for (int32 RewardIdx = 0; RewardIdx < RewardNum; RewardIdx++)
+		{
+			const TLearningArrayView<1, const float> Input = InstanceData->ConstView(Rewards[RewardIdx]->RewardHandle);
+
+			for (const int32 InstanceIdx : Instances)
+			{
+				Reward[InstanceIdx] += Input[InstanceIdx];
+			}
+		}
+	}
+
+	//------------------------------------------------------------------
+
+	FConstantReward::FConstantReward(
+		const FName& InIdentifier,
+		const TSharedRef<FArrayMap>& InInstanceData,
+		const int32 InMaxInstanceNum,
+		const float InValue)
+		: FRewardObject(InIdentifier, InInstanceData, InMaxInstanceNum)
+		, Value(InValue) {}
+
+	void FConstantReward::Evaluate(const FIndexSet Instances)
+	{
+		UE_LEARNING_TRACE_CPUPROFILER_EVENT_SCOPE(FConstantReward::Evaluate);
+
+		Array::Set(InstanceData->View(RewardHandle), Value, Instances);
+	}
+
+	FConstantPenalty::FConstantPenalty(
+		const FName& InIdentifier,
+		const TSharedRef<FArrayMap>& InInstanceData,
+		const int32 InMaxInstanceNum,
+		const float InValue)
+		: FRewardObject(InIdentifier, InInstanceData, InMaxInstanceNum)
+		, Value(InValue) {}
+
+	void FConstantPenalty::Evaluate(const FIndexSet Instances)
+	{
+		UE_LEARNING_TRACE_CPUPROFILER_EVENT_SCOPE(FConstantPenalty::Evaluate);
+
+		Array::Set(InstanceData->View(RewardHandle), -Value, Instances);
+	}
+
+	FConditionalConstantReward::FConditionalConstantReward(
+		const FName& InIdentifier,
+		const TSharedRef<FArrayMap>& InInstanceData,
+		const int32 InMaxInstanceNum,
+		const float InValue)
+		: FRewardObject(InIdentifier, InInstanceData, InMaxInstanceNum)
+		, Value(InValue)
+	{
+		ConditionHandle = InstanceData->Add<1, bool>({ InIdentifier, TEXT("Condition") }, { InMaxInstanceNum }, false);
+	}
+
+	void FConditionalConstantReward::Evaluate(const FIndexSet Instances)
+	{
+		UE_LEARNING_TRACE_CPUPROFILER_EVENT_SCOPE(FConditionalConstantReward::Evaluate);
+
+		const TLearningArrayView<1, const bool> Condition = InstanceData->ConstView(ConditionHandle);
+		TLearningArrayView<1, float> Reward = InstanceData->View(RewardHandle);
+
+		for (const int32 InstanceIdx : Instances)
+		{
+			Reward[InstanceIdx] = Condition[InstanceIdx] ? Value : 0.0f;
+		}
+	}
+
+	FFloatReward::FFloatReward(
+		const FName& InIdentifier,
+		const TSharedRef<FArrayMap>& InInstanceData,
+		const int32 InMaxInstanceNum,
+		const float InWeight)
+		: FRewardObject(InIdentifier, InInstanceData, InMaxInstanceNum)
+	{
+		ValueHandle = InstanceData->Add<1, float>({ InIdentifier, TEXT("Value") }, { InMaxInstanceNum }, 0.0f);
+		WeightHandle = InstanceData->Add<1, float>({ InIdentifier, TEXT("Weight") }, { InMaxInstanceNum }, InWeight);
+	}
+
+	void FFloatReward::Evaluate(const FIndexSet Instances)
+	{
+		UE_LEARNING_TRACE_CPUPROFILER_EVENT_SCOPE(FFloatReward::Evaluate);
+
+		const TLearningArrayView<1, const float> Value = InstanceData->ConstView(ValueHandle);
+		const TLearningArrayView<1, const float> Weight = InstanceData->ConstView(WeightHandle);
+		TLearningArrayView<1, float> Reward = InstanceData->View(RewardHandle);
+
+		for (const int32 InstanceIdx : Instances)
+		{
+			Reward[InstanceIdx] = Weight[InstanceIdx] * Value[InstanceIdx];
+		}
+	}
+
+	//------------------------------------------------------------------
+
+	FVectorAverageMagnitudePenalty::FVectorAverageMagnitudePenalty(
+		const FName& InIdentifier,
+		const TSharedRef<FArrayMap>& InInstanceData,
+		const int32 InMaxInstanceNum,
+		const int32 InDimNum,
+		const float InWeight,
+		const float InScale,
+		const float InThreshold)
+		: FRewardObject(InIdentifier, InInstanceData, InMaxInstanceNum)
+		, DimNum(InDimNum)
+	{
+		ValueHandle = InstanceData->Add<2, float>({ InIdentifier, TEXT("Value") }, { InMaxInstanceNum, DimNum }, 1.0f);
+
+		WeightHandle = InstanceData->Add<1, float>({ InIdentifier, TEXT("Weight") }, { InMaxInstanceNum }, InWeight);
+		ScaleHandle = InstanceData->Add<1, float>({ InIdentifier, TEXT("Scale") }, { InMaxInstanceNum }, InScale);
+		ThresholdHandle = InstanceData->Add<1, float>({ InIdentifier, TEXT("Threshold") }, { InMaxInstanceNum }, InThreshold);
+	}
+
+	void FVectorAverageMagnitudePenalty::Evaluate(const FIndexSet Instances)
+	{
+		UE_LEARNING_TRACE_CPUPROFILER_EVENT_SCOPE(FVectorAverageMagnitudePenalty::Evaluate);
+
+		const TLearningArrayView<2, const float> Value = InstanceData->ConstView(ValueHandle);
+
+		const TLearningArrayView<1, const float> Scale = InstanceData->ConstView(ScaleHandle);
+		const TLearningArrayView<1, const float> Threshold = InstanceData->ConstView(ThresholdHandle);
+		const TLearningArrayView<1, const float> Weight = InstanceData->ConstView(WeightHandle);
+
+		TLearningArrayView<1, float> Reward = InstanceData->View(RewardHandle);
+
+		for (const int32 InstanceIdx : Instances)
+		{
+			Reward[InstanceIdx] = 0.0f;
+
+			for (int32 DimIdx = 0; DimIdx < DimNum; DimIdx++)
+			{
+				Reward[InstanceIdx] += Reward::ScalarMagnitudePenalty(
+					Value[InstanceIdx][DimIdx],
+					Scale[InstanceIdx],
+					Threshold[InstanceIdx],
+					Weight[InstanceIdx]) / DimNum;
+			}
+		}
+	}
+
+	//------------------------------------------------------------------
+
+	FScalarPositionSimilarityReward::FScalarPositionSimilarityReward(
+		const FName& InIdentifier,
+		const TSharedRef<FArrayMap>& InInstanceData,
+		const int32 InMaxInstanceNum,
+		const float InWeight,
+		const float InScale,
+		const float InThreshold)
+		: FRewardObject(InIdentifier, InInstanceData, InMaxInstanceNum)
+	{
+		Position0Handle = InstanceData->Add<1, float>({ InIdentifier, TEXT("Position0") }, { InMaxInstanceNum }, 0.0f);
+		Position1Handle = InstanceData->Add<1, float>({ InIdentifier, TEXT("Position1") }, { InMaxInstanceNum }, 0.0f);
+
+		WeightHandle = InstanceData->Add<1, float>({ InIdentifier, TEXT("Weight") }, { InMaxInstanceNum }, InWeight);
+		ScaleHandle = InstanceData->Add<1, float>({ InIdentifier, TEXT("Scale") }, { InMaxInstanceNum }, InScale);
+		ThresholdHandle = InstanceData->Add<1, float>({ InIdentifier, TEXT("Threshold") }, { InMaxInstanceNum }, InThreshold);
+	}
+
+	void FScalarPositionSimilarityReward::Evaluate(const FIndexSet Instances)
+	{
+		UE_LEARNING_TRACE_CPUPROFILER_EVENT_SCOPE(FScalarPositionSimilarityReward::Evaluate);
+
+		const TLearningArrayView<1, const float> Position0 = InstanceData->ConstView(Position0Handle);
+		const TLearningArrayView<1, const float> Position1 = InstanceData->ConstView(Position1Handle);
+
+		const TLearningArrayView<1, const float> Scale = InstanceData->ConstView(ScaleHandle);
+		const TLearningArrayView<1, const float> Threshold = InstanceData->ConstView(ThresholdHandle);
+		const TLearningArrayView<1, const float> Weight = InstanceData->ConstView(WeightHandle);
+
+		TLearningArrayView<1, float> Reward = InstanceData->View(RewardHandle);
+
+		for (const int32 InstanceIdx : Instances)
+		{
+			Reward[InstanceIdx] = Reward::ScalarPositionSimilarityReward(
+				Position0[InstanceIdx],
+				Position1[InstanceIdx],
+				Scale[InstanceIdx],
+				Threshold[InstanceIdx],
+				Weight[InstanceIdx]);
+		}
+	}
+
+	//------------------------------------------------------------------
+
+	FScalarRotationSimilarityReward::FScalarRotationSimilarityReward(
+		const FName& InIdentifier,
+		const TSharedRef<FArrayMap>& InInstanceData,
+		const int32 InMaxInstanceNum,
+		const float InWeight,
+		const float InScale,
+		const float InThreshold)
+		: FRewardObject(InIdentifier, InInstanceData, InMaxInstanceNum)
+	{
+		Rotation0Handle = InstanceData->Add<1, float>({ InIdentifier, TEXT("Rotation0") }, { InMaxInstanceNum }, 0.0f);
+		Rotation1Handle = InstanceData->Add<1, float>({ InIdentifier, TEXT("Rotation1") }, { InMaxInstanceNum }, 0.0f);
+
+		WeightHandle = InstanceData->Add<1, float>({ InIdentifier, TEXT("Weight") }, { InMaxInstanceNum }, InWeight);
+		ScaleHandle = InstanceData->Add<1, float>({ InIdentifier, TEXT("Scale") }, { InMaxInstanceNum }, InScale);
+		ThresholdHandle = InstanceData->Add<1, float>({ InIdentifier, TEXT("Threshold") }, { InMaxInstanceNum }, InThreshold);
+	}
+
+	void FScalarRotationSimilarityReward::Evaluate(const FIndexSet Instances)
+	{
+		UE_LEARNING_TRACE_CPUPROFILER_EVENT_SCOPE(FScalarRotationSimilarityReward::Evaluate);
+
+		const TLearningArrayView<1, const float> Rotation0 = InstanceData->ConstView(Rotation0Handle);
+		const TLearningArrayView<1, const float> Rotation1 = InstanceData->ConstView(Rotation1Handle);
+
+		const TLearningArrayView<1, const float> Scale = InstanceData->ConstView(ScaleHandle);
+		const TLearningArrayView<1, const float> Threshold = InstanceData->ConstView(ThresholdHandle);
+		const TLearningArrayView<1, const float> Weight = InstanceData->ConstView(WeightHandle);
+
+		TLearningArrayView<1, float> Reward = InstanceData->View(RewardHandle);
+
+		for (const int32 InstanceIdx : Instances)
+		{
+			Reward[InstanceIdx] = Reward::ScalarRotationSimilarityReward(
+				Rotation0[InstanceIdx],
+				Rotation1[InstanceIdx],
+				Scale[InstanceIdx],
+				Threshold[InstanceIdx],
+				Weight[InstanceIdx]);
+		}
+	}
+
+	//------------------------------------------------------------------
+
+	FScalarAngularVelocitySimilarityReward::FScalarAngularVelocitySimilarityReward(
+		const FName& InIdentifier,
+		const TSharedRef<FArrayMap>& InInstanceData,
+		const int32 InMaxInstanceNum,
+		const float InWeight,
+		const float InScale,
+		const float InThreshold)
+		: FRewardObject(InIdentifier, InInstanceData, InMaxInstanceNum)
+	{
+		AngularVelocity0Handle = InstanceData->Add<1, float>({ InIdentifier, TEXT("AngularVelocity0") }, { InMaxInstanceNum }, 0.0f);
+		AngularVelocity1Handle = InstanceData->Add<1, float>({ InIdentifier, TEXT("AngularVelocity1") }, { InMaxInstanceNum }, 0.0f);
+
+		WeightHandle = InstanceData->Add<1, float>({ InIdentifier, TEXT("Weight") }, { InMaxInstanceNum }, InWeight);
+		ScaleHandle = InstanceData->Add<1, float>({ InIdentifier, TEXT("Scale") }, { InMaxInstanceNum }, InScale);
+		ThresholdHandle = InstanceData->Add<1, float>({ InIdentifier, TEXT("Threshold") }, { InMaxInstanceNum }, InThreshold);
+	}
+
+	void FScalarAngularVelocitySimilarityReward::Evaluate(const FIndexSet Instances)
+	{
+		UE_LEARNING_TRACE_CPUPROFILER_EVENT_SCOPE(FScalarAngularVelocitySimilarityReward::Evaluate);
+
+		const TLearningArrayView<1, const float> AngularVelocity0 = InstanceData->ConstView(AngularVelocity0Handle);
+		const TLearningArrayView<1, const float> AngularVelocity1 = InstanceData->ConstView(AngularVelocity1Handle);
+
+		const TLearningArrayView<1, const float> Scale = InstanceData->ConstView(ScaleHandle);
+		const TLearningArrayView<1, const float> Threshold = InstanceData->ConstView(ThresholdHandle);
+		const TLearningArrayView<1, const float> Weight = InstanceData->ConstView(WeightHandle);
+
+		TLearningArrayView<1, float> Reward = InstanceData->View(RewardHandle);
+
+		for (const int32 InstanceIdx : Instances)
+		{
+			Reward[InstanceIdx] = Reward::ScalarAngularVelocitySimilarityReward(
+				AngularVelocity0[InstanceIdx],
+				AngularVelocity1[InstanceIdx],
+				Scale[InstanceIdx],
+				Threshold[InstanceIdx],
+				Weight[InstanceIdx]);
+		}
+	}
+
+	//------------------------------------------------------------------
+
+	FScalarVelocityReward::FScalarVelocityReward(
+		const FName& InIdentifier,
+		const TSharedRef<FArrayMap>& InInstanceData,
+		const int32 InMaxInstanceNum,
+		const float InWeight,
+		const float InScale)
+		: FRewardObject(InIdentifier, InInstanceData, InMaxInstanceNum)
+	{
+		VelocityHandle = InstanceData->Add<1, float>({ InIdentifier, TEXT("Velocity") }, { InMaxInstanceNum }, 0.0f);
+
+		WeightHandle = InstanceData->Add<1, float>({ InIdentifier, TEXT("Weight") }, { InMaxInstanceNum }, InWeight);
+		ScaleHandle = InstanceData->Add<1, float>({ InIdentifier, TEXT("Scale") }, { InMaxInstanceNum }, InScale);
+	}
+
+	void FScalarVelocityReward::Evaluate(const FIndexSet Instances)
+	{
+		UE_LEARNING_TRACE_CPUPROFILER_EVENT_SCOPE(FPlanarVelocityMagnitudeReward::Evaluate);
+
+		const TLearningArrayView<1, const float> Velocity = InstanceData->ConstView(VelocityHandle);
+
+		const TLearningArrayView<1, const float> Scale = InstanceData->ConstView(ScaleHandle);
+		const TLearningArrayView<1, const float> Weight = InstanceData->ConstView(WeightHandle);
+
+		TLearningArrayView<1, float> Reward = InstanceData->View(RewardHandle);
+
+		for (const int32 InstanceIdx : Instances)
+		{
+			Reward[InstanceIdx] = Reward::ScalarVelocityReward(
+				Velocity[InstanceIdx],
+				Scale[InstanceIdx],
+				Weight[InstanceIdx]);
+		}
+	}
+
+	FLocalDirectionalVelocityReward::FLocalDirectionalVelocityReward(
+		const FName& InIdentifier,
+		const TSharedRef<FArrayMap>& InInstanceData,
+		const int32 InMaxInstanceNum,
+		const float InWeight,
+		const float InScale,
+		const FVector InAxis)
+		: FRewardObject(InIdentifier, InInstanceData, InMaxInstanceNum)
+		, Axis(InAxis)
+	{
+		VelocityHandle = InstanceData->Add<1, FVector>({ InIdentifier, TEXT("Velocity") }, { InMaxInstanceNum }, FVector::ZeroVector);
+		RelativeRotationHandle = InstanceData->Add<1, FQuat>({ InIdentifier, TEXT("RelativeRotation") }, { InMaxInstanceNum }, FQuat::Identity);
+
+		WeightHandle = InstanceData->Add<1, float>({ InIdentifier, TEXT("Weight") }, { InMaxInstanceNum }, InWeight);
+		ScaleHandle = InstanceData->Add<1, float>({ InIdentifier, TEXT("Scale") }, { InMaxInstanceNum }, InScale);
+	}
+
+	void FLocalDirectionalVelocityReward::Evaluate(const FIndexSet Instances)
+	{
+		UE_LEARNING_TRACE_CPUPROFILER_EVENT_SCOPE(FLocalDirectionalVelocityReward::Evaluate);
+
+		const TLearningArrayView<1, const FVector> Velocity = InstanceData->ConstView(VelocityHandle);
+		const TLearningArrayView<1, const FQuat> RelativeRotation = InstanceData->ConstView(RelativeRotationHandle);
+
+		const TLearningArrayView<1, const float> Scale = InstanceData->ConstView(ScaleHandle);
+		const TLearningArrayView<1, const float> Weight = InstanceData->ConstView(WeightHandle);
+
+		TLearningArrayView<1, float> Reward = InstanceData->View(RewardHandle);
+
+		for (const int32 InstanceIdx : Instances)
+		{
+			Reward[InstanceIdx] = Reward::LocalDirectionalVelocityReward(
+				Velocity[InstanceIdx],
+				RelativeRotation[InstanceIdx],
+				Scale[InstanceIdx],
+				Weight[InstanceIdx],
+				Axis);
+		}
+	}
+
+	//------------------------------------------------------------------
+
+	FPlanarPositionDifferencePenalty::FPlanarPositionDifferencePenalty(
+		const FName& InIdentifier,
+		const TSharedRef<FArrayMap>& InInstanceData,
+		const int32 InMaxInstanceNum,
+		const float InWeight,
+		const float InScale,
+		const float InThreshold,
+		const FVector InAxis0,
+		const FVector InAxis1)
+		: FRewardObject(InIdentifier, InInstanceData, InMaxInstanceNum)
+		, Axis0(InAxis0)
+		, Axis1(InAxis1)
+	{
+		Position0Handle = InstanceData->Add<1, FVector>({ InIdentifier, TEXT("Position0") }, { InMaxInstanceNum }, FVector::ZeroVector);
+		Position1Handle = InstanceData->Add<1, FVector>({ InIdentifier, TEXT("Position1") }, { InMaxInstanceNum }, FVector::ZeroVector);
+
+		WeightHandle = InstanceData->Add<1, float>({ InIdentifier, TEXT("Weight") }, { InMaxInstanceNum }, InWeight);
+		ScaleHandle = InstanceData->Add<1, float>({ InIdentifier, TEXT("Scale") }, { InMaxInstanceNum }, InScale);
+		ThresholdHandle = InstanceData->Add<1, float>({ InIdentifier, TEXT("Threshold") }, { InMaxInstanceNum }, InThreshold);
+	}
+
+	void FPlanarPositionDifferencePenalty::Evaluate(const FIndexSet Instances)
+	{
+		UE_LEARNING_TRACE_CPUPROFILER_EVENT_SCOPE(FPlanarPositionDifferencePenalty::Evaluate);
+
+		const TLearningArrayView<1, const FVector> Position0 = InstanceData->ConstView(Position0Handle);
+		const TLearningArrayView<1, const FVector> Position1 = InstanceData->ConstView(Position1Handle);
+
+		const TLearningArrayView<1, const float> Scale = InstanceData->ConstView(ScaleHandle);
+		const TLearningArrayView<1, const float> Threshold = InstanceData->ConstView(ThresholdHandle);
+		const TLearningArrayView<1, const float> Weight = InstanceData->ConstView(WeightHandle);
+
+		TLearningArrayView<1, float> Reward = InstanceData->View(RewardHandle);
+
+		for (const int32 InstanceIdx : Instances)
+		{
+			Reward[InstanceIdx] = Reward::PlanarPositionDifferencePenalty(
+				Position0[InstanceIdx],
+				Position1[InstanceIdx],
+				Scale[InstanceIdx],
+				Threshold[InstanceIdx],
+				Weight[InstanceIdx],
+				Axis0,
+				Axis1);
+		}
+	}
+
+	FPositionDifferencePenalty::FPositionDifferencePenalty(
+		const FName& InIdentifier,
+		const TSharedRef<FArrayMap>& InInstanceData,
+		const int32 InMaxInstanceNum,
+		const float InWeight,
+		const float InScale,
+		const float InThreshold)
+		: FRewardObject(InIdentifier, InInstanceData, InMaxInstanceNum)
+	{
+		Position0Handle = InstanceData->Add<1, FVector>({ InIdentifier, TEXT("Position0") }, { InMaxInstanceNum }, FVector::ZeroVector);
+		Position1Handle = InstanceData->Add<1, FVector>({ InIdentifier, TEXT("Position1") }, { InMaxInstanceNum }, FVector::ZeroVector);
+
+		WeightHandle = InstanceData->Add<1, float>({ InIdentifier, TEXT("Weight") }, { InMaxInstanceNum }, InWeight);
+		ScaleHandle = InstanceData->Add<1, float>({ InIdentifier, TEXT("Scale") }, { InMaxInstanceNum }, InScale);
+		ThresholdHandle = InstanceData->Add<1, float>({ InIdentifier, TEXT("Threshold") }, { InMaxInstanceNum }, InThreshold);
+	}
+
+	void FPositionDifferencePenalty::Evaluate(const FIndexSet Instances)
+	{
+		UE_LEARNING_TRACE_CPUPROFILER_EVENT_SCOPE(FPositionDifferencePenalty::Evaluate);
+
+		const TLearningArrayView<1, const FVector> Position0 = InstanceData->ConstView(Position0Handle);
+		const TLearningArrayView<1, const FVector> Position1 = InstanceData->ConstView(Position1Handle);
+
+		const TLearningArrayView<1, const float> Scale = InstanceData->ConstView(ScaleHandle);
+		const TLearningArrayView<1, const float> Threshold = InstanceData->ConstView(ThresholdHandle);
+		const TLearningArrayView<1, const float> Weight = InstanceData->ConstView(WeightHandle);
+
+		TLearningArrayView<1, float> Reward = InstanceData->View(RewardHandle);
+
+		for (const int32 InstanceIdx : Instances)
+		{
+			Reward[InstanceIdx] = Reward::PositionDifferencePenalty(
+				Position0[InstanceIdx],
+				Position1[InstanceIdx],
+				Scale[InstanceIdx],
+				Threshold[InstanceIdx],
+				Weight[InstanceIdx]);
+		}
+	}
+
+
+	//------------------------------------------------------------------
+
+	FVelocityDifferencePenalty::FVelocityDifferencePenalty(
+		const FName& InIdentifier,
+		const TSharedRef<FArrayMap>& InInstanceData,
+		const int32 InMaxInstanceNum,
+		const float InWeight,
+		const float InScale,
+		const float InThreshold)
+		: FRewardObject(InIdentifier, InInstanceData, InMaxInstanceNum)
+	{
+		Velocity0Handle = InstanceData->Add<1, FVector>({ InIdentifier, TEXT("Velocity0") }, { InMaxInstanceNum }, FVector::ZeroVector);
+		Velocity1Handle = InstanceData->Add<1, FVector>({ InIdentifier, TEXT("Velocity1") }, { InMaxInstanceNum }, FVector::ZeroVector);
+
+		WeightHandle = InstanceData->Add<1, float>({ InIdentifier, TEXT("Weight") }, { InMaxInstanceNum }, InWeight);
+		ScaleHandle = InstanceData->Add<1, float>({ InIdentifier, TEXT("Scale") }, { InMaxInstanceNum }, InScale);
+		ThresholdHandle = InstanceData->Add<1, float>({ InIdentifier, TEXT("Threshold") }, { InMaxInstanceNum }, InThreshold);
+	}
+
+	void FVelocityDifferencePenalty::Evaluate(const FIndexSet Instances)
+	{
+		UE_LEARNING_TRACE_CPUPROFILER_EVENT_SCOPE(FVelocityDifferencePenalty::Evaluate);
+
+		const TLearningArrayView<1, const FVector> Velocity0 = InstanceData->ConstView(Velocity0Handle);
+		const TLearningArrayView<1, const FVector> Velocity1 = InstanceData->ConstView(Velocity1Handle);
+
+		const TLearningArrayView<1, const float> Scale = InstanceData->ConstView(ScaleHandle);
+		const TLearningArrayView<1, const float> Threshold = InstanceData->ConstView(ThresholdHandle);
+		const TLearningArrayView<1, const float> Weight = InstanceData->ConstView(WeightHandle);
+
+		TLearningArrayView<1, float> Reward = InstanceData->View(RewardHandle);
+
+		for (const int32 InstanceIdx : Instances)
+		{
+			Reward[InstanceIdx] = Reward::VelocityDifferencePenalty(
+				Velocity0[InstanceIdx],
+				Velocity1[InstanceIdx],
+				Scale[InstanceIdx],
+				Threshold[InstanceIdx],
+				Weight[InstanceIdx]);
+		}
+	}
+
+	//------------------------------------------------------------------
+
+	FFacingDirectionAngularDifferencePenalty::FFacingDirectionAngularDifferencePenalty(
+		const FName& InIdentifier,
+		const TSharedRef<FArrayMap>& InInstanceData,
+		const int32 InMaxInstanceNum,
+		const FVector InFacingDirection0,
+		const FVector InFacingDirection1,
+		const float InWeight,
+		const float InScale,
+		const float InThreshold)
+		: FRewardObject(InIdentifier, InInstanceData, InMaxInstanceNum)
+	{
+		Rotation0Handle = InstanceData->Add<1, FQuat>({ InIdentifier, TEXT("Rotation0") }, { InMaxInstanceNum }, FQuat::Identity);
+		Rotation1Handle = InstanceData->Add<1, FQuat>({ InIdentifier, TEXT("Rotation1") }, { InMaxInstanceNum }, FQuat::Identity);
+		FacingDirection0Handle = InstanceData->Add<1, FVector>({ InIdentifier, TEXT("FacingDirection0") }, { InMaxInstanceNum }, InFacingDirection0);
+		FacingDirection1Handle = InstanceData->Add<1, FVector>({ InIdentifier, TEXT("FacingDirection1") }, { InMaxInstanceNum }, InFacingDirection1);
+
+		WeightHandle = InstanceData->Add<1, float>({ InIdentifier, TEXT("Weight") }, { InMaxInstanceNum }, InWeight);
+		ScaleHandle = InstanceData->Add<1, float>({ InIdentifier, TEXT("Scale") }, { InMaxInstanceNum }, InScale);
+		ThresholdHandle = InstanceData->Add<1, float>({ InIdentifier, TEXT("Threshold") }, { InMaxInstanceNum }, InThreshold);
+	}
+
+	void FFacingDirectionAngularDifferencePenalty::Evaluate(const FIndexSet Instances)
+	{
+		UE_LEARNING_TRACE_CPUPROFILER_EVENT_SCOPE(FFacingDirectionAngularDifferencePenalty::Evaluate);
+
+		const TLearningArrayView<1, const FQuat> Rotation0 = InstanceData->ConstView(Rotation0Handle);
+		const TLearningArrayView<1, const FQuat> Rotation1 = InstanceData->ConstView(Rotation1Handle);
+		const TLearningArrayView<1, const FVector> FacingDirection0 = InstanceData->ConstView(FacingDirection0Handle);
+		const TLearningArrayView<1, const FVector> FacingDirection1 = InstanceData->ConstView(FacingDirection1Handle);
+
+		const TLearningArrayView<1, const float> Scale = InstanceData->ConstView(ScaleHandle);
+		const TLearningArrayView<1, const float> Threshold = InstanceData->ConstView(ThresholdHandle);
+		const TLearningArrayView<1, const float> Weight = InstanceData->ConstView(WeightHandle);
+
+		TLearningArrayView<1, float> Reward = InstanceData->View(RewardHandle);
+
+		for (const int32 InstanceIdx : Instances)
+		{
+			Reward[InstanceIdx] = Reward::DirectionAngleDifferencePenalty(
+				Rotation0[InstanceIdx].RotateVector(FacingDirection0[InstanceIdx]),
+				Rotation1[InstanceIdx].RotateVector(FacingDirection1[InstanceIdx]),
+				Scale[InstanceIdx],
+				Threshold[InstanceIdx],
+				Weight[InstanceIdx]);
+		}
+	}
+
+	FFacingTowardsTargetAngularDifferencePenalty::FFacingTowardsTargetAngularDifferencePenalty(
+		const FName& InIdentifier,
+		const TSharedRef<FArrayMap>& InInstanceData,
+		const int32 InMaxInstanceNum,
+		const FVector InFacingDirection,
+		const float InMinRadius,
+		const float InWeight,
+		const float InScale,
+		const float InThreshold)
+		: FRewardObject(InIdentifier, InInstanceData, InMaxInstanceNum)
+	{
+		TargetHandle = InstanceData->Add<1, FVector>({ InIdentifier, TEXT("Target") }, { InMaxInstanceNum }, FVector::ZeroVector);
+		PositionHandle = InstanceData->Add<1, FVector>({ InIdentifier, TEXT("Position") }, { InMaxInstanceNum }, FVector::ZeroVector);
+		RotationHandle = InstanceData->Add<1, FQuat>({ InIdentifier, TEXT("Rotation") }, { InMaxInstanceNum }, FQuat::Identity);
+		FacingDirectionHandle = InstanceData->Add<1, FVector>({ InIdentifier, TEXT("FacingDirection") }, { InMaxInstanceNum }, InFacingDirection);
+		MinRadiusHandle = InstanceData->Add<1, float>({ InIdentifier, TEXT("MinRadius") }, { InMaxInstanceNum }, InMinRadius);
+
+		WeightHandle = InstanceData->Add<1, float>({ InIdentifier, TEXT("Weight") }, { InMaxInstanceNum }, InWeight);
+		ScaleHandle = InstanceData->Add<1, float>({ InIdentifier, TEXT("Scale") }, { InMaxInstanceNum }, InScale);
+		ThresholdHandle = InstanceData->Add<1, float>({ InIdentifier, TEXT("Threshold") }, { InMaxInstanceNum }, InThreshold);
+	}
+
+	void FFacingTowardsTargetAngularDifferencePenalty::Evaluate(const FIndexSet Instances)
+	{
+		UE_LEARNING_TRACE_CPUPROFILER_EVENT_SCOPE(FFacingTowardsTargetAngularDifferencePenalty::Evaluate);
+
+		const TLearningArrayView<1, const FVector> Target = InstanceData->ConstView(TargetHandle);
+		const TLearningArrayView<1, const FVector> Position = InstanceData->ConstView(PositionHandle);
+		const TLearningArrayView<1, const FQuat> Rotation = InstanceData->ConstView(RotationHandle);
+		const TLearningArrayView<1, const FVector> FacingDirection = InstanceData->ConstView(FacingDirectionHandle);
+		const TLearningArrayView<1, const float> MinRadius = InstanceData->ConstView(MinRadiusHandle);
+
+		const TLearningArrayView<1, const float> Scale = InstanceData->ConstView(ScaleHandle);
+		const TLearningArrayView<1, const float> Threshold = InstanceData->ConstView(ThresholdHandle);
+		const TLearningArrayView<1, const float> Weight = InstanceData->ConstView(WeightHandle);
+
+		TLearningArrayView<1, float> Reward = InstanceData->View(RewardHandle);
+
+		for (const int32 InstanceIdx : Instances)
+		{
+			UE_LEARNING_ARRAY_VALUE_CHECK(MinRadius[InstanceIdx] > 0.0f);
+
+			const FVector TargetDifference = Target[InstanceIdx] - Position[InstanceIdx];
+
+			if (TargetDifference.Length() > FMath::Max(MinRadius[InstanceIdx], UE_SMALL_NUMBER))
+			{
+				Reward[InstanceIdx] = Reward::DirectionAngleDifferencePenalty(
+					Rotation[InstanceIdx].RotateVector(FacingDirection[InstanceIdx]),
+					TargetDifference.GetUnsafeNormal(),
+					Scale[InstanceIdx],
+					Threshold[InstanceIdx],
+					Weight[InstanceIdx]);
+			}
+			else
+			{
+				Reward[InstanceIdx] = 0.0f;
+			}
+		}
+	}
+
+	//------------------------------------------------------------------
+
+	FSpringInertializationPenalty::FSpringInertializationPenalty(
+		const FName& InIdentifier,
+		const TSharedRef<FArrayMap>& InInstanceData,
+		const int32 InMaxInstanceNum,
+		const float InWeight,
+		const float InScale,
+		const float InThreshold)
+		: FRewardObject(InIdentifier, InInstanceData, InMaxInstanceNum)
+	{
+		Position0Handle = InstanceData->Add<1, FVector>({ InIdentifier, TEXT("Position0") }, { InMaxInstanceNum }, FVector::ZeroVector);
+		Position1Handle = InstanceData->Add<1, FVector>({ InIdentifier, TEXT("Position1") }, { InMaxInstanceNum }, FVector::ZeroVector);
+		Velocity0Handle = InstanceData->Add<1, FVector>({ InIdentifier, TEXT("Velocity0") }, { InMaxInstanceNum }, FVector::ZeroVector);
+		Velocity1Handle = InstanceData->Add<1, FVector>({ InIdentifier, TEXT("Velocity1") }, { InMaxInstanceNum }, FVector::ZeroVector);
+
+		HalfLifeHandle = InstanceData->Add<1, float>({ InIdentifier, TEXT("HalfLife") }, { InMaxInstanceNum }, 0.1f);
+
+		WeightHandle = InstanceData->Add<1, float>({ InIdentifier, TEXT("Weight") }, { InMaxInstanceNum }, InWeight);
+		ScaleHandle = InstanceData->Add<1, float>({ InIdentifier, TEXT("Scale") }, { InMaxInstanceNum }, InScale);
+		ThresholdHandle = InstanceData->Add<1, float>({ InIdentifier, TEXT("Threshold") }, { InMaxInstanceNum }, InThreshold);
+	}
+
+	void FSpringInertializationPenalty::Evaluate(const FIndexSet Instances)
+	{
+		UE_LEARNING_TRACE_CPUPROFILER_EVENT_SCOPE(FSpringInertializationPenalty::Evaluate);
+
+		const TLearningArrayView<1, const FVector> Position0 = InstanceData->ConstView(Position0Handle);
+		const TLearningArrayView<1, const FVector> Position1 = InstanceData->ConstView(Position1Handle);
+		const TLearningArrayView<1, const FVector> Velocity0 = InstanceData->ConstView(Velocity0Handle);
+		const TLearningArrayView<1, const FVector> Velocity1 = InstanceData->ConstView(Velocity1Handle);
+
+		const TLearningArrayView<1, const float> HalfLife = InstanceData->ConstView(HalfLifeHandle);
+
+		const TLearningArrayView<1, const float> Scale = InstanceData->ConstView(ScaleHandle);
+		const TLearningArrayView<1, const float> Threshold = InstanceData->ConstView(ThresholdHandle);
+		const TLearningArrayView<1, const float> Weight = InstanceData->ConstView(WeightHandle);
+
+		TLearningArrayView<1, float> Reward = InstanceData->View(RewardHandle);
+
+		for (const int32 InstanceIdx : Instances)
+		{
+			Reward[InstanceIdx] = Reward::SpringInertializationPenalty(
+				Position0[InstanceIdx],
+				Position1[InstanceIdx],
+				Velocity0[InstanceIdx],
+				Velocity1[InstanceIdx],
+				HalfLife[InstanceIdx],
+				Scale[InstanceIdx],
+				Threshold[InstanceIdx],
+				Weight[InstanceIdx]);
+		}
+	}
+
+	//------------------------------------------------------------------
+
+	FPositionArraySimilarityReward::FPositionArraySimilarityReward(
+		const FName& InIdentifier,
+		const TSharedRef<FArrayMap>& InInstanceData,
+		const int32 InMaxInstanceNum,
+		const int32 InPositionNum,
+		const float InWeight,
+		const float InScale,
+		const float InThreshold)
+		: FRewardObject(InIdentifier, InInstanceData, InMaxInstanceNum)
+	{
+		Positions0Handle = InstanceData->Add<2, FVector>({ InIdentifier, TEXT("Positions0") }, { InMaxInstanceNum, InPositionNum }, FVector::ZeroVector);
+		Positions1Handle = InstanceData->Add<2, FVector>({ InIdentifier, TEXT("Positions1") }, { InMaxInstanceNum, InPositionNum }, FVector::ZeroVector);
+		RelativePosition0Handle = InstanceData->Add<1, FVector>({ InIdentifier, TEXT("RelativePosition0") }, { InMaxInstanceNum }, FVector::ZeroVector);
+		RelativePosition1Handle = InstanceData->Add<1, FVector>({ InIdentifier, TEXT("RelativePosition1") }, { InMaxInstanceNum }, FVector::ZeroVector);
+		RelativeRotation0Handle = InstanceData->Add<1, FQuat>({ InIdentifier, TEXT("RelativeRotation0") }, { InMaxInstanceNum }, FQuat::Identity);
+		RelativeRotation1Handle = InstanceData->Add<1, FQuat>({ InIdentifier, TEXT("RelativeRotation1") }, { InMaxInstanceNum }, FQuat::Identity);
+
+		WeightHandle = InstanceData->Add<1, float>({ InIdentifier, TEXT("Weight") }, { InMaxInstanceNum }, InWeight);
+		ScaleHandle = InstanceData->Add<1, float>({ InIdentifier, TEXT("Scale") }, { InMaxInstanceNum }, InScale);
+		ThresholdHandle = InstanceData->Add<1, float>({ InIdentifier, TEXT("Threshold") }, { InMaxInstanceNum }, InThreshold);
+	}
+
+	void FPositionArraySimilarityReward::Evaluate(const FIndexSet Instances)
+	{
+		UE_LEARNING_TRACE_CPUPROFILER_EVENT_SCOPE(FPositionArraySimilarityReward::Evaluate);
+
+		const TLearningArrayView<2, const FVector> Positions0 = InstanceData->ConstView(Positions0Handle);
+		const TLearningArrayView<2, const FVector> Positions1 = InstanceData->ConstView(Positions1Handle);
+		const TLearningArrayView<1, const FVector> RelativePosition0 = InstanceData->ConstView(RelativePosition0Handle);
+		const TLearningArrayView<1, const FVector> RelativePosition1 = InstanceData->ConstView(RelativePosition1Handle);
+		const TLearningArrayView<1, const FQuat> RelativeRotation0 = InstanceData->ConstView(RelativeRotation0Handle);
+		const TLearningArrayView<1, const FQuat> RelativeRotation1 = InstanceData->ConstView(RelativeRotation1Handle);
+
+		const TLearningArrayView<1, const float> Scale = InstanceData->ConstView(ScaleHandle);
+		const TLearningArrayView<1, const float> Threshold = InstanceData->ConstView(ThresholdHandle);
+		const TLearningArrayView<1, const float> Weight = InstanceData->ConstView(WeightHandle);
+
+		TLearningArrayView<1, float> Reward = InstanceData->View(RewardHandle);
+
+		const int32 PositionNum = Positions0.Num<1>();
+
+		for (const int32 InstanceIdx : Instances)
+		{
+			Reward[InstanceIdx] = 0.0f;
+
+			for (int32 PositionIdx = 0; PositionIdx < PositionNum; PositionIdx++)
+			{
+				Reward[InstanceIdx] += Reward::PositionSimilarityReward(
+					RelativeRotation0[InstanceIdx].UnrotateVector(Positions0[InstanceIdx][PositionIdx] - RelativePosition0[InstanceIdx]),
+					RelativeRotation1[InstanceIdx].UnrotateVector(Positions1[InstanceIdx][PositionIdx] - RelativePosition1[InstanceIdx]),
+					Scale[InstanceIdx],
+					Threshold[InstanceIdx],
+					Weight[InstanceIdx]) / PositionNum;
+			}
+		}
+	}
+
+	//------------------------------------------------------------------
+
+}
diff --git a/Engine/Plugins/Experimental/LearningAgents/Source/Learning/Public/LearningArray.h b/Engine/Plugins/Experimental/LearningAgents/Source/Learning/Public/LearningArray.h
--- a/Engine/Plugins/Experimental/LearningAgents/Source/Learning/Public/LearningArray.h
+++ b/Engine/Plugins/Experimental/LearningAgents/Source/Learning/Public/LearningArray.h
@@ -310,16 +310,6 @@ namespace UE::Learning
 		InOutOffset += sizeof(int32);
 	}
 
-	/**
-	* Serialize an integer from bytes
-	*/
-	static inline void DeserializeFromBytes(int32& InOutOffset, TLearningArrayView<1, const uint8> Bytes, uint64& OutValue)
-	{
-		UE_LEARNING_CHECK(InOutOffset + sizeof(uint64) <= Bytes.Num());
-		FMemory::Memcpy((uint8*)&OutValue, &Bytes[InOutOffset], sizeof(uint64));
-		InOutOffset += sizeof(uint64);
-	}
-
 	/**
 	* Some additional functions that act on arrays such as copy, set, zero, etc.
 	* 
@@ -444,32 +434,6 @@ namespace UE::Learning
 #endif
 		}
 
-		/**
-		* Check that an array view does not contain any NaN, Inf, -Inf, MAX_flt, or -MAX_flt
-		*/
-		template<uint8 InDimNum, typename InElementType>
-		inline void Check(const TLearningArray<InDimNum, InElementType>& View, const FIndexSet Indices)
-		{
-#if UE_LEARNING_ARRAY_CHECK_VALUES == 1
-			if constexpr (InDimNum > 1)
-			{
-				for (const int32 Idx : Indices)
-				{
-					Check(View[Idx]);
-				}
-			}
-			else
-			{
-				for (const int32 Idx : Indices)
-				{
-					UE_LEARNING_ARRAY_VALUE_CHECKF(
-						FMath::IsFinite(View.GetData()[Idx]) && View.GetData()[Idx] != MAX_flt && View.GetData()[Idx] != -MAX_flt,
-						TEXT("Invalid value %f found at flat array index %i"), View.GetData()[Idx], Idx);
-				}
-			}
-#endif
-		}
-
 		/**
 		* Zero the memory of an array view
 		*/
@@ -1009,23 +973,6 @@ namespace UE::Learning
 			return sizeof(int32) + sizeof(int32) * InDimNum + sizeof(InElementType) * Shape.Total();
 		}
 
-		/**
-		* Serialize an array to bytes
-		*/
-		template<uint8 InDimNum, typename InElementType, typename Allocator>
-		inline void SerializeToBytes(int32& InOutOffset, TLearningArrayView<1, uint8> Bytes, const TMultiArrayShape<InDimNum> Shape, const TArray<InElementType, Allocator>& InArray)
-		{
-			Learning::SerializeToBytes(InOutOffset, Bytes, (int32)InDimNum);
-			for (uint8 ShapeIdx = 0; ShapeIdx < InDimNum; ShapeIdx++)
-			{
-				Learning::SerializeToBytes(InOutOffset, Bytes, Shape[ShapeIdx]);
-			}
-
-			UE_LEARNING_CHECK(InOutOffset + sizeof(InElementType) * InArray.Num() <= Bytes.Num());
-			FMemory::Memcpy(&Bytes[InOutOffset], (uint8*)InArray.GetData(), sizeof(InElementType) * InArray.Num());
-			InOutOffset += sizeof(InElementType) * InArray.Num();
-		}
-
 		/**
 		* Serialize an array to bytes
 		*/
@@ -1044,30 +991,7 @@ namespace UE::Learning
 		}
 
 		/**
-		* Deserialize an array from bytes
-		*/
-		template<uint8 InDimNum, typename InElementType, typename Allocator>
-		inline void DeserializeFromBytes(int32& InOutOffset, TLearningArrayView<1, const uint8> Bytes, TArray<InElementType, Allocator>& OutArray)
-		{
-			int32 DimNum = INDEX_NONE;
-			Learning::DeserializeFromBytes(InOutOffset, Bytes, DimNum);
-			UE_LEARNING_CHECK(DimNum == InDimNum);
-
-			TLearningArrayShape<InDimNum> Shape;
-			for (uint8 ShapeIdx = 0; ShapeIdx < InDimNum; ShapeIdx++)
-			{
-				Learning::DeserializeFromBytes(InOutOffset, Bytes, Shape[ShapeIdx]);
-			}
-
-			OutArray.SetNumUninitialized(Shape.Total());
-
-			UE_LEARNING_CHECK(InOutOffset + sizeof(InElementType) * OutArray.Num() <= Bytes.Num());
-			FMemory::Memcpy((uint8*)OutArray.GetData(), &Bytes[InOutOffset], sizeof(InElementType) * OutArray.Num());
-			InOutOffset += sizeof(InElementType) * OutArray.Num();
-		}
-
-		/**
-		* Deserialize an array from bytes
+		* Serialize an array from bytes
 		*/
 		template<uint8 InDimNum, typename InElementType, typename Allocator>
 		inline void DeserializeFromBytes(int32& InOutOffset, TLearningArrayView<1, const uint8> Bytes, TLearningArray<InDimNum, InElementType, Allocator>& OutArray)
