[Backout] - CL30154408

diff --git a/Engine/Plugins/Experimental/LearningAgents/Source/Learning/Public/LearningArrayMap.h b/Engine/Plugins/Experimental/LearningAgents/Source/Learning/Public/LearningArrayMap.h
new file mode 100644
--- /dev/null
+++ b/Engine/Plugins/Experimental/LearningAgents/Source/Learning/Public/LearningArrayMap.h
@@ -0,0 +1,532 @@
+// Copyright Epic Games, Inc. All Rights Reserved.
+
+#pragma once
+
+#include "LearningArray.h"
+
+#include "Containers/Map.h"
+#include "Containers/StaticArray.h"
+#include "Misc/GeneratedTypeName.h"
+#include "UObject/NameTypes.h"
+
+/**
+* `FArrayMap` is a Map from `FArrayMapKey` to `TLearningArray` where each array 
+* object stored as a value can be a different shape, dimensionality, and type.
+* 
+* To make this data structure easier to use, a type-safe `TArrayMapHandle` object 
+* is provided which is basically an integer index that can be used to get a view 
+* of an array safely, and without type-casting.
+* 
+* Arrays in this data structure can also be "linked", which makes the two 
+* Keys/Handles which are linked together return a view of the same array data. This 
+* is useful to avoid copying over array data between two entries.
+* 
+* In `Learning` this object is most commonly used to store "InstanceData" - essentially
+* all the data required for multiple instances (such as multiple characters) is 
+* stored together in a single `FArrayMap` as a set of large arrays. This makes this 
+* data extremely efficient and cache-friendly to access when we are processing data for 
+* multiple instances at once.
+* 
+* This structure also provides some basic optimizations for storing single-element arrays
+* to also make it a viable option for when you have (for example) a single instance to store
+* data for.
+*/
+
+class USplineComponent;
+
+namespace UE::Learning
+{
+	enum class ECompletionMode : uint8;
+
+	/**
+	* Key type for FArrayMap. A pair of FNames - generally we can use one to identify the object itself, and one for the variable/member/field.
+	*/
+	struct LEARNING_API FArrayMapKey
+	{
+		FName Namespace;
+		FName Variable;
+
+		bool operator==(const FArrayMapKey& Other) const
+		{
+			return (Namespace == Other.Namespace) && (Variable == Other.Variable);
+		}
+
+		friend uint32 GetTypeHash(const FArrayMapKey& Other)
+		{
+			return GetTypeHash(Other.Namespace) ^ GetTypeHash(Other.Variable);
+		}
+	};
+
+
+	/**
+	* Compile-time function to generate a type id from a C++ type. This is used to check at
+	* runtime that when we access an array from the FArrayMap we are casting it to the correct 
+	* type.
+	*/
+	template<typename T> struct TTypeId
+	{
+		static_assert(sizeof(T) != sizeof(T), "Please specialize this class with a unique id to allow your type to be used in a FArrayMap");
+	};
+	template<> struct TTypeId<float> { enum { Value = 0 }; };
+	template<> struct TTypeId<double> { enum { Value = 1 }; };
+	template<> struct TTypeId<bool> { enum { Value = 2 }; };
+	template<> struct TTypeId<int8> { enum { Value = 3 }; };
+	template<> struct TTypeId<int16> { enum { Value = 4 }; };
+	template<> struct TTypeId<int32> { enum { Value = 5 }; };
+	template<> struct TTypeId<int64> { enum { Value = 6 }; };
+	template<> struct TTypeId<uint8> { enum { Value = 7 }; };
+	template<> struct TTypeId<uint16> { enum { Value = 8 }; };
+	template<> struct TTypeId<uint32> { enum { Value = 9 }; };
+	template<> struct TTypeId<uint64> { enum { Value = 10 }; };
+	template<> struct TTypeId<FVector> { enum { Value = 11 }; };
+	template<> struct TTypeId<FQuat> { enum { Value = 12 }; };
+	template<> struct TTypeId<FTransform> { enum { Value = 13 }; };
+	template<> struct TTypeId<ECompletionMode> { enum { Value = 14 }; };
+	enum { CustomTypeIdStartIndex = 15 };
+
+	/**
+	* Compile-time function to get a debug name for a C++ type
+	*/
+	template<typename T>
+	const TCHAR* CompileTimeTypeName() { return GetGeneratedTypeName<T>(); }
+
+	/**
+	* Untyped handle to an array stored in an FArrayMap
+	*/
+	struct FArrayMapHandle
+	{
+		int16 Index = INDEX_NONE;
+
+		bool operator==(const FArrayMapHandle& Rhs) const { return Index == Rhs.Index; }
+	};
+
+	/**
+	* Type-safe handle to an array stored in an FArrayMap
+	*/
+	template<uint8 DimNum, typename ElementType>
+	struct TArrayMapHandle
+	{
+		int16 Index = INDEX_NONE;
+
+		inline operator FArrayMapHandle() const { return { Index }; }
+	};
+
+	// Some helper functions for FArrayMap
+	namespace ArrayMapPrivate
+	{
+		/**
+		* Destructor for an array of elements of a given type
+		*/
+		template<typename ElementType>
+		void DestructArrayItems(void* Elements, const int32 DimNum, const int32* Shape)
+		{
+			int32 Total = Shape[0];
+			for (uint8 DimIdx = 1; DimIdx < DimNum; DimIdx++)
+			{
+				Total *= Shape[DimIdx];
+			}
+
+			DestructItems<ElementType, int32>((ElementType*)Elements, Total);
+		}
+
+		static inline int32 RoundUpToMultiple(const int32 Value, const int32 Multiple)
+		{
+			return ((Value + Multiple - 1) / Multiple) * Multiple;
+		}
+	}
+
+	/**
+	* A map of differently sized and typed multi-dimensional arrays
+	*/
+	struct LEARNING_API FArrayMap
+	{
+		FArrayMap();
+		~FArrayMap();
+
+		/**
+		* Get a read-only view of an array.
+		*/
+		template<uint8 DimNum, typename ElementType>
+		TLearningArrayView<DimNum, const ElementType> ConstView(const FArrayMapKey Key) const
+		{
+			const FArrayMapHandle Handle = Lookup(Key);
+			CheckArrayWithKey<DimNum, ElementType>(Key, Handle.Index);
+			FDynamicArrayView View = Arrays.GetData()[Handle.Index];
+			return TLearningArrayView<DimNum, const ElementType>((ElementType*)View.Data, TLearningArrayShape<DimNum>(View.Shape));
+		}
+
+		/**
+		* Get a read-only view of an array.
+		*/
+		template<uint8 DimNum, typename ElementType>
+		TLearningArrayView<DimNum, const ElementType> ConstView(const FArrayMapHandle Handle) const
+		{
+			CheckArray<DimNum, ElementType>(Handle.Index);
+			FDynamicArrayView View = Arrays.GetData()[Handle.Index];
+			return TLearningArrayView<DimNum, const ElementType>((ElementType*)View.Data, TLearningArrayShape<DimNum>(View.Shape));
+		}
+
+		/**
+		* Get a read-only view of an array.
+		*/
+		template<uint8 DimNum, typename ElementType>
+		TLearningArrayView<DimNum, const ElementType> ConstView(const TArrayMapHandle<DimNum, ElementType> Handle) const
+		{
+			CheckArray<DimNum, ElementType>(Handle.Index);
+			FDynamicArrayView View = Arrays.GetData()[Handle.Index];
+			return TLearningArrayView<DimNum, const ElementType>((ElementType*)View.Data, TLearningArrayShape<DimNum>(View.Shape));
+		}
+
+		/**
+		* Get a view of an array.
+		*/
+		template<uint8 DimNum, typename ElementType>
+		TLearningArrayView<DimNum, ElementType> View(const FArrayMapKey Key) const
+		{
+			const FArrayMapHandle Handle = Lookup(Key);
+			CheckArrayWithKey<DimNum, ElementType>(Key, Handle.Index);
+			FDynamicArrayView View = Arrays.GetData()[Handle.Index];
+			return TLearningArrayView<DimNum, ElementType>((ElementType*)View.Data, TLearningArrayShape<DimNum>(View.Shape));
+		}
+
+		/**
+		* Get a view of an array.
+		*/
+		template<uint8 DimNum, typename ElementType>
+		TLearningArrayView<DimNum, ElementType> View(const FArrayMapHandle Handle) const
+		{
+			CheckArray<DimNum, ElementType>(Handle.Index);
+			FDynamicArrayView View = Arrays.GetData()[Handle.Index];
+			return TLearningArrayView<DimNum, ElementType>((ElementType*)View.Data, TLearningArrayShape<DimNum>(View.Shape));
+		}
+
+		/**
+		* Get a view of an array.
+		*/
+		template<uint8 DimNum, typename ElementType>
+		TLearningArrayView<DimNum, ElementType> View(const TArrayMapHandle<DimNum, ElementType> Handle) const
+		{
+			CheckArray<DimNum, ElementType>(Handle.Index);
+			FDynamicArrayView View = Arrays.GetData()[Handle.Index];
+			return TLearningArrayView<DimNum, ElementType>((ElementType*)View.Data, TLearningArrayShape<DimNum>(View.Shape));
+		}
+
+		/**
+		* Get the type id of an array in the map.
+		*/
+		int16 TypeId(const FArrayMapKey Key) const;
+
+		/**
+		* Get the type id of an array in the map.
+		*/
+		int16 TypeId(const FArrayMapHandle Handle) const;
+
+		/**
+		* Get the type name of an array in the map.
+		*/
+		const TCHAR* TypeName(const FArrayMapKey Key) const;
+
+		/**
+		* Get the type name of an array in the map.
+		*/
+		const TCHAR* TypeName(const FArrayMapHandle Handle) const;
+
+		/**
+		* Get the number of dimensions of an array in the map.
+		*/
+		uint8 DimNum(const FArrayMapKey Key) const;
+
+		/**
+		* Get the number of dimensions of an array in the map.
+		*/
+		uint8 DimNum(const FArrayMapHandle Handle) const;
+
+		/**
+		* Add a new array to the map
+		*
+		* @param Key		The array key
+		* @param Shape		The shape of the array
+		* @param Default	The default value of the array
+		* @returns			A handle to later retrieve a view of the array for reading or writing
+		*/
+		template<uint8 DimNum, typename ElementType>
+		TArrayMapHandle<DimNum, ElementType> Add(const FArrayMapKey Key, const TLearningArrayShape<DimNum> Shape, const ElementType Default = ElementType())
+		{
+			static_assert(DimNum <= MaxDimNum, "Cannot add an array of that many dimensions.");
+
+			UE_LEARNING_CHECKF(Arrays.Num() < INT16_MAX,
+				TEXT("Maximum number of arrays already registered"));
+
+			UE_LEARNING_CHECKF(!Contains(Key),
+				TEXT("There already exists an array with key { \"%s\", \"%s\" }"),
+				*Key.Namespace.ToString(), *Key.Variable.ToString());
+
+			TArrayMapHandle<DimNum, ElementType> Handle = { Arrays.Num() };
+			Handles.Emplace(Key, Handle);
+
+			/*
+			* The FArrayMap object contains a big buffer of raw data for storing scalar data. This is 
+			* so that when we are allocating arrays with a single element don't end up creating many small 
+			* allocations and indirections. Here we check if the array is a single item and if so we place 
+			* it in the big buffer of scalar data, otherwise we allocate some new memory for it as normal
+			* using Malloc.
+			*/
+
+			const int32 ScalarInsertionIndex = ArrayMapPrivate::RoundUpToMultiple(ScalarOffset, alignof(ElementType));
+
+			FDynamicArrayView ArrayView;
+			if (Shape.Total() == 1 && ScalarInsertionIndex + sizeof(ElementType) <= ScalarData.Num() && alignof(ElementType) <= MaxScalarAlignment)
+			{
+				// Place in Scalar Data
+				ArrayView.Data = &ScalarData[ScalarInsertionIndex];
+				Flags.Emplace(FlagConstructed);
+				ScalarOffset += sizeof(ElementType);
+			}
+			else
+			{
+				// Allocate new array data
+				ArrayView.Data = FMemory::Malloc(Shape.Total() * sizeof(ElementType), alignof(ElementType));
+				Flags.Emplace(FlagConstructed | FlagAllocated);
+			}
+
+			// Set Shape
+			for (uint8 DimIdx = 0; DimIdx < DimNum; DimIdx++)
+			{
+				ArrayView.Shape[DimIdx] = Shape[DimIdx];
+			}
+
+			// Construct elements for this type
+			DefaultConstructItems<ElementType, int32>((ElementType*)ArrayView.Data, Shape.Total());
+
+			// Add array type data to Map
+			Arrays.Emplace(ArrayView);
+			Destructors.Emplace(ArrayMapPrivate::DestructArrayItems<ElementType>);
+			DimNums.Emplace(DimNum);
+			TypeIds.Emplace(TTypeId<ElementType>::Value);
+			TypeNames.Emplace(CompileTimeTypeName<ElementType>());
+
+			// Set to the default value
+			Array::Set(View(Handle), Default);
+
+			return Handle;
+		}
+
+		/**
+		* Add a view to an existing array to the map. Does not allocate any new storage.
+		*
+		* @param Key		The array key
+		* @param View		View to an existing array
+		* @returns			A handle to later retrieve the value of the array for reading or writing
+		*/
+		template<uint8 DimNum, typename ElementType>
+		TArrayMapHandle<DimNum, ElementType> AddView(const FArrayMapKey Key, const TLearningArrayView<DimNum, ElementType> View)
+		{
+			static_assert(DimNum <= MaxDimNum, "Cannot add an array of that many dimensions");
+
+			UE_LEARNING_CHECKF(Arrays.Num() < INT16_MAX,
+				TEXT("Maximum number of arrays already registered"));
+
+			UE_LEARNING_CHECKF(!Contains(Key),
+				TEXT("There already exists an array with key { \"%s\", \"%s\" }"),
+				*Key.Namespace.ToString(), *Key.Variable.ToString());
+
+			TArrayMapHandle<DimNum, ElementType> Handle = { Arrays.Num() };
+			Handles.Emplace(Key, Handle);
+
+			// Since this is an existing array we just create a basic view of it
+			FDynamicArrayView ArrayView;
+			ArrayView.Data = View.GetData();
+			for (uint8 DimIdx = 0; DimIdx < DimNum; DimIdx++)
+			{
+				ArrayView.Shape[DimIdx] = View.Num(DimIdx);
+			}
+
+			// And push the array data as usual
+			Flags.Emplace(FlagNone);
+			Arrays.Emplace(ArrayView);
+			Destructors.Emplace(ArrayMapPrivate::DestructArrayItems<ElementType>);
+			DimNums.Emplace(DimNum);
+			TypeIds.Emplace(TTypeId<ElementType>::Value);
+			TypeNames.Emplace(CompileTimeTypeName<ElementType>());
+
+			return Handle;
+		}
+
+		/**
+		* Check if an array with the given name has already been added.
+		*/
+		bool Contains(const FArrayMapKey Key) const;
+
+		/**
+		* Lookup the array handle associated with a key.
+		*/
+		FArrayMapHandle Lookup(const FArrayMapKey Key) const;
+
+		/**
+		* Find the array handle associated with a key. Returns `nullptr` if Key is not in Map.
+		*/
+		const FArrayMapHandle* Find(const FArrayMapKey Key) const;
+
+		/**
+		* Clear all array data.
+		*/
+		void Empty();
+
+		/**
+		* Link one array to another such that they share the same data. 
+		* Concretely, this function will make asking for a View of the 
+		* Dst array return a view of the Src array instead.
+		*/
+		void LinkKeys(const FArrayMapKey Src, const FArrayMapKey Dst);
+
+		/**
+		* Link one array to another such that they share the same data.
+		* Concretely, this function will make asking for a View of the
+		* Dst array return a view of the Src array instead.
+		*/
+		void LinkHandles(const FArrayMapHandle Src, const FArrayMapHandle Dst);
+
+
+		/**
+		* Link one array to another as if they were flattened. This can be 
+		* useful when the array dimensions don't exactly match such as with
+		* one dimension of a single element.
+		*/
+		void LinkHandlesFlat(const FArrayMapHandle Src, const FArrayMapHandle Dst);
+
+		/**
+		* Link one array to another such that they share the same data.
+		* Concretely, this function will make asking for a View of the
+		* Dst array return a view of the Src array instead. This version
+		* provides a type-safe interface for linking so should be used
+		* by default.
+		*/
+		template<uint8 DimNum, typename ElementType>
+		void Link(const TArrayMapHandle<DimNum, ElementType> Src, const TArrayMapHandle<DimNum, ElementType> Dst)
+		{
+			LinkHandles(Src, Dst);
+		}
+
+		/**
+		* Link one array to another as if they were flattened. This can be
+		* useful when the array dimensions don't exactly match such as with
+		* one dimension of a single element.
+		*/
+		template<uint8 SrcDimNum, uint8 DstDimNum, typename ElementType>
+		void LinkFlat(const TArrayMapHandle<SrcDimNum, ElementType> Src, const TArrayMapHandle<DstDimNum, ElementType> Dst)
+		{
+			LinkHandlesFlat(Src, Dst);
+		}
+
+		/**
+		* Check if the given array has a link.
+		*/
+		bool HasLink(const FArrayMapKey Dst) const;
+
+		/**
+		* Check if the given array has a link.
+		*/
+		bool HasLink(const FArrayMapHandle Dst) const;
+
+		/**
+		* Check if the given array is linked to.
+		*/
+		bool IsLinkedTo(const FArrayMapKey Src) const;
+
+		/**
+		* Check if the given array is linked to.
+		*/
+		bool IsLinkedTo(const FArrayMapHandle Src) const;
+
+		/**
+		* Get all keys in the map.
+		*/
+		template<typename Allocator = FDefaultAllocator>
+		void GetKeys(TArray<FArrayMapKey, Allocator>& OutKeys) const
+		{
+			Handles.GetKeys(OutKeys);
+		}
+
+		/**
+		* Get the key associated with a particular handle
+		*/
+		const FArrayMapKey* FindKey(FArrayMapHandle Handle) const;
+
+	private:
+
+		static constexpr int32 MaxDimNum = 6;
+		static constexpr int32 MaxScalarAlignment = 64;
+		static constexpr int32 InlineArrayNum = 128;
+		static constexpr int32 InlineScalarBytesNum = 2048;
+
+		enum : uint8
+		{
+			FlagNone = 0,
+
+			// Array is allocated by this map, so must be deallocated
+			FlagAllocated = 1 << 0,
+
+			// Array is constructed by this map, so must be destructed
+			FlagConstructed = 1 << 1,
+
+			// Array is a link to another array
+			FlagLinked = 1 << 2,
+		};
+
+		struct FDynamicArrayView
+		{
+			void* Data = nullptr;
+			int32 Shape[MaxDimNum] = { 0 };
+		};
+
+		template<uint8 DimNum, typename ElementType>
+		void CheckArray(const int32 Index) const
+		{
+			UE_LEARNING_CHECKF(Index != INDEX_NONE, TEXT("Invalid Index!"));
+
+			UE_LEARNING_CHECKF(DimNums[Index] == DimNum,
+				TEXT("Asked for array of %i dimensions, but array has %i dimensions"),
+				DimNum, DimNums[Index]);
+
+			UE_LEARNING_CHECKF(TypeIds[Index] == TTypeId<ElementType>::Value,
+				TEXT("Asked for array using type %s, but array data is of type %s"),
+				CompileTimeTypeName<ElementType>(), TypeNames[Index]);
+		}
+
+		template<uint8 DimNum, typename ElementType>
+		void CheckArrayWithKey(const FArrayMapKey Key, const int32 Index) const
+		{
+			UE_LEARNING_CHECKF(Index != INDEX_NONE, TEXT("Invalid Index!"));
+
+			UE_LEARNING_CHECKF(DimNums[Index] == DimNum,
+				TEXT("Asked for array { \"%s\", \"%s\" } with %i dimensions, but array has %i dimensions"),
+				*Key.Namespace.ToString(), *Key.Variable.ToString(), DimNum, DimNums[Index]);
+
+			UE_LEARNING_CHECKF(TypeIds[Index] == TTypeId<ElementType>::Value,
+				TEXT("Asked for array { \"%s\", \"%s\" } using type %s, but array data is of type %s"),
+				*Key.Namespace.ToString(), *Key.Variable.ToString(), CompileTimeTypeName<ElementType>(), TypeNames[Index]);
+		}
+
+		// Array Data
+
+		TArray<FDynamicArrayView, TInlineAllocator<InlineArrayNum>> Arrays;
+
+		// Handle Map
+
+		TMap<FArrayMapKey, FArrayMapHandle, TInlineSetAllocator<InlineArrayNum>> Handles;
+
+		// Array Metadata
+
+		TArray<TFunctionRef<void(void* Elements, const int32 DimNum, const int32* Shape)>, TInlineAllocator<InlineArrayNum>> Destructors;
+		TArray<uint8, TInlineAllocator<InlineArrayNum>> DimNums;
+		TArray<uint8, TInlineAllocator<InlineArrayNum>> Flags;
+		TArray<uint16, TInlineAllocator<InlineArrayNum>> TypeIds;
+		TArray<const TCHAR*, TInlineAllocator<InlineArrayNum>> TypeNames;
+
+		// Scalar Data
+
+		int32 ScalarOffset = 0;
+		TStaticArray<uint8, InlineScalarBytesNum, MaxScalarAlignment> ScalarData;
+	};
+}
+
diff --git a/Engine/Plugins/Experimental/LearningAgents/Source/Learning/Public/LearningCompletion.h b/Engine/Plugins/Experimental/LearningAgents/Source/Learning/Public/LearningCompletion.h
--- a/Engine/Plugins/Experimental/LearningAgents/Source/Learning/Public/LearningCompletion.h
+++ b/Engine/Plugins/Experimental/LearningAgents/Source/Learning/Public/LearningCompletion.h
@@ -27,18 +27,19 @@ namespace UE::Learning
 	*/
 	struct LEARNING_API FResetInstanceBuffer
 	{
-		void Reserve(const int32 InMaxInstanceNum);
+		void Resize(const int32 InMaxInstanceNum);
 
 		void SetResetInstances(const FIndexSet Instances);
 		void SetResetInstancesFromCompletions(const TLearningArrayView<1, const ECompletionMode> Completions, const FIndexSet Instances);
 		const int32 GetResetInstanceNum() const;
 		const FIndexSet GetResetInstances() const;
-		const TArray<int32>& GetResetInstancesArray() const;
 
 	private:
 
+		int32 MaxInstanceNum = 0;
+		int32 ResetInstanceNum = 0;
 		FIndexSet ResetInstancesSet;
-		TArray<int32> ResetInstances;
+		TLearningArray<1, int32, TInlineAllocator<1>> ResetInstances;
 	};
 
 	namespace Completion
@@ -58,16 +59,18 @@ namespace UE::Learning
 		/**
 		* Set completions for all instances whose episode has reached the max number of steps.
 		*
-		* @param OutCompletions					Output buffer to write completions to
+		* @param InOutCompletions				Output buffer to write completions to
 		* @param EpisodeStepNums				Number of steps taken by each instance
 		* @param EpisodeMaxStepNum				Maximum number of allowed steps
+		* @param EpisodeEndCompletionMode		Completion mode to use for instances who reach the end of the episode
 		* @param Instances						Instances to process
 		*
 		*/
 		LEARNING_API void EvaluateEndOfEpisodeCompletions(
-			TLearningArrayView<1, ECompletionMode> OutCompletions,
+			TLearningArrayView<1, ECompletionMode> InOutCompletions,
 			const TLearningArrayView<1, const int32> EpisodeStepNums,
 			const int32 EpisodeMaxStepNum,
+			const ECompletionMode EpisodeEndCompletionMode,
 			const FIndexSet Instances);
 
 	}
diff --git a/Engine/Plugins/Experimental/LearningAgents/Source/Learning/Public/LearningCompletionObject.h b/Engine/Plugins/Experimental/LearningAgents/Source/Learning/Public/LearningCompletionObject.h
new file mode 100644
--- /dev/null
+++ b/Engine/Plugins/Experimental/LearningAgents/Source/Learning/Public/LearningCompletionObject.h
@@ -0,0 +1,245 @@
+// Copyright Epic Games, Inc. All Rights Reserved.
+
+#pragma once
+
+#include "LearningArrayMap.h"
+#include "LearningCompletion.h"
+
+#include "Math/Vector.h"
+#include "Templates/SharedPointer.h"
+
+namespace UE::Learning
+{
+	namespace Completion
+	{
+		static constexpr float DefaultThresholdPosition = 100.0f; // cm
+		static constexpr float DefaultThresholdVelocity = 200.0f; // cm/s
+		static constexpr float DefaultThresholdAngle = UE_PI / 2.0f; // rad
+		static constexpr float DefaultThresholdAngularVelocity = UE_PI; // rad/s
+		static constexpr float DefaultThresholdHeight = 0.0f; // cm
+		static constexpr float DefaultThresholdTime = 10.0f; // s
+	}
+
+	//------------------------------------------------------------------
+
+	/**
+	* Base class for an object which computes if an instance should be completion 
+	* from a set of other arrays. Here, all data is assumed to be stored in a 
+	* `FArrayMap` object to make the processing of multiple instances efficient.
+	*/
+	struct LEARNING_API FCompletionObject
+	{
+		FCompletionObject(
+			const FName& InIdentifier,
+			const TSharedRef<FArrayMap>& InInstanceData,
+			const int32 InMaxInstanceNum,
+			const ECompletionMode InCompletionMode = ECompletionMode::Terminated);
+
+		virtual ~FCompletionObject() {}
+
+		virtual void Evaluate(const FIndexSet Instances) = 0;
+
+		TLearningArrayView<1, ECompletionMode> CompletionBuffer();
+
+		TSharedRef<FArrayMap> InstanceData;
+		ECompletionMode CompletionMode;
+
+		TArrayMapHandle<1, ECompletionMode> CompletionHandle;
+	};
+
+	//------------------------------------------------------------------
+
+	/**
+	* Terminates if any of the provided completions terminate
+	*/
+	struct LEARNING_API FAnyCompletion : public FCompletionObject
+	{
+		FAnyCompletion(
+			const FName& InIdentifier,
+			const TLearningArrayView<1, const TSharedRef<FCompletionObject>> InCompletions,
+			const TSharedRef<FArrayMap>& InInstanceData,
+			const int32 InMaxInstanceNum);
+
+		virtual void Evaluate(const FIndexSet Instances) override final;
+
+		TLearningArray<1, TSharedRef<FCompletionObject>, TInlineAllocator<32>> Completions;
+	};
+
+	//------------------------------------------------------------------
+
+	/**
+	* Completion due to some condition being true
+	*/
+	struct LEARNING_API FConditionalCompletion : public FCompletionObject
+	{
+		FConditionalCompletion(
+			const FName& InIdentifier,
+			const TSharedRef<FArrayMap>& InInstanceData,
+			const int32 InMaxInstanceNum,
+			const ECompletionMode InCompletionMode = ECompletionMode::Terminated);
+
+		virtual void Evaluate(const FIndexSet Instances) override final;
+
+		TArrayMapHandle<1, bool> ConditionHandle;
+	};
+
+	//------------------------------------------------------------------
+
+	/**
+	* Completion due to time having elapsed
+	*/
+	struct LEARNING_API FTimeElapsedCompletion : public FCompletionObject
+	{
+		FTimeElapsedCompletion(
+			const FName& InIdentifier,
+			const TSharedRef<FArrayMap>& InInstanceData,
+			const int32 InMaxInstanceNum,
+			const float InThreshold = Completion::DefaultThresholdTime,
+			const ECompletionMode InCompletionMode = ECompletionMode::Terminated);
+
+		virtual void Evaluate(const FIndexSet Instances) override final;
+
+		TArrayMapHandle<1, float> TimeHandle;
+		TArrayMapHandle<1, float> ThresholdHandle;
+	};
+
+	//------------------------------------------------------------------
+
+	/**
+	* Completion due to the difference between any two scalar positions exceeding a given threshold
+	*/
+	struct LEARNING_API FScalarPositionDifferenceCompletion : public FCompletionObject
+	{
+		FScalarPositionDifferenceCompletion(
+			const FName& InIdentifier,
+			const TSharedRef<FArrayMap>& InInstanceData,
+			const int32 InMaxInstanceNum,
+			const int32 InPositionNum,
+			const float InThreshold = Completion::DefaultThresholdPosition,
+			const ECompletionMode InCompletionMode = ECompletionMode::Terminated);
+
+		virtual void Evaluate(const FIndexSet Instances) override final;
+
+		TArrayMapHandle<2, float> Position0Handle;
+		TArrayMapHandle<2, float> Position1Handle;
+		TArrayMapHandle<1, float> ThresholdHandle;
+	};
+
+	/**
+	* Completion due to the difference between any two planar positions exceeding a given threshold
+	*/
+	struct LEARNING_API FPlanarPositionDifferenceCompletion : public FCompletionObject
+	{
+		FPlanarPositionDifferenceCompletion(
+			const FName& InIdentifier,
+			const TSharedRef<FArrayMap>& InInstanceData,
+			const int32 InMaxInstanceNum,
+			const int32 InPositionNum,
+			const float InThreshold = Completion::DefaultThresholdPosition,
+			const ECompletionMode InCompletionMode = ECompletionMode::Terminated,
+			const FVector InAxis0 = FVector::ForwardVector,
+			const FVector InAxis1 = FVector::RightVector);
+
+		virtual void Evaluate(const FIndexSet Instances) override final;
+
+		FVector Axis0 = FVector::ForwardVector;
+		FVector Axis1 = FVector::RightVector;
+
+		TArrayMapHandle<2, FVector> Position0Handle;
+		TArrayMapHandle<2, FVector> Position1Handle;
+		TArrayMapHandle<1, float> ThresholdHandle;
+	};
+
+	/**
+	* Completion due to the similarity between any two planar positions being below a given threshold
+	*/
+	struct LEARNING_API FPlanarPositionSimilarityCompletion : public FCompletionObject
+	{
+		FPlanarPositionSimilarityCompletion(
+			const FName& InIdentifier,
+			const TSharedRef<FArrayMap>& InInstanceData,
+			const int32 InMaxInstanceNum,
+			const int32 InPositionNum,
+			const float InThreshold = Completion::DefaultThresholdPosition,
+			const ECompletionMode InCompletionMode = ECompletionMode::Terminated,
+			const FVector InAxis0 = FVector::ForwardVector,
+			const FVector InAxis1 = FVector::RightVector);
+
+		virtual void Evaluate(const FIndexSet Instances) override final;
+
+		FVector Axis0 = FVector::ForwardVector;
+		FVector Axis1 = FVector::RightVector;
+
+		TArrayMapHandle<2, FVector> Position0Handle;
+		TArrayMapHandle<2, FVector> Position1Handle;
+		TArrayMapHandle<1, float> ThresholdHandle;
+	};
+
+	//------------------------------------------------------------------
+
+	/**
+	* Completion due to the difference between any two scalar velocities exceeding a given threshold
+	*/
+	struct LEARNING_API FScalarVelocityDifferenceCompletion : public FCompletionObject
+	{
+		FScalarVelocityDifferenceCompletion(
+			const FName& InIdentifier,
+			const TSharedRef<FArrayMap>& InInstanceData,
+			const int32 InMaxInstanceNum,
+			const int32 InVelocityNum,
+			const float InThreshold = Completion::DefaultThresholdVelocity,
+			const ECompletionMode InCompletionMode = ECompletionMode::Terminated);
+
+		virtual void Evaluate(const FIndexSet Instances) override final;
+
+		TArrayMapHandle<2, float> Velocity0Handle;
+		TArrayMapHandle<2, float> Velocity1Handle;
+		TArrayMapHandle<1, float> ThresholdHandle;
+	};
+
+	//------------------------------------------------------------------
+
+	/**
+	* Completion due to the difference between any two scalar rotations exceeding a given threshold
+	*/
+	struct LEARNING_API FScalarRotationDifferenceCompletion : public FCompletionObject
+	{
+		FScalarRotationDifferenceCompletion(
+			const FName& InIdentifier,
+			const TSharedRef<FArrayMap>& InInstanceData,
+			const int32 InMaxInstanceNum,
+			const int32 InRotationNum,
+			const float InThreshold = Completion::DefaultThresholdAngle,
+			const ECompletionMode InCompletionMode = ECompletionMode::Terminated);
+
+		virtual void Evaluate(const FIndexSet Instances) override final;
+
+		TArrayMapHandle<2, float> Rotation0Handle;
+		TArrayMapHandle<2, float> Rotation1Handle;
+		TArrayMapHandle<1, float> ThresholdHandle;
+	};
+
+	//------------------------------------------------------------------
+
+	/**
+	* Completion due to the difference between two scalar angular velocities exceeding a given threshold
+	*/
+	struct LEARNING_API FScalarAngularVelocityDifferenceCompletion : public FCompletionObject
+	{
+		FScalarAngularVelocityDifferenceCompletion(
+			const FName& InIdentifier,
+			const TSharedRef<FArrayMap>& InInstanceData,
+			const int32 InMaxInstanceNum,
+			const int32 InAngularVelocityNum,
+			const float InThreshold = Completion::DefaultThresholdAngularVelocity,
+			const ECompletionMode InCompletionMode = ECompletionMode::Terminated);
+
+		virtual void Evaluate(const FIndexSet Instances) override final;
+
+		TArrayMapHandle<2, float> AngularVelocity0Handle;
+		TArrayMapHandle<2, float> AngularVelocity1Handle;
+		TArrayMapHandle<1, float> ThresholdHandle;
+	};
+
+	//------------------------------------------------------------------
+}
diff --git a/Engine/Plugins/Experimental/LearningAgents/Source/Learning/Public/LearningFeatureObject.h b/Engine/Plugins/Experimental/LearningAgents/Source/Learning/Public/LearningFeatureObject.h
new file mode 100644
--- /dev/null
+++ b/Engine/Plugins/Experimental/LearningAgents/Source/Learning/Public/LearningFeatureObject.h
@@ -0,0 +1,912 @@
+// Copyright Epic Games, Inc. All Rights Reserved.
+
+#pragma once
+
+#include "LearningArrayMap.h"
+
+#include "Math/Quat.h"
+#include "Math/Vector.h"
+#include "Templates/SharedPointer.h"
+
+namespace UE::Learning
+{
+	namespace Feature
+	{
+		static constexpr float DefaultScale = 1.0f;
+		static constexpr float DefaultScaleTime = 1.0f; // s
+		static constexpr float DefaultScaleDirection = 1.0f;
+
+		static constexpr float DefaultScalePosition = 100.0f; // cm
+		static constexpr float DefaultScaleVelocity = 200.0f; // cm/s
+		static constexpr float DefaultScaleAcceleration = 400.0f; // cm/s^2
+
+		static constexpr float DefaultScaleAngle = UE_PI / 2.0f; // rad
+		static constexpr float DefaultScaleAngularVelocity = UE_PI; // rad/s
+		static constexpr float DefaultScaleAngularAcceleration = 2.0f * UE_PI; // rad/s
+
+		static constexpr float DefaultScaleScale = 1.0f;
+
+		namespace Encode
+		{
+			LEARNING_API void Float(
+				float& Output,
+				const float Input,
+				const float Scale = DefaultScale,
+				const float Epsilon = UE_SMALL_NUMBER);
+
+
+			LEARNING_API void Time(
+				float& Output,
+				const float Time,
+				const float RelativeTime = 0.0f,
+				const float Scale = DefaultScaleTime,
+				const float Epsilon = UE_SMALL_NUMBER);
+
+
+			LEARNING_API void PlanarDirection(
+				TLearningArrayView<1, float> Output,
+				const FVector Direction,
+				const FQuat RelativeRotation = FQuat::Identity,
+				const float Scale = DefaultScaleDirection,
+				const FVector Axis0 = FVector::ForwardVector,
+				const FVector Axis1 = FVector::RightVector,
+				const float Epsilon = UE_SMALL_NUMBER);
+
+			LEARNING_API void Direction(
+				TLearningArrayView<1, float> Output,
+				const FVector Direction,
+				const FQuat RelativeRotation = FQuat::Identity,
+				const float Scale = DefaultScaleDirection,
+				const float Epsilon = UE_SMALL_NUMBER);
+
+
+			LEARNING_API void ScalarPosition(
+				float& Output,
+				const float Position,
+				const float RelativePosition = 0.0f,
+				const float Scale = DefaultScalePosition,
+				const float Epsilon = UE_SMALL_NUMBER);
+
+			LEARNING_API void PlanarPosition(
+				TLearningArrayView<1, float> Output,
+				const FVector Position,
+				const FVector RelativePosition = FVector::ZeroVector,
+				const FQuat RelativeRotation = FQuat::Identity,
+				const float Scale = DefaultScalePosition,
+				const FVector Axis0 = FVector::ForwardVector,
+				const FVector Axis1 = FVector::RightVector,
+				const float Epsilon = UE_SMALL_NUMBER);
+
+			LEARNING_API void Position(
+				TLearningArrayView<1, float> Output,
+				const FVector Position,
+				const FVector RelativePosition = FVector::ZeroVector,
+				const FQuat RelativeRotation = FQuat::Identity,
+				const float Scale = DefaultScalePosition,
+				const float Epsilon = UE_SMALL_NUMBER);
+
+
+			LEARNING_API void ScalarVelocity(
+				float& Output,
+				const float Velocity,
+				const float Scale = DefaultScaleVelocity,
+				const float Epsilon = UE_SMALL_NUMBER);
+
+			LEARNING_API void PlanarVelocity(
+				TLearningArrayView<1, float> Output,
+				const FVector Velocity,
+				const FQuat RelativeRotation = FQuat::Identity,
+				const float Scale = DefaultScaleVelocity,
+				const FVector Axis0 = FVector::ForwardVector,
+				const FVector Axis1 = FVector::RightVector,
+				const float Epsilon = UE_SMALL_NUMBER);
+
+			LEARNING_API void Velocity(
+				TLearningArrayView<1, float> Output,
+				const FVector Velocity,
+				const FQuat RelativeRotation = FQuat::Identity,
+				const float Scale = DefaultScaleVelocity,
+				const float Epsilon = UE_SMALL_NUMBER);
+
+
+			LEARNING_API void ScalarAcceleration(
+				float& Output,
+				const float Acceleration,
+				const float Scale = DefaultScaleAcceleration,
+				const float Epsilon = UE_SMALL_NUMBER);
+
+			LEARNING_API void PlanarAcceleration(
+				TLearningArrayView<1, float> Output,
+				const FVector Acceleration,
+				const FQuat RelativeRotation = FQuat::Identity,
+				const float Scale = DefaultScaleAcceleration,
+				const FVector Axis0 = FVector::ForwardVector,
+				const FVector Axis1 = FVector::RightVector,
+				const float Epsilon = UE_SMALL_NUMBER);
+
+			LEARNING_API void Acceleration(
+				TLearningArrayView<1, float> Output,
+				const FVector Acceleration,
+				const FQuat RelativeRotation = FQuat::Identity,
+				const float Scale = DefaultScaleAcceleration,
+				const float Epsilon = UE_SMALL_NUMBER);
+
+
+			LEARNING_API void Angle(
+				TLearningArrayView<1, float> Output,
+				const float Angle,
+				const float RelativeRotation = 0.0f,
+				const float Scale = DefaultScaleDirection, // angles are encoded as directions
+				const float Epsilon = UE_SMALL_NUMBER);
+
+			LEARNING_API void Rotation(
+				TLearningArrayView<1, float> Output,
+				const FQuat Rotation,
+				const FQuat RelativeRotation = FQuat::Identity,
+				const float Scale = DefaultScaleDirection, // rotations are encoded as directions
+				const float Epsilon = UE_SMALL_NUMBER);
+
+			LEARNING_API void RotationVector(
+				TLearningArrayView<1, float> Output,
+				const FVector RotationVector,
+				const float Scale = DefaultScaleAngle,
+				const float Epsilon = UE_SMALL_NUMBER);
+
+
+			LEARNING_API void ScalarAngularVelocity(
+				float& Output,
+				const float AngularVelocity,
+				const float Scale = DefaultScaleAngularVelocity,
+				const float Epsilon = UE_SMALL_NUMBER);
+
+			LEARNING_API void AngularVelocity(
+				TLearningArrayView<1, float> Output,
+				const FVector AngularVelocity,
+				const FQuat RelativeRotation = FQuat::Identity,
+				const float Scale = DefaultScaleAngularVelocity,
+				const float Epsilon = UE_SMALL_NUMBER);
+
+
+			LEARNING_API void ScalarAngularAcceleration(
+				float& Output,
+				const float AngularAcceleration,
+				const float Scale = DefaultScaleAngularAcceleration,
+				const float Epsilon = UE_SMALL_NUMBER);
+
+			LEARNING_API void PlanarAngularAcceleration(
+				float& Output,
+				const FVector AngularAcceleration,
+				const FQuat RelativeRotation = FQuat::Identity,
+				const float Scale = DefaultScaleAngularAcceleration,
+				const FVector Axis = FVector::UpVector,
+				const float Epsilon = UE_SMALL_NUMBER);
+
+			LEARNING_API void AngularAcceleration(
+				TLearningArrayView<1, float> Output,
+				const FVector AngularAcceleration,
+				const FQuat RelativeRotation = FQuat::Identity,
+				const float Scale =DefaultScaleAngularAcceleration,
+				const float Epsilon = UE_SMALL_NUMBER);
+
+
+			LEARNING_API void Scale(
+				float& Output,
+				const float InScale,
+				const float RelativeScale = 1.0f,
+				const float Scale = DefaultScaleScale,
+				const float Epsilon = UE_SMALL_NUMBER);
+		}
+
+		namespace Decode
+		{
+			LEARNING_API void Float(
+				float& Value,
+				const TLearningArrayView<1, const float> Input,
+				const float Scale = DefaultScale);
+
+
+			LEARNING_API void Time(
+				float& OutTime,
+				const TLearningArrayView<1, const float> Input,
+				const float RelativeTime,
+				const float Scale = DefaultScaleTime);
+
+
+			LEARNING_API void PlanarDirection(
+				FVector& OutPlanarDirection,
+				const TLearningArrayView<1, const float> Input,
+				const FQuat RelativeRotation = FQuat::Identity,
+				const float Scale = DefaultScaleDirection,
+				const FVector Axis0 = FVector::ForwardVector,
+				const FVector Axis1 = FVector::RightVector,
+				const float Epsilon = UE_SMALL_NUMBER);
+
+			LEARNING_API void Direction(
+				FVector& OutDirection,
+				const TLearningArrayView<1, const float> Input,
+				const FQuat RelativeRotation = FQuat::Identity,
+				const float Scale = DefaultScaleDirection,
+				const float Epsilon = UE_SMALL_NUMBER);
+
+
+			LEARNING_API void ScalarPosition(
+				float& OutScalarPosition,
+				const float Input,
+				const float RelativePosition = 0.0f,
+				const float Scale = DefaultScalePosition);
+
+			LEARNING_API void PlanarPosition(
+				FVector& OutPlanarPosition,
+				const TLearningArrayView<1, const float> Input,
+				const FVector RelativePosition = FVector::ZeroVector,
+				const FQuat RelativeRotation = FQuat::Identity,
+				const float Scale = DefaultScalePosition,
+				const FVector Axis0 = FVector::ForwardVector,
+				const FVector Axis1 = FVector::RightVector);
+
+			LEARNING_API void Position(
+				FVector& OutPosition,
+				const TLearningArrayView<1, const float> Input,
+				const FVector RelativePosition = FVector::ZeroVector,
+				const FQuat RelativeRotation = FQuat::Identity,
+				const float Scale = DefaultScalePosition);
+
+
+			LEARNING_API void ScalarVelocity(
+				float& OutScalarVelocity,
+				const float Input,
+				const float Scale = DefaultScaleVelocity);
+
+			LEARNING_API void PlanarVelocity(
+				FVector& OutPlanarVelocity,
+				const TLearningArrayView<1, const float> Input,
+				const FQuat RelativeRotation = FQuat::Identity,
+				const float Scale = DefaultScaleVelocity,
+				const FVector Axis0 = FVector::ForwardVector,
+				const FVector Axis1 = FVector::RightVector);
+
+			LEARNING_API void Velocity(
+				FVector& OutVelocity,
+				const TLearningArrayView<1, const float> Input,
+				const FQuat RelativeRotation = FQuat::Identity,
+				const float Scale = DefaultScaleVelocity);
+
+
+			LEARNING_API void ScalarAcceleration(
+				float& OutScalarAcceleration,
+				const float Input,
+				const float Scale = DefaultScaleAcceleration);
+
+			LEARNING_API void PlanarAcceleration(
+				FVector& OutPlanarAcceleration,
+				const TLearningArrayView<1, const float> Input,
+				const FQuat RelativeRotation = FQuat::Identity,
+				const float Scale = DefaultScaleAcceleration,
+				const FVector Axis0 = FVector::ForwardVector,
+				const FVector Axis1 = FVector::RightVector);
+
+			LEARNING_API void Acceleration(
+				FVector& OutAcceleration,
+				const TLearningArrayView<1, const float> Input,
+				const FQuat RelativeRotation = FQuat::Identity,
+				const float Scale = DefaultScaleAcceleration);
+
+
+			LEARNING_API void Angle(
+				float& OutScalarRotation,
+				const TLearningArrayView<1, const float> Input,
+				const float RelativeRotation = 0.0f,
+				const float Scale = DefaultScaleDirection); // angles are encoded as directions
+
+			LEARNING_API void Rotation(
+				FQuat& OutRotation,
+				const TLearningArrayView<1, const float> Input,
+				const FQuat RelativeRotation = FQuat::Identity,
+				const float Scale = DefaultScaleDirection, // rotations are encoded as directions
+				const float Epsilon = UE_SMALL_NUMBER);
+
+			LEARNING_API void RotationVector(
+				FVector& OutRotationVector,
+				const TLearningArrayView<1, const float> Input,
+				const float Scale = DefaultScaleAngle);
+
+
+			LEARNING_API void ScalarAngularVelocity(
+				float& OutScalarVelocity,
+				const float Input,
+				const float Scale = DefaultScaleAngularVelocity);
+
+			LEARNING_API void AngularVelocity(
+				FVector& OutAngularVelocity,
+				const TLearningArrayView<1, const float> Input,
+				const FQuat RelativeRotation = FQuat::Identity,
+				const float Scale = DefaultScaleAngularVelocity);
+
+
+			LEARNING_API void ScalarAngularAcceleration(
+				float& OutScalarAngularAcceleration,
+				const float Input,
+				const float Scale = DefaultScaleAngularAcceleration);
+
+			LEARNING_API void PlanarAngularAcceleration(
+				FVector& OutPlanarAngularAcceleration,
+				const float Input,
+				const FQuat RelativeRotation = FQuat::Identity,
+				const float Scale = DefaultScaleAngularAcceleration,
+				const FVector Axis = FVector::UpVector);
+
+			LEARNING_API void AngularAcceleration(
+				FVector& OutAngularAcceleration,
+				const TLearningArrayView<1, const float> Input,
+				const FQuat RelativeRotation = FQuat::Identity,
+				const float Scale = DefaultScaleAngularAcceleration);
+
+
+			LEARNING_API void Scale(
+				float& OutScale,
+				const float Input,
+				const float RelativeScale = 1.0f,
+				const float Scale = DefaultScaleScale,
+				const float Epsilon = UE_SMALL_NUMBER);
+		}
+	}
+
+	/**
+	* Base class for an object which can do the encoding or decoding of
+	* features from a set of other arrays. Here, all data is assumed
+	* to be stored in a `FArrayMap` object to make the processing of 
+	* multiple instances efficient.
+	*/
+	struct LEARNING_API FFeatureObject
+	{
+		FFeatureObject(
+			const FName& InIdentifier,
+			const TSharedRef<FArrayMap>& InInstanceData,
+			const int32 InMaxInstanceNum,
+			const int32 InFeatureDimNum,
+			const float InScale = Feature::DefaultScale);
+
+		virtual ~FFeatureObject() {}
+
+		int32 DimNum() const;
+		TLearningArrayView<2, float> FeatureBuffer();
+
+		virtual bool IsEncodable() const = 0;
+		virtual bool IsDecodable() const = 0;
+		virtual void Encode(const FIndexSet Instances) = 0;
+		virtual void Decode(const FIndexSet Instances) = 0;
+
+		TSharedRef<FArrayMap> InstanceData;
+		int32 FeatureDimNum = 0;
+		float Scale = Feature::DefaultScale;
+		TArrayMapHandle<2, float> FeatureHandle;
+	};
+
+	//------------------------------------------------------------------
+
+	/**
+	* Feature that concatenates multiple other features
+	*/
+	struct LEARNING_API FConcatenateFeature : public FFeatureObject
+	{
+		FConcatenateFeature(
+			const FName& InIdentifier,
+			const TLearningArrayView<1, const TSharedRef<FFeatureObject>> InFeatures,
+			const TSharedRef<FArrayMap>& InInstanceData,
+			const int32 InMaxInstanceNum,
+			const float InScale = Feature::DefaultScale);
+
+		virtual bool IsEncodable() const override final;
+		virtual bool IsDecodable() const override final;
+		virtual void Encode(const FIndexSet Instances) override final;
+		virtual void Decode(const FIndexSet Instances) override final;
+
+		TLearningArray<1, TSharedRef<FFeatureObject>, TInlineAllocator<32>> Features;
+		TLearningArray<1, int32, TInlineAllocator<32>> Offsets;
+		TLearningArray<1, int32, TInlineAllocator<32>> Sizes;
+	};
+
+	//------------------------------------------------------------------
+
+	/**
+	* Feature for an array of floats
+	*/
+	struct LEARNING_API FFloatFeature : public FFeatureObject
+	{
+		FFloatFeature(
+			const FName& InIdentifier,
+			const TSharedRef<FArrayMap>& InInstanceData,
+			const int32 InMaxInstanceNum,
+			const int32 InFloatNum,
+			const float InScale = Feature::DefaultScale);
+
+		virtual bool IsEncodable() const override final { return true; }
+		virtual bool IsDecodable() const override final { return true; }
+		virtual void Encode(const FIndexSet Instances) override final;
+		virtual void Decode(const FIndexSet Instances) override final;
+
+		TArrayMapHandle<2, float> ValueHandle;
+	};
+
+	//------------------------------------------------------------------
+
+	/**
+	* Feature for an array of times
+	*/
+	struct LEARNING_API FTimeFeature : public FFeatureObject
+	{
+		FTimeFeature(
+			const FName& InIdentifier,
+			const TSharedRef<FArrayMap>& InInstanceData,
+			const int32 InMaxInstanceNum,
+			const int32 InTimeNum,
+			const float InScale = Feature::DefaultScaleTime);
+
+		virtual bool IsEncodable() const override final { return true; }
+		virtual bool IsDecodable() const override final { return true; }
+		virtual void Encode(const FIndexSet Instances) override final;
+		virtual void Decode(const FIndexSet Instances) override final;
+
+		TArrayMapHandle<2, float> TimeHandle;
+		TArrayMapHandle<1, float> RelativeTimeHandle;
+	};
+
+	//------------------------------------------------------------------
+
+	/**
+	* Feature for an array of planar directions
+	*/
+	struct LEARNING_API FPlanarDirectionFeature : public FFeatureObject
+	{
+		FPlanarDirectionFeature(
+			const FName& InIdentifier,
+			const TSharedRef<FArrayMap>& InInstanceData,
+			const int32 InMaxInstanceNum,
+			const int32 InDirectionNum,
+			const float InScale = Feature::DefaultScaleDirection,
+			const FVector InAxis0 = FVector::ForwardVector,
+			const FVector InAxis1 = FVector::RightVector);
+
+		virtual bool IsEncodable() const override final { return true; }
+		virtual bool IsDecodable() const override final { return true; }
+		virtual void Encode(const FIndexSet Instances) override final;
+		virtual void Decode(const FIndexSet Instances) override final;
+
+		FVector Axis0 = FVector::ForwardVector;
+		FVector Axis1 = FVector::RightVector;
+
+		TArrayMapHandle<2, FVector> DirectionHandle;
+		TArrayMapHandle<1, FQuat> RelativeRotationHandle;
+	};
+
+	/**
+	* Feature for an array of directions
+	*/
+	struct LEARNING_API FDirectionFeature : public FFeatureObject
+	{
+		FDirectionFeature(
+			const FName& InIdentifier,
+			const TSharedRef<FArrayMap>& InInstanceData,
+			const int32 InMaxInstanceNum,
+			const int32 InDirectionNum,
+			const float InScale = Feature::DefaultScaleDirection);
+
+		virtual bool IsEncodable() const override final { return true; }
+		virtual bool IsDecodable() const override final { return true; }
+		virtual void Encode(const FIndexSet Instances) override final;
+		virtual void Decode(const FIndexSet Instances) override final;
+
+		TArrayMapHandle<2, FVector> DirectionHandle;
+		TArrayMapHandle<1, FQuat> RelativeRotationHandle;
+	};
+
+	//------------------------------------------------------------------
+
+	/**
+	* Feature for an array of scalar positions
+	*/
+	struct LEARNING_API FScalarPositionFeature : public FFeatureObject
+	{
+		FScalarPositionFeature(
+			const FName& InIdentifier,
+			const TSharedRef<FArrayMap>& InInstanceData,
+			const int32 InMaxInstanceNum,
+			const int32 InPositionNum,
+			const float InScale = Feature::DefaultScalePosition);
+
+		virtual bool IsEncodable() const override final { return true; }
+		virtual bool IsDecodable() const override final { return true; }
+		virtual void Encode(const FIndexSet Instances) override final;
+		virtual void Decode(const FIndexSet Instances) override final;
+
+		TArrayMapHandle<2, float> PositionHandle;
+		TArrayMapHandle<1, float> RelativePositionHandle;
+	};
+
+	/**
+	* Feature for an array of planar positions
+	*/
+	struct LEARNING_API FPlanarPositionFeature : public FFeatureObject
+	{
+		FPlanarPositionFeature(
+			const FName& InIdentifier,
+			const TSharedRef<FArrayMap>& InInstanceData,
+			const int32 InMaxInstanceNum,
+			const int32 InPositionNum,
+			const float InScale = Feature::DefaultScalePosition,
+			const FVector InAxis0 = FVector::ForwardVector,
+			const FVector InAxis1 = FVector::RightVector);
+
+		virtual bool IsEncodable() const override final { return true; }
+		virtual bool IsDecodable() const override final { return true; }
+		virtual void Encode(const FIndexSet Instances) override final;
+		virtual void Decode(const FIndexSet Instances) override final;
+
+		FVector Axis0 = FVector::ForwardVector;
+		FVector Axis1 = FVector::RightVector;
+
+		TArrayMapHandle<2, FVector> PositionHandle;
+		TArrayMapHandle<1, FVector> RelativePositionHandle;
+		TArrayMapHandle<1, FQuat> RelativeRotationHandle;
+	};
+
+	/**
+	* Feature for an array of positions
+	*/
+	struct LEARNING_API FPositionFeature : public FFeatureObject
+	{
+		FPositionFeature(
+			const FName& InIdentifier,
+			const TSharedRef<FArrayMap>& InInstanceData,
+			const int32 InMaxInstanceNum,
+			const int32 InPositionNum,
+			const float InScale = Feature::DefaultScalePosition);
+
+		virtual bool IsEncodable() const override final { return true; }
+		virtual bool IsDecodable() const override final { return true; }
+		virtual void Encode(const FIndexSet Instances) override final;
+		virtual void Decode(const FIndexSet Instances) override final;
+
+		TArrayMapHandle<2, FVector> PositionHandle;
+		TArrayMapHandle<1, FVector> RelativePositionHandle;
+		TArrayMapHandle<1, FQuat> RelativeRotationHandle;
+	};
+
+	//------------------------------------------------------------------
+
+	/**
+	* Feature for an array of scalar velocities
+	*/
+	struct LEARNING_API FScalarVelocityFeature : public FFeatureObject
+	{
+		FScalarVelocityFeature(
+			const FName& InIdentifier,
+			const TSharedRef<FArrayMap>& InInstanceData,
+			const int32 InMaxInstanceNum,
+			const int32 InVelocityNum,
+			const float InScale = Feature::DefaultScaleVelocity);
+
+		virtual bool IsEncodable() const override final { return true; }
+		virtual bool IsDecodable() const override final { return true; }
+		virtual void Encode(const FIndexSet Instances) override final;
+		virtual void Decode(const FIndexSet Instances) override final;
+
+		TArrayMapHandle<2, float> VelocityHandle;
+	};
+
+	/**
+	* Feature for an array of planar velocities
+	*/
+	struct LEARNING_API FPlanarVelocityFeature : public FFeatureObject
+	{
+		FPlanarVelocityFeature(
+			const FName& InIdentifier,
+			const TSharedRef<FArrayMap>& InInstanceData,
+			const int32 InMaxInstanceNum,
+			const int32 InVelocityNum,
+			const float InScale = Feature::DefaultScaleVelocity,
+			const FVector InAxis0 = FVector::ForwardVector,
+			const FVector InAxis1 = FVector::RightVector);
+
+		virtual bool IsEncodable() const override final { return true; }
+		virtual bool IsDecodable() const override final { return true; }
+		virtual void Encode(const FIndexSet Instances) override final;
+		virtual void Decode(const FIndexSet Instances) override final;
+
+		FVector Axis0 = FVector::ForwardVector;
+		FVector Axis1 = FVector::RightVector;
+
+		TArrayMapHandle<2, FVector> VelocityHandle;
+		TArrayMapHandle<1, FQuat> RelativeRotationHandle;
+	};
+
+	/**
+	* Feature for an array of velocities
+	*/
+	struct LEARNING_API FVelocityFeature : public FFeatureObject
+	{
+		FVelocityFeature(
+			const FName& InIdentifier,
+			const TSharedRef<FArrayMap>& InInstanceData,
+			const int32 InMaxInstanceNum,
+			const int32 InVelocityNum,
+			const float InScale = Feature::DefaultScaleVelocity);
+
+		virtual bool IsEncodable() const override final { return true; }
+		virtual bool IsDecodable() const override final { return true; }
+		virtual void Encode(const FIndexSet Instances) override final;
+		virtual void Decode(const FIndexSet Instances) override final;
+
+		TArrayMapHandle<2, FVector> VelocityHandle;
+		TArrayMapHandle<1, FQuat> RelativeRotationHandle;
+	};
+
+	//------------------------------------------------------------------
+
+	/**
+	* Feature for an array of scalar accelerations
+	*/
+	struct LEARNING_API FScalarAccelerationFeature : public FFeatureObject
+	{
+		FScalarAccelerationFeature(
+			const FName& InIdentifier,
+			const TSharedRef<FArrayMap>& InInstanceData,
+			const int32 InMaxInstanceNum,
+			const int32 InAccelerationNum,
+			const float InScale = Feature::DefaultScaleAcceleration);
+
+		virtual bool IsEncodable() const override final { return true; }
+		virtual bool IsDecodable() const override final { return true; }
+		virtual void Encode(const FIndexSet Instances) override final;
+		virtual void Decode(const FIndexSet Instances) override final;
+
+		TArrayMapHandle<2, float> AccelerationHandle;
+	};
+
+	/**
+	* Feature for an array of planar accelerations
+	*/
+	struct LEARNING_API FPlanarAccelerationFeature : public FFeatureObject
+	{
+		FPlanarAccelerationFeature(
+			const FName& InIdentifier,
+			const TSharedRef<FArrayMap>& InInstanceData,
+			const int32 InMaxInstanceNum,
+			const int32 InAccelerationNum,
+			const float InScale = Feature::DefaultScaleAcceleration,
+			const FVector InAxis0 = FVector::ForwardVector,
+			const FVector InAxis1 = FVector::RightVector);
+
+		virtual bool IsEncodable() const override final { return true; }
+		virtual bool IsDecodable() const override final { return true; }
+		virtual void Encode(const FIndexSet Instances) override final;
+		virtual void Decode(const FIndexSet Instances) override final;
+
+		FVector Axis0 = FVector::ForwardVector;
+		FVector Axis1 = FVector::RightVector;
+
+		TArrayMapHandle<2, FVector> AccelerationHandle;
+		TArrayMapHandle<1, FQuat> RelativeRotationHandle;
+	};
+
+	/**
+	* Feature for an array of accelerations
+	*/
+	struct LEARNING_API FAccelerationFeature : public FFeatureObject
+	{
+		FAccelerationFeature(
+			const FName& InIdentifier,
+			const TSharedRef<FArrayMap>& InInstanceData,
+			const int32 InMaxInstanceNum,
+			const int32 InAccelerationNum,
+			const float InScale = Feature::DefaultScaleAcceleration);
+
+		virtual bool IsEncodable() const override final { return true; }
+		virtual bool IsDecodable() const override final { return true; }
+		virtual void Encode(const FIndexSet Instances) override final;
+		virtual void Decode(const FIndexSet Instances) override final;
+
+		TArrayMapHandle<2, FVector> AccelerationHandle;
+		TArrayMapHandle<1, FQuat> RelativeRotationHandle;
+	};
+
+	//------------------------------------------------------------------
+
+	/**
+	* Feature for an array of angles
+	*/
+	struct LEARNING_API FAngleFeature : public FFeatureObject
+	{
+		FAngleFeature(
+			const FName& InIdentifier,
+			const TSharedRef<FArrayMap>& InInstanceData,
+			const int32 InMaxInstanceNum,
+			const int32 InAngleNum,
+			const float InScale = Feature::DefaultScaleDirection); // angles are encoded as directions
+
+		virtual bool IsEncodable() const override final { return true; }
+		virtual bool IsDecodable() const override final { return true; }
+		virtual void Encode(const FIndexSet Instances) override final;
+		virtual void Decode(const FIndexSet Instances) override final;
+
+		TArrayMapHandle<2, float> AngleHandle;
+		TArrayMapHandle<1, float> RelativeAngleHandle;
+	};
+
+	/**
+	* Feature for an array of rotations
+	*/
+	struct LEARNING_API FRotationFeature : public FFeatureObject
+	{
+		FRotationFeature(
+			const FName& InIdentifier,
+			const TSharedRef<FArrayMap>& InInstanceData,
+			const int32 InMaxInstanceNum,
+			const int32 InRotationNum,
+			const float InScale = Feature::DefaultScaleDirection); // rotations are encoded as directions
+
+		virtual bool IsEncodable() const override final { return true; }
+		virtual bool IsDecodable() const override final { return true; }
+		virtual void Encode(const FIndexSet Instances) override final;
+		virtual void Decode(const FIndexSet Instances) override final;
+
+		TArrayMapHandle<2, FQuat> RotationHandle;
+		TArrayMapHandle<1, FQuat> RelativeRotationHandle;
+	};
+
+	/**
+	* Feature for an array of rotation vectors
+	*/
+	struct LEARNING_API FRotationVectorFeature : public FFeatureObject
+	{
+		FRotationVectorFeature(
+			const FName& InIdentifier,
+			const TSharedRef<FArrayMap>& InInstanceData,
+			const int32 InMaxInstanceNum,
+			const int32 InRotationVectorNum,
+			const float InScale = Feature::DefaultScaleAngle);
+
+		virtual bool IsEncodable() const override final { return true; }
+		virtual bool IsDecodable() const override final { return true; }
+		virtual void Encode(const FIndexSet Instances) override final;
+		virtual void Decode(const FIndexSet Instances) override final;
+
+		TArrayMapHandle<2, FVector> RotationVectorsHandle;
+	};
+
+	//------------------------------------------------------------------
+
+	/**
+	* Feature for an array of scalar angular velocities
+	*/
+	struct LEARNING_API FScalarAngularVelocityFeature : public FFeatureObject
+	{
+		FScalarAngularVelocityFeature(
+			const FName& InIdentifier,
+			const TSharedRef<FArrayMap>& InInstanceData,
+			const int32 InMaxInstanceNum,
+			const int32 InAngularAccelerationNum,
+			const float InScale = Feature::DefaultScaleAngularVelocity);
+
+		virtual bool IsEncodable() const override final { return true; }
+		virtual bool IsDecodable() const override final { return true; }
+		virtual void Encode(const FIndexSet Instances) override final;
+		virtual void Decode(const FIndexSet Instances) override final;
+
+		TArrayMapHandle<2, float> AngularVelocityHandle;
+	};
+
+	/**
+	* Feature for an array of angular velocities
+	*/
+	struct LEARNING_API FAngularVelocityFeature : public FFeatureObject
+	{
+		FAngularVelocityFeature(
+			const FName& InIdentifier,
+			const TSharedRef<FArrayMap>& InInstanceData,
+			const int32 InMaxInstanceNum,
+			const int32 InAngularAccelerationNum,
+			const float InScale = Feature::DefaultScaleAngularVelocity);
+
+		virtual bool IsEncodable() const override final { return true; }
+		virtual bool IsDecodable() const override final { return true; }
+		virtual void Encode(const FIndexSet Instances) override final;
+		virtual void Decode(const FIndexSet Instances) override final;
+
+		TArrayMapHandle<2, FVector> AngularVelocityHandle;
+		TArrayMapHandle<1, FQuat> RelativeRotationHandle;
+	};
+
+	//------------------------------------------------------------------
+
+	/**
+	* Feature for an array of scalar angular accelerations
+	*/
+	struct LEARNING_API FScalarAngularAccelerationFeature : public FFeatureObject
+	{
+		FScalarAngularAccelerationFeature(
+			const FName& InIdentifier,
+			const TSharedRef<FArrayMap>& InInstanceData,
+			const int32 InMaxInstanceNum,
+			const int32 InAngularAccelerationNum,
+			const float InScale = Feature::DefaultScaleAngularAcceleration);
+
+		virtual bool IsEncodable() const override final { return true; }
+		virtual bool IsDecodable() const override final { return true; }
+		virtual void Encode(const FIndexSet Instances) override final;
+		virtual void Decode(const FIndexSet Instances) override final;
+
+		TArrayMapHandle<2, float> AngularAccelerationHandle;
+	};
+
+	/**
+	* Feature for an array of planar angular accelerations
+	*/
+	struct LEARNING_API FPlanarAngularAccelerationFeature : public FFeatureObject
+	{
+		FPlanarAngularAccelerationFeature(
+			const FName& InIdentifier,
+			const TSharedRef<FArrayMap>& InInstanceData,
+			const int32 InMaxInstanceNum,
+			const int32 InAngularAccelerationNum,
+			const float InScale = Feature::DefaultScaleAngularAcceleration,
+			const FVector InAxis = FVector::UpVector);
+
+		virtual bool IsEncodable() const override final { return true; }
+		virtual bool IsDecodable() const override final { return true; }
+		virtual void Encode(const FIndexSet Instances) override final;
+		virtual void Decode(const FIndexSet Instances) override final;
+
+		FVector Axis = FVector::UpVector;
+
+		TArrayMapHandle<2, FVector> AngularAccelerationHandle;
+		TArrayMapHandle<1, FQuat> RelativeRotationHandle;
+	};
+
+	/**
+	* Feature for an array of angular accelerations
+	*/
+	struct LEARNING_API FAngularAccelerationFeature : public FFeatureObject
+	{
+		FAngularAccelerationFeature(
+			const FName& InIdentifier,
+			const TSharedRef<FArrayMap>& InInstanceData,
+			const int32 InMaxInstanceNum,
+			const int32 InAngularAccelerationNum,
+			const float InScale = Feature::DefaultScaleAngularAcceleration);
+
+		virtual bool IsEncodable() const override final { return true; }
+		virtual bool IsDecodable() const override final { return true; }
+		virtual void Encode(const FIndexSet Instances) override final;
+		virtual void Decode(const FIndexSet Instances) override final;
+
+		TArrayMapHandle<2, FVector> AngularAccelerationHandle;
+		TArrayMapHandle<1, FQuat> RelativeRotationHandle;
+	};
+
+	//------------------------------------------------------------------
+
+	/**
+	* Feature for an array of scales
+	*/
+	struct LEARNING_API FScaleFeature : public FFeatureObject
+	{
+		FScaleFeature(
+			const FName& InIdentifier,
+			const TSharedRef<FArrayMap>& InInstanceData,
+			const int32 InMaxInstanceNum,
+			const int32 InScaleNum,
+			const float InScale = Feature::DefaultScaleScale);
+
+		virtual bool IsEncodable() const override final { return true; }
+		virtual bool IsDecodable() const override final { return true; }
+		virtual void Encode(const FIndexSet Instances) override final;
+		virtual void Decode(const FIndexSet Instances) override final;
+
+		TArrayMapHandle<2, float> ScaleHandle;
+		TArrayMapHandle<1, float> RelativeScaleHandle;
+	};
+
+	//------------------------------------------------------------------
+
+}
diff --git a/Engine/Plugins/Experimental/LearningAgents/Source/Learning/Public/LearningFunctionObject.h b/Engine/Plugins/Experimental/LearningAgents/Source/Learning/Public/LearningFunctionObject.h
new file mode 100644
--- /dev/null
+++ b/Engine/Plugins/Experimental/LearningAgents/Source/Learning/Public/LearningFunctionObject.h
@@ -0,0 +1,198 @@
+// Copyright Epic Games, Inc. All Rights Reserved.
+
+#pragma once
+
+#include "LearningArrayMap.h"
+
+#include "Math/Quat.h"
+#include "Math/Transform.h"
+#include "Math/Vector.h"
+#include "Templates/SharedPointer.h"
+
+namespace UE::Learning
+{
+	/**
+	* Base class for an object which outputs to some arrays stored in
+	* a `FArrayMap`, using other arrays in the `FArrayMap` as input. 
+	* Here, all data is assumed to be stored in a `FArrayMap` object 
+	* to make the processing of multiple instances efficient.
+	*/
+	struct LEARNING_API FFunctionObject
+	{
+		FFunctionObject(const TSharedRef<FArrayMap>& InInstanceData);
+		virtual ~FFunctionObject() {}
+
+		virtual void Evaluate(const FIndexSet Instances) = 0;
+
+		TSharedRef<FArrayMap> InstanceData;
+	};
+
+	/**
+	* Executes another set of function objects in sequence
+	*/
+	struct LEARNING_API FSequentialFunction : public FFunctionObject
+	{
+		FSequentialFunction(
+			const TLearningArrayView<1, const TSharedRef<FFunctionObject>> InFunctions,
+			const TSharedRef<FArrayMap>& InInstanceData);
+
+		virtual void Evaluate(const FIndexSet Instances) override final;
+
+		TLearningArray<1, TSharedRef<FFunctionObject>, TInlineAllocator<32>> Functions;
+	};
+
+	/**
+	* Copies vector data from one array to another
+	*/
+	struct LEARNING_API FCopyVectorsFunction : public FFunctionObject
+	{
+		FCopyVectorsFunction(
+			const FName& InIdentifier,
+			const TSharedRef<FArrayMap>& InInstanceData,
+			const int32 InMaxInstanceNum,
+			const int32 InDimensionsNum);
+
+		virtual void Evaluate(const FIndexSet Instances) override final;
+
+		TArrayMapHandle<2, float> InputHandle;
+		TArrayMapHandle<2, float> OutputHandle;
+	};
+
+	/**
+	* Computes rotations given transforms
+	*/
+	struct LEARNING_API FExtractRotationsFromTransformsFunction : public FFunctionObject
+	{
+		FExtractRotationsFromTransformsFunction(
+			const FName& InIdentifier,
+			const TSharedRef<FArrayMap>& InInstanceData,
+			const int32 InMaxInstanceNum);
+
+		virtual void Evaluate(const FIndexSet Instances) override final;
+
+		TArrayMapHandle<1, FTransform> TransformHandle;
+		TArrayMapHandle<1, FQuat> RotationHandle;
+	};
+
+	/**
+	* Computes positions and rotations given transforms
+	*/
+	struct LEARNING_API FExtractPositionsAndRotationsFromTransformsFunction : public FFunctionObject
+	{
+		FExtractPositionsAndRotationsFromTransformsFunction(
+			const FName& InIdentifier,
+			const TSharedRef<FArrayMap>& InInstanceData,
+			const int32 InMaxInstanceNum);
+
+		virtual void Evaluate(const FIndexSet Instances) override final;
+
+		TArrayMapHandle<1, FTransform> TransformHandle;
+		TArrayMapHandle<1, FQuat> RotationHandle;
+		TArrayMapHandle<1, FVector> PositionHandle;
+	};
+
+	/**
+	* Samples scalar values from a uniform distribution
+	*/
+	struct LEARNING_API FRandomUniformFunction : public FFunctionObject
+	{
+		FRandomUniformFunction(
+			const FName& InIdentifier,
+			const TSharedRef<FArrayMap>& InInstanceData,
+			const int32 InMaxInstanceNum,
+			const int32 InValueNum,
+			const uint32 InSeed,
+			const float InMin = 0.0f,
+			const float InMax = 1.0f);
+
+		virtual void Evaluate(const FIndexSet Instances) override final;
+
+		TArrayMapHandle<1, uint32> SeedHandle;
+		TArrayMapHandle<1, float> MinHandle;
+		TArrayMapHandle<1, float> MaxHandle;
+		TArrayMapHandle<2, float> ValueHandle;
+	};
+
+	/**
+	* Samples values from a planar clipped Gaussian distribution.
+	*/
+	struct LEARNING_API FRandomPlanarClippedGaussianFunction : public FFunctionObject
+	{
+		FRandomPlanarClippedGaussianFunction(
+			const FName& InIdentifier,
+			const TSharedRef<FArrayMap>& InInstanceData,
+			const int32 InMaxInstanceNum,
+			const int32 InValueNum,
+			const uint32 InSeed,
+			const float InMean = 0.0f,
+			const float InStd = 1.0f,
+			const float InClip = 10.0f,
+			const FVector InAxis0 = FVector::ForwardVector,
+			const FVector InAxis1 = FVector::RightVector);
+
+		virtual void Evaluate(const FIndexSet Instances) override final;
+
+		FVector Axis0 = FVector::ForwardVector;
+		FVector Axis1 = FVector::RightVector;
+
+		TArrayMapHandle<1, uint32> SeedHandle;
+		TArrayMapHandle<1, float> MeanHandle;
+		TArrayMapHandle<1, float> StdHandle;
+		TArrayMapHandle<1, float> ClipHandle;
+		TArrayMapHandle<2, FVector> ValueHandle;
+	};
+
+	/**
+	* Samples planar directions.
+	*/
+	struct LEARNING_API FRandomPlanarDirectionFunction : public FFunctionObject
+	{
+		FRandomPlanarDirectionFunction(
+			const FName& InIdentifier,
+			const TSharedRef<FArrayMap>& InInstanceData,
+			const int32 InMaxInstanceNum,
+			const int32 InDirectionNum,
+			const uint32 InSeed,
+			const FVector InAxis0 = FVector::ForwardVector,
+			const FVector InAxis1 = FVector::RightVector);
+
+		virtual void Evaluate(const FIndexSet Instances) override final;
+
+		FVector Axis0 = FVector::ForwardVector;
+		FVector Axis1 = FVector::RightVector;
+
+		TArrayMapHandle<1, uint32> SeedHandle;
+		TArrayMapHandle<2, FVector> DirectionHandle;
+	};
+
+	/**
+	* Samples planar directions with corresponding velocities.
+	*/
+	struct LEARNING_API FRandomPlanarDirectionVelocityFunction : public FFunctionObject
+	{
+		FRandomPlanarDirectionVelocityFunction(
+			const FName& InIdentifier,
+			const TSharedRef<FArrayMap>& InInstanceData,
+			const int32 InMaxInstanceNum,
+			const int32 InDirectionVelocityNum,
+			const uint32 InSeed,
+			const float InVelocityScale = 1.0f,
+			const FVector InAxis0 = FVector::ForwardVector,
+			const FVector InAxis1 = FVector::RightVector);
+
+		virtual void Evaluate(const FIndexSet Instances) override final;
+
+		FVector Axis0 = FVector::ForwardVector;
+		FVector Axis1 = FVector::RightVector;
+
+		TArrayMapHandle<1, uint32> SeedHandle;
+		TArrayMapHandle<1, float> VelocityScaleHandle;
+		TArrayMapHandle<2, FVector> DirectionHandle;
+		TArrayMapHandle<2, FVector> VelocityHandle;
+	};
+
+}
+
+
+
+
diff --git a/Engine/Plugins/Experimental/LearningAgents/Source/Learning/Public/LearningNeuralNetwork.h b/Engine/Plugins/Experimental/LearningAgents/Source/Learning/Public/LearningNeuralNetwork.h
--- a/Engine/Plugins/Experimental/LearningAgents/Source/Learning/Public/LearningNeuralNetwork.h
+++ b/Engine/Plugins/Experimental/LearningAgents/Source/Learning/Public/LearningNeuralNetwork.h
@@ -6,114 +6,10 @@
 
 #include "Templates/SharedPointer.h"
 
-#include "LearningNeuralNetwork.generated.h"
-
-class UNNEModelData;
-
-namespace UE::NNE
-{
-	class IModelCPU;
-	class IModelInstanceCPU;
-}
-
 namespace UE::Learning
 {
-	struct FNeuralNetwork;
-	struct FNeuralNetworkInference;
-}
-
-/**
- * Neural Network Data Object
- * 
- * This is the UObject which contains the actual data used by a Neural Network. It stores the raw FileData used to construct the network using NNE,
- * as well as the input and output sizes and a compatibility hash that can be used to quickly check if two networks may be compatible in terms of 
- * inputs and outputs.
- * 
- * Internally this also stores the various things required to map between NNE style inference and the style of inference used in Learning via 
- * FNeuralNetwork and FNeuralNetworkInference.
- */
-UCLASS(BlueprintType)
-class LEARNING_API ULearningNeuralNetworkData : public UObject
-{
-	GENERATED_BODY()
-
-public:
-
-	// Get the FNeuralNetwork object that can be used to do inference.
-	TSharedPtr<UE::Learning::FNeuralNetwork>& GetNetwork();
-
-	/**
-	 * Initialize the ULearningNeuralNetworkData object given some input and output sizes as well as a compatibility hash and FileData.
-	 * 
-	 * @param InInputSize			Network Input Size
-	 * @param InOutputSize			Network Output Size
-	 * @param InCompatibilityHash	Compatibility Hash
-	 * @param InFileData			Network FileData
-	 */
-	void Init(
-		const int32 InInputSize, 
-		const int32 InOutputSize, 
-		const int32 InCompatibilityHash, 
-		const TArrayView<const uint8> InFileData);
-
-	// Initialize this network from another ULearningNeuralNetworkData object
-	void InitFrom(const ULearningNeuralNetworkData* OtherNetworkData);
-
-	// Load a snapshot of this network from the given array of bytes
-	bool LoadFromSnapshot(const TArrayView<const uint8> InBytes);
-
-	// Save a snapshot of this network to the given array of bytes
-	void SaveToSnapshot(TArrayView<uint8> OutBytes) const;
-
-	// Number of bytes required to save or load a snapshot of this network
-	int32 GetSnapshotByteNum() const;
-
-	// If this network is empty or not
-	bool IsEmpty() const;
-
-	// Gets the network input size
-	int32 GetInputSize() const;
-	
-	// Gets the network output size
-	int32 GetOutputSize() const;
+	struct INeuralNetworkInference;
 	
-	// Gets the compatibility hash
-	int32 GetCompatibilityHash() const;
-
-private:
-
-	// Calls UpdateNetwork to update the internal representation.
-	virtual void PostLoad() override final;
-
-	// Uploads the FileData to NNE and updates the internal in-memory representation of the network used for inference.
-	void UpdateNetwork();
-
-	// Size of the inputs expected by this network
-	UPROPERTY(VisibleAnywhere, Category = "Network Properties")
-	int32 InputSize = 0;
-
-	// Size of the output produced by this network
-	UPROPERTY(VisibleAnywhere, Category = "Network Properties")
-	int32 OutputSize = 0;
-
-	// Compatibility hash used for testing if inputs to one network are compatible with another
-	UPROPERTY(VisibleAnywhere, Category = "Network Properties")
-	int32 CompatibilityHash = 0;
-
-	// File data used by NNE
-	UPROPERTY()
-	TArray<uint8> FileData;
-
-	// Model Data used by NNE
-	UPROPERTY()
-	TObjectPtr<UNNEModelData> ModelData;
-
-	// Internal in-memory network representation
-	TSharedPtr<UE::Learning::FNeuralNetwork> Network;
-};
-
-namespace UE::Learning
-{
 	/**
 	* Settings object for a neural network instance
 	*/
@@ -127,116 +23,117 @@ namespace UE::Learning
 	};
 
 	/**
-	* Neural Network Object
+	* Interface for a Neural Network
 	*/
-	struct LEARNING_API FNeuralNetwork
+	struct INeuralNetwork
 	{
+		virtual ~INeuralNetwork() {}
+
+		/** Deserialize the network from the raw bytes starting at the given offset. Returns true if successful. */
+		virtual bool DeserializeFromBytes(int32& InOutOffset, const TLearningArrayView<1, const uint8> RawBytes) = 0;
+
+		/** Serialize the network to the raw bytes starting at the given offset. */
+		virtual void SerializeToBytes(int32& InOutOffset, TLearningArrayView<1, uint8> OutRawBytes) const = 0;
+
+		/** Get the number of bytes required to serialize this network. */
+		virtual int32 GetSerializationByteNum() const = 0;
+
 		/** Create a new inference object for this network with the given maximum batch size and inference settings. */
-		TSharedRef<FNeuralNetworkInference> CreateInferenceObject(
+		virtual TSharedRef<INeuralNetworkInference> CreateInferenceObject(
 			const int32 MaxBatchSize,
-			const FNeuralNetworkInferenceSettings& Settings = FNeuralNetworkInferenceSettings());
+			const FNeuralNetworkInferenceSettings& Settings = FNeuralNetworkInferenceSettings()) = 0;
 
-		bool IsEmpty() const;
-		int32 GetInputSize() const;
-		int32 GetOutputSize() const;
+		/** Get the number of inputs expected by this network */
+		virtual int32 GetInputNum() const = 0;
 
-		void UpdateModel(const TSharedPtr<NNE::IModelCPU>& InModel, const int32 InInputSize, const int32 InOutputSize);
+		/** Get the number of outputs expected by this network */
+		virtual int32 GetOutputNum() const = 0;
 
-	private:
-
-		int32 InputSize = 0;
-		int32 OutputSize = 0;
-		TSharedPtr<NNE::IModelCPU> Model;
-		TArray<TWeakPtr<FNeuralNetworkInference>, TInlineAllocator<64>> InferenceObjects;
+		/** Get the name of the python class used to represent this type of network during training. */
+		virtual const TCHAR* GetPythonClassName() const = 0;
 	};
 
 	/**
-	* Neural Network Inference Object
+	* Interface for a Neural Network Inference
 	*/
-	struct LEARNING_API FNeuralNetworkInference
+	struct INeuralNetworkInference
 	{
-		/**
-		 * Constructs a new network inference object.
-		 * 
-		 * @param InModel			NNE Model
-		 * @param MaxBatchSize		Maximum batch size
-		 * @param InInputSize		Network input size
-		 * @param InOutputSize		Network output size
-		 * @param InSettings		Inference settings
-		 */
-		FNeuralNetworkInference(
-			UE::NNE::IModelCPU& InModel,
-			const int32 MaxBatchSize,
-			const int32 InInputSize,
-			const int32 InOutputSize,
-			const FNeuralNetworkInferenceSettings& InSettings = FNeuralNetworkInferenceSettings());
-
-		/**
-		 * Evaluate this network for the given instances.
-		 * 
-		 * Note: this function takes a lock, so although it can be called from multiple threads it will not benefit from multi-threading. This is 
-		 * because to use the NNE interface all the instances that need to be evaluated must be gathered and scattered into internal buffers. To
-		 * compensate for this, this call will do multi-threading internally based on the inference settings used.
-		 * 
-		 * @param Output		The Output Buffer of shape (MaxBatchSize, InOutputSize)
-		 * @param Input			The Input Buffer of shape (MaxBatchSize, InInputSize)
-		 * @param Instances		The batch instances to evaluate
-		 */
-		void Evaluate(
+		virtual ~INeuralNetworkInference() {}
+
+		virtual void Evaluate(
 			TLearningArrayView<2, float> Output,
 			const TLearningArrayView<2, const float> Input,
-			const FIndexSet Instances);
+			const FIndexSet Instances) = 0;
+	};
 
-		// This function will re-build the internal Model Instances used for multi-threading. It should be called whenever the given Model is updated. 
-		void ReloadModelInstances(NNE::IModelCPU& Model);
+	/**
+	* Activation Function for use in basic MLP Neural Network
+	*/
+	enum class EActivationFunction : uint8
+	{
+		// ReLU Activation - Fast to train and evaluate but occasionally causes gradient collapse and untrainable networks.
+		ReLU = 0,
 
-	private:
+		// ELU Activation - Generally performs better than ReLU and is not prone to gradient collapse but slower to evaluate.
+		ELU = 1,
 
-		FNeuralNetworkInferenceSettings Settings;
-		FRWLock EvaluationLock;
-		TLearningArray<2, float> InputBuffer;
-		TLearningArray<2, float> OutputBuffer;
-		TArray<TSharedPtr<NNE::IModelInstanceCPU>, TInlineAllocator<64>> ModelInstances;
+		// TanH Activation - Smooth activation function that is slower to train and evaluate but sometimes more stable for certain tasks.
+		TanH = 2,
 	};
 
 	/**
-	* Neural-network based function object.
+	* Basic Implementation for a MLP Neural Network
 	*/
-	struct LEARNING_API FNeuralNetworkFunction
+	struct LEARNING_API FNeuralNetworkMLP : public INeuralNetwork
 	{
-		/**
-		 * Constructs a Neural Network Function from the given Neural Network
-		 *
-		 * @param InMaxInstanceNum		Maximum number of instances to evaluate for
-		 * @param InNeuralNetwork		Neural network to use
-		 * @param InInferenceSettings	Inference settings
-		 */
-		FNeuralNetworkFunction(
-			const int32 InMaxInstanceNum,
-			const TSharedPtr<FNeuralNetwork>& InNeuralNetwork,
-			const FNeuralNetworkInferenceSettings& InInferenceSettings = FNeuralNetworkInferenceSettings());
-
-		/**
-		 * Evaluate this network for the given instances.
-		 *
-		 * @param Output		The Output Buffer of shape (MaxInstanceNum, InOutputSize)
-		 * @param Input			The Input Buffer of shape (MaxInstanceNum, InInputSize)
-		 * @param Instances		The instances to evaluate
-		 */
-		void Evaluate(
-			TLearningArrayView<2, float> Output,
-			const TLearningArrayView<2, const float> Input,
-			const FIndexSet Instances);
+		friend struct FNeuralNetworkInstanceMLP;
+
+		//~ Begin INeuralNetwork Interface
+		virtual bool DeserializeFromBytes(int32& InOutOffset, const TLearningArrayView<1, const uint8> RawBytes) override final;
+		virtual void SerializeToBytes(int32& InOutOffset, TLearningArrayView<1, uint8> OutRawBytes) const override final;
+		virtual int32 GetSerializationByteNum() const override final;
+		virtual TSharedRef<INeuralNetworkInference> CreateInferenceObject(
+			const int32 MaxBatchSize,
+			const FNeuralNetworkInferenceSettings& Settings = FNeuralNetworkInferenceSettings()) override final;
+		virtual int32 GetInputNum() const override final;
+		virtual int32 GetOutputNum() const override final;
+		virtual const TCHAR* GetPythonClassName() const override final { return TEXT("NeuralNetworkMLP"); };
+		//~ End INeuralNetwork Interface
+
+		int32 GetHiddenNum() const;
+		int32 GetLayerNum() const;
+
+		void Resize(
+			const int32 InputNum,
+			const int32 OutputNum,
+			const int32 HiddenNum,
+			const int32 LayerNum);
+
+		EActivationFunction ActivationFunction = EActivationFunction::ELU;
+		TArray<TLearningArray<2, float>, TInlineAllocator<16>> Weights;
+		TArray<TLearningArray<1, float>, TInlineAllocator<16>> Biases;
+	};
 
-		/** Sets the NeuralNetwork and re-creates the NeuralNetworkInference object. */
-		void UpdateNeuralNetwork(const TSharedPtr<FNeuralNetwork>& NewNeuralNetwork);
+	/**
+	* Basic Implementation for a MLP Neural Network Inference
+	*/
+	struct LEARNING_API FNeuralNetworkInferenceMLP : public INeuralNetworkInference
+	{
+		FNeuralNetworkInferenceMLP(
+			const FNeuralNetworkMLP& InNetwork,
+			const int32 MaxBatchSize,
+			const FNeuralNetworkInferenceSettings& InSettings);
 
-	private:
+		//~ Begin INeuralNetworkInference Interface
+		virtual void Evaluate(
+			TLearningArrayView<2, float> Output,
+			const TLearningArrayView<2, const float> Input,
+			const FIndexSet Instances) override final;
+		//~ End INeuralNetworkInference Interface
 
-		int32 MaxInstanceNum = 0;
-		TSharedPtr<FNeuralNetwork> NeuralNetwork;
-		TSharedPtr<FNeuralNetworkInference> NeuralNetworkInference;
-		FNeuralNetworkInferenceSettings InferenceSettings;
+		const FNeuralNetworkMLP& Network;
+		FNeuralNetworkInferenceSettings Settings;
+		TLearningArray<2, float> FrontBuffer;
+		TLearningArray<2, float> BackBuffer;
 	};
-
 }
diff --git a/Engine/Plugins/Experimental/LearningAgents/Source/Learning/Public/LearningNeuralNetworkObject.h b/Engine/Plugins/Experimental/LearningAgents/Source/Learning/Public/LearningNeuralNetworkObject.h
new file mode 100644
--- /dev/null
+++ b/Engine/Plugins/Experimental/LearningAgents/Source/Learning/Public/LearningNeuralNetworkObject.h
@@ -0,0 +1,107 @@
+// Copyright Epic Games, Inc. All Rights Reserved.
+
+#pragma once
+
+#include "LearningArrayMap.h"
+#include "LearningFunctionObject.h"
+#include "LearningNeuralNetwork.h" // Included for FNeuralNetworkInferenceSettings
+
+#include "Templates/SharedPointer.h"
+
+namespace UE::Learning
+{
+	/**
+	* Settings object for a neural network based policy
+	*/
+	struct LEARNING_API FNeuralNetworkPolicyFunctionSettings
+	{
+		// Minimum amount of action noise to allow
+		float ActionNoiseMin = 0.0f;
+
+		// Maximum amount of action noise to allow
+		float ActionNoiseMax = 0.0f;
+
+		// Overall scale of the action noise
+		float ActionNoiseScale = 1.0f;
+	};
+
+	/**
+	* Neural-network based policy object. Stores various settings and intermediate
+	* storage required to evaluate the given network for the provided number of instances.
+	*/
+	struct LEARNING_API FNeuralNetworkPolicyFunction : public FFunctionObject
+	{
+		FNeuralNetworkPolicyFunction(
+			const FName& InIdentifier,
+			const TSharedRef<FArrayMap>& InInstanceData,
+			const int32 InMaxInstanceNum,
+			const int32 InObservationNum,
+			const int32 InActionNum,
+			const int32 InMemoryStateNum,
+			const TSharedPtr<INeuralNetwork>& InNeuralNetwork,
+			const uint32 InSeed,
+			const FNeuralNetworkInferenceSettings& InInferenceSettings = FNeuralNetworkInferenceSettings(),
+			const FNeuralNetworkPolicyFunctionSettings& InSettings = FNeuralNetworkPolicyFunctionSettings());
+
+		virtual void Evaluate(const FIndexSet Instances) override final;
+
+		/** Sets the NeuralNetwork and re-creates the NeuralNetworkInference object */
+		void UpdateNeuralNetwork(const TSharedPtr<INeuralNetwork>& NewNeuralNetwork);
+
+		int32 MaxInstanceNum = 0;
+		int32 ObservationNum = 0;
+		int32 ActionNum = 0;
+		int32 MemoryStateNum = 0;
+		TSharedPtr<INeuralNetwork> NeuralNetwork;
+		TSharedPtr<INeuralNetworkInference> NeuralNetworkInference;
+		FNeuralNetworkInferenceSettings InferenceSettings;
+		FNeuralNetworkPolicyFunctionSettings Settings;
+
+		TArrayMapHandle<1, uint32> SeedHandle;
+
+		TArrayMapHandle<2, float> InputObservationHandle;
+		TArrayMapHandle<2, float> InputMemoryStateHandle;
+		TArrayMapHandle<2, float> InputNetworkHandle;
+
+		TArrayMapHandle<2, float> OutputNetworkHandle;
+		TArrayMapHandle<2, float> OutputMemoryStateHandle;
+		TArrayMapHandle<2, float> OutputActionHandle;
+		TArrayMapHandle<2, float> OutputActionMeanHandle;
+		TArrayMapHandle<2, float> OutputActionStdHandle;
+		
+		TArrayMapHandle<1, float> ActionNoiseScaleHandle;
+	};
+
+	/**
+	* Neural-network based critic object. Stores various settings and intermediate
+	* storage required to evaluate the given network for the provided number of instances.
+	*/
+	struct LEARNING_API FNeuralNetworkCriticFunction : public FFunctionObject
+	{
+		FNeuralNetworkCriticFunction(
+			const FName& InIdentifier,
+			const TSharedRef<FArrayMap>& InInstanceData,
+			const int32 InMaxInstanceNum,
+			const int32 InObservationNum,
+			const int32 InMemoryStateNum,
+			const TSharedPtr<INeuralNetwork>& InNeuralNetwork,
+			const FNeuralNetworkInferenceSettings& InInferenceSettings = FNeuralNetworkInferenceSettings());
+
+		virtual void Evaluate(const FIndexSet Instances) override final;
+
+		/** Sets the NeuralNetwork and re-creates the NeuralNetworkInference object */
+		void UpdateNeuralNetwork(const TSharedPtr<INeuralNetwork>& NewNeuralNetwork);
+
+		int32 MaxInstanceNum = 0;
+		int32 ObservationNum = 0;
+		int32 MemoryStateNum = 0;
+		TSharedPtr<INeuralNetwork> NeuralNetwork;
+		TSharedPtr<INeuralNetworkInference> NeuralNetworkInference;
+		FNeuralNetworkInferenceSettings InferenceSettings;
+
+		TArrayMapHandle<2, float> InputObservationHandle;
+		TArrayMapHandle<2, float> InputMemoryStateHandle;
+		TArrayMapHandle<2, float> InputNetworkHandle;
+		TArrayMapHandle<1, float> OutputHandle;
+	};
+}
diff --git a/Engine/Plugins/Experimental/LearningAgents/Source/Learning/Public/LearningRandom.h b/Engine/Plugins/Experimental/LearningAgents/Source/Learning/Public/LearningRandom.h
--- a/Engine/Plugins/Experimental/LearningAgents/Source/Learning/Public/LearningRandom.h
+++ b/Engine/Plugins/Experimental/LearningAgents/Source/Learning/Public/LearningRandom.h
@@ -118,11 +118,6 @@ namespace UE::Learning::Random
 		TLearningArrayView<1, uint32> Output,
 		const uint32 State);
 
-	LEARNING_API void IntArray(
-		TLearningArrayView<1, uint32> Output,
-		const uint32 State,
-		const FIndexSet Indices);
-
 	LEARNING_API void FloatArray(
 		TLearningArrayView<1, float> Output,
 		const uint32 State);
@@ -154,25 +149,6 @@ namespace UE::Learning::Random
 		const FVector Axis0 = FVector::ForwardVector,
 		const FVector Axis1 = FVector::RightVector);
 
-	LEARNING_API void DistributionIndependantNormal(
-		TLearningArrayView<1, float> Output,
-		const uint32 State,
-		const TLearningArrayView<1, const float> Mean,
-		const TLearningArrayView<1, const float> LogStd,
-		const float Scale);
-
-	LEARNING_API void DistributionMultinoulli(
-		TLearningArrayView<1, float> Output,
-		const uint32 State,
-		const TLearningArrayView<1, const float> Logits,
-		const float Scale);
-
-	LEARNING_API void DistributionBernoulli(
-		TLearningArrayView<1, float> Output,
-		const uint32 State,
-		const TLearningArrayView<1, const float> Logits,
-		const float Scale);
-
 	/////////////////////////////////////////////////////////
 	//
 	// Here we provide a more normal interface which updates
@@ -235,11 +211,6 @@ namespace UE::Learning::Random
 		TLearningArrayView<1, uint32> Output,
 		uint32& State);
 
-	LEARNING_API void SampleIntArray(
-		TLearningArrayView<1, uint32> Output,
-		uint32& State,
-		const FIndexSet Indices);
-
 	LEARNING_API void SampleFloatArray(
 		TLearningArrayView<1, float> Output,
 		uint32& State);
@@ -256,6 +227,10 @@ namespace UE::Learning::Random
 		const float Mean = 0.0f,
 		const float Std = 1.0f);
 
+	LEARNING_API void ResampleStateArray(
+		TLearningArrayView<1, uint32> InOutStates,
+		const FIndexSet Indices);
+
 	LEARNING_API void SamplePlanarClippedGaussianArray(
 		TLearningArrayView<1, FVector> Output,
 		uint32& State,
@@ -270,23 +245,4 @@ namespace UE::Learning::Random
 		uint32& State,
 		const FVector Axis0 = FVector::ForwardVector,
 		const FVector Axis1 = FVector::RightVector);
-
-	LEARNING_API void SampleDistributionIndependantNormal(
-		TLearningArrayView<1, float> Output,
-		uint32& State,
-		const TLearningArrayView<1, const float> Mean,
-		const TLearningArrayView<1, const float> LogStd,
-		const float Scale);
-
-	LEARNING_API void SampleDistributionMultinoulli(
-		TLearningArrayView<1, float> Output,
-		uint32& State,
-		const TLearningArrayView<1, const float> Logits,
-		const float Scale);
-
-	LEARNING_API void SampleDistributionBernoulli(
-		TLearningArrayView<1, float> Output,
-		uint32& State,
-		const TLearningArrayView<1, const float> Logits,
-		const float Scale);
 }
\ No newline at end of file
diff --git a/Engine/Plugins/Experimental/LearningAgents/Source/Learning/Public/LearningRecorderObject.h b/Engine/Plugins/Experimental/LearningAgents/Source/Learning/Public/LearningRecorderObject.h
new file mode 100644
--- /dev/null
+++ b/Engine/Plugins/Experimental/LearningAgents/Source/Learning/Public/LearningRecorderObject.h
@@ -0,0 +1,123 @@
+// Copyright Epic Games, Inc. All Rights Reserved.
+
+#pragma once
+
+#include "LearningArrayMap.h"
+
+#include "Templates/SharedPointer.h"
+
+namespace UE::Learning
+{
+	namespace RecorderObjectPrivate
+	{
+		template<uint8 DimNum, typename ElementType>
+		void Copy(
+			TLearningArrayView<1, int32> InOutStepNums,
+			const FArrayMap& SrcArrayMap,
+			const FArrayMap& LogArrayMap,
+			const FArrayMapHandle SrcHandle,
+			const FArrayMapHandle LogHandle,
+			const int32 MaxStepNum,
+			const FIndexSet Instances)
+		{
+			const TLearningArrayView<DimNum, const ElementType> SrcView = SrcArrayMap.ConstView<DimNum, ElementType>(SrcHandle);
+			TLearningArrayView<DimNum + 1, ElementType> LogView = LogArrayMap.View<DimNum + 1, ElementType>(LogHandle);
+
+			for (const int32 InstanceIdx : Instances)
+			{
+				const int32 StepNum = InOutStepNums[InstanceIdx];
+
+				if (StepNum < MaxStepNum)
+				{
+					if constexpr (DimNum == 1)
+					{
+						LogView[StepNum][InstanceIdx] = SrcView[InstanceIdx];
+					}
+					else
+					{
+						Array::Copy(LogView[StepNum][InstanceIdx], SrcView[InstanceIdx]);
+					}
+
+					InOutStepNums[InstanceIdx]++;
+				}
+			}
+		}
+	}
+
+	/**
+	* The `RecorderObject` can be used to record the values of certain arrays in an `FArrayMap`. 
+	* This can be useful for debugging since we can keep track of certain arrays over time and 
+	* compute statistics or visualize them.
+	* 
+	* This object assumes the first dimension on all arrays corresponds to `Instances`, as it works by
+	* only record updates and resets for certain `Instances` given by an `FIndexSet`. This makes it 
+	* easier to use in conjunction with the other objects that use `FArrayMap` such as `FRewardObject`.
+	*/
+	struct LEARNING_API FRecorderObject
+	{
+		FRecorderObject(
+			const int32 InMaxInstanceNum,
+			const int32 InMaxStepNum,
+			const TSharedRef<FArrayMap>& InSrcArrayMap,
+			const TSharedRef<FArrayMap>& InLogArrayMap);
+
+		~FRecorderObject();
+
+		/**
+		* Record an array from the provided `SrcArrayMap` in the `LogArrayMap`
+		*/
+		template<uint8 DimNum, typename ElementType>
+		void Add(const FArrayMapKey Key)
+		{
+			const FArrayMapHandle SrcHandle = SrcArrayMap->Lookup(Key);
+
+			const TLearningArrayView<DimNum, const ElementType> SrcView = SrcArrayMap->ConstView<DimNum, ElementType>(SrcHandle);
+			UE_LEARNING_ARRAY_SHAPE_CHECK(SrcView.template Num<0>() == MaxInstanceNum);
+
+			TMultiArrayShape<DimNum + 1> LogShape;
+			LogShape[0] = MaxStepNum;
+			for (int32 DimIdx = 0; DimIdx < DimNum; DimIdx++)
+			{
+				LogShape[1 + DimIdx] = SrcView.Num(DimIdx);
+			}
+
+			const FArrayMapHandle LogHandle = LogArrayMap->Add<DimNum + 1, ElementType>(Key, LogShape);
+
+			SrcHandles.Emplace(SrcHandle);
+			LogHandles.Emplace(LogHandle);
+			CopyFunctions.Emplace(RecorderObjectPrivate::Copy<DimNum, ElementType>);
+		}
+
+		TLearningArrayView<1, const int32> GetStepNums() const { return StepNums; }
+
+		/**
+		* Reset all recorded arrays
+		*/
+		void Reset(const FIndexSet Instances);
+
+		/**
+		* Update all recorded arrays (record their current value)
+		*/
+		void Update(const FIndexSet Instances);
+
+	private:
+
+		int32 MaxInstanceNum = 0;
+		int32 MaxStepNum = 0;
+		TSharedRef<FArrayMap> SrcArrayMap;
+		TSharedRef<FArrayMap> LogArrayMap;
+
+		TLearningArray<1, int32> StepNums;
+		TArray<FArrayMapHandle, TInlineAllocator<32>> SrcHandles;
+		TArray<FArrayMapHandle, TInlineAllocator<32>> LogHandles;
+
+		TArray<TFunctionRef<void(
+			TLearningArrayView<1, int32> InOutStepNums,
+			const FArrayMap& SrcArrayMap,
+			const FArrayMap& LogArrayMap,
+			const FArrayMapHandle SrcHandle,
+			const FArrayMapHandle LogHandle,
+			const int32 MaxStepNum,
+			const FIndexSet Instances)>, TInlineAllocator<32>> CopyFunctions;
+	};
+}
