[Backout] - CL30154408

diff --git a/Engine/Plugins/Experimental/LearningAgents/Source/Learning/Public/LearningRewardObject.h b/Engine/Plugins/Experimental/LearningAgents/Source/Learning/Public/LearningRewardObject.h
new file mode 100644
--- /dev/null
+++ b/Engine/Plugins/Experimental/LearningAgents/Source/Learning/Public/LearningRewardObject.h
@@ -0,0 +1,605 @@
+// Copyright Epic Games, Inc. All Rights Reserved.
+
+#pragma once
+
+#include "LearningArrayMap.h"
+
+#include "Math/Quat.h"
+#include "Math/Vector.h"
+#include "Templates/SharedPointer.h"
+
+namespace UE::Learning
+{
+	namespace Reward
+	{
+		static constexpr float DefaultScale = 1.0f;
+		static constexpr float DefaultScalePosition = 100.0f; // cm
+		static constexpr float DefaultScaleVelocity = 200.0f; // cm/s
+		static constexpr float DefaultScaleAngle = UE_PI / 2.0f; // rad
+		static constexpr float DefaultScaleAngularVelocity = UE_PI; // rad/s
+		static constexpr float DefaultScaleSpringInertialization = 100.0f; // cm
+
+		LEARNING_API float DistanceToReward(const float Distance);
+		LEARNING_API float DistanceToPenalty(const float Distance);
+		LEARNING_API float SquaredDistanceToReward(const float SquaredDistance);
+		LEARNING_API float SquaredDistanceToPenalty(const float SquaredDistance);
+
+		// Reward Functions
+
+		LEARNING_API float ScalarSimilarityReward(
+			const float Scalar0,
+			const float Scalar1,
+			const float Scale = DefaultScale,
+			const float Threshold = 0.0f,
+			const float Weight = 1.0f,
+			const float Epsilon = UE_SMALL_NUMBER);
+
+		LEARNING_API float ScalarPositionSimilarityReward(
+			const float Position0,
+			const float Position1,
+			const float Scale = DefaultScalePosition,
+			const float Threshold = 0.0f,
+			const float Weight = 1.0f,
+			const float Epsilon = UE_SMALL_NUMBER);
+
+		LEARNING_API float PositionSimilarityReward(
+			const FVector Position0,
+			const FVector Position1,
+			const float Scale = DefaultScalePosition,
+			const float Threshold = 0.0f,
+			const float Weight = 1.0f,
+			const float Epsilon = UE_SMALL_NUMBER);
+
+		LEARNING_API float ScalarRotationSimilarityReward(
+			const float Rotation0,
+			const float Rotation1,
+			const float Scale = DefaultScaleAngle,
+			const float Threshold = 0.0f,
+			const float Weight = 1.0f,
+			const float Epsilon = UE_SMALL_NUMBER);
+
+		LEARNING_API float ScalarAngularVelocitySimilarityReward(
+			const float AngularVelocity0,
+			const float AngularVelocity1,
+			const float Scale = DefaultScaleAngularVelocity,
+			const float Threshold = 0.0f,
+			const float Weight = 1.0f,
+			const float Epsilon = UE_SMALL_NUMBER);
+
+		LEARNING_API float ScalarVelocityReward(
+			const float Velocity,
+			const float Scale = DefaultScaleVelocity,
+			const float Weight = 1.0f,
+			const float Epsilon = UE_SMALL_NUMBER);
+
+		LEARNING_API float LocalDirectionalVelocityReward(
+			const FVector Velocity,
+			const FQuat RelativeRotation,
+			const float Scale = DefaultScaleVelocity,
+			const float Weight = 1.0f,
+			const FVector Axis = FVector::ForwardVector,
+			const float Epsilon = UE_SMALL_NUMBER);
+
+		// Penalty Functions
+
+		LEARNING_API float ScalarDifferencePenalty(
+			const float Scalar0,
+			const float Scalar1,
+			const float Scale = DefaultScale,
+			const float Threshold = 0.0f,
+			const float Weight = 1.0f,
+			const float Epsilon = UE_SMALL_NUMBER);
+
+		LEARNING_API float ScalarMagnitudePenalty(
+			const float Scalar,
+			const float Scale = DefaultScale,
+			const float Threshold = 0.0f,
+			const float Weight = 1.0f,
+			const float Epsilon = UE_SMALL_NUMBER);
+
+		LEARNING_API float DirectionAngleDifferencePenalty(
+			const FVector Direction0,
+			const FVector Direction1,
+			const float Scale = DefaultScaleAngle, // difference between directions is an angle
+			const float Threshold = 0.0f,
+			const float Weight = 1.0f,
+			const float Epsilon = UE_SMALL_NUMBER);
+
+		LEARNING_API float PlanarPositionDifferencePenalty(
+			const FVector Position0,
+			const FVector Position1,
+			const float Scale = DefaultScalePosition,
+			const float Threshold = 0.0f,
+			const float Weight = 1.0f,
+			const FVector Axis0 = FVector::ForwardVector,
+			const FVector Axis1 = FVector::UpVector,
+			const float Epsilon = UE_SMALL_NUMBER);
+
+		LEARNING_API float PositionDifferencePenalty(
+			const FVector Position0,
+			const FVector Position1,
+			const float Scale = DefaultScalePosition,
+			const float Threshold = 0.0f,
+			const float Weight = 1.0f,
+			const float Epsilon = UE_SMALL_NUMBER);
+
+		LEARNING_API float VelocityDifferencePenalty(
+			const FVector Velocity0,
+			const FVector Velocity1,
+			const float Scale = DefaultScaleVelocity,
+			const float Threshold = 0.0f,
+			const float Weight = 1.0f,
+			const float Epsilon = UE_SMALL_NUMBER);
+
+		LEARNING_API float SpringInertializationPenalty(
+			const FVector Position0,
+			const FVector Position1,
+			const FVector Velocity0,
+			const FVector Velocity1,
+			const float HalfLife,
+			const float Scale = DefaultScaleSpringInertialization,
+			const float Threshold = 0.0f,
+			const float Weight = 1.0f,
+			const float Epsilon = UE_SMALL_NUMBER);
+	}
+
+	//------------------------------------------------------------------
+
+	/**
+	* Base class for an object which computes a reward or penalty value 
+	* from a set of other arrays. Here, all data is assumed to be stored 
+	* in a `FArrayMap` object to make the processing of multiple instances efficient.
+	*/
+	struct LEARNING_API FRewardObject
+	{
+		FRewardObject(
+			const FName& InIdentifier,
+			const TSharedRef<FArrayMap>& InInstanceData,
+			const int32 InMaxInstanceNum);
+
+		virtual ~FRewardObject() {}
+
+		virtual void Evaluate(const FIndexSet Instances) = 0;
+
+		TLearningArrayView<1, float> RewardBuffer();
+
+		TSharedRef<FArrayMap> InstanceData;
+		TArrayMapHandle<1, float> RewardHandle;
+	};
+
+	//------------------------------------------------------------------
+
+	/**
+	* Reward that sums the results of multiple other reward functions
+	*/
+	struct LEARNING_API FSumReward : public FRewardObject
+	{
+		FSumReward(
+			const FName& InIdentifier,
+			const TLearningArrayView<1, const TSharedRef<FRewardObject>> InRewards,
+			const TSharedRef<FArrayMap>& InInstanceData,
+			const int32 InMaxInstanceNum);
+
+		virtual void Evaluate(const FIndexSet Instances) override final;
+
+		TLearningArray<1, TSharedRef<FRewardObject>, TInlineAllocator<32>> Rewards;
+	};
+
+	//------------------------------------------------------------------
+
+	/**
+	* Constant Reward
+	*/
+	struct LEARNING_API FConstantReward : public FRewardObject
+	{
+		FConstantReward(
+			const FName& InIdentifier,
+			const TSharedRef<FArrayMap>& InInstanceData,
+			const int32 InMaxInstanceNum,
+			const float InValue = 1.0f);
+
+		virtual void Evaluate(const FIndexSet Instances) override final;
+
+		float Value = 1.0f;
+	};
+
+	/**
+	* Constant Penalty
+	*/
+	struct LEARNING_API FConstantPenalty : public FRewardObject
+	{
+		FConstantPenalty(
+			const FName& InIdentifier,
+			const TSharedRef<FArrayMap>& InInstanceData,
+			const int32 InMaxInstanceNum,
+			const float InValue = 1.0f);
+
+		virtual void Evaluate(const FIndexSet Instances) override final;
+
+		float Value = 1.0f;
+	};
+
+	/**
+	* Conditional Constant Reward
+	*/
+	struct LEARNING_API FConditionalConstantReward : public FRewardObject
+	{
+		FConditionalConstantReward(
+			const FName& InIdentifier,
+			const TSharedRef<FArrayMap>& InInstanceData,
+			const int32 InMaxInstanceNum,
+			const float InValue = 1.0f);
+
+		virtual void Evaluate(const FIndexSet Instances) override final;
+
+		float Value = 1.0f;
+
+		TArrayMapHandle<1, bool> ConditionHandle;
+	};
+
+	//------------------------------------------------------------------
+
+	/**
+	* Basic Float Reward
+	*/
+	struct LEARNING_API FFloatReward : public FRewardObject
+	{
+		FFloatReward(
+			const FName& InIdentifier,
+			const TSharedRef<FArrayMap>& InInstanceData,
+			const int32 InMaxInstanceNum,
+			const float InWeight = 1.0f);
+
+		virtual void Evaluate(const FIndexSet Instances) override final;
+
+		TArrayMapHandle<1, float> ValueHandle;
+		TArrayMapHandle<1, float> WeightHandle;
+	};
+
+	//------------------------------------------------------------------
+
+	/**
+	* Penalty based on the average absolute magnitude of each dimension of a vector
+	*/
+	struct LEARNING_API FVectorAverageMagnitudePenalty : public FRewardObject
+	{
+		FVectorAverageMagnitudePenalty(
+			const FName& InIdentifier,
+			const TSharedRef<FArrayMap>& InInstanceData,
+			const int32 InMaxInstanceNum,
+			const int32 InDimNum,
+			const float InWeight = 1.0f,
+			const float InScale = Reward::DefaultScale,
+			const float InThreshold = 0.0f);
+
+		virtual void Evaluate(const FIndexSet Instances) override final;
+
+		int32 DimNum = 0;
+
+		TArrayMapHandle<2, float> ValueHandle;
+
+		TArrayMapHandle<1, float> WeightHandle;
+		TArrayMapHandle<1, float> ScaleHandle;
+		TArrayMapHandle<1, float> ThresholdHandle;
+	};
+
+	//------------------------------------------------------------------
+
+	/**
+	* Reward based on the similarity between two scalar position
+	*/
+	struct LEARNING_API FScalarPositionSimilarityReward : public FRewardObject
+	{
+		FScalarPositionSimilarityReward(
+			const FName& InIdentifier,
+			const TSharedRef<FArrayMap>& InInstanceData,
+			const int32 InMaxInstanceNum,
+			const float InWeight = 1.0f,
+			const float InScale = Reward::DefaultScalePosition,
+			const float InThreshold = 0.0f);
+
+		virtual void Evaluate(const FIndexSet Instances) override final;
+
+		TArrayMapHandle<1, float> Position0Handle;
+		TArrayMapHandle<1, float> Position1Handle;
+
+		TArrayMapHandle<1, float> WeightHandle;
+		TArrayMapHandle<1, float> ScaleHandle;
+		TArrayMapHandle<1, float> ThresholdHandle;
+	};
+
+	//------------------------------------------------------------------
+
+	/**
+	* Reward based on the similarity between two rotations
+	*/
+	struct LEARNING_API FScalarRotationSimilarityReward : public FRewardObject
+	{
+		FScalarRotationSimilarityReward(
+			const FName& InIdentifier,
+			const TSharedRef<FArrayMap>& InInstanceData,
+			const int32 InMaxInstanceNum,
+			const float InWeight = 1.0f,
+			const float InScale = Reward::DefaultScaleAngle,
+			const float InThreshold = 0.0f);
+
+		virtual void Evaluate(const FIndexSet Instances) override final;
+
+		TArrayMapHandle<1, float> Rotation0Handle;
+		TArrayMapHandle<1, float> Rotation1Handle;
+
+		TArrayMapHandle<1, float> WeightHandle;
+		TArrayMapHandle<1, float> ScaleHandle;
+		TArrayMapHandle<1, float> ThresholdHandle;
+	};
+
+	//------------------------------------------------------------------
+
+	/**
+	* Reward based on the similarity between two scalar angular velocities
+	*/
+	struct LEARNING_API FScalarAngularVelocitySimilarityReward : public FRewardObject
+	{
+		FScalarAngularVelocitySimilarityReward(
+			const FName& InIdentifier,
+			const TSharedRef<FArrayMap>& InInstanceData,
+			const int32 InMaxInstanceNum,
+			const float InWeight = 1.0f,
+			const float InScale = Reward::DefaultScaleAngularVelocity,
+			const float InThreshold = 0.0f);
+
+		virtual void Evaluate(const FIndexSet Instances) override final;
+
+		TArrayMapHandle<1, float> AngularVelocity0Handle;
+		TArrayMapHandle<1, float> AngularVelocity1Handle;
+
+		TArrayMapHandle<1, float> WeightHandle;
+		TArrayMapHandle<1, float> ScaleHandle;
+		TArrayMapHandle<1, float> ThresholdHandle;
+	};
+
+	//------------------------------------------------------------------
+
+	/**
+	* Reward based on how large (and positive) a scalar velocity is
+	*/
+	struct LEARNING_API FScalarVelocityReward : public FRewardObject
+	{
+		FScalarVelocityReward(
+			const FName& InIdentifier,
+			const TSharedRef<FArrayMap>& InInstanceData,
+			const int32 InMaxInstanceNum,
+			const float InWeight = 1.0f,
+			const float InScale = Reward::DefaultScaleVelocity);
+
+		virtual void Evaluate(const FIndexSet Instances) override final;
+
+		TArrayMapHandle<1, float> VelocityHandle;
+
+		TArrayMapHandle<1, float> WeightHandle;
+		TArrayMapHandle<1, float> ScaleHandle;
+	};
+
+	/**
+	* Reward based on how large a velocity is in a particular local direction
+	*/
+	struct LEARNING_API FLocalDirectionalVelocityReward : public FRewardObject
+	{
+		FLocalDirectionalVelocityReward(
+			const FName& InIdentifier,
+			const TSharedRef<FArrayMap>& InInstanceData,
+			const int32 InMaxInstanceNum,
+			const float InWeight = 1.0f,
+			const float InScale = Reward::DefaultScaleVelocity,
+			const FVector InAxis = FVector::ForwardVector);
+
+		virtual void Evaluate(const FIndexSet Instances) override final;
+
+		FVector Axis = FVector::ForwardVector;
+
+		TArrayMapHandle<1, FVector> VelocityHandle;
+		TArrayMapHandle<1, FQuat> RelativeRotationHandle;
+
+		TArrayMapHandle<1, float> WeightHandle;
+		TArrayMapHandle<1, float> ScaleHandle;
+	};
+
+	//------------------------------------------------------------------
+
+	/**
+	* Penalty based on the difference between two planar positions
+	*/
+	struct LEARNING_API FPlanarPositionDifferencePenalty : public FRewardObject
+	{
+		FPlanarPositionDifferencePenalty(
+			const FName& InIdentifier,
+			const TSharedRef<FArrayMap>& InInstanceData,
+			const int32 InMaxInstanceNum,
+			const float InWeight = 1.0f,
+			const float InScale = Reward::DefaultScalePosition,
+			const float InThreshold = 0.0f,
+			const FVector InAxis0 = FVector::ForwardVector,
+			const FVector InAxis1 = FVector::RightVector);
+
+		virtual void Evaluate(const FIndexSet Instances) override final;
+
+		FVector Axis0 = FVector::ForwardVector;
+		FVector Axis1 = FVector::RightVector;
+
+		TArrayMapHandle<1, FVector> Position0Handle;
+		TArrayMapHandle<1, FVector> Position1Handle;
+
+		TArrayMapHandle<1, float> WeightHandle;
+		TArrayMapHandle<1, float> ScaleHandle;
+		TArrayMapHandle<1, float> ThresholdHandle;
+	};
+
+	/**
+	* Penalty based on the difference between two positions
+	*/
+	struct LEARNING_API FPositionDifferencePenalty : public FRewardObject
+	{
+		FPositionDifferencePenalty(
+			const FName& InIdentifier,
+			const TSharedRef<FArrayMap>& InInstanceData,
+			const int32 InMaxInstanceNum,
+			const float InWeight = 1.0f,
+			const float InScale = Reward::DefaultScalePosition,
+			const float InThreshold = 0.0f);
+
+		virtual void Evaluate(const FIndexSet Instances) override final;
+
+		TArrayMapHandle<1, FVector> Position0Handle;
+		TArrayMapHandle<1, FVector> Position1Handle;
+
+		TArrayMapHandle<1, float> WeightHandle;
+		TArrayMapHandle<1, float> ScaleHandle;
+		TArrayMapHandle<1, float> ThresholdHandle;
+	};
+
+	//------------------------------------------------------------------
+
+	/**
+	* Penalty based on the difference between two velocities
+	*/
+	struct LEARNING_API FVelocityDifferencePenalty : public FRewardObject
+	{
+		FVelocityDifferencePenalty(
+			const FName& InIdentifier,
+			const TSharedRef<FArrayMap>& InInstanceData,
+			const int32 InMaxInstanceNum,
+			const float InWeight = 1.0f,
+			const float InScale = Reward::DefaultScaleVelocity,
+			const float InThreshold = 0.0f);
+
+		virtual void Evaluate(const FIndexSet Instances) override final;
+
+		TArrayMapHandle<1, FVector> Velocity0Handle;
+		TArrayMapHandle<1, FVector> Velocity1Handle;
+
+		TArrayMapHandle<1, float> WeightHandle;
+		TArrayMapHandle<1, float> ScaleHandle;
+		TArrayMapHandle<1, float> ThresholdHandle;
+	};
+
+	//------------------------------------------------------------------
+
+	/**
+	* Penalty based on the difference between two rotations facing directions
+	*/
+	struct LEARNING_API FFacingDirectionAngularDifferencePenalty : public FRewardObject
+	{
+		FFacingDirectionAngularDifferencePenalty(
+			const FName& InIdentifier,
+			const TSharedRef<FArrayMap>& InInstanceData,
+			const int32 InMaxInstanceNum,
+			const FVector InFacingDirection0 = FVector::ForwardVector,
+			const FVector InFacingDirection1 = FVector::ForwardVector,
+			const float InWeight = 1.0f,
+			const float InScale = Reward::DefaultScaleAngle, // difference between directions is an angle
+			const float InThreshold = 0.0f);
+
+		virtual void Evaluate(const FIndexSet Instances) override final;
+
+		TArrayMapHandle<1, FQuat> Rotation0Handle;
+		TArrayMapHandle<1, FQuat> Rotation1Handle;
+		TArrayMapHandle<1, FVector> FacingDirection0Handle;
+		TArrayMapHandle<1, FVector> FacingDirection1Handle;
+
+		TArrayMapHandle<1, float> WeightHandle;
+		TArrayMapHandle<1, float> ScaleHandle;
+		TArrayMapHandle<1, float> ThresholdHandle;
+	};
+
+	/**
+	* Penalty based on the difference between a rotation facing direction and a position
+	*/
+	struct LEARNING_API FFacingTowardsTargetAngularDifferencePenalty : public FRewardObject
+	{
+		FFacingTowardsTargetAngularDifferencePenalty(
+			const FName& InIdentifier,
+			const TSharedRef<FArrayMap>& InInstanceData,
+			const int32 InMaxInstanceNum,
+			const FVector InFacingDirection = FVector::ForwardVector,
+			const float InMinRadius = UE_KINDA_SMALL_NUMBER,
+			const float InWeight = 1.0f,
+			const float InScale = Reward::DefaultScaleAngle, // difference between directions is an angle
+			const float InThreshold = 0.0f); 
+
+		virtual void Evaluate(const FIndexSet Instances) override final;
+
+		TArrayMapHandle<1, FVector> PositionHandle;
+		TArrayMapHandle<1, FVector> TargetHandle;
+		TArrayMapHandle<1, FQuat> RotationHandle;
+		TArrayMapHandle<1, FVector> FacingDirectionHandle;
+		TArrayMapHandle<1, float> MinRadiusHandle;
+
+		TArrayMapHandle<1, float> WeightHandle;
+		TArrayMapHandle<1, float> ScaleHandle;
+		TArrayMapHandle<1, float> ThresholdHandle;
+	};
+
+	//------------------------------------------------------------------
+
+	/**
+	* Penalty based on the average spring inertialization displacement introduced when transitioning between two moving points
+	*
+	* See: https://theorangeduck.com/page/inertialization-transition-cost
+	*/
+	struct LEARNING_API FSpringInertializationPenalty : public FRewardObject
+	{
+		FSpringInertializationPenalty(
+			const FName& InIdentifier,
+			const TSharedRef<FArrayMap>& InInstanceData,
+			const int32 InMaxInstanceNum,
+			const float InWeight = 1.0f,
+			const float InScale = Reward::DefaultScaleSpringInertialization,
+			const float InThreshold = 0.0f);
+
+		virtual void Evaluate(const FIndexSet Instances) override final;
+
+		TArrayMapHandle<1, FVector> Position0Handle;
+		TArrayMapHandle<1, FVector> Position1Handle;
+		TArrayMapHandle<1, FVector> Velocity0Handle;
+		TArrayMapHandle<1, FVector> Velocity1Handle;
+
+		TArrayMapHandle<1, float> HalfLifeHandle;
+
+		TArrayMapHandle<1, float> WeightHandle;
+		TArrayMapHandle<1, float> ScaleHandle;
+		TArrayMapHandle<1, float> ThresholdHandle;
+	};
+
+	//------------------------------------------------------------------
+
+	/**
+	* Reward based on the similarity between two arrays of position
+	*/
+	struct LEARNING_API FPositionArraySimilarityReward : public FRewardObject
+	{
+		FPositionArraySimilarityReward(
+			const FName& InIdentifier,
+			const TSharedRef<FArrayMap>& InInstanceData,
+			const int32 InMaxInstanceNum,
+			const int32 InPositionNum,
+			const float InWeight = 1.0f,
+			const float InScale = Reward::DefaultScalePosition,
+			const float InThreshold = 0.0f);
+
+		virtual void Evaluate(const FIndexSet Instances) override final;
+
+		TArrayMapHandle<2, FVector> Positions0Handle;
+		TArrayMapHandle<2, FVector> Positions1Handle;
+		TArrayMapHandle<1, FVector> RelativePosition0Handle;
+		TArrayMapHandle<1, FVector> RelativePosition1Handle;
+		TArrayMapHandle<1, FQuat> RelativeRotation0Handle;
+		TArrayMapHandle<1, FQuat> RelativeRotation1Handle;
+
+		TArrayMapHandle<1, float> WeightHandle;
+		TArrayMapHandle<1, float> ScaleHandle;
+		TArrayMapHandle<1, float> ThresholdHandle;
+	};
+
+	//------------------------------------------------------------------
+
+}
