[Backout] - CL30154408

diff --git a/Engine/Plugins/Experimental/LearningAgents/Source/LearningAgents/Private/LearningAgentsController.cpp b/Engine/Plugins/Experimental/LearningAgents/Source/LearningAgents/Private/LearningAgentsController.cpp
--- a/Engine/Plugins/Experimental/LearningAgents/Source/LearningAgents/Private/LearningAgentsController.cpp
+++ b/Engine/Plugins/Experimental/LearningAgents/Source/LearningAgents/Private/LearningAgentsController.cpp
@@ -5,81 +5,19 @@
 #include "LearningAgentsManager.h"
 #include "LearningAgentsInteractor.h"
 #include "LearningAgentsActions.h"
+#include "LearningFeatureObject.h"
 #include "LearningLog.h"
 
 ULearningAgentsController::ULearningAgentsController() : Super(FObjectInitializer::Get()) {}
 ULearningAgentsController::ULearningAgentsController(FVTableHelper& Helper) : Super(Helper) {}
 ULearningAgentsController::~ULearningAgentsController() = default;
 
-void ULearningAgentsController::EvaluateAgentController_Implementation(
-	TArray<FLearningAgentsActionObjectElement>& OutActionObjectElements,
-	ULearningAgentsActionObject* InActionObject,
-	const TArray<FLearningAgentsObservationObjectElement>& InObservationObjectElements,
-	const ULearningAgentsObservationObject* InObservationObject,
-	const TArray<int32>& AgentIds)
+void ULearningAgentsController::SetActions_Implementation(const TArray<int32>& AgentIds)
 {
-	UE_LOG(LogLearning, Error, TEXT("%s: EvaluateAgentController function must be overridden!"), *GetName());
-	OutActionObjectElements.Empty();
+	// Can be overridden to get actions without blueprints
 }
 
-void ULearningAgentsController::EvaluateAgentControllerUsingDelegate(
-	TArray<FLearningAgentsActionObjectElement>& OutActionObjectElements,
-	ULearningAgentsActionObject* InActionObject,
-	const TArray<FLearningAgentsObservationObjectElement>& InObservationObjectElements,
-	const ULearningAgentsObservationObject* InObservationObject,
-	const TArray<int32>& AgentIds,
-	const FEvaluateAgentControllerDelegate& Delegate)
-{
-	if (!Delegate.IsBound())
-	{
-		UE_LOG(LogLearning, Error, TEXT("%s: Delegate Not Bound."), *GetName());
-		return;
-	}
-
-	const int32 AgentNum = AgentIds.Num();
-
-	if (AgentNum != InObservationObjectElements.Num())
-	{
-		UE_LOG(LogLearning, Error, TEXT("%s: Not enough Observation Objects. Expected %i, Got %i."), *GetName(), AgentNum, InObservationObjectElements.Num());
-		return;
-	}
-
-	OutActionObjectElements.Empty(AgentNum);
-	for (int32 AgentIdx = 0; AgentIdx < AgentNum; AgentIdx++)
-	{
-		OutActionObjectElements.Add(Delegate.Execute(InActionObject, InObservationObjectElements[AgentIdx], InObservationObject, AgentIds[AgentIdx]));
-	}
-}
-
-ULearningAgentsController* ULearningAgentsController::MakeController(
-	ULearningAgentsManager* InManager, 
-	ULearningAgentsInteractor* InInteractor, 
-	TSubclassOf<ULearningAgentsController> Class,
-	const FName Name)
-{
-	if (!InManager)
-	{
-		UE_LOG(LogLearning, Error, TEXT("MakeController: InManager is nullptr."));
-		return nullptr;
-	}
-
-	if (!Class)
-	{
-		UE_LOG(LogLearning, Error, TEXT("MakeController: Class is nullptr."));
-		return nullptr;
-	}
-
-	const FName UniqueName = MakeUniqueObjectName(InManager, Class, Name, EUniqueObjectNameOptions::GloballyUnique);
-
-	ULearningAgentsController* Controller = NewObject<ULearningAgentsController>(InManager, Class, UniqueName);
-	if (!Controller) { return nullptr; }
-	
-	Controller->SetupController(InManager, InInteractor);
-
-	return Controller->IsSetup() ? Controller : nullptr;
-}
-
-void ULearningAgentsController::SetupController(ULearningAgentsManager* InManager, ULearningAgentsInteractor* InInteractor)
+void ULearningAgentsController::SetupController(ULearningAgentsInteractor* InInteractor)
 {
 	if (IsSetup())
 	{
@@ -87,9 +25,9 @@ void ULearningAgentsController::SetupController(ULearningAgentsManager* InManage
 		return;
 	}
 
-	if (!InManager)
+	if (!Manager)
 	{
-		UE_LOG(LogLearning, Error, TEXT("%s: InManager is nullptr."), *GetName());
+		UE_LOG(LogLearning, Error, TEXT("%s: Must be attached to a LearningAgentsManager Actor."), *GetName());
 		return;
 	}
 
@@ -105,16 +43,16 @@ void ULearningAgentsController::SetupController(ULearningAgentsManager* InManage
 		return;
 	}
 
-	Manager = InManager;
 	Interactor = InInteractor;
+
 	bIsSetup = true;
 
-	InManager->AddListener(this);
+	OnAgentsAdded(Manager->GetAllAgentIds());
 }
 
-void ULearningAgentsController::EvaluateController()
+void ULearningAgentsController::EncodeActions()
 {
-	UE_LEARNING_TRACE_CPUPROFILER_EVENT_SCOPE(ULearningAgentsController::EvaluateController);
+	UE_LEARNING_TRACE_CPUPROFILER_EVENT_SCOPE(ULearningAgentsController::EncodeActions);
 
 	if (!IsSetup())
 	{
@@ -122,60 +60,80 @@ void ULearningAgentsController::EvaluateController()
 		return;
 	}
 
-	if (Manager->GetAgentNum() == 0)
+	const TLearningArrayView<1, uint64> ActionEncodingAgentIteration = Interactor->GetActionEncodingAgentIteration();
+
+	// Run Set Actions Callback
+
+	SetActions(Manager->GetAllAgentIds());
+
+	// Check that all actions have had their setter run
+
+	ValidAgentStatus.SetNumUninitialized(Manager->GetMaxAgentNum());
+	ValidAgentStatus.SetRange(0, Manager->GetMaxAgentNum(), true);
+
+	for (ULearningAgentsAction* ActionObject : Interactor->GetActionObjects())
 	{
-		UE_LOG(LogLearning, Warning, TEXT("%s: No agents added to Manager."), *GetName());
+		for (const int32 AgentId : Manager->GetAllAgentSet())
+		{
+			if (ActionObject->GetAgentSetIteration(AgentId) == ActionEncodingAgentIteration[AgentId])
+			{
+				UE_LOG(LogLearning, Warning, TEXT("%s: Action %s for agent with id %i has not been set (got iteration %i, expected iteration %i) and so agent will not have actions encoded."), *GetName(), *ActionObject->GetName(), AgentId, ActionObject->GetAgentSetIteration(AgentId), ActionEncodingAgentIteration[AgentId] + 1);
+				ValidAgentStatus[AgentId] = false;
+				continue;
+			}
+
+			if (ActionObject->GetAgentSetIteration(AgentId) > ActionEncodingAgentIteration[AgentId] + 1)
+			{
+				UE_LOG(LogLearning, Warning, TEXT("%s: Action %s for agent with id %i appears to have been set multiple times (got iteration %i, expected iteration %i) and so agent will not have actions encoded."), *GetName(), *ActionObject->GetName(), AgentId, ActionObject->GetAgentSetIteration(AgentId), ActionEncodingAgentIteration[AgentId] + 1);
+				ValidAgentStatus[AgentId] = false;
+				continue;
+			}
+
+			if (ActionObject->GetAgentSetIteration(AgentId) != ActionEncodingAgentIteration[AgentId] + 1)
+			{
+				UE_LOG(LogLearning, Warning, TEXT("%s: Action %s for agent with id %i does not have a matching iteration number (got iteration %i, expected iteration %i) and so agent will not have actions encoded."), *GetName(), *ActionObject->GetName(), AgentId, ActionObject->GetAgentSetIteration(AgentId), ActionEncodingAgentIteration[AgentId] + 1);
+				ValidAgentStatus[AgentId] = false;
+				continue;
+			}
+		}
 	}
 
-	if (Manager->GetAgentNum() != Interactor->ObservationObjectElements.Num())
+	ValidAgentIds.Empty(Manager->GetAgentNum());
+
+	for (const int32 AgentId : Manager->GetAllAgentSet())
 	{
-		UE_LOG(LogLearning, Error, TEXT("%s: Not enough Observation Objects added by GatherAgentObservations. Expected %i, Got %i."), *GetName(), Manager->GetAgentNum(), Interactor->ObservationObjectElements.Num());
-		return;
+		if (ValidAgentStatus[AgentId]) { ValidAgentIds.Add(AgentId); }
 	}
 
-	// Run EvaluateAgentController Callback
+	ValidAgentSet = ValidAgentIds;
+	ValidAgentSet.TryMakeSlice();
+
+	// Encode Actions
+
+	Interactor->GetActionFeature().Encode(ValidAgentSet);
 
-	Interactor->ActionObject->GetActionObject().Reset();
-	Interactor->ActionObjectElements.Empty(Manager->GetMaxAgentNum());
-	EvaluateAgentController(
-		Interactor->ActionObjectElements, 
-		Interactor->ActionObject, 
-		Interactor->ObservationObjectElements,
-		Interactor->ObservationObject,
-		Manager->GetAllAgentIds());
+	// Increment Action Encoding Iteration
 
-	if (Manager->GetAgentNum() != Interactor->ActionObjectElements.Num())
+	for (const int32 AgentId : ValidAgentSet)
 	{
-		UE_LOG(LogLearning, Error, TEXT("%s: Not enough Action Objects added by EvaluateAgentController. Expected %i, Got %i."), *GetName(), Manager->GetAgentNum(), Interactor->ActionObjectElements.Num());
-		return;
+		ActionEncodingAgentIteration[AgentId]++;
 	}
 
-	// Check Action Objects are Valid and if so convert to action vectors
+	// Visual Logger
 
-	for (int32 AgentIdx = 0; AgentIdx < Manager->GetAgentNum(); AgentIdx++)
+#if UE_LEARNING_AGENTS_ENABLE_VISUAL_LOG
+	for (const ULearningAgentsAction* ActionObject : Interactor->GetActionObjects())
 	{
-		if (Interactor->ActionSchema->ValidateObjectMatchesSchema(
-			Interactor->ActionSchemaElement, 
-			Interactor->ActionObject, 
-			Interactor->ActionObjectElements[AgentIdx]))
-		{
-			UE::Learning::Action::SetVectorFromObject(
-				Interactor->ActionVectors[Manager->GetAllAgentSet()[AgentIdx]],
-				Interactor->ActionSchema->GetActionSchema(),
-				Interactor->ActionSchemaElement.SchemaElement,
-				Interactor->ActionObject->GetActionObject(),
-				Interactor->ActionObjectElements[AgentIdx].ObjectElement);
-
-			Interactor->ActionVectorIteration[Manager->GetAllAgentSet()[AgentIdx]]++;
-		}
+		ActionObject->VisualLog(ValidAgentSet);
 	}
+#endif
 }
 
 ULearningAgentsInteractor* ULearningAgentsController::GetInteractor(const TSubclassOf<ULearningAgentsInteractor> InteractorClass) const
 {
-	if (!IsSetup())
+	if (!Interactor)
 	{
-		UE_LOG(LogLearning, Error, TEXT("%s: Setup not complete."), *GetName());
+		UE_LOG(LogLearning, Error, TEXT("%s: Interactor is nullptr. Did we forget to call Setup on this component?"), *GetName());
 		return nullptr;
 	}
 
@@ -192,7 +150,7 @@ void ULearningAgentsController::RunController()
 		return;
 	}
 
-	Interactor->GatherObservations();
-	EvaluateController();
-	Interactor->ScatterActions();
+	Interactor->EncodeObservations();
+	EncodeActions();
+	Interactor->DecodeActions();
 }
diff --git a/Engine/Plugins/Experimental/LearningAgents/Source/LearningAgents/Private/LearningAgentsCritic.cpp b/Engine/Plugins/Experimental/LearningAgents/Source/LearningAgents/Private/LearningAgentsCritic.cpp
--- a/Engine/Plugins/Experimental/LearningAgents/Source/LearningAgents/Private/LearningAgentsCritic.cpp
+++ b/Engine/Plugins/Experimental/LearningAgents/Source/LearningAgents/Private/LearningAgentsCritic.cpp
@@ -5,8 +5,11 @@
 #include "LearningAgentsManager.h"
 #include "LearningAgentsInteractor.h"
 #include "LearningAgentsPolicy.h"
+#include "LearningAgentsHelpers.h"
+#include "LearningAgentsNeuralNetworkData.h"
+#include "LearningFeatureObject.h"
 #include "LearningNeuralNetwork.h"
-#include "LearningCritic.h"
+#include "LearningNeuralNetworkObject.h"
 #include "LearningLog.h"
 
 #include "UObject/Package.h"
@@ -14,76 +17,15 @@
 #include "Misc/Paths.h"
 #include "GameFramework/Actor.h"
 
-#include "NNERuntimeBasicCpuBuilder.h"
-
-namespace UE::Learning::Agents::Critic::Private
-{
-	static inline UE::NNE::RuntimeBasic::FModelBuilder::EActivationFunction GetBuilderActivationFunction(const ELearningAgentsActivationFunction ActivationFunction)
-	{
-		switch (ActivationFunction)
-		{
-		case ELearningAgentsActivationFunction::ReLU: return UE::NNE::RuntimeBasic::FModelBuilder::EActivationFunction::ReLU;
-		case ELearningAgentsActivationFunction::ELU: return UE::NNE::RuntimeBasic::FModelBuilder::EActivationFunction::ELU;
-		case ELearningAgentsActivationFunction::TanH: return UE::NNE::RuntimeBasic::FModelBuilder::EActivationFunction::TanH;
-		default:
-			UE_LOG(LogLearning, Error, TEXT("Unknown Activation Function"));
-			return UE::NNE::RuntimeBasic::FModelBuilder::EActivationFunction::ReLU;
-		}
-	}
-}
-
 ULearningAgentsCritic::ULearningAgentsCritic() : Super(FObjectInitializer::Get()) {}
 ULearningAgentsCritic::ULearningAgentsCritic(FVTableHelper& Helper) : Super(Helper) {}
 ULearningAgentsCritic::~ULearningAgentsCritic() = default;
 
-ULearningAgentsCritic* ULearningAgentsCritic::MakeCritic(
-	ULearningAgentsManager* InManager,
-	ULearningAgentsInteractor* InInteractor,
-	ULearningAgentsPolicy* InPolicy,
-	TSubclassOf<ULearningAgentsCritic> Class,
-	const FName Name,
-	ULearningAgentsNeuralNetwork* CriticNeuralNetworkAsset,
-	const bool bReinitializeCriticNetwork,
-	const FLearningAgentsCriticSettings& CriticSettings,
-	const int32 Seed)
-{
-	if (!InManager)
-	{
-		UE_LOG(LogLearning, Error, TEXT("MakeCritic: InManager is nullptr."));
-		return nullptr;
-	}
-
-	if (!Class)
-	{
-		UE_LOG(LogLearning, Error, TEXT("MakeCritic: Class is nullptr."));
-		return nullptr;
-	}
-
-	const FName UniqueName = MakeUniqueObjectName(InManager, Class, Name, EUniqueObjectNameOptions::GloballyUnique);
-
-	ULearningAgentsCritic* Critic = NewObject<ULearningAgentsCritic>(InManager, Class, UniqueName);
-	if (!Critic) { return nullptr; }
-
-	Critic->SetupCritic(
-		InManager, 
-		InInteractor,
-		InPolicy,
-		CriticNeuralNetworkAsset,
-		bReinitializeCriticNetwork,
-		CriticSettings,
-		Seed);
-
-	return Critic->IsSetup() ? Critic : nullptr;
-}
-
 void ULearningAgentsCritic::SetupCritic(
-	ULearningAgentsManager* InManager,
 	ULearningAgentsInteractor* InInteractor, 
 	ULearningAgentsPolicy* InPolicy,
-	ULearningAgentsNeuralNetwork* CriticNeuralNetworkAsset,
-	const bool bReinitializeCriticNetwork,
 	const FLearningAgentsCriticSettings& CriticSettings,
-	const int32 Seed)
+	ULearningAgentsNeuralNetwork* NeuralNetworkAsset)
 {
 	if (IsSetup())
 	{
@@ -91,9 +33,9 @@ void ULearningAgentsCritic::SetupCritic(
 		return;
 	}
 
-	if (!InManager)
+	if (!Manager)
 	{
-		UE_LOG(LogLearning, Error, TEXT("%s: InManager is nullptr."), *GetName());
+		UE_LOG(LogLearning, Error, TEXT("%s: Must be attached to a LearningAgentsManager Actor."), *GetName());
 		return;
 	}
 
@@ -109,6 +51,8 @@ void ULearningAgentsCritic::SetupCritic(
 		return;
 	}
 
+	Interactor = InInteractor;
+
 	if (!InPolicy)
 	{
 		UE_LOG(LogLearning, Error, TEXT("%s: InPolicy is nullptr."), *GetName());
@@ -121,97 +65,132 @@ void ULearningAgentsCritic::SetupCritic(
 		return;
 	}
 
-	Manager = InManager;
-	Interactor = InInteractor;
 	Policy = InPolicy;
 
-	const int32 ObservationEncodedVectorSize = Interactor->GetObservationEncodedVectorSize();
-	const int32 MemoryStateSize = Policy->GetMemoryStateSize();
-
-	const int32 CriticHashData[2] = { MemoryStateSize, ObservationEncodedVectorSize };
-	const int32 CriticCompatibilityHash = CityHash32((const char*)CriticHashData, 2 * sizeof(int32));
-
-	// Try to use existing Neural Network Asset
+	const int32 NetworkInputNum = Interactor->GetObservationFeature().DimNum() + Policy->GetMemoryStateSize();
+	const int32 NetworkOutputNum = 1;
 
-	if (CriticNeuralNetworkAsset)
+	if (NeuralNetworkAsset)
 	{
-		CriticNetwork = CriticNeuralNetworkAsset;
+		// Use Existing Neural Network Asset
 
-		if (CriticNeuralNetworkAsset->NeuralNetworkData && !bReinitializeCriticNetwork)
+		if (NeuralNetworkAsset->NeuralNetworkData)
 		{
-			if (CriticNeuralNetworkAsset->NeuralNetworkData->GetCompatibilityHash() != CriticCompatibilityHash)
+			if (NeuralNetworkAsset->NeuralNetworkData->GetNetworkInterface()->GetInputNum() != NetworkInputNum ||
+				NeuralNetworkAsset->NeuralNetworkData->GetNetworkInterface()->GetOutputNum() != NetworkOutputNum)
 			{
-				UE_LOG(LogLearning, Error, TEXT("%s: Critic Network Asset provided during Setup is incompatible with Schema. Network hash is %i vs Schema hash %i."), *GetName(),
-					CriticNeuralNetworkAsset->NeuralNetworkData->GetCompatibilityHash(),
-					CriticCompatibilityHash);
+				UE_LOG(LogLearning, Error, TEXT("%s: Neural Network Asset provided during Setup is incorrect size: Inputs and outputs don't match."), *GetName());
 				return;
 			}
 
-			if (CriticNeuralNetworkAsset->NeuralNetworkData->GetInputSize() != ObservationEncodedVectorSize + MemoryStateSize ||
-				CriticNeuralNetworkAsset->NeuralNetworkData->GetOutputSize() != 1)
-			{
-				UE_LOG(LogLearning, Error, TEXT("%s: Critic Network Asset provided during Setup is incorrect size: Got inputs of size %i, expected %i. Got outputs of size %i, expected %i."), *GetName(),
-					CriticNeuralNetworkAsset->NeuralNetworkData->GetInputSize(), ObservationEncodedVectorSize + MemoryStateSize,
-					CriticNeuralNetworkAsset->NeuralNetworkData->GetOutputSize(), 1);
-				return;
-			}
+			Network = NeuralNetworkAsset;
+		}
+		else
+		{
+			Network = NeuralNetworkAsset;
+			Network->NeuralNetworkData = NewObject<ULearningAgentsNeuralNetworkData>(Network);
+			Network->NeuralNetworkData->CreateMLP(
+				NetworkInputNum,
+				NetworkOutputNum,
+				CriticSettings.HiddenLayerSize,
+				CriticSettings.LayerNum,
+				CriticSettings.ActivationFunction);
 		}
 	}
-
-	if (!CriticNetwork)
+	else
 	{
+		// Create New Neural Network Asset
+
 		const FName UniqueName = MakeUniqueObjectName(this, ULearningAgentsNeuralNetwork::StaticClass(), TEXT("CriticNetwork"), EUniqueObjectNameOptions::GloballyUnique);
-		CriticNetwork = NewObject<ULearningAgentsNeuralNetwork>(this, UniqueName);
+
+		Network = NewObject<ULearningAgentsNeuralNetwork>(this, UniqueName);
+		Network->NeuralNetworkData = NewObject<ULearningAgentsNeuralNetworkData>(Network);
+		Network->NeuralNetworkData->CreateMLP(
+			NetworkInputNum,
+			NetworkOutputNum,
+			CriticSettings.HiddenLayerSize,
+			CriticSettings.LayerNum,
+			CriticSettings.ActivationFunction);
 	}
 
-	if (!CriticNetwork->NeuralNetworkData || bReinitializeCriticNetwork)
+	// Create Critic Object
+	CriticObject = MakeShared<UE::Learning::FNeuralNetworkCriticFunction>(
+		TEXT("CriticObject"),
+		Manager->GetInstanceData().ToSharedRef(),
+		Manager->GetMaxAgentNum(),
+		Interactor->GetObservationFeature().DimNum(),
+		Policy->GetMemoryStateSize(),
+		Network->NeuralNetworkData->GetNetworkInterface());
+
+	CriticAgentIteration.SetNumUninitialized({ Manager->GetMaxAgentNum() });
+	UE::Learning::Array::Set<1, uint64>(CriticAgentIteration, INDEX_NONE);
+
+	bIsSetup = true;
+
+	OnAgentsAdded(Manager->GetAllAgentIds());
+}
+
+void ULearningAgentsCritic::OnAgentsAdded(const TArray<int32>& AgentIds)
+{
+	if (IsSetup())
 	{
-		// Create New Neural Network Asset
+		UE::Learning::Array::Set<1, uint64>(CriticAgentIteration, 0, AgentIds);
 
-		if (!CriticNetwork->NeuralNetworkData)
+		for (ULearningAgentsHelper* Helper : HelperObjects)
 		{
-			CriticNetwork->NeuralNetworkData = NewObject<ULearningNeuralNetworkData>(CriticNetwork);
+			Helper->OnAgentsAdded(AgentIds);
 		}
 
-		UE::NNE::RuntimeBasic::FModelBuilder Builder;
+		AgentsAdded(AgentIds);
+	}
+}
 
-		TArray<uint8> FileData;
-		uint32 CriticInputSize, CriticOutputSize;
-		Builder.WriteFileDataAndReset(FileData, CriticInputSize, CriticOutputSize,
-			Builder.MakeMLPWithRandomKaimingWeights(
-				ObservationEncodedVectorSize + MemoryStateSize,
-				1,
-				CriticSettings.HiddenLayerSize,
-				FMath::Max(CriticSettings.LayerNum, 2),
-				UE::Learning::Agents::Critic::Private::GetBuilderActivationFunction(CriticSettings.ActivationFunction)));
+void ULearningAgentsCritic::OnAgentsRemoved(const TArray<int32>& AgentIds)
+{
+	if (IsSetup())
+	{
+		UE::Learning::Array::Set<1, uint64>(CriticAgentIteration, INDEX_NONE, AgentIds);
 
-		UE_LEARNING_CHECK(CriticInputSize == ObservationEncodedVectorSize + MemoryStateSize);
-		UE_LEARNING_CHECK(CriticOutputSize == 1);
+		for (ULearningAgentsHelper* Helper : HelperObjects)
+		{
+			Helper->OnAgentsRemoved(AgentIds);
+		}
 
-		CriticNetwork->NeuralNetworkData->Init(CriticInputSize, CriticOutputSize, CriticCompatibilityHash, FileData);
-		CriticNetwork->ForceMarkDirty();
+		AgentsRemoved(AgentIds);
 	}
+}
 
-	// Create Critic Object
+void ULearningAgentsCritic::OnAgentsReset(const TArray<int32>& AgentIds)
+{
+	if (IsSetup())
+	{
+		UE::Learning::Array::Set<1, uint64>(CriticAgentIteration, 0, AgentIds);
 
-	CriticObject = MakeShared<UE::Learning::FNeuralNetworkCritic>(
-		Manager->GetMaxAgentNum(),
-		ObservationEncodedVectorSize,
-		MemoryStateSize,
-		CriticNetwork->NeuralNetworkData->GetNetwork());
+		for (ULearningAgentsHelper* Helper : HelperObjects)
+		{
+			Helper->OnAgentsReset(AgentIds);
+		}
 
-	Returns.SetNumUninitialized({ Manager->GetMaxAgentNum() });
-	UE::Learning::Array::Set<1, float>(Returns, FLT_MAX);
+		AgentsReset(AgentIds);
+	}
+}
 
-	ReturnsIteration.SetNumUninitialized({ Manager->GetMaxAgentNum() });
-	UE::Learning::Array::Set<1, uint64>(ReturnsIteration, INDEX_NONE);
+ULearningAgentsNeuralNetwork* ULearningAgentsCritic::GetNetworkAsset()
+{
+	return Network;
+}
 
-	bIsSetup = true;
+UE::Learning::INeuralNetwork& ULearningAgentsCritic::GetCriticNetwork()
+{
+	return *Network->NeuralNetworkData->GetNetworkInterface();
+}
 
-	Manager->AddListener(this);
+UE::Learning::FNeuralNetworkCriticFunction& ULearningAgentsCritic::GetCriticObject()
+{
+	return *CriticObject;
 }
 
-void ULearningAgentsCritic::OnAgentsAdded_Implementation(const TArray<int32>& AgentIds)
+void ULearningAgentsCritic::LoadCriticFromSnapshot(const FFilePath& File)
 {
 	if (!IsSetup())
 	{
@@ -219,11 +198,10 @@ void ULearningAgentsCritic::OnAgentsAdded_Implementation(const TArray<int32>& Ag
 		return;
 	}
 
-	UE::Learning::Array::Set<1, uint64>(ReturnsIteration, 0, AgentIds);
-	UE::Learning::Array::Set<1, float>(Returns, 0.0f, AgentIds);
+	Network->LoadNetworkFromSnapshot(File);
 }
 
-void ULearningAgentsCritic::OnAgentsRemoved_Implementation(const TArray<int32>& AgentIds)
+void ULearningAgentsCritic::SaveCriticToSnapshot(const FFilePath& File) const
 {
 	if (!IsSetup())
 	{
@@ -231,11 +209,10 @@ void ULearningAgentsCritic::OnAgentsRemoved_Implementation(const TArray<int32>&
 		return;
 	}
 
-	UE::Learning::Array::Set<1, uint64>(ReturnsIteration, INDEX_NONE, AgentIds);
-	UE::Learning::Array::Set<1, float> (Returns, FLT_MAX, AgentIds);
+	Network->SaveNetworkToSnapshot(File);
 }
 
-void ULearningAgentsCritic::OnAgentsReset_Implementation(const TArray<int32>& AgentIds)
+void ULearningAgentsCritic::UseCriticFromAsset(ULearningAgentsNeuralNetwork* NeuralNetworkAsset)
 {
 	if (!IsSetup())
 	{
@@ -243,24 +220,49 @@ void ULearningAgentsCritic::OnAgentsReset_Implementation(const TArray<int32>& Ag
 		return;
 	}
 
-	UE::Learning::Array::Set<1, uint64>(ReturnsIteration, 0, AgentIds);
-	UE::Learning::Array::Set<1, float>(Returns, 0.0f, AgentIds);
+	if (!NeuralNetworkAsset || !NeuralNetworkAsset->NeuralNetworkData)
+	{
+		UE_LOG(LogLearning, Error, TEXT("%s: Asset is invalid."), *GetName());
+		return;
+	}
+
+	if (NeuralNetworkAsset == Network)
+	{
+		UE_LOG(LogLearning, Error, TEXT("%s: Asset is same as the current network."), *GetName());
+		return;
+	}
+
+	if (NeuralNetworkAsset->NeuralNetworkData->GetNetworkInterface()->GetInputNum() != Network->NeuralNetworkData->GetNetworkInterface()->GetInputNum() ||
+		NeuralNetworkAsset->NeuralNetworkData->GetNetworkInterface()->GetOutputNum() != Network->NeuralNetworkData->GetNetworkInterface()->GetOutputNum())
+	{
+		UE_LOG(LogLearning, Error, TEXT("%s: Failed to use asset as network settings don't match."), *GetName());
+		return;
+	}
+
+	Network = NeuralNetworkAsset;
+	CriticObject->UpdateNeuralNetwork(Network->NeuralNetworkData->GetNetworkInterface());
 }
 
-UE::Learning::FNeuralNetworkCritic& ULearningAgentsCritic::GetCriticObject()
+void ULearningAgentsCritic::LoadCriticFromAsset(ULearningAgentsNeuralNetwork* NeuralNetworkAsset)
 {
-	return *CriticObject;
+	if (!IsSetup())
+	{
+		UE_LOG(LogLearning, Error, TEXT("%s: Setup not complete."), *GetName());
+		return;
+	}
+
+	Network->LoadNetworkFromAsset(NeuralNetworkAsset);
 }
 
-ULearningAgentsNeuralNetwork* ULearningAgentsCritic::GetCriticNetworkAsset()
+void ULearningAgentsCritic::SaveCriticToAsset(ULearningAgentsNeuralNetwork* NeuralNetworkAsset)
 {
 	if (!IsSetup())
 	{
 		UE_LOG(LogLearning, Error, TEXT("%s: Setup not complete."), *GetName());
-		return nullptr;
+		return;
 	}
 
-	return CriticNetwork;
+	Network->SaveNetworkToAsset(NeuralNetworkAsset);
 }
 
 void ULearningAgentsCritic::EvaluateCritic()
@@ -273,21 +275,16 @@ void ULearningAgentsCritic::EvaluateCritic()
 		return;
 	}
 
-	if (Manager->GetAgentNum() == 0)
-	{
-		UE_LOG(LogLearning, Warning, TEXT("%s: No agents added to Manager."), *GetName());
-	}
-
 	// Check Agents actually have encoded observations.
 	// All added agents should already have some memory state even if it is zero.
 
-	ValidAgentIds.Empty(Manager->GetMaxAgentNum());
+	ValidAgentIds.Empty(Manager->GetAgentNum());
 
 	for (const int32 AgentId : Manager->GetAllAgentSet())
 	{
-		if (Policy->ObservationVectorEncodedIteration[AgentId] == 0)
+		if (Interactor->GetObservationEncodingAgentIteration()[AgentId] == 0)
 		{
-			UE_LOG(LogLearning, Warning, TEXT("%s: Agent with id %i has not made observations so critic will not be evaluated for it. Was EncodeObservations run without error?"), *GetName(), AgentId);
+			UE_LOG(LogLearning, Warning, TEXT("%s: Agent with id %i has not made observations so critic will not be evaluated for it."), *GetName(), AgentId);
 			continue;
 		}
 
@@ -297,20 +294,32 @@ void ULearningAgentsCritic::EvaluateCritic()
 	ValidAgentSet = ValidAgentIds;
 	ValidAgentSet.TryMakeSlice();
 
+	// Get views of Observations and Memory State and copy into network input buffers
+
+	TLearningArrayView<2, const float> MemoryStateView = Policy->GetMemoryStateView();
+	TLearningArrayView<2, const float> ObservationsView = Manager->GetInstanceData()->ConstView(Interactor->GetObservationFeature().FeatureHandle);
+	TLearningArrayView<2, float> InputObservationView = Manager->GetInstanceData()->View(CriticObject->InputObservationHandle);
+	TLearningArrayView<2, float> InputMemoryStateView = Manager->GetInstanceData()->View(CriticObject->InputMemoryStateHandle);
+
+	UE::Learning::Array::Copy<2, float>(InputObservationView, ObservationsView, ValidAgentSet);
+	UE::Learning::Array::Copy<2, float>(InputMemoryStateView, MemoryStateView, ValidAgentSet);
+
 	// Evaluate Critic
 
-	CriticObject->Evaluate(
-		Returns,
-		Policy->ObservationVectorsEncoded,
-		Policy->MemoryState,
-		ValidAgentSet);
+	CriticObject->Evaluate(ValidAgentSet);
 
 	// Increment Discounted Return Iteration
 
 	for (const int32 AgentId : ValidAgentSet)
 	{
-		ReturnsIteration[AgentId]++;
+		CriticAgentIteration[AgentId]++;
 	}
+		
+	// Visual Logger
+
+#if UE_LEARNING_AGENTS_ENABLE_VISUAL_LOG
+	VisualLog(ValidAgentSet);
+#endif
 }
 
 float ULearningAgentsCritic::GetEstimatedDiscountedReturn(const int32 AgentId) const
@@ -327,12 +336,39 @@ float ULearningAgentsCritic::GetEstimatedDiscountedReturn(const int32 AgentId) c
 		return 0.0f;
 	}
 
-	if (ReturnsIteration[AgentId] == 0)
+	if (CriticAgentIteration[AgentId] == 0)
 	{
 		UE_LOG(LogLearning, Error, TEXT("%s: Agent with id %d has not yet computed the estimated discounted return. Did you run EvaluateCritic?"), *GetName(), AgentId);
 		return 0.0f;
 	}
 
-	return Returns[AgentId];
+	return CriticObject->InstanceData->ConstView(CriticObject->OutputHandle)[AgentId];
 }
 
+#if UE_LEARNING_AGENTS_ENABLE_VISUAL_LOG
+void ULearningAgentsCritic::VisualLog(const UE::Learning::FIndexSet AgentSet) const
+{
+	UE_LEARNING_TRACE_CPUPROFILER_EVENT_SCOPE(ULearningAgentsCritic::VisualLog);
+
+	const TLearningArrayView<2, const float> InputObservationView = CriticObject->InstanceData->ConstView(CriticObject->InputObservationHandle);
+	const TLearningArrayView<2, const float> InputMemoryStateView = CriticObject->InstanceData->ConstView(CriticObject->InputMemoryStateHandle);
+	const TLearningArrayView<1, const float> OutputView = CriticObject->InstanceData->ConstView(CriticObject->OutputHandle);
+
+	for (const int32 AgentId : AgentSet)
+	{
+		if (const AActor* Actor = Cast<AActor>(GetAgent(AgentId)))
+		{
+			UE_LEARNING_AGENTS_VLOG_STRING(this, LogLearning, Display,
+				Actor->GetActorLocation(),
+				VisualLogColor.ToFColor(true),
+				TEXT("Agent %i\nObservation Input: %s\nObservation Input Stats (Min/Max/Mean/Std): %s\nMemory State Input: %s\nMemory State Input Stats (Min/Max/Mean/Std): %s\nOutput: [% 6.3f]"),
+				AgentId,
+				*UE::Learning::Array::FormatFloat(InputObservationView[AgentId]),
+				*UE::Learning::Agents::Debug::FloatArrayToStatsString(InputObservationView[AgentId]),
+				*UE::Learning::Array::FormatFloat(InputMemoryStateView[AgentId]),
+				*UE::Learning::Agents::Debug::FloatArrayToStatsString(InputMemoryStateView[AgentId]),
+				OutputView[AgentId]);
+		}
+	}
+}
+#endif
diff --git a/Engine/Plugins/Experimental/LearningAgents/Source/LearningAgents/Private/LearningAgentsDebug.cpp b/Engine/Plugins/Experimental/LearningAgents/Source/LearningAgents/Private/LearningAgentsDebug.cpp
new file mode 100644
--- /dev/null
+++ b/Engine/Plugins/Experimental/LearningAgents/Source/LearningAgents/Private/LearningAgentsDebug.cpp
@@ -0,0 +1,45 @@
+// Copyright Epic Games, Inc. All Rights Reserved.
+
+#include "LearningAgentsDebug.h"
+
+namespace UE::Learning::Agents::Debug
+{
+	FMatrix PlaneMatrix(const FQuat Rotation, const FVector Position, const FVector Axis0, const FVector Axis1)
+	{
+		// Here we reconstruct the transformation matrix from the plane axes assuming the first axis is forward, and second is to the right.
+		const FVector ForwardAxis = Rotation.RotateVector(Axis0.GetSafeNormal(UE_SMALL_NUMBER, FVector::ForwardVector));
+		const FVector InitialRightAxis = Rotation.RotateVector(Axis1.GetSafeNormal(UE_SMALL_NUMBER, FVector::RightVector));
+		const FVector UpAxis = ForwardAxis.Cross(InitialRightAxis).GetSafeNormal(UE_SMALL_NUMBER, FVector::UpVector);
+		const FVector RightAxis = UpAxis.Cross(ForwardAxis).GetSafeNormal(UE_SMALL_NUMBER, FVector::RightVector);
+		return FMatrix(ForwardAxis, RightAxis, UpAxis, Position);
+	}
+
+	FVector GridOffsetForIndex(const int32 Idx, const int32 Num, const float Spacing, const float Height)
+	{
+		const int32 RowNum = FMath::Max(FMath::CeilToInt(FMath::Sqrt((float)Num)), 1);
+
+		return Spacing * (FVector((Idx / RowNum), (Idx % RowNum), 0.0f) - FVector(RowNum - 1, RowNum - 1, 0.0f) / 2.0f) + FVector(0.0f, 0.0f, Height);
+	}
+
+	FString FloatArrayToStatsString(const TLearningArrayView<1, const float> Array)
+	{
+		const int32 ItemNum = Array.Num();
+
+		float Min = +FLT_MAX, Max = -FLT_MAX, Mean = 0.0f;
+		for (int32 Idx = 0; Idx < ItemNum; Idx++)
+		{
+			Min = FMath::Min(Min, Array[Idx]);
+			Max = FMath::Max(Max, Array[Idx]);
+			Mean += Array[Idx] / ItemNum;
+		}
+
+		float Var = 0.0f;
+		for (int32 Idx = 0; Idx < ItemNum; Idx++)
+		{
+			Var += FMath::Square(Array[Idx] - Mean) / ItemNum;
+		}
+
+		return FString::Printf(TEXT("[% 6.3f/% 6.3f/% 6.3f/% 6.3f]"),
+			Min, Max, Mean, FMath::Sqrt(Var));
+	}
+}
diff --git a/Engine/Plugins/Experimental/LearningAgents/Source/LearningAgents/Private/LearningAgentsHelpers.cpp b/Engine/Plugins/Experimental/LearningAgents/Source/LearningAgents/Private/LearningAgentsHelpers.cpp
new file mode 100644
--- /dev/null
+++ b/Engine/Plugins/Experimental/LearningAgents/Source/LearningAgents/Private/LearningAgentsHelpers.cpp
@@ -0,0 +1,976 @@
+// Copyright Epic Games, Inc. All Rights Reserved.
+
+#include "LearningAgentsHelpers.h"
+
+#include "LearningAgentsManager.h"
+#include "LearningAgentsManagerComponent.h"
+#include "LearningLog.h"
+
+#include "Components/MeshComponent.h"
+#include "Components/SplineComponent.h"
+#include "Engine/World.h"
+
+//------------------------------------------------------------------
+
+void ULearningAgentsHelper::OnAgentsAdded(const TArray<int32>& AgentIds) {}
+void ULearningAgentsHelper::OnAgentsRemoved(const TArray<int32>& AgentIds) {}
+void ULearningAgentsHelper::OnAgentsReset(const TArray<int32>& AgentIds) {}
+
+//------------------------------------------------------------------
+
+USplineComponentHelper* USplineComponentHelper::AddSplineComponentHelper(ULearningAgentsManagerComponent* InManagerComponent, const FName Name)
+{
+	if (!InManagerComponent)
+	{
+		UE_LOG(LogLearning, Error, TEXT("AddSplineComponentHelper: InManagerComponent is nullptr."));
+		return nullptr;
+	}
+
+	const FName UniqueName = MakeUniqueObjectName(InManagerComponent, USplineComponentHelper::StaticClass(), Name, EUniqueObjectNameOptions::GloballyUnique);
+
+	USplineComponentHelper* Helper = NewObject<USplineComponentHelper>(InManagerComponent, UniqueName);
+	Helper->ManagerComponent = InManagerComponent;
+	InManagerComponent->AddHelper(Helper);
+	return Helper;
+}
+
+FVector USplineComponentHelper::GetNearestPositionOnSpline(const int32 AgentId, const USplineComponent* SplineComponent, const FVector Position, const ESplineCoordinateSpace::Type CoordinateSpace) const
+{
+	if (!ManagerComponent->HasAgent(AgentId))
+	{
+		UE_LOG(LogLearning, Error, TEXT("%s: AgentId %d not found in the agents set."), *GetName(), AgentId);
+		return FVector::ZeroVector;
+	}
+
+	if (!SplineComponent)
+	{
+		UE_LOG(LogLearning, Error, TEXT("%s: SplineComponent is nullptr."), *GetName());
+		return FVector::ZeroVector;
+	}
+
+	const FVector SplinePosition = SplineComponent->FindLocationClosestToWorldLocation(Position, CoordinateSpace);
+
+#if UE_LEARNING_AGENTS_ENABLE_VISUAL_LOG
+	{
+		UE_LEARNING_AGENTS_VLOG_LOCATION(this, LogLearning, Display,
+			SplinePosition,
+			10,
+			VisualLogColor.ToFColor(true),
+			TEXT(""));
+
+		UE_LEARNING_AGENTS_VLOG_SEGMENT(this, LogLearning, Display,
+			Position,
+			SplinePosition,
+			VisualLogColor.ToFColor(true),
+			TEXT("GetNearestPositionOnSpline\nAgent %i\nPosition: [% 6.1f % 6.1f % 6.1f]\nSpline Position: [% 6.1f % 6.1f % 6.1f]"),
+			AgentId,
+			Position.X, Position.Y, Position.Z,
+			SplinePosition.X, SplinePosition.Y, SplinePosition.Z);
+
+		UE_LEARNING_AGENTS_VLOG_LOCATION(this, LogLearning, Display,
+			Position,
+			10,
+			VisualLogColor.ToFColor(true),
+			TEXT(""));
+	}
+#endif
+
+	return SplinePosition;
+}
+
+float USplineComponentHelper::GetDistanceAlongSplineAtPosition(const int32 AgentId, const USplineComponent* SplineComponent, const FVector Position, const ESplineCoordinateSpace::Type CoordinateSpace) const
+{
+	if (!ManagerComponent->HasAgent(AgentId))
+	{
+		UE_LOG(LogLearning, Error, TEXT("%s: AgentId %d not found in the agents set."), *GetName(), AgentId);
+		return 0.0f;
+	}
+
+	if (!SplineComponent)
+	{
+		UE_LOG(LogLearning, Error, TEXT("%s: SplineComponent is nullptr."), *GetName());
+		return 0.0f;
+	}
+
+	const float DistanceAlongSpline = SplineComponent->GetDistanceAlongSplineAtLocation(Position, CoordinateSpace);
+
+#if UE_LEARNING_AGENTS_ENABLE_VISUAL_LOG
+	{
+		const FVector SplinePosition = SplineComponent->GetLocationAtDistanceAlongSpline(DistanceAlongSpline, ESplineCoordinateSpace::World);
+
+		UE_LEARNING_AGENTS_VLOG_LOCATION(this, LogLearning, Display,
+			SplinePosition,
+			10,
+			VisualLogColor.ToFColor(true),
+			TEXT(""));
+
+		UE_LEARNING_AGENTS_VLOG_SEGMENT(this, LogLearning, Display,
+			Position,
+			SplinePosition,
+			VisualLogColor.ToFColor(true),
+			TEXT("GetDistanceAlongSplineAtPosition\nAgent %i\nPosition: [% 6.1f % 6.1f % 6.1f]\nDistance Along Spline: [% 6.1f]"),
+			AgentId,
+			Position.X, Position.Y, Position.Z,
+			DistanceAlongSpline);
+
+		UE_LEARNING_AGENTS_VLOG_LOCATION(this, LogLearning, Display,
+			Position,
+			10,
+			VisualLogColor.ToFColor(true),
+			TEXT(""));
+	}
+#endif
+
+	return DistanceAlongSpline;
+}
+
+FVector USplineComponentHelper::GetPositionAtDistanceAlongSpline(const int32 AgentId, const USplineComponent* SplineComponent, const float DistanceAlongSpline, const ESplineCoordinateSpace::Type CoordinateSpace) const
+{
+	if (!ManagerComponent->HasAgent(AgentId))
+	{
+		UE_LOG(LogLearning, Error, TEXT("%s: AgentId %d not found in the agents set."), *GetName(), AgentId);
+		return FVector::ZeroVector;
+	}
+
+	if (!SplineComponent)
+	{
+		UE_LOG(LogLearning, Error, TEXT("%s: SplineComponent is nullptr."), *GetName());
+		return FVector::ZeroVector;
+	}
+
+	const FVector Position = SplineComponent->GetLocationAtDistanceAlongSpline(DistanceAlongSpline, CoordinateSpace);
+
+#if UE_LEARNING_AGENTS_ENABLE_VISUAL_LOG
+	{
+		UE_LEARNING_AGENTS_VLOG_LOCATION(this, LogLearning, Display,
+			Position,
+			10,
+			VisualLogColor.ToFColor(true),
+			TEXT("GetPositionAtDistanceAlongSpline\nAgent %i\nDistance Along Spline: [% 6.1f]\nPosition: [% 6.1f % 6.1f % 6.1f]"),
+			AgentId,
+			DistanceAlongSpline,
+			Position.X, Position.Y, Position.Z);
+	}
+#endif
+
+	return Position;
+}
+
+FVector USplineComponentHelper::GetDirectionAtDistanceAlongSpline(const int32 AgentId, const USplineComponent* SplineComponent, const float DistanceAlongSpline, const ESplineCoordinateSpace::Type CoordinateSpace) const
+{
+	if (!ManagerComponent->HasAgent(AgentId))
+	{
+		UE_LOG(LogLearning, Error, TEXT("%s: AgentId %d not found in the agents set."), *GetName(), AgentId);
+		return FVector::ZeroVector;
+	}
+
+	if (!SplineComponent)
+	{
+		UE_LOG(LogLearning, Error, TEXT("%s: SplineComponent is nullptr."), *GetName());
+		return FVector::ZeroVector;
+	}
+
+	const FVector Direction = SplineComponent->GetDirectionAtDistanceAlongSpline(DistanceAlongSpline, CoordinateSpace);
+
+#if UE_LEARNING_AGENTS_ENABLE_VISUAL_LOG
+	{
+		const FVector Position = SplineComponent->GetLocationAtDistanceAlongSpline(DistanceAlongSpline, ESplineCoordinateSpace::World);
+
+		UE_LEARNING_AGENTS_VLOG_ARROW(this, LogLearning, Display,
+			Position,
+			Position + Direction * 100.0f,
+			VisualLogColor.ToFColor(true),
+			TEXT("GetDirectionAtDistanceAlongSpline\nAgent %i\nDistance Along Spline: [% 6.1f]\nDirection: [% 6.3f % 6.3f % 6.3f]"),
+			AgentId,
+			DistanceAlongSpline,
+			Direction.X, Direction.Y, Direction.Z);
+	}
+#endif
+
+	return Direction;
+}
+
+float USplineComponentHelper::GetProportionAlongSpline(const int32 AgentId, const USplineComponent* SplineComponent, const float DistanceAlongSpline) const
+{
+	if (!ManagerComponent->HasAgent(AgentId))
+	{
+		UE_LOG(LogLearning, Error, TEXT("%s: AgentId %d not found in the agents set."), *GetName(), AgentId);
+		return 0.0f;
+	}
+
+	if (!SplineComponent)
+	{
+		UE_LOG(LogLearning, Error, TEXT("%s: SplineComponent is nullptr."), *GetName());
+		return 0.0f;
+	}
+
+	const float Proportion = FMath::Clamp(DistanceAlongSpline / FMath::Max(SplineComponent->GetSplineLength(), UE_SMALL_NUMBER), 0.0f, 1.0f);
+
+#if UE_LEARNING_AGENTS_ENABLE_VISUAL_LOG
+	{
+		const FVector Position = SplineComponent->GetLocationAtDistanceAlongSpline(DistanceAlongSpline, ESplineCoordinateSpace::World);
+
+		UE_LEARNING_AGENTS_VLOG_LOCATION(this, LogLearning, Display,
+			Position,
+			10,
+			VisualLogColor.ToFColor(true),
+			TEXT("GetProportionAlongSpline\nAgent %i\nDistance Along Spline: [% 6.1f]\nProportion: [% 6.3f]"),
+			AgentId,
+			DistanceAlongSpline,
+			Proportion);
+	}
+#endif
+
+	return Proportion;
+}
+
+float USplineComponentHelper::GetProportionAlongSplineAsAngle(const int32 AgentId, const USplineComponent* SplineComponent, const float DistanceAlongSpline) const
+{
+	if (!ManagerComponent->HasAgent(AgentId))
+	{
+		UE_LOG(LogLearning, Error, TEXT("%s: AgentId %d not found in the agents set."), *GetName(), AgentId);
+		return 0.0f;
+	}
+
+	if (!SplineComponent)
+	{
+		UE_LOG(LogLearning, Error, TEXT("%s: SplineComponent is nullptr."), *GetName());
+		return 0.0f;
+	}
+
+	if (!SplineComponent->IsClosedLoop())
+	{
+		UE_LOG(LogLearning, Warning, TEXT("%s: Getting proportion along spline as angle, but spline is not closed loop. Consider using GetPropotionAlongSpline instead."), *GetName());
+	}
+
+	const float TotalDistance = SplineComponent->GetSplineLength();
+	const float WrapDistance = FMath::Wrap(DistanceAlongSpline, 0.0f, TotalDistance);
+	const float Proportion = WrapDistance / FMath::Max(TotalDistance, UE_SMALL_NUMBER);
+	const float Angle = FMath::Wrap(UE_TWO_PI * Proportion, -UE_PI, UE_PI);
+
+#if UE_LEARNING_AGENTS_ENABLE_VISUAL_LOG
+	{
+		const FVector Position = SplineComponent->GetLocationAtDistanceAlongSpline(WrapDistance, ESplineCoordinateSpace::World);
+
+		UE_LEARNING_AGENTS_VLOG_ANGLE(this, LogLearning, Display,
+			Angle,
+			0.0f,
+			Position,
+			50.0f,
+			VisualLogColor.ToFColor(true),
+			TEXT("GetProportionAlongSplineAsAngle\nAgent %i\nDistance Along Spline: [% 6.1f]\nProportion: [% 6.3f]\nAngle: [% 6.1f]"),
+			AgentId,
+			WrapDistance,
+			Proportion,
+			FMath::RadiansToDegrees(Angle));
+	}
+#endif
+
+	return FMath::RadiansToDegrees(Angle);
+}
+
+void USplineComponentHelper::GetPositionsAlongSpline(TArray<FVector>& OutPositions, const int32 AgentId, const USplineComponent* SplineComponent, const int32 PositionNum, const float StartDistanceAlongSpline, const float StopDistanceAlongSpline, const ESplineCoordinateSpace::Type CoordinateSpace) const
+{
+	if (!ManagerComponent->HasAgent(AgentId))
+	{
+		UE_LOG(LogLearning, Error, TEXT("%s: AgentId %d not found in the agents set."), *GetName(), AgentId);
+		OutPositions.Empty();
+		return;
+	}
+
+	if (!SplineComponent)
+	{
+		UE_LOG(LogLearning, Error, TEXT("%s: SplineComponent is nullptr."), *GetName());
+		OutPositions.Empty();
+		return;
+	}
+
+	if (PositionNum < 1)
+	{
+		UE_LOG(LogLearning, Warning, TEXT("%s: Expected at least one position along the spline got %i."), *GetName(), PositionNum);
+		OutPositions.Empty();
+		return;
+	}
+
+	OutPositions.SetNumUninitialized(PositionNum);
+
+	const float TotalDistance = SplineComponent->GetSplineLength();
+	const bool bIsClosedLoop = SplineComponent->IsClosedLoop();
+
+	for (int32 PositionIdx = 0; PositionIdx < PositionNum; PositionIdx++)
+	{
+		float PositionDistance = PositionNum == 1 ? 
+			(StartDistanceAlongSpline + StopDistanceAlongSpline) / 2.0f : FMath::Lerp(
+				StartDistanceAlongSpline,
+				StopDistanceAlongSpline,
+				((float)PositionIdx) / (PositionNum - 1));
+
+		PositionDistance = bIsClosedLoop ? FMath::Wrap(PositionDistance, 0.0f, TotalDistance) : PositionDistance;
+
+		OutPositions[PositionIdx] = SplineComponent->GetLocationAtDistanceAlongSpline(PositionDistance, CoordinateSpace);
+	}
+
+#if UE_LEARNING_AGENTS_ENABLE_VISUAL_LOG
+	{
+		for (int32 PositionIdx = 0; PositionIdx < PositionNum; PositionIdx++)
+		{
+			UE_LEARNING_AGENTS_VLOG_LOCATION(this, LogLearning, Display,
+				OutPositions[PositionIdx],
+				2,
+				VisualLogColor.ToFColor(true),
+				TEXT("GetPositionsAlongSpline\nAgent: %i\nPosition: [% 6.1f % 6.1f % 6.1f]"),
+				AgentId,
+				OutPositions[PositionIdx].X, OutPositions[PositionIdx].Y, OutPositions[PositionIdx].Z);
+		}
+	}
+#endif
+}
+
+float USplineComponentHelper::GetVelocityAlongSpline(const int32 AgentId, const USplineComponent* SplineComponent, const FVector Position, const FVector Velocity, const float FiniteDifferenceDelta, const ESplineCoordinateSpace::Type CoordinateSpace) const
+{
+	if (!ManagerComponent->HasAgent(AgentId))
+	{
+		UE_LOG(LogLearning, Error, TEXT("%s: AgentId %d not found in the agents set."), *GetName(), AgentId);
+		return 0.0f;
+	}
+
+	if (!SplineComponent)
+	{
+		UE_LOG(LogLearning, Error, TEXT("%s: SplineComponent is nullptr."), *GetName());
+		return 0.0f;
+	}
+
+	float FiniteDiff = FiniteDifferenceDelta;
+
+	if (FiniteDiff < UE_KINDA_SMALL_NUMBER)
+	{
+		UE_LOG(LogLearning, Warning, TEXT("%s: FiniteDifferenceDelta is too small (%6.4f). Clamping to %6.4f."), *GetName(), FiniteDifferenceDelta, UE_KINDA_SMALL_NUMBER);
+		FiniteDiff = UE_KINDA_SMALL_NUMBER;
+	}
+
+	const float RawDistance0 = SplineComponent->GetDistanceAlongSplineAtLocation(Position, CoordinateSpace);
+	const float RawDistance1 = SplineComponent->GetDistanceAlongSplineAtLocation(Position + FiniteDiff * Velocity.GetSafeNormal(), CoordinateSpace);
+
+	float Distance0 = RawDistance0, Distance1 = RawDistance1;
+
+	if (SplineComponent->IsClosedLoop())
+	{
+		const float SplineDistance = SplineComponent->GetSplineLength();
+
+		if (FMath::Abs(Distance0 - (Distance1 + SplineDistance)) < FMath::Abs(Distance0 - Distance1))
+		{
+			Distance1 = Distance1 + SplineDistance;
+		}
+		else if (FMath::Abs((Distance0 + SplineDistance) - Distance1) < FMath::Abs(Distance0 - Distance1))
+		{
+			Distance0 = Distance0 + SplineDistance;
+		}
+	}
+
+	const float SplineVelocity = ((Distance1 - Distance0) / FiniteDiff) * Velocity.Length();
+
+#if UE_LEARNING_AGENTS_ENABLE_VISUAL_LOG
+	{
+		const FVector VelocityEnd = Position + Velocity;
+		const FVector VelocityDiff = Position + FiniteDiff * Velocity.GetSafeNormal();
+
+		UE_LEARNING_AGENTS_VLOG_ARROW(this, LogLearning, Display,
+			Position,
+			VelocityEnd,
+			VisualLogColor.ToFColor(true),
+			TEXT("GetVelocityAlongSpline\nAgent %i\nPosition [% 6.1f % 6.1f % 6.1f]\nVelocity: [% 6.1f % 6.1f % 6.1f]\nMagnitude: [% 6.1f]"),
+			AgentId,
+			Position.X, Position.Y, Position.Z,
+			Velocity.X, Velocity.Y, Velocity.Z,
+			Velocity.Length());
+
+		UE_LEARNING_AGENTS_VLOG_SEGMENT(this, LogLearning, Display,
+			Position,
+			VelocityDiff,
+			VisualLogColor.ToFColor(true),
+			TEXT(""));
+
+		const FVector Position0 = SplineComponent->GetLocationAtDistanceAlongSpline(RawDistance0, CoordinateSpace);
+		const FVector Position1 = SplineComponent->GetLocationAtDistanceAlongSpline(RawDistance1, CoordinateSpace);
+
+		UE_LEARNING_AGENTS_VLOG_SEGMENT(this, LogLearning, Display,
+			Position0,
+			Position1,
+			VisualLogColor.ToFColor(true),
+			TEXT(""));
+
+		UE_LEARNING_AGENTS_VLOG_SEGMENT(this, LogLearning, Display,
+			Position,
+			Position0,
+			VisualLogColor.ToFColor(true),
+			TEXT(""));
+
+		UE_LEARNING_AGENTS_VLOG_SEGMENT(this, LogLearning, Display,
+			VelocityDiff,
+			Position1,
+			VisualLogColor.ToFColor(true),
+			TEXT(""));
+
+		UE_LEARNING_AGENTS_VLOG_ARROW(this, LogLearning, Display,
+			Position0,
+			Position0 + (Position1 - Position0).GetSafeNormal() * FMath::Abs(SplineVelocity),
+			VisualLogColor.ToFColor(true),
+			TEXT("GetVelocityAlongSpline\nAgent %i\nSpline Velocity: [% 6.1f]"),
+			AgentId,
+			SplineVelocity);
+	}
+#endif
+
+	return SplineVelocity;
+}
+
+//------------------------------------------------------------------
+
+UProjectionHelper* UProjectionHelper::AddProjectionHelper(ULearningAgentsManagerComponent* InManagerComponent, const FName Name)
+{
+	if (!InManagerComponent)
+	{
+		UE_LOG(LogLearning, Error, TEXT("AddProjectionHelper: InManagerComponent is nullptr."));
+		return nullptr;
+	}
+
+	const FName UniqueName = MakeUniqueObjectName(InManagerComponent, UProjectionHelper::StaticClass(), Name, EUniqueObjectNameOptions::GloballyUnique);
+
+	UProjectionHelper* Helper = NewObject<UProjectionHelper>(InManagerComponent, UniqueName);
+	Helper->ManagerComponent = InManagerComponent;
+	InManagerComponent->AddHelper(Helper);
+	return Helper;
+}
+
+FTransform UProjectionHelper::ProjectTransformOntoGroundPlane(const int32 AgentId, const FTransform Transform, const FVector LocalForwardVector) const
+{
+	if (!ManagerComponent->HasAgent(AgentId))
+	{
+		UE_LOG(LogLearning, Error, TEXT("%s: AgentId %d not found in the agents set."), *GetName(), AgentId);
+		return FTransform::Identity;
+	}
+
+	FVector Position = Transform.GetLocation();
+	Position.Z = 0.0f;
+
+	const FVector Direction = (FVector(1.0f, 1.0f, 0.0f) * Transform.TransformVectorNoScale(LocalForwardVector)).GetSafeNormal(UE_SMALL_NUMBER, FVector::ForwardVector);
+
+	const FTransform Projected = FTransform(FQuat::FindBetweenNormals(FVector::ForwardVector, Direction), Position);
+
+#if UE_LEARNING_AGENTS_ENABLE_VISUAL_LOG
+	{
+		UE_LEARNING_AGENTS_VLOG_TRANSFORM(this, LogLearning, Display,
+			Transform.GetLocation(),
+			Transform.GetRotation(),
+			VisualLogColor.ToFColor(true),
+			TEXT("ProjectTransformOntoGroundPlane\nAgent %i\nTransform\nLocal Forward: [% 6.3f % 6.3f % 6.3f]"),
+			LocalForwardVector.X, LocalForwardVector.Y, LocalForwardVector.Z);
+
+		UE_LEARNING_AGENTS_VLOG_ARROW(this, LogLearning, Display,
+			Transform.GetLocation(),
+			Transform.GetLocation() + 25.0f * Transform.GetRotation().RotateVector(LocalForwardVector),
+			VisualLogColor.ToFColor(true),
+			TEXT(""));
+
+		UE_LEARNING_AGENTS_VLOG_SEGMENT(this, LogLearning, Display,
+			Transform.GetLocation(),
+			Projected.GetLocation(),
+			VisualLogColor.ToFColor(true),
+			TEXT(""));
+
+		UE_LEARNING_AGENTS_VLOG_TRANSFORM(this, LogLearning, Display,
+			Projected.GetLocation(),
+			Projected.GetRotation(),
+			VisualLogColor.ToFColor(true),
+			TEXT("ProjectTransformOntoGroundPlane\nAgent %i\nProjected"));
+	}
+#endif
+
+	return Projected;
+}
+
+void UProjectionHelper::ProjectPositionRotationOntoGroundPlane(FVector& OutPosition, FRotator& OutRotation, const int32 AgentId, const FVector InPosition, const FRotator InRotation, const FVector LocalForwardVector) const
+{
+	if (!ManagerComponent->HasAgent(AgentId))
+	{
+		UE_LOG(LogLearning, Error, TEXT("%s: AgentId %d not found in the agents set."), *GetName(), AgentId);
+		OutPosition = FVector::ZeroVector;
+		OutRotation = FRotator::ZeroRotator;
+		return;
+	}
+
+	OutPosition = InPosition;
+	OutPosition.Z = 0.0f;
+
+	const FVector Direction = (FVector(1.0f, 1.0f, 0.0f) * InRotation.RotateVector(LocalForwardVector)).GetSafeNormal(UE_SMALL_NUMBER, FVector::ForwardVector);
+	OutRotation = FQuat::FindBetweenNormals(FVector::ForwardVector, Direction).Rotator();
+
+#if UE_LEARNING_AGENTS_ENABLE_VISUAL_LOG
+	{
+		UE_LEARNING_AGENTS_VLOG_TRANSFORM(this, LogLearning, Display,
+			InPosition,
+			InRotation,
+			VisualLogColor.ToFColor(true),
+			TEXT("ProjectPositionRotationOntoGroundPlane\nAgent %i\nTransform\nLocal Forward: [% 6.3f % 6.3f % 6.3f]"),
+			LocalForwardVector.X, LocalForwardVector.Y, LocalForwardVector.Z);
+
+		UE_LEARNING_AGENTS_VLOG_ARROW(this, LogLearning, Display,
+			InPosition,
+			InPosition + 25.0f * InRotation.RotateVector(LocalForwardVector),
+			VisualLogColor.ToFColor(true),
+			TEXT(""));
+
+		UE_LEARNING_AGENTS_VLOG_SEGMENT(this, LogLearning, Display,
+			InPosition,
+			OutPosition,
+			VisualLogColor.ToFColor(true),
+			TEXT(""));
+
+		UE_LEARNING_AGENTS_VLOG_TRANSFORM(this, LogLearning, Display,
+			OutPosition,
+			OutRotation,
+			VisualLogColor.ToFColor(true),
+			TEXT("ProjectPositionRotationOntoGroundPlane\nAgent %i\nProjected"));
+	}
+#endif
+}
+
+//------------------------------------------------------------------
+
+UMeshComponentHelper* UMeshComponentHelper::AddMeshComponentHelper(ULearningAgentsManagerComponent* InManagerComponent, const FName Name)
+{
+	if (!InManagerComponent)
+	{
+		UE_LOG(LogLearning, Error, TEXT("AddMeshComponentHelper: InManagerComponent is nullptr."));
+		return nullptr;
+	}
+
+	const FName UniqueName = MakeUniqueObjectName(InManagerComponent, UMeshComponentHelper::StaticClass(), Name, EUniqueObjectNameOptions::GloballyUnique);
+
+	UMeshComponentHelper* Helper = NewObject<UMeshComponentHelper>(InManagerComponent, UniqueName);
+	Helper->ManagerComponent = InManagerComponent;
+	InManagerComponent->AddHelper(Helper);
+	return Helper;
+}
+
+void UMeshComponentHelper::GetMeshBonePositions(TArray<FVector>& OutBonePositions, const int32 AgentId, const UMeshComponent* MeshComponent, const TArray<FName>& BoneNames) const
+{
+	if (!ManagerComponent->HasAgent(AgentId))
+	{
+		UE_LOG(LogLearning, Error, TEXT("%s: AgentId %d not found in the agents set."), *GetName(), AgentId);
+		OutBonePositions.Empty();
+		return;
+	}
+
+	if (!MeshComponent)
+	{
+		UE_LOG(LogLearning, Error, TEXT("%s: Mesh Component is nullptr."), *GetName());
+		OutBonePositions.Empty();
+		return;
+	}
+
+	if (BoneNames.IsEmpty())
+	{
+		UE_LOG(LogLearning, Warning, TEXT("%s: No bone names provided!"), *GetName());
+		OutBonePositions.Empty();
+		return;
+	}
+
+	const int32 BoneNum = BoneNames.Num();
+
+	OutBonePositions.SetNumUninitialized(BoneNum);
+
+	for (int32 BoneIdx = 0; BoneIdx < BoneNum; BoneIdx++)
+	{
+		if (MeshComponent->DoesSocketExist(BoneNames[BoneIdx]))
+		{
+			OutBonePositions[BoneIdx] = MeshComponent->GetSocketLocation(BoneNames[BoneIdx]);
+		}
+		else
+		{
+			OutBonePositions[BoneIdx] = FVector::ZeroVector;
+			UE_LOG(LogLearning, Warning, TEXT("%s: Bone \"%s\" does not exist!"), *GetName(), *BoneNames[BoneIdx].ToString());
+		}
+	}
+
+#if UE_LEARNING_AGENTS_ENABLE_VISUAL_LOG
+	{
+		for (int32 BoneIdx = 0; BoneIdx < BoneNum; BoneIdx++)
+		{
+			UE_LEARNING_AGENTS_VLOG_LOCATION(this, LogLearning, Display,
+				OutBonePositions[BoneIdx],
+				2,
+				VisualLogColor.ToFColor(true),
+				TEXT("Agent %i\nBone: \"%s\"\nPosition: [% 6.1f % 6.1f % 6.1f]"),
+				AgentId,
+				*BoneNames[BoneIdx].ToString(),
+				OutBonePositions[BoneIdx].X, OutBonePositions[BoneIdx].Y, OutBonePositions[BoneIdx].Z);
+		}
+	}
+#endif
+}
+
+//------------------------------------------------------------------
+
+URayCastHelper* URayCastHelper::AddRayCastHelper(ULearningAgentsManagerComponent* InManagerComponent, const FName Name)
+{
+	if (!InManagerComponent)
+	{
+		UE_LOG(LogLearning, Error, TEXT("AddRayCastHelper: InManagerComponent is nullptr."));
+		return nullptr;
+	}
+
+	const FName UniqueName = MakeUniqueObjectName(InManagerComponent, URayCastHelper::StaticClass(), Name, EUniqueObjectNameOptions::GloballyUnique);
+
+	URayCastHelper* Helper = NewObject<URayCastHelper>(InManagerComponent, UniqueName);
+	Helper->ManagerComponent = InManagerComponent;
+	InManagerComponent->AddHelper(Helper);
+	return Helper;
+}
+
+void URayCastHelper::RayCastGridHeights(
+	TArray<float>& OutHeights, 
+	const int32 AgentId,
+	const FVector Position, 
+	const FRotator Rotation, 
+	const int32 RowNum, 
+	const int32 ColNum, 
+	const float RowWidth, 
+	const float ColWidth, 
+	const float MaxHeight, 
+	const float MinHeight,
+	const ECollisionChannel CollisionChannel) const
+{
+	if (!ManagerComponent->HasAgent(AgentId))
+	{
+		UE_LOG(LogLearning, Error, TEXT("%s: AgentId %d not found in the agents set."), *GetName(), AgentId);
+		OutHeights.Empty();
+		return;
+	}
+
+	if (RowNum < 1 || ColNum < 1)
+	{
+		UE_LOG(LogLearning, Error, TEXT("%s: Need at least 1 row and 1 col in grid."), *GetName());
+		OutHeights.Empty();
+		return;
+	}
+
+	OutHeights.SetNumUninitialized(ColNum * RowNum);
+
+	TArray<FHitResult, TInlineAllocator<25>> TraceHits;
+	TraceHits.SetNum(ColNum * RowNum);
+
+	const FRotator YawRotator(0.0f, Rotation.Yaw, 0.0f);
+
+	for (int32 RowIdx = 0; RowIdx < RowNum; RowIdx++)
+	{
+		for (int32 ColIdx = 0; ColIdx < ColNum; ColIdx++)
+		{
+			const FVector Start = Position + YawRotator.RotateVector(FVector(
+				RowWidth * (((float)RowIdx) / (RowNum - 1) - 0.5f),
+				ColWidth * (((float)ColIdx) / (ColNum - 1) - 0.5f),
+				MaxHeight));
+
+			const FVector End = Position + YawRotator.RotateVector(FVector(
+				RowWidth * (((float)RowIdx) / (RowNum - 1) - 0.5f),
+				ColWidth * (((float)ColIdx) / (ColNum - 1) - 0.5f),
+				MinHeight));
+
+			FCollisionObjectQueryParams ObjectQueryParams;
+			ObjectQueryParams.AddObjectTypesToQuery(CollisionChannel);
+
+			const bool bHit = GetWorld()->LineTraceSingleByObjectType(TraceHits[RowIdx * ColNum + ColIdx], Start, End, ObjectQueryParams);
+
+			if (bHit)
+			{
+				OutHeights[RowIdx * ColNum + ColIdx] = TraceHits[RowIdx * ColNum + ColIdx].ImpactPoint.Z;
+			}
+			else
+			{
+				OutHeights[RowIdx * ColNum + ColIdx] = MinHeight;
+			}
+		}
+	}
+
+#if UE_LEARNING_AGENTS_ENABLE_VISUAL_LOG
+	{
+		UE_LEARNING_AGENTS_VLOG_TRANSFORM(this, LogLearning, Display,
+			Position,
+			Rotation,
+			VisualLogColor.ToFColor(true),
+			TEXT(""));
+
+		for (int32 RowIdx = 0; RowIdx < RowNum; RowIdx++)
+		{
+			for (int32 ColIdx = 0; ColIdx < ColNum; ColIdx++)
+			{
+				const FVector Start = Position + YawRotator.RotateVector(FVector(
+					RowWidth * (((float)RowIdx) / (RowNum - 1) - 0.5f),
+					ColWidth * (((float)ColIdx) / (ColNum - 1) - 0.5f),
+					MaxHeight));
+
+				const FVector End = Position + YawRotator.RotateVector(FVector(
+					RowWidth * (((float)RowIdx) / (RowNum - 1) - 0.5f),
+					ColWidth * (((float)ColIdx) / (ColNum - 1) - 0.5f),
+					MinHeight));
+
+				const FVector ImpactPoint = 
+					TraceHits[RowIdx * ColNum + ColIdx].bBlockingHit ?
+					TraceHits[RowIdx * ColNum + ColIdx].ImpactPoint : End;
+
+				UE_LEARNING_AGENTS_VLOG_ARROW(this, LogLearning, Display,
+					Start,
+					End,
+					VisualLogColor.ToFColor(true),
+					TEXT(""));
+
+				UE_LEARNING_AGENTS_VLOG_LOCATION(this, LogLearning, Display,
+					ImpactPoint,
+					2,
+					VisualLogColor.ToFColor(true),
+					TEXT("Agent %i\nRow: %i\nCol: %i\nImpact Point: [% 6.1f % 6.1f % 6.1f]\nHeight: [% 6.1f]"),
+					AgentId,
+					RowIdx, ColIdx,
+					ImpactPoint.X, ImpactPoint.Y, ImpactPoint.Z,
+					OutHeights[RowIdx * ColNum + ColIdx]);
+			}
+		}
+	}
+#endif
+}
+
+void URayCastHelper::RayCastRadial(
+	TArray<float>& OutDistances,
+	const int32 AgentId,
+	const FVector Position,
+	const FRotator Rotation,
+	const int32 RayNum,
+	const float MinAngle,
+	const float MaxAngle,
+	const float MaxRayDist,
+	const FVector LocalForward,
+	const ECollisionChannel CollisionChannel) const
+{
+	if (!ManagerComponent->HasAgent(AgentId))
+	{
+		UE_LOG(LogLearning, Error, TEXT("%s: AgentId %d not found in the agents set."), *GetName(), AgentId);
+		OutDistances.Empty();
+		return;
+	}
+
+	if (RayNum < 1)
+	{
+		UE_LOG(LogLearning, Error, TEXT("%s: Need at least 1 ray."), *GetName());
+		OutDistances.Empty();
+		return;
+	}
+
+	OutDistances.SetNumUninitialized(RayNum);
+
+	TArray<FHitResult, TInlineAllocator<25>> TraceHits;
+	TraceHits.SetNum(RayNum);
+
+
+	for (int32 RayIdx = 0; RayIdx < RayNum; RayIdx++)
+	{
+		const float Alpha = RayNum == 1 ? 0.0f : ((float)RayIdx) / (RayNum - 1);
+		const FRotator YawRotator(0.0f, Rotation.Yaw + Alpha * (MaxAngle - MinAngle) + MinAngle, 0.0f);
+
+		const FVector Start = Position;
+		const FVector End = Position + YawRotator.RotateVector(MaxRayDist * LocalForward.GetSafeNormal(UE_SMALL_NUMBER, FVector::ForwardVector));
+
+		FCollisionObjectQueryParams ObjectQueryParams;
+		ObjectQueryParams.AddObjectTypesToQuery(CollisionChannel);
+
+		const bool bHit = GetWorld()->LineTraceSingleByObjectType(TraceHits[RayIdx], Start, End, ObjectQueryParams);
+
+		if (bHit)
+		{
+			OutDistances[RayIdx] = TraceHits[RayIdx].Distance;
+		}
+		else
+		{
+			OutDistances[RayIdx] = MaxRayDist;
+		}
+	}
+
+#if UE_LEARNING_AGENTS_ENABLE_VISUAL_LOG
+	{
+		UE_LEARNING_AGENTS_VLOG_TRANSFORM(this, LogLearning, Display,
+			Position,
+			Rotation,
+			VisualLogColor.ToFColor(true),
+			TEXT(""));
+
+		for (int32 RayIdx = 0; RayIdx < RayNum; RayIdx++)
+		{
+			const float Alpha = RayNum == 1 ? 0.0f : ((float)RayIdx) / (RayNum - 1);
+			const FRotator YawRotator(0.0f, Rotation.Yaw + Alpha * (MaxAngle - MinAngle) + MinAngle, 0.0f);
+
+			const FVector Start = Position;
+			const FVector End = Position + YawRotator.RotateVector(MaxRayDist * LocalForward.GetSafeNormal(UE_SMALL_NUMBER, FVector::ForwardVector));
+
+			const FVector ImpactPoint =
+				TraceHits[RayIdx].bBlockingHit ?
+				TraceHits[RayIdx].ImpactPoint : End;
+
+			UE_LEARNING_AGENTS_VLOG_ARROW(this, LogLearning, Display,
+				Start,
+				End,
+				VisualLogColor.ToFColor(true),
+				TEXT(""));
+
+			UE_LEARNING_AGENTS_VLOG_LOCATION(this, LogLearning, Display,
+				ImpactPoint,
+				2,
+				VisualLogColor.ToFColor(true),
+				TEXT("Agent %i\nRay: %i\nImpact Point: [% 6.1f % 6.1f % 6.1f]\nDistance: [% 6.1f]"),
+				AgentId,
+				RayIdx,
+				ImpactPoint.X, ImpactPoint.Y, ImpactPoint.Z,
+				OutDistances[RayIdx]);
+		}
+	}
+#endif
+}
+
+
+
+UCollisionMonitorHelper* UCollisionMonitorHelper::AddCollisionMonitorHelper(
+	ULearningAgentsManagerComponent* InManagerComponent,
+	const FName Name)
+{
+	if (!InManagerComponent)
+	{
+		UE_LOG(LogLearning, Error, TEXT("AddCollisionMonitorHelper: InManagerComponent is nullptr."));
+		return nullptr;
+	}
+
+	const FName UniqueName = MakeUniqueObjectName(InManagerComponent, UCollisionMonitorHelper::StaticClass(), Name, EUniqueObjectNameOptions::GloballyUnique);
+
+	UCollisionMonitorHelper* Helper = NewObject<UCollisionMonitorHelper>(InManagerComponent, UniqueName);
+	Helper->Init(InManagerComponent->GetAgentManager()->GetMaxAgentNum());
+	Helper->ManagerComponent = InManagerComponent;
+
+	InManagerComponent->AddHelper(Helper);
+	return Helper;
+}
+
+void UCollisionMonitorHelper::Init(const int32 MaxAgentNum)
+{
+	Components.SetNumUninitialized({ MaxAgentNum });
+	OtherComponentTags.SetNumUninitialized({ MaxAgentNum });
+	CollisionsOccured.SetNumUninitialized({ MaxAgentNum });
+
+	UE::Learning::Array::Set<1, UPrimitiveComponent*>(Components, nullptr);
+	UE::Learning::Array::Set<1, FName>(OtherComponentTags, NAME_None);
+	UE::Learning::Array::Set<1, bool>(CollisionsOccured, false);
+}
+
+void UCollisionMonitorHelper::OnAgentsRemoved(const TArray<int32>& AgentIds)
+{
+	ULearningAgentsHelper::OnAgentsRemoved(AgentIds);
+
+	for (const int32 AgentId : AgentIds)
+	{
+		if (Components[AgentId])
+		{
+			Components[AgentId]->OnComponentHit.RemoveDynamic(this, &UCollisionMonitorHelper::HandleOnHit);
+		}
+	}
+
+	UE::Learning::Array::Set<1, UPrimitiveComponent*>(Components, nullptr);
+	UE::Learning::Array::Set<1, FName>(OtherComponentTags, NAME_None, AgentIds);
+	UE::Learning::Array::Set<1, bool>(CollisionsOccured, false, AgentIds);
+}
+
+void UCollisionMonitorHelper::OnAgentsReset(const TArray<int32>& AgentIds)
+{
+	ULearningAgentsHelper::OnAgentsReset(AgentIds);
+
+	UE::Learning::Array::Set(CollisionsOccured, false, AgentIds);
+}
+
+void UCollisionMonitorHelper::SetComponent(const int32 AgentId, UPrimitiveComponent* Component, const FName OtherComponentTag)
+{
+	if (!ManagerComponent->HasAgent(AgentId))
+	{
+		UE_LOG(LogLearning, Error, TEXT("%s: AgentId %d not found in the agents set."), *GetName(), AgentId);
+		return;
+	}
+
+	if (Components[AgentId] != Component)
+	{
+		if (Components[AgentId])
+		{
+			Components[AgentId]->OnComponentHit.RemoveDynamic(this, &UCollisionMonitorHelper::HandleOnHit);
+		}
+
+		Components[AgentId] = Component;
+		OtherComponentTags[AgentId] = OtherComponentTag;
+		Component->OnComponentHit.AddUniqueDynamic(this, &UCollisionMonitorHelper::HandleOnHit);
+	}
+}
+
+bool UCollisionMonitorHelper::GetCollisionOccurred(const int32 AgentId) const
+{
+	if (!ManagerComponent->HasAgent(AgentId))
+	{
+		UE_LOG(LogLearning, Error, TEXT("%s: AgentId %d not found in the agents set."), *GetName(), AgentId);
+		return false;
+	}
+
+	return CollisionsOccured[AgentId];
+}
+
+bool UCollisionMonitorHelper::GetAndResetCollisionOccurred(const int32 AgentId)
+{
+	if (!ManagerComponent->HasAgent(AgentId))
+	{
+		UE_LOG(LogLearning, Error, TEXT("%s: AgentId %d not found in the agents set."), *GetName(), AgentId);
+		return false;
+	}
+
+	bool bOccurred = CollisionsOccured[AgentId];
+	CollisionsOccured[AgentId] = false;
+	return bOccurred;
+}
+
+void UCollisionMonitorHelper::ResetCollisionOccurred(const int32 AgentId)
+{
+	if (!ManagerComponent->HasAgent(AgentId))
+	{
+		UE_LOG(LogLearning, Error, TEXT("%s: AgentId %d not found in the agents set."), *GetName(), AgentId);
+		return;
+	}
+
+	CollisionsOccured[AgentId] = false;
+}
+
+void UCollisionMonitorHelper::HandleOnHit(UPrimitiveComponent* HitComponent, AActor* OtherActor, UPrimitiveComponent* OtherComp, FVector NormalImpulse, const FHitResult& Hit)
+{
+	for (int32 AgentId = 0; AgentId < Components.Num(); AgentId++)
+	{
+		if (Components[AgentId] == HitComponent)
+		{
+			CollisionsOccured[AgentId] = true;
+
+#if UE_LEARNING_AGENTS_ENABLE_VISUAL_LOG
+			{
+				UE_LEARNING_AGENTS_VLOG_ARROW(this, LogLearning, Display,
+					Hit.ImpactPoint,
+					Hit.ImpactPoint + Hit.ImpactNormal * 100.0f,
+					VisualLogColor.ToFColor(true),
+					TEXT("Agent %i\nCollision \"%s\" and \"%s\"\nTag: \"%s\""),
+					AgentId,
+					*HitComponent->GetName(),
+					*OtherComp->GetName(),
+					*OtherComponentTags[AgentId].ToString());
+			}
+#endif
+			return;
+		}
+	}
+
+	UE_LOG(LogLearning, Error, TEXT("%s: HitComponent not found in the set of agent's Components."), *GetName());
+}
\ No newline at end of file
diff --git a/Engine/Plugins/Experimental/LearningAgents/Source/LearningAgents/Private/LearningAgentsInteractor.cpp b/Engine/Plugins/Experimental/LearningAgents/Source/LearningAgents/Private/LearningAgentsInteractor.cpp
--- a/Engine/Plugins/Experimental/LearningAgents/Source/LearningAgents/Private/LearningAgentsInteractor.cpp
+++ b/Engine/Plugins/Experimental/LearningAgents/Source/LearningAgents/Private/LearningAgentsInteractor.cpp
@@ -5,46 +5,19 @@
 #include "LearningAgentsManager.h"
 #include "LearningAgentsObservations.h"
 #include "LearningAgentsActions.h"
-#include "LearningAgentsNeuralNetwork.h"
+#include "LearningAgentsHelpers.h"
 
-#include "LearningNeuralNetwork.h"
 #include "LearningArray.h"
+#include "LearningArrayMap.h"
+#include "LearningFeatureObject.h"
 #include "LearningLog.h"
-
 #include "EngineDefines.h"
 
 ULearningAgentsInteractor::ULearningAgentsInteractor() : Super(FObjectInitializer::Get()) {}
 ULearningAgentsInteractor::ULearningAgentsInteractor(FVTableHelper& Helper) : Super(Helper) {}
 ULearningAgentsInteractor::~ULearningAgentsInteractor() = default;
 
-ULearningAgentsInteractor* ULearningAgentsInteractor::MakeInteractor(
-	ULearningAgentsManager* InManager,
-	TSubclassOf<ULearningAgentsInteractor> Class,
-	const FName Name)
-{
-	if (!InManager)
-	{
-		UE_LOG(LogLearning, Error, TEXT("MakeInteractor: InManager is nullptr."));
-		return nullptr;
-	}
-
-	if (!Class)
-	{
-		UE_LOG(LogLearning, Error, TEXT("MakeInteractor: Class is nullptr."));
-		return nullptr;
-	}
-
-	const FName UniqueName = MakeUniqueObjectName(InManager, Class, Name, EUniqueObjectNameOptions::GloballyUnique);
-
-	ULearningAgentsInteractor* Interactor = NewObject<ULearningAgentsInteractor>(InManager, Class, UniqueName);
-	if (!Interactor) { return nullptr; }
-
-	Interactor->SetupInteractor(InManager);
-
-	return Interactor->IsSetup() ? Interactor : nullptr;
-}
-
-void ULearningAgentsInteractor::SetupInteractor(ULearningAgentsManager* InManager)
+void ULearningAgentsInteractor::SetupInteractor()
 {
 	if (IsSetup())
 	{
@@ -52,260 +25,298 @@ void ULearningAgentsInteractor::SetupInteractor(ULearningAgentsManager* InManage
 		return;
 	}
 
-	if (!InManager)
+	if (!Manager)
 	{
-		UE_LOG(LogLearning, Error, TEXT("%s: InManager is nullptr."), *GetName());
+		UE_LOG(LogLearning, Error, TEXT("%s: Must be attached to a LearningAgentsManager Actor."), *GetName());
 		return;
 	}
 
-	Manager = InManager;
+	// Setup Observations
+	ObservationObjects.Empty();
+	ObservationFeatures.Empty();
+	SetupObservations();
 
-	// Specify Observations
-
-	const FName ObservationSchemaUniqueName = MakeUniqueObjectName(this, ULearningAgentsObservationSchema::StaticClass(), TEXT("ObservationSchema"), EUniqueObjectNameOptions::GloballyUnique);
-
-	ObservationSchema = NewObject<ULearningAgentsObservationSchema>(this, ObservationSchemaUniqueName);
-	SpecifyAgentObservations(ObservationSchemaElement, ObservationSchema);
-
-	if (!ObservationSchema->GetObservationSchema().IsValid(ObservationSchemaElement.SchemaElement))
+	if (ObservationObjects.Num() == 0)
 	{
-		UE_LOG(LogLearning, Error, TEXT("%s: Invalid observation provided to Interactor during SpecifyObservations."), *GetName());
+		UE_LOG(LogLearning, Error, TEXT("%s: No observations added to Interactor during SetupObservations."), *GetName());
 		return;
 	}
 
-	const int32 ObservationVectorSize = ObservationSchema->GetObservationSchema().GetObservationVectorSize(ObservationSchemaElement.SchemaElement);
+	Observations = MakeShared<UE::Learning::FConcatenateFeature>(*(GetName() + TEXT("Observations")),
+		TLearningArrayView<1, const TSharedRef<UE::Learning::FFeatureObject>>(ObservationFeatures),
+		Manager->GetInstanceData().ToSharedRef(),
+		Manager->GetMaxAgentNum());
 
-	if (ObservationVectorSize == 0)
+	if (Observations->DimNum() == 0)
 	{
-		UE_LOG(LogLearning, Error, TEXT("%s: Observation vector is zero-sized - specified observations have no size."), *GetName());
+		UE_LOG(LogLearning, Error, TEXT("%s: Observation vector is zero-sized - all added observations have no size."), *GetName());
 		return;
 	}
 
-	const int32 ObservationEncodedVectorSize = ObservationSchema->GetObservationSchema().GetEncodedVectorSize(ObservationSchemaElement.SchemaElement);
+	// Setup Actions
+	ActionObjects.Empty();
+	ActionFeatures.Empty();
+	SetupActions();
 
-	if (ObservationEncodedVectorSize == 0)
+	if (ActionObjects.Num() == 0)
 	{
-		UE_LOG(LogLearning, Error, TEXT("%s: Observation encoded vector is zero-sized - observations map to empty encoding."), *GetName());
+		UE_LOG(LogLearning, Error, TEXT("%s: No actions added to Interactor during SetupActions."), *GetName());
 		return;
 	}
 
-	ObservationVectors.SetNumUninitialized({ Manager->GetMaxAgentNum(), ObservationVectorSize  });
-	ObservationCompatibilityHash = UE::Learning::Observation::GetSchemaObjectsCompatibilityHash(ObservationSchema->GetObservationSchema(), ObservationSchemaElement.SchemaElement);
+	Actions = MakeShared<UE::Learning::FConcatenateFeature>(*(GetName() + TEXT("Actions")),
+		TLearningArrayView<1, const TSharedRef<UE::Learning::FFeatureObject>>(ActionFeatures),
+		Manager->GetInstanceData().ToSharedRef(),
+		Manager->GetMaxAgentNum());
 
-	const FName ObservationObjectUniqueName = MakeUniqueObjectName(this, ULearningAgentsObservationObject::StaticClass(), TEXT("ObservationObject"), EUniqueObjectNameOptions::GloballyUnique);
+	if (Actions->DimNum() == 0)
+	{
+		UE_LOG(LogLearning, Error, TEXT("%s: Action vector is zero-sized -  all added actions have no size."), *GetName());
+		return;
+	}
 
-	ObservationObject = NewObject<ULearningAgentsObservationObject>(this, ObservationObjectUniqueName);
-	ObservationObjectElements.Empty(Manager->GetMaxAgentNum());
+	// Reset Agent iteration
 
-	// Specify Actions
+	ObservationEncodingAgentIteration.SetNumUninitialized({ Manager->GetMaxAgentNum() });
+	ActionEncodingAgentIteration.SetNumUninitialized({ Manager->GetMaxAgentNum() });
+	UE::Learning::Array::Set<1, uint64>(ObservationEncodingAgentIteration, INDEX_NONE);
+	UE::Learning::Array::Set<1, uint64>(ActionEncodingAgentIteration, INDEX_NONE);
 
-	const FName ActionSchemaUniqueName = MakeUniqueObjectName(this, ULearningAgentsActionSchema::StaticClass(), TEXT("ActionSchema"), EUniqueObjectNameOptions::GloballyUnique);
+	bIsSetup = true;
 
-	ActionSchema = NewObject<ULearningAgentsActionSchema>(this, ActionSchemaUniqueName);
-	SpecifyAgentActions(ActionSchemaElement, ActionSchema);
+	OnAgentsAdded(Manager->GetAllAgentIds());
+}
 
-	if (!ActionSchema->GetActionSchema().IsValid(ActionSchemaElement.SchemaElement))
+void ULearningAgentsInteractor::OnAgentsAdded(const TArray<int32>& AgentIds)
+{
+	if (IsSetup())
 	{
-		UE_LOG(LogLearning, Error, TEXT("%s: Invalid action provided to Interactor during SpecifyActions."), *GetName());
-		return;
-	}
+		UE::Learning::Array::Set<1, uint64>(ObservationEncodingAgentIteration, 0, AgentIds);
+		UE::Learning::Array::Set<1, uint64>(ActionEncodingAgentIteration, 0, AgentIds);
 
-	const int32 ActionVectorSize = ActionSchema->GetActionSchema().GetActionVectorSize(ActionSchemaElement.SchemaElement);
+		for (ULearningAgentsHelper* Helper : HelperObjects)
+		{
+			Helper->OnAgentsAdded(AgentIds);
+		}
 
-	if (ActionVectorSize == 0)
-	{
-		UE_LOG(LogLearning, Error, TEXT("%s: Action vector is zero-sized - specified actions have no size."), *GetName());
-		return;
-	}
+		for (ULearningAgentsObservation* ObservationObject : ObservationObjects)
+		{
+			ObservationObject->OnAgentsAdded(AgentIds);
+		}
 
-	const int32 ActionEncodedVectorSize = ActionSchema->GetActionSchema().GetEncodedVectorSize(ActionSchemaElement.SchemaElement);
+		for (ULearningAgentsAction* ActionObject : ActionObjects)
+		{
+			ActionObject->OnAgentsAdded(AgentIds);
+		}
 
-	if (ActionEncodedVectorSize == 0)
-	{
-		UE_LOG(LogLearning, Error, TEXT("%s: Action encoded vector is zero-sized - actions map to empty encoding."), *GetName());
-		return;
+		AgentsAdded(AgentIds);
 	}
+}
 
-	const int32 ActionDistributionVectorSize = ActionSchema->GetActionSchema().GetActionDistributionVectorSize(ActionSchemaElement.SchemaElement);
+void ULearningAgentsInteractor::OnAgentsRemoved(const TArray<int32>& AgentIds)
+{
+	if (IsSetup())
+	{
+		UE::Learning::Array::Set<1, uint64>(ObservationEncodingAgentIteration, INDEX_NONE, AgentIds);
+		UE::Learning::Array::Set<1, uint64>(ActionEncodingAgentIteration, INDEX_NONE, AgentIds);
 
-	ActionVectors.SetNumUninitialized({ Manager->GetMaxAgentNum(), ActionVectorSize });
-	ActionCompatibilityHash = UE::Learning::Action::GetSchemaObjectsCompatibilityHash(ActionSchema->GetActionSchema(), ActionSchemaElement.SchemaElement);
+		for (ULearningAgentsHelper* Helper : HelperObjects)
+		{
+			Helper->OnAgentsRemoved(AgentIds);
+		}
 
-	const FName ActionObjectUniqueName = MakeUniqueObjectName(this, ULearningAgentsActionObject::StaticClass(), TEXT("ActionObject"), EUniqueObjectNameOptions::GloballyUnique);
+		for (ULearningAgentsObservation* ObservationObject : ObservationObjects)
+		{
+			ObservationObject->OnAgentsRemoved(AgentIds);
+		}
 
-	ActionObject = NewObject<ULearningAgentsActionObject>(this, ActionObjectUniqueName);
-	ActionObjectElements.Empty(Manager->GetMaxAgentNum());
+		for (ULearningAgentsAction* ActionObject : ActionObjects)
+		{
+			ActionObject->OnAgentsRemoved(AgentIds);
+		}
 
-	// Reset Agent iteration
+		AgentsRemoved(AgentIds);
+	}
+}
 
-	ObservationVectorIteration.SetNumUninitialized({ Manager->GetMaxAgentNum() });
-	ActionVectorIteration.SetNumUninitialized({ Manager->GetMaxAgentNum() });
+void ULearningAgentsInteractor::OnAgentsReset(const TArray<int32>& AgentIds)
+{
+	if (IsSetup())
+	{
+		UE::Learning::Array::Set<1, uint64>(ObservationEncodingAgentIteration, 0, AgentIds);
+		UE::Learning::Array::Set<1, uint64>(ActionEncodingAgentIteration, 0, AgentIds);
+		
+		for (ULearningAgentsHelper* Helper : HelperObjects)
+		{
+			Helper->OnAgentsReset(AgentIds);
+		}
 
-	UE::Learning::Array::Set<1, uint64>(ObservationVectorIteration, INDEX_NONE);
-	UE::Learning::Array::Set<1, uint64>(ActionVectorIteration, INDEX_NONE);
+		for (ULearningAgentsObservation* ObservationObject : ObservationObjects)
+		{
+			ObservationObject->OnAgentsReset(AgentIds);
+		}
 
-	bIsSetup = true;
+		for (ULearningAgentsAction* ActionObject : ActionObjects)
+		{
+			ActionObject->OnAgentsReset(AgentIds);
+		}
 
-	Manager->AddListener(this);
+		AgentsReset(AgentIds);
+	}
 }
 
-void ULearningAgentsInteractor::OnAgentsAdded_Implementation(const TArray<int32>& AgentIds)
+UE::Learning::FFeatureObject& ULearningAgentsInteractor::GetObservationFeature() const
 {
-	if (!IsSetup())
-	{
-		UE_LOG(LogLearning, Error, TEXT("%s: Setup not complete."), *GetName());
-		return;
-	}
+	return *Observations;
+}
 
-	UE::Learning::Array::Set<1, uint64>(ObservationVectorIteration, 0, AgentIds);
-	UE::Learning::Array::Set<1, uint64>(ActionVectorIteration, 0, AgentIds);
+UE::Learning::FFeatureObject& ULearningAgentsInteractor::GetActionFeature() const
+{
+	return *Actions;
 }
 
-void ULearningAgentsInteractor::OnAgentsRemoved_Implementation(const TArray<int32>& AgentIds)
+TConstArrayView<ULearningAgentsObservation*> ULearningAgentsInteractor::GetObservationObjects() const
 {
-	if (!IsSetup())
-	{
-		UE_LOG(LogLearning, Error, TEXT("%s: Setup not complete."), *GetName());
-		return;
-	}
+	return ObservationObjects;
+}
 
-	UE::Learning::Array::Set<1, uint64>(ObservationVectorIteration, INDEX_NONE, AgentIds);
-	UE::Learning::Array::Set<1, uint64>(ActionVectorIteration, INDEX_NONE, AgentIds);
+TConstArrayView<ULearningAgentsAction*> ULearningAgentsInteractor::GetActionObjects() const
+{
+	return ActionObjects;
 }
 
-void ULearningAgentsInteractor::OnAgentsReset_Implementation(const TArray<int32>& AgentIds)
+TLearningArrayView<1, uint64> ULearningAgentsInteractor::GetObservationEncodingAgentIteration()
 {
-	if (!IsSetup())
-	{
-		UE_LOG(LogLearning, Error, TEXT("%s: Setup not complete."), *GetName());
-		return;
-	}
+	return ObservationEncodingAgentIteration;
+}
 
-	UE::Learning::Array::Set<1, uint64>(ObservationVectorIteration, 0, AgentIds);
-	UE::Learning::Array::Set<1, uint64>(ActionVectorIteration, 0, AgentIds);
+TLearningArrayView<1, uint64> ULearningAgentsInteractor::GetActionEncodingAgentIteration()
+{
+	return ActionEncodingAgentIteration;
 }
 
-void ULearningAgentsInteractor::SpecifyAgentObservations_Implementation(FLearningAgentsObservationSchemaElement& OutObservationSchemaElement, ULearningAgentsObservationSchema* InObservationSchema)
+void ULearningAgentsInteractor::SetupObservations_Implementation()
 {
-	UE_LOG(LogLearning, Error, TEXT("%s: SpecifyAgentObservations function must be overridden!"), *GetName());
-	OutObservationSchemaElement = FLearningAgentsObservationSchemaElement();
+	// Can be overridden to setup observations without blueprints
 }
 
-void ULearningAgentsInteractor::GatherAgentObservations_Implementation(TArray<FLearningAgentsObservationObjectElement>& OutObservationObjectElements, ULearningAgentsObservationObject* InObservationObject, const TArray<int32>& AgentIds)
+void ULearningAgentsInteractor::SetObservations_Implementation(const TArray<int32>& AgentIds)
 {
-	UE_LOG(LogLearning, Error, TEXT("%s: GatherAgentObservations function must be overridden!"), *GetName());
-	OutObservationObjectElements.Empty();
+	// Can be overridden to set observations without blueprints
 }
 
-void ULearningAgentsInteractor::GatherAgentObservationsUsingDelegate(
-	TArray<FLearningAgentsObservationObjectElement>& OutObservationElements,
-	ULearningAgentsObservationObject* InObservationObject,
-	const TArray<int32>& AgentIds,
-	const FGatherAgentObservationDelegate& Delegate)
+void ULearningAgentsInteractor::AddObservation(TObjectPtr<ULearningAgentsObservation> Object, const TSharedRef<UE::Learning::FFeatureObject>& Feature)
 {
-	if (!Delegate.IsBound())
-	{
-		UE_LOG(LogLearning, Error, TEXT("%s: Delegate Not Bound."), *GetName());
-		OutObservationElements.Empty();
-		return;
-	}
+	UE_LEARNING_CHECK(!IsSetup());
+	UE_LEARNING_CHECK(Object);
+	ObservationObjects.Add(Object);
+	ObservationFeatures.Add(Feature);
+}
 
-	OutObservationElements.Empty(AgentIds.Num());
-	for (const int32 AgentId : AgentIds)
-	{
-		OutObservationElements.Add(Delegate.Execute(InObservationObject, AgentId));
-	}
+void ULearningAgentsInteractor::SetupActions_Implementation()
+{
+	// Can be overridden to setup actions without blueprints
 }
 
-void ULearningAgentsInteractor::SpecifyAgentActions_Implementation(FLearningAgentsActionSchemaElement& OutActionSchemaElement, ULearningAgentsActionSchema* InActionSchema)
+void ULearningAgentsInteractor::GetActions_Implementation(const TArray<int32>& AgentIds)
 {
-	UE_LOG(LogLearning, Error, TEXT("%s: SpecifyAgentActions function must be overridden!"), *GetName());
-	OutActionSchemaElement = FLearningAgentsActionSchemaElement();
+	// Can be overridden to get actions without blueprints
 }
 
-void ULearningAgentsInteractor::ScatterAgentActions_Implementation(const TArray<FLearningAgentsActionObjectElement>& InActionObjectElements, const ULearningAgentsActionObject* InActionObject, const TArray<int32>& AgentIds)
+void ULearningAgentsInteractor::AddAction(TObjectPtr<ULearningAgentsAction> Object, const TSharedRef<UE::Learning::FFeatureObject>& Feature)
 {
-	UE_LOG(LogLearning, Error, TEXT("%s: ScatterAgentActions function must be overridden!"), *GetName());
+	UE_LEARNING_CHECK(!IsSetup());
+	UE_LEARNING_CHECK(Object);
+	ActionObjects.Add(Object);
+	ActionFeatures.Add(Feature);
 }
 
-void  ULearningAgentsInteractor::ScatterAgentObservationsUsingDelegate(
-	const TArray<FLearningAgentsActionObjectElement>& InActionObjectElements,
-	const ULearningAgentsActionObject* InActionObject,
-	const TArray<int32>& AgentIds,
-	const FScatterAgentActionDelegate& Delegate)
+void ULearningAgentsInteractor::EncodeObservations(const UE::Learning::FIndexSet AgentSet)
 {
-	if (!Delegate.IsBound())
+	UE_LEARNING_TRACE_CPUPROFILER_EVENT_SCOPE(ULearningAgentsInteractor::EncodeObservations);
+
+	if (!IsSetup())
 	{
-		UE_LOG(LogLearning, Error, TEXT("%s: Delegate Not Bound."), *GetName());
+		UE_LOG(LogLearning, Error, TEXT("%s: Setup not complete."), *GetName());
 		return;
 	}
 
-	const int32 AgentNum = AgentIds.Num();
+	// Run Set Observations Callback
 
-	if (AgentNum != InActionObjectElements.Num())
+	ValidAgentIds.Empty(AgentSet.Num());
+	for (const int32 AgentId : AgentSet)
 	{
-		UE_LOG(LogLearning, Error, TEXT("%s: Not enough Action Objects. Expected %i, Got %i."), *GetName(), AgentNum, InActionObjectElements.Num());
-		return;
+		ValidAgentIds.Add(AgentId);
 	}
 
-	for (int32 AgentIdx = 0; AgentIdx < AgentNum; AgentIdx++)
-	{
-		Delegate.Execute(InActionObjectElements[AgentIdx], InActionObject, AgentIds[AgentIdx]);
-	}
-}
+	SetObservations(ValidAgentIds);
 
-void ULearningAgentsInteractor::GatherObservations(const UE::Learning::FIndexSet AgentSet, bool bIncrementIteration)
-{
-	UE_LEARNING_TRACE_CPUPROFILER_EVENT_SCOPE(ULearningAgentsInteractor::GatherObservations);
+	// Check that all observations have had their setter run
 
-	if (!IsSetup())
+	ValidAgentStatus.SetNumUninitialized(Manager->GetMaxAgentNum());
+	ValidAgentStatus.SetRange(0, Manager->GetMaxAgentNum(), true);
+
+	for (ULearningAgentsObservation* ObservationObject : ObservationObjects)
 	{
-		UE_LOG(LogLearning, Error, TEXT("%s: Setup not complete."), *GetName());
-		return;
+		for (const int32 AgentId : AgentSet)
+		{
+			if (ObservationObject->GetAgentIteration(AgentId) <= ObservationEncodingAgentIteration[AgentId])
+			{
+				UE_LOG(LogLearning, Warning, TEXT("%s: Observation %s for agent with id %i has not been set (got iteration %i, expected iteration %i) and so agent will not have observations encoded."), *GetName(), *ObservationObject->GetName(), AgentId, ObservationObject->GetAgentIteration(AgentId), ObservationEncodingAgentIteration[AgentId] + 1);
+				ValidAgentStatus[AgentId] = false;
+				continue;
+			}
+
+			if (ObservationObject->GetAgentIteration(AgentId) > ObservationEncodingAgentIteration[AgentId] + 1)
+			{
+				UE_LOG(LogLearning, Warning, TEXT("%s: Observation %s for agent with id %i appears to have been set multiple times (got iteration %i, expected iteration %i) and so agent will not have observations encoded."), *GetName(), *ObservationObject->GetName(), AgentId, ObservationObject->GetAgentIteration(AgentId), ObservationEncodingAgentIteration[AgentId] + 1);
+				ValidAgentStatus[AgentId] = false;
+				continue;
+			}
+
+			if (ObservationObject->GetAgentIteration(AgentId) != ObservationEncodingAgentIteration[AgentId] + 1)
+			{
+				UE_LOG(LogLearning, Warning, TEXT("%s: Observation %s for agent with id %i does not have a matching iteration number (got iteration %i, expected iteration %i) and so agent will not have observations encoded."), *GetName(), *ObservationObject->GetName(), AgentId, ObservationObject->GetAgentIteration(AgentId), ObservationEncodingAgentIteration[AgentId] + 1);
+				ValidAgentStatus[AgentId] = false;
+				continue;
+			}
+		}
 	}
 
-	// Run GatherAgentObservations Callback
+	ValidAgentIds.Empty(Manager->GetAgentNum());
 
-	ValidAgentIds.Empty(Manager->GetMaxAgentNum());
 	for (const int32 AgentId : AgentSet)
 	{
-		ValidAgentIds.Add(AgentId);
+		if (ValidAgentStatus[AgentId]) { ValidAgentIds.Add(AgentId); }
 	}
 
-	ObservationObject->GetObservationObject().Reset();
-	ObservationObjectElements.Empty(Manager->GetMaxAgentNum());
-	GatherAgentObservations(ObservationObjectElements, ObservationObject, ValidAgentIds);
+	ValidAgentSet = ValidAgentIds;
+	ValidAgentSet.TryMakeSlice();
+
+	// Encode Observations
+
+	Observations->Encode(ValidAgentSet);
+
+	// Increment Observation Encoding Iteration
 
-	if (ValidAgentIds.Num() != ObservationObjectElements.Num())
+	for (const int32 AgentId : ValidAgentSet)
 	{
-		UE_LOG(LogLearning, Error, TEXT("%s: Not enough Observation Objects added by GatherAgentObservations. Expected %i, Got %i."), *GetName(), ValidAgentIds.Num(), ObservationObjectElements.Num());
-		return;
+		ObservationEncodingAgentIteration[AgentId]++;
 	}
 
-	// Check Observation Objects are Valid and if so convert to observation vectors
+	// Visual Logger
 
-	for (int32 AgentIdx = 0; AgentIdx < AgentSet.Num(); AgentIdx++)
+#if UE_LEARNING_AGENTS_ENABLE_VISUAL_LOG
+	for (const ULearningAgentsObservation* ObservationObject : ObservationObjects)
 	{
-		if (ObservationSchema->ValidateObjectMatchesSchema(ObservationSchemaElement, ObservationObject, ObservationObjectElements[AgentIdx]))
-		{
-			UE::Learning::Observation::SetVectorFromObject(
-				ObservationVectors[AgentSet[AgentIdx]],
-				ObservationSchema->GetObservationSchema(),
-				ObservationSchemaElement.SchemaElement,
-				ObservationObject->GetObservationObject(),
-				ObservationObjectElements[AgentIdx].ObjectElement);
-
-			if (bIncrementIteration)
-			{
-				ObservationVectorIteration[AgentSet[AgentIdx]]++;
-			}
-		}
+		ObservationObject->VisualLog(ValidAgentSet);
 	}
+#endif
 }
 
-void ULearningAgentsInteractor::ScatterActions(const UE::Learning::FIndexSet AgentSet)
+void ULearningAgentsInteractor::DecodeActions(const UE::Learning::FIndexSet AgentSet)
 {
-	UE_LEARNING_TRACE_CPUPROFILER_EVENT_SCOPE(ULearningAgentsInteractor::ScatterActions);
+	UE_LEARNING_TRACE_CPUPROFILER_EVENT_SCOPE(ULearningAgentsInteractor::DecodeActions);
 
 	if (!IsSetup())
 	{
@@ -313,15 +324,15 @@ void ULearningAgentsInteractor::ScatterActions(const UE::Learning::FIndexSet Age
 		return;
 	}
 
-	// Check which agents have had their action vector set
+	// Check Agents actually have encoded actions.
 
-	ValidAgentIds.Empty(Manager->GetMaxAgentNum());
+	ValidAgentIds.Empty(Manager->GetAgentNum());
 
 	for (const int32 AgentId : AgentSet)
 	{
-		if (ActionVectorIteration[AgentId] == 0)
+		if (ActionEncodingAgentIteration[AgentId] == 0)
 		{
-			UE_LOG(LogLearning, Warning, TEXT("%s: Agent with id %i does not have an action vector so actions will not be scattered for it. Was DecodeAndSampleActions run without error?"), *GetName(), AgentId);
+			UE_LOG(LogLearning, Warning, TEXT("%s: Agent with id %i does not have an encoded action vector so actions will not be decoded for it. Was EvaluatePolicy or EncodeActions run?"), *GetName(), AgentId);
 			continue;
 		}
 
@@ -331,58 +342,65 @@ void ULearningAgentsInteractor::ScatterActions(const UE::Learning::FIndexSet Age
 	ValidAgentSet = ValidAgentIds;
 	ValidAgentSet.TryMakeSlice();
 
-	// Generate Action Objects
-
-	ActionObject->GetActionObject().Reset();
-	ActionObjectElements.Empty(Manager->GetMaxAgentNum());
-
-	for (int32 AgentIdx = 0; AgentIdx < ValidAgentSet.Num(); AgentIdx++)
-	{
-		FLearningAgentsActionObjectElement ActionObjectElement;
+	// Decode Actions
 
-		UE::Learning::Action::GetObjectFromVector(
-			ActionObject->GetActionObject(),
-			ActionObjectElement.ObjectElement,
-			ActionSchema->GetActionSchema(),
-			ActionSchemaElement.SchemaElement,
-			ActionVectors[ValidAgentSet[AgentIdx]]);
+	Actions->Decode(ValidAgentSet);
 
-		ActionObjectElements.Add(ActionObjectElement);
-	}
+	// Run Get Actions Callback
 
-	// Scatter Action Objects
+	GetActions(ValidAgentIds);
 
-	ScatterAgentActions(ActionObjectElements, ActionObject, ValidAgentIds);
-}
+	// Check that all actions have had their getter run
 
-void ULearningAgentsInteractor::GatherObservations()
-{
-	if (Manager->GetAgentNum() == 0)
+	for (ULearningAgentsAction* ActionObject : ActionObjects)
 	{
-		UE_LOG(LogLearning, Warning, TEXT("%s: No agents added to Manager."), *GetName());
+		for (const int32 AgentId : ValidAgentSet)
+		{
+			if (ActionObject->GetAgentGetIteration(AgentId) == ActionEncodingAgentIteration[AgentId] - 1)
+			{
+				UE_LOG(LogLearning, Warning, TEXT("%s: Action %s for agent with id %i appears to have not been got."), *GetName(), *ActionObject->GetName(), AgentId);
+				continue;
+			}
+				
+			if (ActionObject->GetAgentGetIteration(AgentId) > ActionEncodingAgentIteration[AgentId])
+			{
+				UE_LOG(LogLearning, Warning, TEXT("%s: Action %s for agent with id %i appears to have been got multiple times."), *GetName(), *ActionObject->GetName(), AgentId);
+				continue;
+			}
+				
+			if (ActionObject->GetAgentGetIteration(AgentId) != ActionEncodingAgentIteration[AgentId])
+			{
+				UE_LOG(LogLearning, Warning, TEXT("%s: Action %s for agent with id %i does not have a matching iteration number (got %i, expected %i)."), *GetName(), *ActionObject->GetName(), AgentId, ActionObject->GetAgentGetIteration(AgentId), ActionEncodingAgentIteration[AgentId]);
+				continue;
+			}
+		}
 	}
 
-	GatherObservations(Manager->GetAllAgentSet());
-}
+	// Visual Logger
 
-void ULearningAgentsInteractor::ScatterActions()
-{
-	if (Manager->GetAgentNum() == 0)
+#if UE_LEARNING_AGENTS_ENABLE_VISUAL_LOG
+	for (const ULearningAgentsAction* ActionObject : ActionObjects)
 	{
-		UE_LOG(LogLearning, Warning, TEXT("%s: No agents added to Manager."), *GetName());
+		ActionObject->VisualLog(ValidAgentSet);
 	}
-
-	ScatterActions(Manager->GetAllAgentSet());
+#endif
 }
 
+void ULearningAgentsInteractor::EncodeObservations()
+{
+	EncodeObservations(Manager->GetAllAgentSet());
+}
 
-void ULearningAgentsInteractor::GetObservationVector(TArray<float>& OutObservationVector, int32& OutObservationCompatibilityHash, const int32 AgentId)
+void ULearningAgentsInteractor::DecodeActions()
+{
+	DecodeActions(Manager->GetAllAgentSet());
+}
+void ULearningAgentsInteractor::GetObservationVector(const int32 AgentId, TArray<float>& OutObservationVector) const
 {
 	if (!IsSetup())
 	{
 		UE_LOG(LogLearning, Error, TEXT("%s: Setup not complete."), *GetName());
 		OutObservationVector.Empty();
-		OutObservationCompatibilityHash = 0;
 		return;
 	}
 
@@ -390,30 +408,26 @@ void ULearningAgentsInteractor::GetObservationVector(TArray<float>& OutObservati
 	{
 		UE_LOG(LogLearning, Error, TEXT("%s: AgentId %d not found in the agents set."), *GetName(), AgentId);
 		OutObservationVector.Empty();
-		OutObservationCompatibilityHash = 0;
 		return;
 	}
 
-	if (ObservationVectorIteration[AgentId] == 0)
+	if (ObservationEncodingAgentIteration[AgentId] == 0)
 	{
-		UE_LOG(LogLearning, Error, TEXT("%s: Observation vector not set for agent %i."), *GetName(), AgentId);
+		UE_LOG(LogLearning, Error, TEXT("%s: Agent with id %d has not yet computed an observation vector. Have you run EncodeObservations?"), *GetName(), AgentId);
 		OutObservationVector.Empty();
-		OutObservationCompatibilityHash = 0;
 		return;
 	}
 
-	OutObservationVector.SetNumUninitialized(GetObservationVectorSize());
-	UE::Learning::Array::Copy<1, float>(OutObservationVector, ObservationVectors[AgentId]);
-	OutObservationCompatibilityHash = ObservationCompatibilityHash;
+	OutObservationVector.SetNumUninitialized(Observations->DimNum());
+	UE::Learning::Array::Copy<1, float>(OutObservationVector, Observations->FeatureBuffer()[AgentId]);
 }
 
-void ULearningAgentsInteractor::GetActionVector(TArray<float>& OutActionVector, int32& OutActionCompatibilityHash, const int32 AgentId)
+void ULearningAgentsInteractor::GetActionVector(const int32 AgentId, TArray<float>& OutActionVector) const
 {
 	if (!IsSetup())
 	{
 		UE_LOG(LogLearning, Error, TEXT("%s: Setup not complete."), *GetName());
 		OutActionVector.Empty();
-		OutActionCompatibilityHash = 0;
 		return;
 	}
 
@@ -421,188 +435,16 @@ void ULearningAgentsInteractor::GetActionVector(TArray<float>& OutActionVector,
 	{
 		UE_LOG(LogLearning, Error, TEXT("%s: AgentId %d not found in the agents set."), *GetName(), AgentId);
 		OutActionVector.Empty();
-		OutActionCompatibilityHash = 0;
 		return;
 	}
 
-	if (ActionVectorIteration[AgentId] == 0)
+	if (ActionEncodingAgentIteration[AgentId] == 0)
 	{
-		UE_LOG(LogLearning, Error, TEXT("%s: Action vector not set for agent %i."), *GetName(), AgentId);
+		UE_LOG(LogLearning, Error, TEXT("%s: Agent with id %d has not yet computed an action vector. Have you run EvaluatePolicy or EncodeActions?"), *GetName(), AgentId);
 		OutActionVector.Empty();
-		OutActionCompatibilityHash = 0;
-		return;
-	}
-
-	OutActionVector.SetNumUninitialized(GetActionVectorSize());
-	UE::Learning::Array::Copy<1, float>(OutActionVector, ActionVectors[AgentId]);
-	OutActionCompatibilityHash = ActionCompatibilityHash;
-}
-
-void ULearningAgentsInteractor::SetObservationVector(const TArray<float>& ObservationVector, const int32 InObservationCompatibilityHash, const int32 AgentId, bool bIncrementIteration)
-{
-	if (!IsSetup())
-	{
-		UE_LOG(LogLearning, Error, TEXT("%s: Setup not complete."), *GetName());
 		return;
 	}
 
-	if (!HasAgent(AgentId))
-	{
-		UE_LOG(LogLearning, Error, TEXT("%s: AgentId %d not found in the agents set."), *GetName(), AgentId);
-		return;
-	}
-
-	if (InObservationCompatibilityHash != ObservationCompatibilityHash)
-	{
-		UE_LOG(LogLearning, Error, TEXT("%s: Observation Compatibility hash incompatible. Got %i, expected %i."), *GetName(), InObservationCompatibilityHash, ObservationCompatibilityHash);
-		return;
-	}
-
-	if (ObservationVector.Num() != GetObservationVectorSize())
-	{
-		UE_LOG(LogLearning, Error, TEXT("%s: Observation Vector size incompatible. Got %i, expected %i."), *GetName(), ObservationVector.Num(), GetObservationVectorSize());
-		return;
-	}
-
-	UE::Learning::Array::Copy<1, float>(ObservationVectors[AgentId], ObservationVector);
-	if (bIncrementIteration) { ObservationVectorIteration[AgentId]++; }
-}
-
-void ULearningAgentsInteractor::SetActionVector(const TArray<float>& ActionVector, const int32 InActionCompatibilityHash, const int32 AgentId, bool bIncrementIteration)
-{
-	if (!IsSetup())
-	{
-		UE_LOG(LogLearning, Error, TEXT("%s: Setup not complete."), *GetName());
-		return;
-	}
-
-	if (!HasAgent(AgentId))
-	{
-		UE_LOG(LogLearning, Error, TEXT("%s: AgentId %d not found in the agents set."), *GetName(), AgentId);
-		return;
-	}
-
-	if (InActionCompatibilityHash != ActionCompatibilityHash)
-	{
-		UE_LOG(LogLearning, Error, TEXT("%s: Action Compatibility hash incompatible. Got %i, expected %i."), *GetName(), InActionCompatibilityHash, ActionCompatibilityHash);
-		return;
-	}
-
-	if (ActionVector.Num() != GetObservationVectorSize())
-	{
-		UE_LOG(LogLearning, Error, TEXT("%s: Action Vector size incompatible. Got %i, expected %i."), *GetName(), ActionVector.Num(), GetActionVectorSize());
-		return;
-	}
-
-	UE::Learning::Array::Copy<1, float>(ActionVectors[AgentId], ActionVector);
-	if (bIncrementIteration) { ActionVectorIteration[AgentId]++; }
-}
-
-bool ULearningAgentsInteractor::HasObservationVector(const int32 AgentId) const
-{
-	if (!IsSetup())
-	{
-		UE_LOG(LogLearning, Error, TEXT("%s: Setup not complete."), *GetName());
-		return false;
-	}
-
-	if (!HasAgent(AgentId))
-	{
-		UE_LOG(LogLearning, Error, TEXT("%s: AgentId %d not found in the agents set."), *GetName(), AgentId);
-		return false;
-	}
-
-	return ObservationVectorIteration[AgentId] > 0;
-}
-
-bool ULearningAgentsInteractor::HasActionVector(const int32 AgentId) const
-{
-	if (!IsSetup())
-	{
-		UE_LOG(LogLearning, Error, TEXT("%s: Setup not complete."), *GetName());
-		return false;
-	}
-
-	if (!HasAgent(AgentId))
-	{
-		UE_LOG(LogLearning, Error, TEXT("%s: AgentId %d not found in the agents set."), *GetName(), AgentId);
-		return false;
-	}
-
-	return ActionVectorIteration[AgentId] > 0;
-}
-
-int32 ULearningAgentsInteractor::GetObservationVectorSize() const
-{
-	if (!IsSetup())
-	{
-		UE_LOG(LogLearning, Error, TEXT("%s: Setup not complete."), *GetName());
-		return 0;
-	}
-
-	return ObservationSchema->GetObservationSchema().GetObservationVectorSize(ObservationSchemaElement.SchemaElement);
-}
-
-int32 ULearningAgentsInteractor::GetObservationEncodedVectorSize() const
-{
-	if (!IsSetup())
-	{
-		UE_LOG(LogLearning, Error, TEXT("%s: Setup not complete."), *GetName());
-		return 0;
-	}
-
-	return ObservationSchema->GetObservationSchema().GetEncodedVectorSize(ObservationSchemaElement.SchemaElement);
-}
-
-int32 ULearningAgentsInteractor::GetActionVectorSize() const
-{
-	if (!IsSetup())
-	{
-		UE_LOG(LogLearning, Error, TEXT("%s: Setup not complete."), *GetName());
-		return 0;
-	}
-
-	return ActionSchema->GetActionSchema().GetActionVectorSize(ActionSchemaElement.SchemaElement);
-}
-
-int32 ULearningAgentsInteractor::GetActionDistributionVectorSize() const
-{
-	if (!IsSetup())
-	{
-		UE_LOG(LogLearning, Error, TEXT("%s: Setup not complete."), *GetName());
-		return 0;
-	}
-
-	return ActionSchema->GetActionSchema().GetActionDistributionVectorSize(ActionSchemaElement.SchemaElement);
-}
-
-int32 ULearningAgentsInteractor::GetActionEncodedVectorSize() const
-{
-	if (!IsSetup())
-	{
-		UE_LOG(LogLearning, Error, TEXT("%s: Setup not complete."), *GetName());
-		return 0;
-	}
-
-	return ActionSchema->GetActionSchema().GetEncodedVectorSize(ActionSchemaElement.SchemaElement);
-}
-
-const UE::Learning::Observation::FSchema& ULearningAgentsInteractor::GetObservationSchema() const
-{
-	return ObservationSchema->GetObservationSchema();
-}
-
-UE::Learning::Observation::FSchemaElement ULearningAgentsInteractor::GetObservationSchemaElement() const
-{
-	return ObservationSchemaElement.SchemaElement;
-}
-
-const UE::Learning::Action::FSchema& ULearningAgentsInteractor::GetActionSchema() const
-{
-	return ActionSchema->GetActionSchema();
-}
-
-UE::Learning::Action::FSchemaElement ULearningAgentsInteractor::GetActionSchemaElement() const
-{
-	return ActionSchemaElement.SchemaElement;
+	OutActionVector.SetNumUninitialized(Actions->DimNum());
+	UE::Learning::Array::Copy<1, float>(OutActionVector, Actions->FeatureBuffer()[AgentId]);
 }
\ No newline at end of file
