[Backout] - CL30154408

diff --git a/Engine/Plugins/Experimental/LearningAgents/Source/LearningAgents/Private/LearningAgentsManager.cpp b/Engine/Plugins/Experimental/LearningAgents/Source/LearningAgents/Private/LearningAgentsManager.cpp
--- a/Engine/Plugins/Experimental/LearningAgents/Source/LearningAgents/Private/LearningAgentsManager.cpp
+++ b/Engine/Plugins/Experimental/LearningAgents/Source/LearningAgents/Private/LearningAgentsManager.cpp
@@ -1,19 +1,20 @@
 // Copyright Epic Games, Inc. All Rights Reserved.
 
 #include "LearningAgentsManager.h"
-#include "LearningAgentsManagerListener.h"
+#include "LearningAgentsManagerComponent.h"
 
 #include "LearningArray.h"
+#include "LearningArrayMap.h"
 #include "LearningLog.h"
 
 #include "UObject/ScriptInterface.h"
 #include "EngineDefines.h"
 
-ULearningAgentsManager::ULearningAgentsManager() : Super(FObjectInitializer::Get()) { PrimaryComponentTick.bCanEverTick = true; }
-ULearningAgentsManager::ULearningAgentsManager(FVTableHelper& Helper) : Super(Helper) {}
-ULearningAgentsManager::~ULearningAgentsManager() = default;
+ALearningAgentsManager::ALearningAgentsManager() : Super(FObjectInitializer::Get()) {}
+ALearningAgentsManager::ALearningAgentsManager(FVTableHelper& Helper) : Super(Helper) {}
+ALearningAgentsManager::~ALearningAgentsManager() = default;
 
-void ULearningAgentsManager::PostInitProperties()
+void ALearningAgentsManager::PostInitProperties()
 {
 	Super::PostInitProperties();
 
@@ -22,40 +23,29 @@ void ULearningAgentsManager::PostInitProperties()
 		return;
 	}
 
+	// Allocate Instance Data
+	InstanceData = MakeShared<UE::Learning::FArrayMap>();
+
 	// Pre-populate the vacant ids
 	OnEventAgentIds.Reserve(MaxAgentNum);
 	OccupiedAgentIds.Reserve(MaxAgentNum);
 	VacantAgentIds.Reserve(MaxAgentNum);
 	Agents.Reserve(MaxAgentNum);
-	for (int32 AgentId = MaxAgentNum - 1; AgentId >= 0; AgentId--)
+	for (int32 i = MaxAgentNum - 1; i >= 0; i--)
 	{
-		VacantAgentIds.Push(AgentId);
+		VacantAgentIds.Push(i);
 		Agents.Add(nullptr);
 	}
 
 	UpdateAgentSets();
 }
 
-void ULearningAgentsManager::TickComponent(float DeltaTime, enum ELevelTick TickType, FActorComponentTickFunction* ThisTickFunction)
-{
-	Super::TickComponent(DeltaTime, TickType, ThisTickFunction);
-
-	// We only want to give agent tick events when the actual level is ticking - not when paused or just the viewport is ticking.
-	if (TickType == ELevelTick::LEVELTICK_All)
-	{
-		for (ULearningAgentsManagerListener* Listener : Listeners)
-		{
-			Listener->OnAgentsManagerTick(OccupiedAgentIds, DeltaTime);
-		}
-	}
-}
-
-int32 ULearningAgentsManager::GetMaxAgentNum() const
+int32 ALearningAgentsManager::GetMaxAgentNum() const
 {
 	return MaxAgentNum;
 }
 
-int32 ULearningAgentsManager::AddAgent(UObject* Agent)
+int32 ALearningAgentsManager::AddAgent(UObject* Agent)
 {
 	if (Agent == nullptr)
 	{
@@ -76,16 +66,19 @@ int32 ULearningAgentsManager::AddAgent(UObject* Agent)
 	OccupiedAgentIds.Add(AgentId);
 	UpdateAgentSets();
 
-	OnEventAgentIds.Reset();
+	OnEventAgentIds.SetNumUninitialized(0, false);
 	OnEventAgentIds.Add(AgentId);
 
 	// Call OnAgentsAdded Event
 
 	UE_LOG(LogLearning, Display, TEXT("%s: Adding Agent %s with id %i."), *GetName(), *Agent->GetName(), AgentId);
 
-	for (ULearningAgentsManagerListener* Listener : Listeners)
+	TInlineComponentArray<ULearningAgentsManagerComponent*> ManagerComponents;
+	GetComponents<ULearningAgentsManagerComponent>(ManagerComponents);
+
+	for (ULearningAgentsManagerComponent* ManagerComponent : ManagerComponents)
 	{
-		Listener->OnAgentsAdded(OnEventAgentIds);
+		ManagerComponent->OnAgentsAdded(OnEventAgentIds);
 	}
 
 	// Return new id
@@ -93,9 +86,9 @@ int32 ULearningAgentsManager::AddAgent(UObject* Agent)
 	return AgentId;
 }
 
-void ULearningAgentsManager::AddAgents(TArray<int32>& OutAgentIds, const TArray<UObject*>& InAgents)
+void ALearningAgentsManager::AddAgents(TArray<int32>& OutAgentIds, const TArray<UObject*>& InAgents)
 {
-	OnEventAgentIds.Reset();
+	OnEventAgentIds.SetNumUninitialized(0, false);
 	OutAgentIds.SetNumUninitialized(InAgents.Num());
 
 	for (int32 AgentIdx = 0; AgentIdx < InAgents.Num(); AgentIdx++)
@@ -128,13 +121,16 @@ void ULearningAgentsManager::AddAgents(TArray<int32>& OutAgentIds, const TArray<
 
 	UE_LOG(LogLearning, Display, TEXT("%s: Adding Agents %s."), *GetName(), *UE::Learning::Array::FormatInt32(OnEventAgentIds));
 
-	for (ULearningAgentsManagerListener* Listener : Listeners)
+	TInlineComponentArray<ULearningAgentsManagerComponent*> ManagerComponents;
+	GetComponents<ULearningAgentsManagerComponent>(ManagerComponents);
+
+	for (ULearningAgentsManagerComponent* ManagerComponent : ManagerComponents)
 	{
-		Listener->OnAgentsAdded(OnEventAgentIds);
+		ManagerComponent->OnAgentsAdded(OnEventAgentIds);
 	}
 }
 
-void ULearningAgentsManager::RemoveAgent(const int32 AgentId)
+void ALearningAgentsManager::RemoveAgent(const int32 AgentId)
 {
 	if (AgentId == INDEX_NONE)
 	{
@@ -158,22 +154,25 @@ void ULearningAgentsManager::RemoveAgent(const int32 AgentId)
 	Agents[AgentId] = nullptr;
 	UpdateAgentSets();
 
-	OnEventAgentIds.Reset();
+	OnEventAgentIds.SetNumUninitialized(0, false);
 	OnEventAgentIds.Add(AgentId);
 
 	// Call OnAgentsRemoved Event
 
 	UE_LOG(LogLearning, Display, TEXT("%s: Removing Agent %i."), *GetName(), AgentId);
 
-	for (ULearningAgentsManagerListener* Listener : Listeners)
+	TInlineComponentArray<ULearningAgentsManagerComponent*> ManagerComponents;
+	GetComponents<ULearningAgentsManagerComponent>(ManagerComponents);
+
+	for (ULearningAgentsManagerComponent* ManagerComponent : ManagerComponents)
 	{
-		Listener->OnAgentsRemoved(OnEventAgentIds);
+		ManagerComponent->OnAgentsRemoved(OnEventAgentIds);
 	}
 }
 
-void ULearningAgentsManager::RemoveAgents(const TArray<int32>& AgentIds)
+void ALearningAgentsManager::RemoveAgents(const TArray<int32>& AgentIds)
 {
-	OnEventAgentIds.Reset();
+	OnEventAgentIds.SetNumUninitialized(0, false);
 
 	for (int32 AgentIdx = 0; AgentIdx < AgentIds.Num(); AgentIdx++)
 	{
@@ -206,20 +205,23 @@ void ULearningAgentsManager::RemoveAgents(const TArray<int32>& AgentIds)
 
 	UE_LOG(LogLearning, Display, TEXT("%s: Removing Agents %s."), *GetName(), *UE::Learning::Array::FormatInt32(OnEventAgentIds));
 
-	for (ULearningAgentsManagerListener* Listener : Listeners)
+	TInlineComponentArray<ULearningAgentsManagerComponent*> ManagerComponents;
+	GetComponents<ULearningAgentsManagerComponent>(ManagerComponents);
+
+	for (ULearningAgentsManagerComponent* ManagerComponent : ManagerComponents)
 	{
-		Listener->OnAgentsRemoved(OnEventAgentIds);
+		ManagerComponent->OnAgentsRemoved(OnEventAgentIds);
 	}
 }
 
-void ULearningAgentsManager::RemoveAllAgents()
+void ALearningAgentsManager::RemoveAllAgents()
 {
 	// We need to make a copy as OccupiedAgentIds will be modified during removal.
 	const TArray<int32> RemoveAgentIds = OccupiedAgentIds;
 	RemoveAgents(RemoveAgentIds);
 }
 
-void ULearningAgentsManager::ResetAgent(const int32 AgentId)
+void ALearningAgentsManager::ResetAgent(const int32 AgentId)
 {
 	if (AgentId == INDEX_NONE)
 	{
@@ -233,22 +235,25 @@ void ULearningAgentsManager::ResetAgent(const int32 AgentId)
 		return;
 	}
 
-	OnEventAgentIds.Reset();
+	OnEventAgentIds.SetNumUninitialized(0, false);
 	OnEventAgentIds.Add(AgentId);
 
 	// Call OnAgentsReset Event
 
 	UE_LOG(LogLearning, Display, TEXT("%s: Resetting Agent %i."), *GetName(), AgentId);
 
-	for (ULearningAgentsManagerListener* Listener : Listeners)
+	TInlineComponentArray<ULearningAgentsManagerComponent*> ManagerComponents;
+	GetComponents<ULearningAgentsManagerComponent>(ManagerComponents);
+
+	for (ULearningAgentsManagerComponent* ManagerComponent : ManagerComponents)
 	{
-		Listener->OnAgentsReset(OnEventAgentIds);
+		ManagerComponent->OnAgentsReset(OnEventAgentIds);
 	}
 }
 
-void ULearningAgentsManager::ResetAgents(const TArray<int32>& AgentIds)
+void ALearningAgentsManager::ResetAgents(const TArray<int32>& AgentIds)
 {
-	OnEventAgentIds.Reset();
+	OnEventAgentIds.SetNumUninitialized(0, false);
 
 	for (int32 AgentIdx = 0; AgentIdx < AgentIds.Num(); AgentIdx++)
 	{
@@ -271,41 +276,34 @@ void ULearningAgentsManager::ResetAgents(const TArray<int32>& AgentIds)
 
 	UE_LOG(LogLearning, Display, TEXT("%s: Resetting Agents %s."), *GetName(), *UE::Learning::Array::FormatInt32(OnEventAgentIds));
 
-	for (ULearningAgentsManagerListener* Listener : Listeners)
+	TInlineComponentArray<ULearningAgentsManagerComponent*> ManagerComponents;
+	GetComponents<ULearningAgentsManagerComponent>(ManagerComponents);
+
+	for (ULearningAgentsManagerComponent* ManagerComponent : ManagerComponents)
 	{
-		Listener->OnAgentsReset(OnEventAgentIds);
+		ManagerComponent->OnAgentsReset(OnEventAgentIds);
 	}
 }
 
-void ULearningAgentsManager::ResetAllAgents()
+void ALearningAgentsManager::ResetAllAgents()
 {
 	ResetAgents(OccupiedAgentIds);
 }
 
-UObject* ULearningAgentsManager::GetAgent(const int32 AgentId, const TSubclassOf<UObject> AgentClass) const
+UObject* ALearningAgentsManager::GetAgent(const int32 AgentId, const TSubclassOf<UObject> AgentClass) const
 {
 	if (!HasAgent(AgentId))
 	{
 		UE_LOG(LogLearning, Error, TEXT("%s: AgentId %d not found. Be sure to only use AgentIds returned by AddAgent and check that the agent has not be removed."), *GetName(), AgentId);
 		return nullptr;
 	}
-	
-	if (Agents[AgentId] == nullptr)
-	{
-		UE_LOG(LogLearning, Warning, TEXT("%s: AgentId %d object is nullptr."), *GetName(), AgentId);
-		return nullptr;
-	}
-	
-	if (!Agents[AgentId].IsA(AgentClass))
+	else
 	{
-		UE_LOG(LogLearning, Error, TEXT("%s: AgentId %d cast invalid. Agent was of class '%s', not '%s'."), *GetName(), AgentId, *Agents[AgentId]->GetClass()->GetName(), *AgentClass->GetName());
-		return nullptr;
+		return Agents[AgentId];
 	}
-	
-	return Agents[AgentId];
 }
 
-void ULearningAgentsManager::GetAgents(TArray<UObject*>& OutAgents, const TArray<int32>& AgentIds, const TSubclassOf<UObject> AgentClass) const
+void ALearningAgentsManager::GetAgents(TArray<UObject*>& OutAgents, const TArray<int32>& AgentIds, const TSubclassOf<UObject> AgentClass) const
 {
 	OutAgents.SetNumUninitialized(AgentIds.Num());
 
@@ -315,34 +313,21 @@ void ULearningAgentsManager::GetAgents(TArray<UObject*>& OutAgents, const TArray
 		{
 			UE_LOG(LogLearning, Error, TEXT("%s: AgentId %d not found. Be sure to only use AgentIds returned by AddAgent and check that the agent has not be removed."), *GetName(), AgentIds[AgentIdIdx]);
 			OutAgents[AgentIdIdx] = nullptr;
-			continue;
-		}
-		
-		if (Agents[AgentIds[AgentIdIdx]] == nullptr)
-		{
-			UE_LOG(LogLearning, Warning, TEXT("%s: AgentId %d object is nullptr."), *GetName(), AgentIds[AgentIdIdx]);
-			OutAgents[AgentIdIdx] = nullptr;
-			continue;
 		}
-		
-		if (!Agents[AgentIds[AgentIdIdx]].IsA(AgentClass))
+		else
 		{
-			UE_LOG(LogLearning, Error, TEXT("%s: AgentId %d cast invalid. Agent was of class '%s', not '%s'."), *GetName(), AgentIds[AgentIdIdx], *Agents[AgentIds[AgentIdIdx]]->GetClass()->GetName(), *AgentClass->GetName());
-			OutAgents[AgentIdIdx] = nullptr;
-			continue;
+			OutAgents[AgentIdIdx] = Agents[AgentIds[AgentIdIdx]];
 		}
-		
-		OutAgents[AgentIdIdx] = Agents[AgentIds[AgentIdIdx]];
 	}
 }
 
-void ULearningAgentsManager::GetAllAgents(TArray<UObject*>& OutAgents, TArray<int32>& OutAgentIds, const TSubclassOf<UObject> AgentClass) const
+void ALearningAgentsManager::GetAllAgents(TArray<UObject*>& OutAgents, TArray<int32>& OutAgentIds, const TSubclassOf<UObject> AgentClass) const
 {
 	OutAgentIds = OccupiedAgentIds;
 	GetAgents(OutAgents, OutAgentIds, AgentClass);
 }
 
-int32 ULearningAgentsManager::GetAgentId(UObject* Agent) const
+int32 ALearningAgentsManager::GetAgentId(UObject* Agent) const
 {
 	const int32 AgentId = Agents.Find(Agent);
 
@@ -357,7 +342,7 @@ int32 ULearningAgentsManager::GetAgentId(UObject* Agent) const
 	}
 }
 
-void ULearningAgentsManager::GetAgentIds(TArray<int32>& OutAgentIds, const TArray<UObject*>& InAgents) const
+void ALearningAgentsManager::GetAgentIds(TArray<int32>& OutAgentIds, const TArray<UObject*>& InAgents) const
 {
 	OutAgentIds.SetNumUninitialized(InAgents.Num());
 
@@ -377,63 +362,70 @@ void ULearningAgentsManager::GetAgentIds(TArray<int32>& OutAgentIds, const TArra
 	}
 }
 
-int32 ULearningAgentsManager::GetAgentNum() const
+int32 ALearningAgentsManager::GetAgentNum() const
 {
 	return OccupiedAgentIds.Num();
 }
 
-bool ULearningAgentsManager::HasAgentObject(UObject* Agent) const
+bool ALearningAgentsManager::HasAgentObject(UObject* Agent) const
 {
 	return Agents.Contains(Agent);
 }
 
-bool ULearningAgentsManager::HasAgent(const int32 AgentId) const
+bool ALearningAgentsManager::HasAgent(const int32 AgentId) const
 {
 	return OccupiedAgentSet.Contains(AgentId);
 }
 
-const UObject* ULearningAgentsManager::GetAgent(const int32 AgentId) const
+void ALearningAgentsManager::AddManagerAsTickPrerequisiteOfAgents(const TArray<AActor*>& InAgents)
+{
+	for (AActor* Agent : InAgents)
+	{
+		Agent->AddTickPrerequisiteActor(this);
+	}
+}
+
+void ALearningAgentsManager::AddAgentsAsTickPrerequisiteOfManager(const TArray<AActor*>& InAgents)
+{
+	for (AActor* Agent : InAgents)
+	{
+		this->AddTickPrerequisiteActor(Agent);
+	}
+}
+
+const UObject* ALearningAgentsManager::GetAgent(const int32 AgentId) const
 {
 	UE_LEARNING_CHECK(HasAgent(AgentId));
 	return Agents[AgentId];
 }
 
-UObject* ULearningAgentsManager::GetAgent(const int32 AgentId)
+UObject* ALearningAgentsManager::GetAgent(const int32 AgentId)
 {
 	UE_LEARNING_CHECK(HasAgent(AgentId));
 	return Agents[AgentId];
 }
 
-const TArray<int32>& ULearningAgentsManager::GetAllAgentIds() const
+const TArray<int32>& ALearningAgentsManager::GetAllAgentIds() const
 {
 	return OccupiedAgentIds;
 }
 
-UE::Learning::FIndexSet ULearningAgentsManager::GetAllAgentSet() const
+UE::Learning::FIndexSet ALearningAgentsManager::GetAllAgentSet() const
 {
 	return OccupiedAgentSet;
 }
 
-TConstArrayView<TObjectPtr<UObject>> ULearningAgentsManager::GetAgents() const
+const TSharedPtr<UE::Learning::FArrayMap>& ALearningAgentsManager::GetInstanceData() const
 {
-	return Agents;
+	return InstanceData;
 }
 
-void ULearningAgentsManager::AddListener(ULearningAgentsManagerListener* Listener)
+TConstArrayView<TObjectPtr<UObject>> ALearningAgentsManager::GetAgents() const
 {
-	Listeners.Add(Listener);
-	OnEventAgentIds = OccupiedAgentIds;
-	Listener->OnAgentsAdded(OnEventAgentIds);
-}
-
-void ULearningAgentsManager::RemoveListener(ULearningAgentsManagerListener* Listener)
-{
-	OnEventAgentIds = OccupiedAgentIds;
-	Listener->OnAgentsRemoved(OnEventAgentIds);
-	Listeners.Remove(Listener);
+	return Agents;
 }
 
-void ULearningAgentsManager::UpdateAgentSets()
+void ALearningAgentsManager::UpdateAgentSets()
 {
 	// Keep OccupiedAgentIds sorted in ascending order for better memory access patterns.
 	OccupiedAgentIds.Sort();
diff --git a/Engine/Plugins/Experimental/LearningAgents/Source/LearningAgents/Private/LearningAgentsManagerComponent.cpp b/Engine/Plugins/Experimental/LearningAgents/Source/LearningAgents/Private/LearningAgentsManagerComponent.cpp
new file mode 100644
--- /dev/null
+++ b/Engine/Plugins/Experimental/LearningAgents/Source/LearningAgents/Private/LearningAgentsManagerComponent.cpp
@@ -0,0 +1,173 @@
+// Copyright Epic Games, Inc. All Rights Reserved.
+
+#include "LearningAgentsManagerComponent.h"
+#include "LearningAgentsManager.h"
+#include "LearningAgentsHelpers.h"
+#include "LearningLog.h"
+
+ULearningAgentsManagerComponent::ULearningAgentsManagerComponent(const FObjectInitializer& ObjectInitializer) : Super(ObjectInitializer) {}
+ULearningAgentsManagerComponent::ULearningAgentsManagerComponent(FVTableHelper& Helper) : Super(Helper) {}
+ULearningAgentsManagerComponent::~ULearningAgentsManagerComponent() = default;
+
+void ULearningAgentsManagerComponent::PostInitProperties()
+{
+	Super::PostInitProperties();
+
+	if (HasAnyFlags(RF_ClassDefaultObject | RF_ArchetypeObject))
+	{
+		return;
+	}
+
+	Manager = GetOwner<ALearningAgentsManager>();
+
+	if (!Manager)
+	{
+		UE_LOG(LogLearning, Error, TEXT("%s: Must be attached to a LearningAgentsManager Actor."), *GetName());
+	}
+}
+
+void ULearningAgentsManagerComponent::AgentsAdded_Implementation(const TArray<int32>& AgentIds)
+{
+	// Can be overridden for additional logic when agents are added
+}
+
+void ULearningAgentsManagerComponent::AgentsRemoved_Implementation(const TArray<int32>& AgentIds)
+{
+	// Can be overridden for additional logic when agents are removed
+}
+
+void ULearningAgentsManagerComponent::AgentsReset_Implementation(const TArray<int32>& AgentIds)
+{
+	// Can be overridden for additional logic when agents are reset
+}
+
+void ULearningAgentsManagerComponent::OnAgentsAdded(const TArray<int32>& AgentIds)
+{
+	if (IsSetup())
+	{
+		for (ULearningAgentsHelper* Helper : HelperObjects)
+		{
+			Helper->OnAgentsAdded(AgentIds);
+		}
+
+		AgentsAdded(AgentIds);
+	}
+}
+
+void ULearningAgentsManagerComponent::OnAgentsRemoved(const TArray<int32>& AgentIds)
+{
+	if (IsSetup())
+	{
+		for (ULearningAgentsHelper* Helper : HelperObjects)
+		{
+			Helper->OnAgentsRemoved(AgentIds);
+		}
+
+		AgentsRemoved(AgentIds);
+	}
+}
+
+void ULearningAgentsManagerComponent::OnAgentsReset(const TArray<int32>& AgentIds)
+{
+	if (IsSetup())
+	{
+		for (ULearningAgentsHelper* Helper : HelperObjects)
+		{
+			Helper->OnAgentsReset(AgentIds);
+		}	
+	
+		AgentsReset(AgentIds);
+	}
+}
+
+UObject* ULearningAgentsManagerComponent::GetAgent(const int32 AgentId, const TSubclassOf<UObject> AgentClass) const
+{
+	if (!Manager)
+	{
+		UE_LOG(LogLearning, Error, TEXT("%s: Must be attached to a LearningAgentsManager Actor."), *GetName());
+		return nullptr;
+	}
+
+	return Manager->GetAgent(AgentId, AgentClass);
+}
+
+void ULearningAgentsManagerComponent::GetAgents(const TArray<int32>& AgentIds, const TSubclassOf<UObject> AgentClass, TArray<UObject*>& OutAgents) const
+{
+	if (!Manager)
+	{
+		UE_LOG(LogLearning, Error, TEXT("%s: Must be attached to a LearningAgentsManager Actor."), *GetName());
+		OutAgents.Empty();
+		return;
+	}
+
+	Manager->GetAgents(OutAgents, AgentIds, AgentClass);
+}
+
+void ULearningAgentsManagerComponent::GetAllAgents(TArray<UObject*>& OutAgents, TArray<int32>& OutAgentIds, const TSubclassOf<UObject> AgentClass) const
+{
+	if (!Manager)
+	{
+		UE_LOG(LogLearning, Error, TEXT("%s: Must be attached to a LearningAgentsManager Actor."), *GetName());
+		OutAgents.Empty();
+		return;
+	}
+
+	Manager->GetAllAgents(OutAgents, OutAgentIds, AgentClass);
+}
+
+ALearningAgentsManager* ULearningAgentsManagerComponent::GetAgentManager(const TSubclassOf<ALearningAgentsManager> AgentManagerClass) const
+{
+	if (!Manager)
+	{
+		UE_LOG(LogLearning, Error, TEXT("%s: Must be attached to a LearningAgentsManager Actor."), *GetName());
+		return nullptr;
+	}
+	
+	return Manager;
+}
+
+const UObject* ULearningAgentsManagerComponent::GetAgent(const int32 AgentId) const
+{
+	UE_LEARNING_CHECKF(Manager, TEXT("Must be attached to a LearningAgentsManager Actor."));
+	return Manager->GetAgent(AgentId);
+}
+
+UObject* ULearningAgentsManagerComponent::GetAgent(const int32 AgentId)
+{
+	UE_LEARNING_CHECKF(Manager, TEXT("Must be attached to a LearningAgentsManager Actor."));
+	return Manager->GetAgent(AgentId);
+}
+
+bool ULearningAgentsManagerComponent::HasAgent(const int32 AgentId) const
+{
+	UE_LEARNING_CHECKF(Manager, TEXT("Must be attached to a LearningAgentsManager Actor."));
+	return Manager->HasAgent(AgentId);
+}
+
+const ALearningAgentsManager* ULearningAgentsManagerComponent::GetAgentManager() const
+{
+	UE_LEARNING_CHECKF(Manager, TEXT("Must be attached to a LearningAgentsManager Actor."));
+	return Manager;
+}
+
+ALearningAgentsManager* ULearningAgentsManagerComponent::GetAgentManager()
+{
+	UE_LEARNING_CHECKF(Manager, TEXT("Must be attached to a LearningAgentsManager Actor."));
+	return Manager;
+}
+
+bool ULearningAgentsManagerComponent::HasAgentManager() const
+{
+	return Manager != nullptr;
+}
+
+bool ULearningAgentsManagerComponent::IsSetup() const
+{
+	return bIsSetup;
+}
+
+void ULearningAgentsManagerComponent::AddHelper(TObjectPtr<ULearningAgentsHelper> Object)
+{
+	UE_LEARNING_CHECK(Object);
+	HelperObjects.Add(Object);
+}
\ No newline at end of file
diff --git a/Engine/Plugins/Experimental/LearningAgents/Source/LearningAgents/Private/LearningAgentsNeuralNetwork.cpp b/Engine/Plugins/Experimental/LearningAgents/Source/LearningAgents/Private/LearningAgentsNeuralNetwork.cpp
--- a/Engine/Plugins/Experimental/LearningAgents/Source/LearningAgents/Private/LearningAgentsNeuralNetwork.cpp
+++ b/Engine/Plugins/Experimental/LearningAgents/Source/LearningAgents/Private/LearningAgentsNeuralNetwork.cpp
@@ -2,6 +2,8 @@
 
 #include "LearningAgentsNeuralNetwork.h"
 
+#include "LearningAgentsNeuralNetworkData.h"
+
 #include "LearningLog.h"
 #include "LearningArray.h"
 #include "LearningNeuralNetwork.h"
@@ -15,11 +17,8 @@ ULearningAgentsNeuralNetwork::~ULearningAgentsNeuralNetwork() = default;
 
 void ULearningAgentsNeuralNetwork::ResetNetwork()
 {
-	if (NeuralNetworkData)
-	{
-		NeuralNetworkData->ConditionalBeginDestroy();
-		NeuralNetworkData = nullptr;
-	}
+	NeuralNetworkData->ConditionalBeginDestroy();
+	NeuralNetworkData = nullptr;
 	
 	ForceMarkDirty();
 }
@@ -30,11 +29,39 @@ void ULearningAgentsNeuralNetwork::LoadNetworkFromSnapshot(const FFilePath& File
 
 	if (FFileHelper::LoadFileToArray(NetworkData, *File.FilePath))
 	{
-		if (!NeuralNetworkData->LoadFromSnapshot(NetworkData))
+		ULearningAgentsNeuralNetworkData* TempNeuralNetworkData = NewObject<ULearningAgentsNeuralNetworkData>(this);
+		int32 Offset = 0;
+		bool bSuccess = TempNeuralNetworkData->GetNetworkInterface()->DeserializeFromBytes(Offset, NetworkData);
+		
+		if (!bSuccess)
 		{
 			UE_LOG(LogLearning, Error, TEXT("%s: Failed to load network. Invalid Format: \"%s\""), *GetName(), *File.FilePath);
 			return;
 		}
+		
+		NetworkData.Empty();
+
+		if (NeuralNetworkData)
+		{
+			// If we already have a neural network check settings match
+
+			if (TempNeuralNetworkData->GetNetworkInterface()->GetInputNum() != NeuralNetworkData->GetNetworkInterface()->GetInputNum() ||
+				TempNeuralNetworkData->GetNetworkInterface()->GetOutputNum() != NeuralNetworkData->GetNetworkInterface()->GetOutputNum())
+			{
+				UE_LOG(LogLearning, Error, TEXT("%s: Failed to load network from snapshot as settings don't match."), *GetName());
+				return;
+			}
+			else
+			{
+				NeuralNetworkData->CopyFrom(TempNeuralNetworkData);
+			}
+		}
+		else
+		{
+			// Otherwise use loaded neural network as-is
+
+			NeuralNetworkData = TempNeuralNetworkData;
+		}
 
 		ForceMarkDirty();
 	}
@@ -53,8 +80,11 @@ void ULearningAgentsNeuralNetwork::SaveNetworkToSnapshot(const FFilePath& File)
 	}
 
 	TArray<uint8> NetworkData;
-	NetworkData.SetNumUninitialized(NeuralNetworkData->GetSnapshotByteNum());
-	NeuralNetworkData->SaveToSnapshot(NetworkData);
+	NetworkData.SetNumUninitialized(NeuralNetworkData->GetNetworkInterface()->GetSerializationByteNum());
+
+	int32 Offset = 0;
+	NeuralNetworkData->GetNetworkInterface()->SerializeToBytes(Offset, NetworkData);
+	UE_LEARNING_CHECK(Offset == NetworkData.Num());
 
 	if (!FFileHelper::SaveArrayToFile(NetworkData, *File.FilePath))
 	{
@@ -76,12 +106,21 @@ void ULearningAgentsNeuralNetwork::LoadNetworkFromAsset(ULearningAgentsNeuralNet
 		return;
 	}
 
-	if (!NeuralNetworkData)
+	if (NeuralNetworkData)
+	{
+		if (NeuralNetworkAsset->NeuralNetworkData->GetNetworkInterface()->GetInputNum() != NeuralNetworkData->GetNetworkInterface()->GetInputNum() ||
+			NeuralNetworkAsset->NeuralNetworkData->GetNetworkInterface()->GetOutputNum() != NeuralNetworkData->GetNetworkInterface()->GetOutputNum())
+		{
+			UE_LOG(LogLearning, Error, TEXT("%s: Failed to load network from asset as settings don't match."), *GetName());
+			return;
+		}
+	}
+	else
 	{
-		NeuralNetworkData = NewObject<ULearningNeuralNetworkData>(this);
+		NeuralNetworkData = NewObject<ULearningAgentsNeuralNetworkData>(this);
 	}
 
-	NeuralNetworkData->InitFrom(NeuralNetworkAsset->NeuralNetworkData);
+	NeuralNetworkData->CopyFrom(NeuralNetworkAsset->NeuralNetworkData);
 	ForceMarkDirty();
 }
 
@@ -101,10 +140,10 @@ void ULearningAgentsNeuralNetwork::SaveNetworkToAsset(ULearningAgentsNeuralNetwo
 
 	if (!NeuralNetworkAsset->NeuralNetworkData)
 	{
-		NeuralNetworkAsset->NeuralNetworkData = NewObject<ULearningNeuralNetworkData>(NeuralNetworkAsset);
+		NeuralNetworkAsset->NeuralNetworkData = NewObject<ULearningAgentsNeuralNetworkData>(NeuralNetworkAsset);
 	}
 
-	NeuralNetworkAsset->NeuralNetworkData->InitFrom(NeuralNetworkData);
+	NeuralNetworkAsset->NeuralNetworkData->CopyFrom(NeuralNetworkData);
 	NeuralNetworkAsset->ForceMarkDirty();
 }
 
diff --git a/Engine/Plugins/Experimental/LearningAgents/Source/LearningAgents/Private/LearningAgentsNeuralNetworkData.cpp b/Engine/Plugins/Experimental/LearningAgents/Source/LearningAgents/Private/LearningAgentsNeuralNetworkData.cpp
new file mode 100644
--- /dev/null
+++ b/Engine/Plugins/Experimental/LearningAgents/Source/LearningAgents/Private/LearningAgentsNeuralNetworkData.cpp
@@ -0,0 +1,360 @@
+// Copyright Epic Games, Inc. All Rights Reserved.
+
+#include "LearningAgentsNeuralNetworkData.h"
+
+#include "LearningAgentsNeuralNetwork.h"
+#include "LearningNeuralNetwork.h"
+
+#include "Async/ParallelFor.h"
+#include "Modules/ModuleManager.h"
+
+#include "NNE.h"
+#include "NNEModelData.h"
+#include "NNERuntimeBasicCpu.h"
+
+namespace UE::Learning::Agents
+{
+	FNeuralNetworkInference::FNeuralNetworkInference(
+		UE::NNE::IModelCPU& InModel,
+		const FNeuralNetworkInferenceSettings& InSettings,
+		const int32 MaxBatchSize,
+		const int32 InputSize,
+		const int32 OutputSize)
+		: Settings(InSettings)
+	{
+		InputBuffer.SetNumUninitialized({ MaxBatchSize, InputSize });
+		OutputBuffer.SetNumUninitialized({ MaxBatchSize, OutputSize });
+
+		CreateInstances(InModel);
+	}
+
+	void FNeuralNetworkInference::CreateInstances(NNE::IModelCPU& InModel)
+	{
+		UE_LEARNING_TRACE_CPUPROFILER_EVENT_SCOPE(Learning::Agents::FNeuralNetworkInference::CreateInstances);
+
+		const int32 MaxBatchSize = InputBuffer.Num<0>();
+		const int32 InputSize = InputBuffer.Num<1>();
+		const int32 OutputSize = OutputBuffer.Num<1>();
+
+		if (Settings.bParallelEvaluation)
+		{
+			const int32 IdealInferenceInstanceNum = ParallelForImpl::GetNumberOfThreadTasks(MaxBatchSize, Settings.MinParallelBatchSize, EParallelForFlags::None);
+			const int32 InferenceInstanceSliceLength = FMath::DivideAndRoundUp((int32)MaxBatchSize, IdealInferenceInstanceNum);
+			const int32 InferenceInstanceNum = FMath::DivideAndRoundUp((int32)MaxBatchSize, InferenceInstanceSliceLength);
+
+			InferenceInstances.Empty(InferenceInstanceNum);
+
+			for (int32 InferenceInstanceIdx = 0; InferenceInstanceIdx < InferenceInstanceNum; InferenceInstanceIdx++)
+			{
+				InferenceInstances.Emplace(InModel.CreateModelInstanceCPU());
+				InferenceInstances.Last()->SetInputTensorShapes({ NNE::FTensorShape::Make({(uint32)InferenceInstanceSliceLength, (uint32)InputSize}) });
+			}
+		}
+		else
+		{
+			InferenceInstances.Empty(1);
+			InferenceInstances.Emplace(InModel.CreateModelInstanceCPU());
+			InferenceInstances.Last()->SetInputTensorShapes({ NNE::FTensorShape::Make({(uint32)MaxBatchSize, (uint32)InputSize}) });
+		}
+	}
+
+	void FNeuralNetworkInference::Evaluate(
+		TLearningArrayView<2, float> Output,
+		const TLearningArrayView<2, const float> Input,
+		const FIndexSet Instances)
+	{
+		UE_LEARNING_TRACE_CPUPROFILER_EVENT_SCOPE(Learning::Agents::FNeuralNetworkInference::Evaluate);
+
+		const int32 InstanceNum = Instances.Num();
+		const int32 InputNum = Input.Num<1>();
+		const int32 OutputNum = Output.Num<1>();
+
+		// Gather Inputs
+
+		Array::Check(Input, Instances);
+
+		for (int32 InstanceIdx = 0; InstanceIdx < InstanceNum; InstanceIdx++)
+		{
+			Array::Copy(InputBuffer[InstanceIdx], Input[Instances[InstanceIdx]]);
+		}
+
+		// Run Inference
+
+		Array::Check(InputBuffer.Slice(0, InstanceNum));
+
+		if (Settings.bParallelEvaluation && InstanceNum > Settings.MinParallelBatchSize)
+		{
+			const int32 IdealInferenceInstanceNum = ParallelForImpl::GetNumberOfThreadTasks(InstanceNum, Settings.MinParallelBatchSize, EParallelForFlags::None);
+			const int32 InferenceInstanceSliceLength = FMath::DivideAndRoundUp(InstanceNum, IdealInferenceInstanceNum);
+			const int32 InferenceInstanceNum = FMath::DivideAndRoundUp(InstanceNum, InferenceInstanceSliceLength);
+
+			check(InferenceInstanceNum <= InferenceInstances.Num());
+
+			ParallelFor(InferenceInstanceNum, [this, InstanceNum, InputNum, InferenceInstanceSliceLength](int32 InferenceInstanceIdx)
+			{
+				const int32 StartIndex = InferenceInstanceIdx * InferenceInstanceSliceLength;
+				const int32 StopIndex = FMath::Min((InferenceInstanceIdx + 1) * InferenceInstanceSliceLength, InstanceNum);
+
+				TLearningArrayView<2, float> InputBufferSlice = InputBuffer.Slice(StartIndex, StopIndex - StartIndex);
+				TLearningArrayView<2, float> OutputBufferSlice = OutputBuffer.Slice(StartIndex, StopIndex - StartIndex);
+
+				InferenceInstances[InferenceInstanceIdx]->SetInputTensorShapes({ NNE::FTensorShape::Make({(uint32)(StopIndex - StartIndex), (uint32)InputNum}) });
+				InferenceInstances[InferenceInstanceIdx]->RunSync(
+					{ { (void*)InputBufferSlice.GetData(), InputBufferSlice.Num() * sizeof(float) } },
+					{ { (void*)OutputBufferSlice.GetData(), OutputBufferSlice.Num() * sizeof(float) } });
+			});
+		}
+		else
+		{
+			InferenceInstances[0]->SetInputTensorShapes({ NNE::FTensorShape::Make({(uint32)InstanceNum, (uint32)InputNum}) });
+			InferenceInstances[0]->RunSync(
+				{ { (void*)InputBuffer.GetData(), InputBuffer.Slice(0, InstanceNum).Num() * sizeof(float) } },
+				{ { (void*)OutputBuffer.GetData(), OutputBuffer.Slice(0, InstanceNum).Num() * sizeof(float) } });
+		}
+
+		Array::Check(OutputBuffer.Slice(0, InstanceNum));
+
+		// Scatter Outputs
+
+		for (int32 InstanceIdx = 0; InstanceIdx < InstanceNum; InstanceIdx++)
+		{
+			Array::Copy(Output[Instances[InstanceIdx]], OutputBuffer[InstanceIdx]);
+		}
+
+		Array::Check(Output, Instances);
+	}
+
+	FNeuralNetwork::FNeuralNetwork(ULearningAgentsNeuralNetworkData& InParent) : Parent(InParent) {}
+
+	void FNeuralNetwork::ReloadFromFileData()
+	{
+		UE_LEARNING_TRACE_CPUPROFILER_EVENT_SCOPE(Learning::Agents::FNeuralNetwork::ReloadFromFileData);
+
+		if (!Parent.ModelData)
+		{
+			Parent.ModelData = NewObject<UNNEModelData>(&Parent);
+		}
+
+		Parent.ModelData->Init(TEXT("ubnne"), Parent.FileData);
+
+		ensureMsgf(FModuleManager::Get().LoadModule(TEXT("NNERuntimeBasicCpu")), TEXT("Unable to load module for NNE runtime NNERuntimeBasicCpu."));
+
+		TWeakInterfacePtr<INNERuntimeCPU> RuntimeCPU = NNE::GetRuntime<INNERuntimeCPU>(TEXT("NNERuntimeBasicCpu"));
+
+		if (ensureMsgf(RuntimeCPU.IsValid(), TEXT("Could not find requested NNE Runtime")))
+		{
+			Model = RuntimeCPU->CreateModelCPU(Parent.ModelData);
+		}
+
+		// If we are not in the editor we can now clear the FileData and FileType since these will be
+		// using additional memory and we are not going to save this asset and so don't require them.
+
+#if !WITH_EDITOR
+		Parent.ModelData->ClearFileDataAndFileType();
+#endif
+
+		// We need to tell all the created inference objects to update their instances since the Model has changed
+		// and so the instance objects themselves are either no longer valid or tied to the old model.
+
+		for (const TWeakPtr<FNeuralNetworkInference>& InferenceObject : InferenceObjects)
+		{
+			if (TSharedPtr<FNeuralNetworkInference> InferenceObjectPtr = InferenceObject.Pin())
+			{
+				InferenceObjectPtr->CreateInstances(*Model);
+			}
+		}
+	}
+
+	bool FNeuralNetwork::DeserializeFromBytes(int32& InOutOffset, const TLearningArrayView<1, const uint8> RawBytes)
+	{
+		UE_LEARNING_TRACE_CPUPROFILER_EVENT_SCOPE(Learning::Agents::FNeuralNetwork::DeserializeFromBytes);
+
+		if (RawBytes.Num() - InOutOffset < Parent.FileData.Num())
+		{
+			return false;
+		}
+
+		Array::Copy(
+			TLearningArrayView<1, uint8>(Parent.FileData),
+			RawBytes.Slice(InOutOffset, Parent.FileData.Num()));
+
+		InOutOffset += Parent.FileData.Num();
+
+		ReloadFromFileData();
+
+		return true;
+	}
+
+	void FNeuralNetwork::SerializeToBytes(int32& InOutOffset, TLearningArrayView<1, uint8> OutRawBytes) const
+	{
+		UE_LEARNING_TRACE_CPUPROFILER_EVENT_SCOPE(Learning::Agents::FNeuralNetwork::SerializeToBytes);
+
+		Array::Copy(
+			OutRawBytes.Slice(InOutOffset, Parent.FileData.Num()),
+			TLearningArrayView<1, const uint8>(Parent.FileData));
+
+		InOutOffset += Parent.FileData.Num();
+	}
+
+	int32 FNeuralNetwork::GetSerializationByteNum() const
+	{
+		return Parent.FileData.Num();
+	}
+
+	TSharedRef<INeuralNetworkInference> FNeuralNetwork::CreateInferenceObject(
+		const int32 MaxBatchSize,
+		const FNeuralNetworkInferenceSettings& InSettings)
+	{
+		TSharedRef<FNeuralNetworkInference> InferenceObject = MakeShared<FNeuralNetworkInference>(
+			*Model,
+			InSettings,
+			MaxBatchSize,
+			Parent.InputNum,
+			Parent.OutputNum);
+
+		InferenceObjects.Emplace(InferenceObject.ToWeakPtr());
+
+		return InferenceObject;
+	}
+
+	int32 FNeuralNetwork::GetInputNum() const { return Parent.InputNum; }
+	int32 FNeuralNetwork::GetOutputNum() const { return Parent.OutputNum; }
+}
+
+
+void ULearningAgentsNeuralNetworkData::PostLoad()
+{
+	Super::PostLoad();
+
+	if (FileData.Num() > 0)
+	{
+		Network = MakeShared<UE::Learning::Agents::FNeuralNetwork>(*this);
+		Network->ReloadFromFileData();
+	}
+}
+
+TSharedPtr<UE::Learning::INeuralNetwork> ULearningAgentsNeuralNetworkData::GetNetworkInterface()
+{
+	return Network;
+}
+
+void ULearningAgentsNeuralNetworkData::CreateMLP(
+	const uint32 InputSize,
+	const uint32 OutputSize,
+	const uint32 HiddenUnitNum,
+	const uint32 LayerNum,
+	const ELearningAgentsActivationFunction Activation)
+{
+	UE_LEARNING_TRACE_CPUPROFILER_EVENT_SCOPE(ULearningAgentsNeuralNetworkData::CreateMLP);
+
+	UE_LEARNING_CHECK(InputSize > 0);
+	UE_LEARNING_CHECK(OutputSize > 0);
+	UE_LEARNING_CHECK(HiddenUnitNum > 0);
+	UE_LEARNING_CHECK(LayerNum > 0);
+
+	TArray<float> ZerosData;
+	ZerosData.Init(0.0f, FMath::Max(FMath::Max(InputSize * HiddenUnitNum, OutputSize * HiddenUnitNum), HiddenUnitNum * HiddenUnitNum));
+	TConstArrayView<float> ZerosDataView = ZerosData;
+
+	UE::NNE::RuntimeBasic::FSequentialModelBuilder Builder;
+
+	for (uint32 LayerIdx = 0; LayerIdx < LayerNum; LayerIdx++)
+	{
+		const int32 LayerInputNum = LayerIdx == 0 ? InputSize : HiddenUnitNum;
+		const int32 LayerOutputNum = LayerIdx == LayerNum - 1 ? OutputSize : HiddenUnitNum;
+
+		Builder.AddLinear(LayerInputNum, LayerOutputNum, ZerosDataView.Slice(0, LayerInputNum * LayerOutputNum), ZerosDataView.Slice(0, LayerOutputNum));
+
+		if (LayerIdx != LayerNum - 1)
+		{
+			switch (Activation)
+			{
+			case ELearningAgentsActivationFunction::ELU: Builder.AddELU(); break;
+			case ELearningAgentsActivationFunction::ReLU: Builder.AddReLU(); break;
+			case ELearningAgentsActivationFunction::TanH: Builder.AddTanH(); break;
+			default: UE_LEARNING_CHECKF(false, TEXT("Unimplemented")); break;
+			}
+		}
+	}
+
+	FileData.SetNumUninitialized(Builder.GetWriteByteNum());
+	Builder.WriteAndReset(FileData);
+
+	ZerosData.Empty();
+
+	InputNum = InputSize;
+	OutputNum = OutputSize;
+
+	if (!Network)
+	{
+		Network = MakeShared<UE::Learning::Agents::FNeuralNetwork>(*this);
+	}
+
+	Network->ReloadFromFileData();
+}
+
+void ULearningAgentsNeuralNetworkData::CreateMemoryBackbone(
+	const uint32 InputSize,
+	const uint32 OutputSize,
+	const uint32 MemorySize,
+	const uint32 HiddenUnitNum,
+	const uint32 PrefixLayerNum,
+	const uint32 PostfixLayerNum)
+{
+	UE_LEARNING_TRACE_CPUPROFILER_EVENT_SCOPE(ULearningAgentsNeuralNetworkData::CreateMemoryBackbone);
+
+	UE_LEARNING_CHECK(InputSize > 0);
+	UE_LEARNING_CHECK(OutputSize > 0);
+	UE_LEARNING_CHECK(MemorySize > 0);
+	UE_LEARNING_CHECK(HiddenUnitNum > 0);
+	UE_LEARNING_CHECK(PrefixLayerNum > 0);
+	UE_LEARNING_CHECK(PostfixLayerNum > 0);
+
+	UE::NNE::RuntimeBasic::FMemoryBackboneModelBuilder Builder;
+	Builder.BuildEmptyModel(
+		InputSize, 
+		OutputSize, 
+		MemorySize, 
+		HiddenUnitNum, 
+		PrefixLayerNum, 
+		PostfixLayerNum);
+
+	FileData.SetNumUninitialized(Builder.GetWriteByteNum());
+	Builder.WriteAndReset(FileData);
+
+	InputNum = InputSize + MemorySize;
+	OutputNum = OutputSize + MemorySize;
+
+	if (!Network)
+	{
+		Network = MakeShared<UE::Learning::Agents::FNeuralNetwork>(*this);
+	}
+
+	Network->ReloadFromFileData();
+}
+
+void ULearningAgentsNeuralNetworkData::CopyFrom(const ULearningAgentsNeuralNetworkData* Other)
+{
+	UE_LEARNING_TRACE_CPUPROFILER_EVENT_SCOPE(ULearningAgentsNeuralNetworkData::CopyFrom);
+
+	if (const ULearningAgentsNeuralNetworkData* OtherData = Cast<ULearningAgentsNeuralNetworkData>(Other))
+	{
+		FileData = OtherData->FileData;
+		InputNum = OtherData->InputNum;
+		OutputNum = OtherData->OutputNum;
+
+		if (!Network)
+		{
+			Network = MakeShared<UE::Learning::Agents::FNeuralNetwork>(*this);
+		}
+
+		Network->ReloadFromFileData();
+	}
+	else
+	{
+		ensureMsgf(false, TEXT("Cannot copy from object of different derived type."));
+	}
+}
+
+
+
diff --git a/Engine/Plugins/Experimental/LearningAgents/Source/LearningAgents/Private/LearningAgentsNeuralNetworkData.h b/Engine/Plugins/Experimental/LearningAgents/Source/LearningAgents/Private/LearningAgentsNeuralNetworkData.h
new file mode 100644
--- /dev/null
+++ b/Engine/Plugins/Experimental/LearningAgents/Source/LearningAgents/Private/LearningAgentsNeuralNetworkData.h
@@ -0,0 +1,119 @@
+// Copyright Epic Games, Inc. All Rights Reserved.
+
+#pragma once
+
+#include "LearningNeuralNetwork.h"
+
+#include "LearningAgentsNeuralNetworkData.generated.h"
+
+enum class ELearningAgentsActivationFunction : uint8;
+
+class UNNEModelData;
+
+namespace UE::NNE
+{
+	class IModelCPU;
+	class IModelInstanceCPU;
+}
+
+class ULearningAgentsNeuralNetworkData;
+
+namespace UE::Learning::Agents
+{
+	struct FNeuralNetwork;
+	struct FNeuralNetworkInference;
+
+	struct FNeuralNetwork : public INeuralNetwork
+	{
+		FNeuralNetwork(ULearningAgentsNeuralNetworkData& InParent);
+
+		//~ Begin INeuralNetwork Interface
+		virtual bool DeserializeFromBytes(int32& InOutOffset, const TLearningArrayView<1, const uint8> RawBytes) override final;
+		virtual void SerializeToBytes(int32& InOutOffset, TLearningArrayView<1, uint8> OutRawBytes) const override final;
+		virtual int32 GetSerializationByteNum() const override final;
+		virtual TSharedRef<INeuralNetworkInference> CreateInferenceObject(
+			const int32 MaxBatchSize,
+			const FNeuralNetworkInferenceSettings& Settings = FNeuralNetworkInferenceSettings()) override final;
+		virtual int32 GetInputNum() const override final;
+		virtual int32 GetOutputNum() const override final;
+		virtual const TCHAR* GetPythonClassName() const override final { return TEXT("NeuralNetworkNNE"); };
+		//~ End INeuralNetwork Interface
+
+		void ReloadFromFileData();
+
+		ULearningAgentsNeuralNetworkData& Parent;
+		TSharedPtr<NNE::IModelCPU> Model;
+		TArray<TWeakPtr<FNeuralNetworkInference>, TInlineAllocator<64>> InferenceObjects;
+	};
+
+	struct FNeuralNetworkInference : public INeuralNetworkInference
+	{
+		FNeuralNetworkInference(
+			NNE::IModelCPU& InModel,
+			const FNeuralNetworkInferenceSettings& InSettings,
+			const int32 MaxBatchSize,
+			const int32 InputSize,
+			const int32 OutputSize);
+
+		//~ Begin INeuralNetworkInference Interface
+		virtual void Evaluate(
+			TLearningArrayView<2, float> Output,
+			const TLearningArrayView<2, const float> Input,
+			const FIndexSet Instances) override final;
+		//~ End INeuralNetworkInference Interface
+
+		void CreateInstances(NNE::IModelCPU& InModel);
+
+		FNeuralNetworkInferenceSettings Settings;
+		TLearningArray<2, float> InputBuffer;
+		TLearningArray<2, float> OutputBuffer;
+		TArray<TSharedPtr<NNE::IModelInstanceCPU>, TInlineAllocator<64>> InferenceInstances;
+	};
+}
+
+UCLASS()
+class LEARNINGAGENTS_API ULearningAgentsNeuralNetworkData : public UObject
+{
+	GENERATED_BODY()
+
+	friend struct UE::Learning::Agents::FNeuralNetwork;
+
+public:
+
+	virtual void PostLoad() override final;
+
+	TSharedPtr<UE::Learning::INeuralNetwork> GetNetworkInterface();
+
+	void CopyFrom(const ULearningAgentsNeuralNetworkData* Other);
+
+	void CreateMLP(
+		const uint32 InputSize,
+		const uint32 OutputSize,
+		const uint32 HiddenUnitNum,
+		const uint32 LayerNum,
+		const ELearningAgentsActivationFunction Activation);
+
+	void CreateMemoryBackbone(
+		const uint32 InputSize,
+		const uint32 OutputSize,
+		const uint32 MemorySize,
+		const uint32 HiddenUnitNum,
+		const uint32 PrefixLayerNum,
+		const uint32 PostfixLayerNum);
+
+private:
+
+	UPROPERTY()
+	int32 InputNum = 0;
+
+	UPROPERTY()
+	int32 OutputNum = 0;
+
+	UPROPERTY()
+	TArray<uint8> FileData;
+
+	UPROPERTY()
+	TObjectPtr<UNNEModelData> ModelData;
+
+	TSharedPtr<UE::Learning::Agents::FNeuralNetwork> Network;
+};
