[Backout] - CL30154408

diff --git a/Engine/Plugins/Experimental/LearningAgents/Source/LearningAgents/Public/LearningAgentsCritic.h b/Engine/Plugins/Experimental/LearningAgents/Source/LearningAgents/Public/LearningAgentsCritic.h
--- a/Engine/Plugins/Experimental/LearningAgents/Source/LearningAgents/Public/LearningAgentsCritic.h
+++ b/Engine/Plugins/Experimental/LearningAgents/Source/LearningAgents/Public/LearningAgentsCritic.h
@@ -2,9 +2,10 @@
 
 #pragma once
 
-#include "LearningAgentsManagerListener.h"
+#include "LearningAgentsManagerComponent.h"
 
 #include "LearningAgentsNeuralNetwork.h" // Included for ELearningAgentsActivationFunction
+#include "LearningAgentsDebug.h"
 #include "LearningArray.h"
 #include "UObject/ObjectPtr.h"
 
@@ -12,9 +13,11 @@
 
 namespace UE::Learning
 {
-	struct FNeuralNetworkCritic;
+	struct INeuralNetwork;
+	struct FNeuralNetworkCriticFunction;
 }
 
+struct FFilePath;
 class ULearningAgentsNeuralNetwork;
 class ULearningAgentsInteractor;
 class ULearningAgentsPolicy;
@@ -40,9 +43,9 @@ public:
 	ELearningAgentsActivationFunction ActivationFunction = ELearningAgentsActivationFunction::ELU;
 };
 
-/** A critic used for training the policy. Can be used at inference time to estimate the discounted returns.  */
+/** A critic used by some algorithms for training the managed agents. */
 UCLASS(BlueprintType, Blueprintable, meta = (BlueprintSpawnableComponent))
-class LEARNINGAGENTS_API ULearningAgentsCritic : public ULearningAgentsManagerListener
+class LEARNINGAGENTS_API ULearningAgentsCritic : public ULearningAgentsManagerComponent
 {
 	GENERATED_BODY()
 
@@ -54,64 +57,69 @@ public:
 	virtual ~ULearningAgentsCritic();
 
 	/**
-	 * Constructs a Critic to be used with the given agent interactor and critic settings.
-	 * 
-	 * @param InManager						The input Manager
-	 * @param InInteractor					The input Interactor object
-	 * @param InPolicy						The input Policy object
-	 * @param Class							The critic class
-	 * @param Name							The critic name
-	 * @param CriticNeuralNetworkAsset		Optional Network Asset to use. If not provided, asset is empty, or
-	 *										bReinitializeNetwork is set then a new neural network object will be created
-	 *										according to the given CriticSettings.
-	 * @param bReinitializeCriticNetwork	If to reinitialize the critic network
-	 * @param CriticSettings				The critic settings to use on creation of a new critic network
-	 * @param Seed							Random seed to use for initializing the critic weights
+	 * Initializes this object to be used with the given agent interactor and critic settings.
+	 * @param InInteractor The input Interactor component
+	 * @param InPolicy The input Policy component
+	 * @param CriticSettings The critic settings to use
+	 * @param NeuralNetworkAsset Optional Network Asset to use. If provided must match the given CriticSettings. If not
+	 * provided or asset is empty then a new neural network object will be created according to the given 
+	 * CriticSettings and used.
 	 */
-	UFUNCTION(BlueprintCallable, Category = "LearningAgents", meta = (Class = "/Script/LearningAgents.LearningAgentsCritic", DeterminesOutputType = "Class"))
-	static ULearningAgentsCritic* MakeCritic(
-		ULearningAgentsManager* InManager,
+	UFUNCTION(BlueprintCallable, Category = "LearningAgents")
+	void SetupCritic(
 		ULearningAgentsInteractor* InInteractor,
 		ULearningAgentsPolicy* InPolicy,
-		TSubclassOf<ULearningAgentsCritic> Class,
-		const FName Name = TEXT("Critic"),
-		ULearningAgentsNeuralNetwork* CriticNeuralNetworkAsset = nullptr,
-		const bool bReinitializeCriticNetwork = true,
 		const FLearningAgentsCriticSettings& CriticSettings = FLearningAgentsCriticSettings(),
-		const int32 Seed = 1234);
+		ULearningAgentsNeuralNetwork* NeuralNetworkAsset = nullptr);
+
+public:
+
+	//~ Begin ULearningAgentsManagerComponent Interface
+	virtual void OnAgentsAdded(const TArray<int32>& AgentIds) override;
+	virtual void OnAgentsRemoved(const TArray<int32>& AgentIds) override;
+	virtual void OnAgentsReset(const TArray<int32>& AgentIds) override;
+	//~ End ULearningAgentsManagerComponent Interface
+
+// ----- Load / Save -----
+public:
 
 	/**
-	 * Initializes a critic to be used with the given agent interactor and critic settings.
-	 * 
-	 * @param InManager						The input Manager
-	 * @param InInteractor					The input Interactor object
-	 * @param InPolicy						The input Policy object
-	 * @param CriticNeuralNetworkAsset		Optional Network Asset to use. If not provided, asset is empty, or
-	 *										bReinitializeNetwork is set then a new neural network object will be created
-	 *										according to the given CriticSettings.
-	 * @param bReinitializeCriticNetwork	If to reinitialize the critic network
-	 * @param CriticSettings				The critic settings to use on creation of a new critic network
-	 * @param Seed							Random seed to use for initializing the critic weights
+	 * Load a snapshot's weights into this critic.
+	 * @param File The snapshot file.
+	 */
+	UFUNCTION(BlueprintCallable, Category = "LearningAgents", meta = (RelativePath))
+	void LoadCriticFromSnapshot(const FFilePath& File);
+
+	/**
+	 * Save this critic's weights into a snapshot.
+	 * @param File The snapshot file.
+	 */
+	UFUNCTION(BlueprintCallable, BlueprintPure = false, Category = "LearningAgents", meta = (RelativePath))
+	void SaveCriticToSnapshot(const FFilePath& File) const;
+
+	/**
+	 * Use a ULearningAgentsNeuralNetwork asset directly for this critic rather than making a copy. If used 
+	 * during training then this asset's weights will be updated as training progresses.
+	 * @param NeuralNetworkAsset The asset to use.
 	 */
 	UFUNCTION(BlueprintCallable, Category = "LearningAgents")
-	void SetupCritic(
-		ULearningAgentsManager* InManager,
-		ULearningAgentsInteractor* InInteractor,
-		ULearningAgentsPolicy* InPolicy,
-		ULearningAgentsNeuralNetwork* CriticNeuralNetworkAsset = nullptr,
-		const bool bReinitializeCriticNetwork = true,
-		const FLearningAgentsCriticSettings& CriticSettings = FLearningAgentsCriticSettings(),
-		const int32 Seed = 1234);
+	void UseCriticFromAsset(ULearningAgentsNeuralNetwork* NeuralNetworkAsset);
 
-public:
+	/**
+	 * Load a ULearningAgentsNeuralNetwork asset's weights into this critic.
+	 * @param NeuralNetworkAsset The asset to load from.
+	 */
+	UFUNCTION(BlueprintCallable, Category = "LearningAgents")
+	void LoadCriticFromAsset(ULearningAgentsNeuralNetwork* NeuralNetworkAsset);
 
-	//~ Begin ULearningAgentsManagerListener Interface
-	virtual void OnAgentsAdded_Implementation(const TArray<int32>& AgentIds) override;
-	virtual void OnAgentsRemoved_Implementation(const TArray<int32>& AgentIds) override;
-	virtual void OnAgentsReset_Implementation(const TArray<int32>& AgentIds) override;
-	//~ End ULearningAgentsManagerListener Interface
+	/**
+	 * Save this critic's weights to a ULearningAgentsNeuralNetwork asset.
+	 * @param NeuralNetworkAsset The asset to save to.
+	 */
+	UFUNCTION(BlueprintCallable, Category = "LearningAgents", Meta=(DevelopmentOnly))
+	void SaveCriticToAsset(ULearningAgentsNeuralNetwork* NeuralNetworkAsset);
 
-// ----- Blueprint public interface -----
+// ----- Evaluation -----
 public:
 
 	/**
@@ -133,15 +141,17 @@ public:
 	UFUNCTION(BlueprintPure, Category = "LearningAgents", Meta=(AgentId = -1))
 	float GetEstimatedDiscountedReturn(const int32 AgentId) const;
 
-	/** Gets the current Network Asset being used */
-	UFUNCTION(BlueprintPure, Category = "LearningAgents")
-	ULearningAgentsNeuralNetwork* GetCriticNetworkAsset();
-
 // ----- Non-blueprint public interface -----
 public:
 
+	/** Gets the current Network Asset being used */
+	ULearningAgentsNeuralNetwork* GetNetworkAsset();
+
+	/** Get a reference to this critic's neural network. */
+	UE::Learning::INeuralNetwork& GetCriticNetwork();
+	
 	/** Get a reference to this critic's critic function object. */
-	UE::Learning::FNeuralNetworkCritic& GetCriticObject();
+	UE::Learning::FNeuralNetworkCriticFunction& GetCriticObject();
 
 // ----- Private Data -----
 private:
@@ -156,18 +166,27 @@ private:
 
 	/** The underlying neural network. */
 	UPROPERTY(VisibleAnywhere, Transient, Category = "LearningAgents")
-	TObjectPtr<ULearningAgentsNeuralNetwork> CriticNetwork;
+	TObjectPtr<ULearningAgentsNeuralNetwork> Network;
 
 	/** Internal Critic Function Object */
-	TSharedPtr<UE::Learning::FNeuralNetworkCritic> CriticObject;
+	TSharedPtr<UE::Learning::FNeuralNetworkCriticFunction> CriticObject;
 
-	/** Buffer to store the computed returns for each agent */
-	TLearningArray<1, float, TInlineAllocator<32>> Returns;
+// ----- Private Iteration Checks ----- 
+private:
 
-	/** Number of times the returns have been computed for all agents */
-	TLearningArray<1, uint64, TInlineAllocator<32>> ReturnsIteration;
+	/** Number of times critic has been evaluated for all agents */
+	TLearningArray<1, uint64, TInlineAllocator<32>> CriticAgentIteration;
 
 	/** Temp buffers used to record the set of agents that are valid for evaluation */
 	TArray<int32> ValidAgentIds;
 	UE::Learning::FIndexSet ValidAgentSet;
+
+private:
+#if UE_LEARNING_AGENTS_ENABLE_VISUAL_LOG
+	/** Color used to draw this action in the visual log */
+	FLinearColor VisualLogColor = FColor::Orange;
+
+	/** Describes this critic to the visual logger for debugging purposes. */
+	void VisualLog(const UE::Learning::FIndexSet AgentSet) const;
+#endif
 };
diff --git a/Engine/Plugins/Experimental/LearningAgents/Source/LearningAgents/Public/LearningAgentsHelpers.h b/Engine/Plugins/Experimental/LearningAgents/Source/LearningAgents/Public/LearningAgentsHelpers.h
new file mode 100644
--- /dev/null
+++ b/Engine/Plugins/Experimental/LearningAgents/Source/LearningAgents/Public/LearningAgentsHelpers.h
@@ -0,0 +1,419 @@
+// Copyright Epic Games, Inc. All Rights Reserved.
+
+#pragma once
+
+#include "LearningAgentsDebug.h"
+
+#include "Templates/SharedPointer.h"
+#include "UObject/Object.h"
+#include "Components/SplineComponent.h" // Required for ESplineCoordinateSpace::World
+
+#include "LearningAgentsHelpers.generated.h"
+
+class ULearningAgentsManagerComponent;
+class UMeshComponent;
+
+//------------------------------------------------------------------
+
+/**
+ * The base class for all helpers. Helpers are additional objects that can be used in getting or setting observations,
+ * actions, rewards, and completions.
+ */
+UCLASS(Abstract, BlueprintType)
+class LEARNINGAGENTS_API ULearningAgentsHelper : public UObject
+{
+	GENERATED_BODY()
+
+public:
+
+	/** Reference to the Manager Component this helper is associated with. */
+	UPROPERTY(VisibleAnywhere, Transient, Category = "LearningAgents")
+	TObjectPtr<ULearningAgentsManagerComponent> ManagerComponent;
+
+public:
+
+	/**
+	 * Called whenever agents are added to the associated ULearningAgentsManagerComponent object.
+	 * @param AgentIds Array of agent ids which have been added
+	 */
+	virtual void OnAgentsAdded(const TArray<int32>& AgentIds);
+
+	/**
+	 * Called whenever agents are removed from the associated ULearningAgentsManagerComponent object.
+	 * @param AgentIds Array of agent ids which have been removed
+	 */
+	virtual void OnAgentsRemoved(const TArray<int32>& AgentIds);
+
+	/**
+	 * Called whenever agents are reset on the associated ULearningAgentsManagerComponent object.
+	 * @param AgentIds Array of agent ids which have been reset
+	 */
+	virtual void OnAgentsReset(const TArray<int32>& AgentIds);
+
+public:
+#if UE_LEARNING_AGENTS_ENABLE_VISUAL_LOG
+	/** Color used to draw this helper in the visual log */
+	FLinearColor VisualLogColor = FColor::Magenta;
+#endif
+};
+
+//------------------------------------------------------------------
+
+/** A helper for computing various properties from a SplineComponent. */
+UCLASS()
+class LEARNINGAGENTS_API USplineComponentHelper : public ULearningAgentsHelper
+{
+	GENERATED_BODY()
+
+public:
+
+	/**
+	 * Adds a new spline component helper to the given manager component.
+	 * @param InManagerComponent The manager component to add this helper to (such as an Interactor or Trainer).
+	 * @param Name The name of this new helper. Used for debugging.
+	 * @return The newly created helper.
+	 */
+	UFUNCTION(BlueprintCallable, Category = "LearningAgents", meta = (DefaultToSelf = "InManagerComponent"))
+	static USplineComponentHelper* AddSplineComponentHelper(ULearningAgentsManagerComponent* InManagerComponent, const FName Name = NAME_None);
+
+	/**
+	 * Gets the position on a spline closest to the provided position.
+	 * @param AgentId The agent id to run the helper for.
+	 * @param SplineComponent The spline component to use.
+	 * @param Position The position to find the closest position to.
+	 * @param CoordinateSpace The coordinate space to use for the spline.
+	 * @return The closest position on the spline.
+	 */
+	UFUNCTION(BlueprintPure = false, Category = "LearningAgents", meta = (AgentId = "-1"))
+	FVector GetNearestPositionOnSpline(const int32 AgentId, const USplineComponent* SplineComponent, const FVector Position, const ESplineCoordinateSpace::Type CoordinateSpace = ESplineCoordinateSpace::World) const;
+
+	/**
+	 * Gets the distance along a spline closest to the provided position.
+	 * @param AgentId The agent id to run the helper for.
+	 * @param SplineComponent The spline component to use.
+	 * @param Position The position to find the closest position to.
+	 * @param CoordinateSpace The coordinate space to use for the spline.
+	 * @return The distance along the spline of the closest position.
+	 */
+	UFUNCTION(BlueprintPure = false, Category = "LearningAgents", meta = (AgentId = "-1"))
+	float GetDistanceAlongSplineAtPosition(const int32 AgentId, const USplineComponent* SplineComponent, const FVector Position, const ESplineCoordinateSpace::Type CoordinateSpace = ESplineCoordinateSpace::World) const;
+
+	/**
+	 * Gets the position along a spline at the given distance.
+	 * @param AgentId The agent id to run the helper for.
+	 * @param SplineComponent The spline component to use.
+	 * @param DistanceAlongSpline The distance along the spline to use.
+	 * @param CoordinateSpace The coordinate space to use for the spline.
+	 * @return The position at the given distance along the spline.
+	 */
+	UFUNCTION(BlueprintPure = false, Category = "LearningAgents", meta = (AgentId = "-1"))
+	FVector GetPositionAtDistanceAlongSpline(const int32 AgentId, const USplineComponent* SplineComponent, const float DistanceAlongSpline, const ESplineCoordinateSpace::Type CoordinateSpace = ESplineCoordinateSpace::World) const;
+
+	/**
+	 * Gets the direction along a spline at the given distance.
+	 * @param AgentId The agent id to run the helper for.
+	 * @param SplineComponent The spline component to use.
+	 * @param DistanceAlongSpline The distance along the spline to use.
+	 * @param CoordinateSpace The coordinate space to use for the spline.
+	 * @return The direction at the given distance along the spline.
+	 */
+	UFUNCTION(BlueprintPure = false, Category = "LearningAgents", meta = (AgentId = "-1"))
+	FVector GetDirectionAtDistanceAlongSpline(const int32 AgentId, const USplineComponent* SplineComponent, const float DistanceAlongSpline, const ESplineCoordinateSpace::Type CoordinateSpace = ESplineCoordinateSpace::World) const;
+
+	/**
+	 * Gets the proportion along a spline in the range 0-1 for a given distance.
+	 * @param AgentId The agent id to run the helper for.
+	 * @param SplineComponent The spline component to use.
+	 * @param DistanceAlongSpline The distance along the spline to use.
+	 * @return The proportion along the spline in the range 0-1.
+	 */
+	UFUNCTION(BlueprintPure = false, Category = "LearningAgents", meta = (AgentId = "-1"))
+	float GetProportionAlongSpline(const int32 AgentId, const USplineComponent* SplineComponent, const float DistanceAlongSpline) const;
+
+	/**
+	 * Gets the proportion along a spline encoded as an angle between -180 and 180 degrees for a given distance. 
+	 * This should be used for looped splines.
+	 * @param AgentId The agent id to run the helper for.
+	 * @param SplineComponent The spline component to use.
+	 * @param DistanceAlongSpline The distance along the spline to use.
+	 * @return The proportion along the spline encoded as an angle between -180 and 180 degrees
+	 */
+	UFUNCTION(BlueprintPure = false, Category = "LearningAgents", meta = (AgentId = "-1"))
+	float GetProportionAlongSplineAsAngle(const int32 AgentId, const USplineComponent* SplineComponent, const float DistanceAlongSpline) const;
+
+	/**
+	 * Gets an array of positions sampled along a spline between some starting and stopping distance. Deals properly
+	 * for splines which are looped.
+	 * @param OutPositions Output array of spline positions.
+	 * @param AgentId The agent id to run the helper for.
+	 * @param SplineComponent The spline component to use.
+	 * @param PositionNum The number of positions to sample along the spline.
+	 * @param StartDistanceAlongSpline The starting distance along the spline.
+	 * @param StopDistanceAlongSpline The stopping distance along the spline.
+	 * @param CoordinateSpace The coordinate space to use for the spline.
+	 */
+	UFUNCTION(BlueprintPure = false, Category = "LearningAgents", meta = (AgentId = "-1"))
+	void GetPositionsAlongSpline(
+		TArray<FVector>& OutPositions,
+		const int32 AgentId,
+		const USplineComponent* SplineComponent,
+		const int32 PositionNum,
+		const float StartDistanceAlongSpline,
+		const float StopDistanceAlongSpline,
+		const ESplineCoordinateSpace::Type CoordinateSpace = ESplineCoordinateSpace::World) const;
+
+	/**
+	 * Gets the velocity a point is traveling along a spline.
+	 * @param AgentId The agent id to run the helper for.
+	 * @param SplineComponent The spline component to use.
+	 * @param Position The position of the point.
+	 * @param Velocity The velocity of the point.
+	 * @param FiniteDifferenceDelta The finite difference delta used to estimate the velocity along the spline.
+	 * @param CoordinateSpace The coordinate space to use for the spline.
+	 * @return The scalar velocity the point is traveling along the spline.
+	 */
+	UFUNCTION(BlueprintPure = false, Category = "LearningAgents", meta = (AgentId = "-1"))
+	float GetVelocityAlongSpline(
+		const int32 AgentId,
+		const USplineComponent* SplineComponent,
+		const FVector Position,
+		const FVector Velocity,
+		const float FiniteDifferenceDelta = 10.0f,
+		const ESplineCoordinateSpace::Type CoordinateSpace = ESplineCoordinateSpace::World) const;
+};
+
+//------------------------------------------------------------------
+
+/** A helper for projecting onto surfaces. */
+UCLASS()
+class LEARNINGAGENTS_API UProjectionHelper : public ULearningAgentsHelper
+{
+	GENERATED_BODY()
+
+public:
+
+	/**
+	 * Adds a new projection helper to the given manager component.
+	 * @param InManagerComponent The manager component to add this helper to (such as an Interactor or Trainer).
+	 * @param Name The name of this new helper. Used for debugging.
+	 * @return The newly created helper.
+	 */
+	UFUNCTION(BlueprintCallable, Category = "LearningAgents", meta = (DefaultToSelf = "InManagerComponent"))
+	static UProjectionHelper* AddProjectionHelper(ULearningAgentsManagerComponent* InManagerComponent, const FName Name = NAME_None);
+
+	/**
+	 * Projects a transform onto the zero height ground plane resulting in translation only in XY and rotation only around the Z.
+	 * @param AgentId The agent id to run the helper for.
+	 * @param Transform The transform to project.
+	 * @param LocalForwardVector The local forward direction.
+	 * @return The transform projected onto the ground plane.
+	 */
+	UFUNCTION(BlueprintPure = false, Category = "LearningAgents", meta = (AgentId = "-1"))
+	FTransform ProjectTransformOntoGroundPlane(const int32 AgentId, const FTransform Transform, const FVector LocalForwardVector = FVector::ForwardVector) const;
+
+	/**
+	 * Projects a position and rotation onto the zero height ground plane resulting in translation only in XY and rotation only around the Z.
+	 * @param OutPosition The output projected position.
+	 * @param OutRotation The output projected rotation.
+	 * @param AgentId The agent id to run the helper for.
+	 * @param InPosition The input position.
+	 * @param InRotation The input rotation.
+	 * @param LocalForwardVector The local forward direction.
+	 */
+	UFUNCTION(BlueprintPure = false, Category = "LearningAgents", meta = (AgentId = "-1"))
+	void ProjectPositionRotationOntoGroundPlane(
+		FVector& OutPosition, 
+		FRotator& OutRotation, 
+		const int32 AgentId, 
+		const FVector InPosition, 
+		const FRotator InRotation, 
+		const FVector LocalForwardVector = FVector::ForwardVector) const;
+};
+
+//------------------------------------------------------------------
+
+/** A helper for getting various properties from a MeshComponent. */
+UCLASS()
+class LEARNINGAGENTS_API UMeshComponentHelper : public ULearningAgentsHelper
+{
+	GENERATED_BODY()
+
+public:
+
+	/**
+	 * Adds a new mesh component helper to the given manager component.
+	 * @param InManagerComponent The manager component to add this helper to (such as an Interactor or Trainer).
+	 * @param Name The name of this new helper. Used for debugging.
+	 * @return The newly created helper.
+	 */
+	UFUNCTION(BlueprintCallable, Category = "LearningAgents", meta = (DefaultToSelf = "InManagerComponent"))
+	static UMeshComponentHelper* AddMeshComponentHelper(ULearningAgentsManagerComponent* InManagerComponent, const FName Name = NAME_None);
+
+	/**
+	 * Gets the bone positions for a set of bones of the mesh component in world space.
+	 * @param OutBonePositions The output array of bone positions.
+	 * @param AgentId The agent id to run the helper for.
+	 * @param MeshComponent The mesh component to use.
+	 * @param BoneNames The names of the bones to use.
+	 */
+	UFUNCTION(BlueprintPure = false, Category = "LearningAgents", meta = (AgentId = "-1"))
+	void GetMeshBonePositions(TArray<FVector>& OutBonePositions, const int32 AgentId, const UMeshComponent* MeshComponent, const TArray<FName>& BoneNames) const;
+};
+
+//------------------------------------------------------------------
+
+/** A helper for performing various kinds of ray cast. */
+UCLASS()
+class LEARNINGAGENTS_API URayCastHelper : public ULearningAgentsHelper
+{
+	GENERATED_BODY()
+
+public:
+
+	/**
+	 * Adds a new ray cast helper to the given manager component.
+	 * @param InManagerComponent The manager component to add this helper to (such as an Interactor or Trainer).
+	 * @param Name The name of this new helper. Used for debugging.
+	 * @return The newly created helper.
+	 */
+	UFUNCTION(BlueprintCallable, Category = "LearningAgents", meta = (DefaultToSelf = "InManagerComponent"))
+	static URayCastHelper* AddRayCastHelper(ULearningAgentsManagerComponent* InManagerComponent, const FName Name = NAME_None);
+
+	/**
+	 * Samples a grid of heights from above.
+	 * @param OutHeights The output array of heights for each ray in the grid.
+	 * @param AgentId The agent id to run the helper for.
+	 * @param Position The central position of the grid.
+	 * @param Rotation The orientation of the grid.
+	 * @param RowNum The number of rows in the grid.
+	 * @param ColNum The number of columns in the grid.
+	 * @param RowWidth The width of grid rows.
+	 * @param ColWidth The width of grid columns.
+	 * @param MaxHeight The height at which to start ray casts from.
+	 * @param MinHeight The height at which to send ray casts to.
+	 * @param CollisionChannel The collision channel to use for the ray-casts.
+	 */
+	UFUNCTION(BlueprintPure=false, Category = "LearningAgents", meta = (AgentId = "-1"))
+	void RayCastGridHeights(
+		TArray<float>& OutHeights,
+		const int32 AgentId,
+		const FVector Position,
+		const FRotator Rotation,
+		const int32 RowNum = 5,
+		const int32 ColNum = 5,
+		const float RowWidth = 1000.0f,
+		const float ColWidth = 1000.0f,
+		const float MaxHeight = 10000.0f,
+		const float MinHeight = -10000.0f,
+		const ECollisionChannel CollisionChannel = ECollisionChannel::ECC_WorldStatic) const;
+
+	/**
+	 * Casts rays out in a radius from a given position.
+	 * @param OutDistances The output array of distances for each ray from the start position.
+	 * @param AgentId The agent id to run the helper for.
+	 * @param Position The central position of the ray cast.
+	 * @param Rotation The orientation of the ray cast.
+	 * @param RayNum The number of rays to cast.
+	 * @param MinAngle The minimum angle to spread the rays.
+	 * @param MaxAngle The maximum angle to spread the rays.
+	 * @param MaxRayDist The maximum distance at which to shoot rays.
+	 * @param CollisionChannel The collision channel to use for the ray-casts.
+	 */
+	UFUNCTION(BlueprintPure = false, Category = "LearningAgents", meta = (AgentId = "-1"))
+	void RayCastRadial(
+		TArray<float>& OutDistances,
+		const int32 AgentId,
+		const FVector Position,
+		const FRotator Rotation,
+		const int32 RayNum = 7,
+		const float MinAngle = -90.0f,
+		const float MaxAngle = 90.0f,
+		const float MaxRayDist = 500.0f,
+		const FVector LocalForward = FVector::ForwardVector,
+		const ECollisionChannel CollisionChannel = ECollisionChannel::ECC_WorldStatic) const;
+
+};
+
+//------------------------------------------------------------------
+
+/**
+ * A helper for monitoring collisions between components.
+ *
+ * This object works by monitoring for collisions between a given component and any other components with a given tag.
+ * To indicate which component to monitor for collisions `SetComponent` should be called (generally on the 
+ * `ULearningAgentsManagerComponent::AgentsAdded` event) with the corresponding agent id.
+ * 
+ * If any collisions occur `GetCollisionOccurred` will report true until it is reset to false either by calling 
+ * `ResetCollisionOccurred` or `GetAndResetCollisionOccurred`. It will also be automatically reset when the agent is 
+ * reset.
+ */
+UCLASS()
+class LEARNINGAGENTS_API UCollisionMonitorHelper : public ULearningAgentsHelper
+{
+	GENERATED_BODY()
+
+public:
+
+	/**
+	 * Adds a new collision monitor helper to the given manager component.
+	 * @param InManagerComponent The manager component to add this helper to (such as an Interactor or Trainer).
+	 * @param Name The name of this new helper. Used for debugging.
+	 * @return The newly created helper.
+	 */
+	UFUNCTION(BlueprintCallable, Category = "LearningAgents", meta = (DefaultToSelf = "InManagerComponent"))
+	static UCollisionMonitorHelper* AddCollisionMonitorHelper(
+		ULearningAgentsManagerComponent* InManagerComponent, 
+		const FName Name = NAME_None);
+
+	/**
+	 * Sets the component to monitor for collisions. Generally this should be called on the 
+	 * `ULearningAgentsManagerComponent::AgentsAdded` event.
+	 * @param AgentId The agent id to associate this component with.
+	 * @param Component The component to monitor for collisions.
+	 * @param OtherComponentTag The tag other components must have to trigger collisions.
+	 */
+	UFUNCTION(BlueprintCallable, Category = "LearningAgents", meta = (AgentId = "-1"))
+	void SetComponent(const int32 AgentId, UPrimitiveComponent* Component, const FName OtherComponentTag = NAME_None);
+
+	/**
+	 * Gets if a collision has occurred for this agent.
+	 * @param AgentId The agent id to check for collisions for.
+	 */
+	UFUNCTION(BlueprintPure = false, Category = "LearningAgents", meta = (AgentId = "-1"))
+	bool GetCollisionOccurred(const int32 AgentId) const;
+	
+	/**
+	 * Gets if a collision has occurred for this agent and resets the state of if a collision has occurred back to false.
+	 * @param AgentId The agent id to check for collisions for.
+	 */
+	UFUNCTION(BlueprintCallable, Category = "LearningAgents", meta = (AgentId = "-1"))
+	bool GetAndResetCollisionOccurred(const int32 AgentId);
+
+	/**
+	 * Resets the state of if a collision has occurred back to false.
+	 * @param AgentId The agent id to reset the state for.
+	 */
+	UFUNCTION(BlueprintCallable, Category = "LearningAgents", meta = (AgentId = "-1"))
+	void ResetCollisionOccurred(const int32 AgentId);
+
+public:
+
+	/** Initialize the internal state for a given maximum number of agents */
+	void Init(const int32 MaxAgentNum);
+
+	//~ Begin ULearningAgentsHelper Interface
+	virtual void OnAgentsRemoved(const TArray<int32>& AgentIds) override;
+	virtual void OnAgentsReset(const TArray<int32>& AgentIds) override;
+	//~ End ULearningAgentsHelper Interface
+
+private:
+
+	TLearningArray<1, UPrimitiveComponent*, TInlineAllocator<32>> Components;
+	TLearningArray<1, FName, TInlineAllocator<32>> OtherComponentTags;
+	TLearningArray<1, bool, TInlineAllocator<32>> CollisionsOccured;
+
+private:
+
+	UFUNCTION()
+	void HandleOnHit(UPrimitiveComponent* HitComponent, AActor* OtherActor, UPrimitiveComponent* OtherComp, FVector NormalImpulse, const FHitResult& Hit);
+};
diff --git a/Engine/Plugins/Experimental/LearningAgents/Source/LearningAgents/Public/LearningAgentsInteractor.h b/Engine/Plugins/Experimental/LearningAgents/Source/LearningAgents/Public/LearningAgentsInteractor.h
--- a/Engine/Plugins/Experimental/LearningAgents/Source/LearningAgents/Public/LearningAgentsInteractor.h
+++ b/Engine/Plugins/Experimental/LearningAgents/Source/LearningAgents/Public/LearningAgentsInteractor.h
@@ -2,9 +2,7 @@
 
 #pragma once
 
-#include "LearningAgentsManagerListener.h"
-#include "LearningAgentsObservations.h"
-#include "LearningAgentsActions.h"
+#include "LearningAgentsManagerComponent.h"
 
 #include "LearningArray.h"
 #include "Containers/Array.h"
@@ -16,306 +14,195 @@
 
 #include "LearningAgentsInteractor.generated.h"
 
-class ULearningAgentsNeuralNetwork;
+namespace UE::Learning
+{
+	struct FArrayMap;
+	struct FConcatenateFeature;
+	struct FFeatureObject;
+}
 
-DECLARE_DYNAMIC_DELEGATE_RetVal_TwoParams(FLearningAgentsObservationObjectElement, FGatherAgentObservationDelegate, ULearningAgentsObservationObject*, InObservationObject, const int32, AgentId);
-DECLARE_DYNAMIC_DELEGATE_ThreeParams(FScatterAgentActionDelegate, const FLearningAgentsActionObjectElement, InActionObjectElement, const ULearningAgentsActionObject*, InActionObject, const int32, AgentId);
+class ALearningAgentsManager;
+class ULearningAgentsAction;
+class ULearningAgentsObservation;
 
 /**
  * ULearningAgentsInteractor defines how agents interact with the environment through their observations and actions.
+ * Additionally, it provides methods to be called during the inference process of those agents.
  *
- * To use this class, you need to implement `SpecifyAgentObservations` and `SpecifyAgentActions`, which will define 
- * the structure of inputs and outputs to your policy. You also need to implement `GatherAgentObservations` and 
- * `ScatterAgentActions` which will dictate how those observations are gathered, and actions actuated in your
- * environment.
+ * To use this class, you need to implement the SetupObservations and SetupActions functions (as well as their
+ * corresponding SetObservations and SetActions functions), which will define the size of inputs and outputs to your
+ * policy. Before you can do inference, you need to call SetupInteractor, which will initialize the underlying data
+ * structure, and you need to call AddAgent for each object you want controlled by this agent interactor.
  */
-UCLASS(Abstract, HideDropdown, BlueprintType, Blueprintable)
-class LEARNINGAGENTS_API ULearningAgentsInteractor : public ULearningAgentsManagerListener
+UCLASS(Abstract, BlueprintType, Blueprintable)
+class LEARNINGAGENTS_API ULearningAgentsInteractor : public ULearningAgentsManagerComponent
 {
 	GENERATED_BODY()
 
 public:
 
-	friend class ULearningAgentsController;
-	friend class ULearningAgentsPolicy;
-	friend class ULearningAgentsRecorder;
-	friend class ULearningAgentsTrainer;
-
 	// These constructors/destructors are needed to make forward declarations happy
 	ULearningAgentsInteractor();
 	ULearningAgentsInteractor(FVTableHelper& Helper);
 	virtual ~ULearningAgentsInteractor();
 
-	/**
-	 * Constructs an Interactor.
-	 *
-	 * @param InManager						The input Manager
-	 * @param Class							The interactor class
-	 * @param Name							The interactor name
-	 */
-	UFUNCTION(BlueprintCallable, Category = "LearningAgents", meta = (DeterminesOutputType = "Class"))
-	static ULearningAgentsInteractor* MakeInteractor(
-		ULearningAgentsManager* InManager, 
-		TSubclassOf<ULearningAgentsInteractor> Class,
-		const FName Name = TEXT("Interactor"));
-
-	/**
-	 * Initializes an Interactor.
-	 *
-	 * @param InManager						The input Manager
-	 */
+	/** Initializes this object and runs the setup events for observations and actions. */
 	UFUNCTION(BlueprintCallable, Category = "LearningAgents")
-	void SetupInteractor(ULearningAgentsManager* InManager);
+	void SetupInteractor();
 
 public:
 
-	//~ Begin ULearningAgentsManagerListener Interface
-	virtual void OnAgentsAdded_Implementation(const TArray<int32>& AgentIds) override;
-	virtual void OnAgentsRemoved_Implementation(const TArray<int32>& AgentIds) override;
-	virtual void OnAgentsReset_Implementation(const TArray<int32>& AgentIds) override;
-	//~ End ULearningAgentsManagerListener Interface
+	//~ Begin ULearningAgentsManagerComponent Interface
+	virtual void OnAgentsAdded(const TArray<int32>& AgentIds) override;
+	virtual void OnAgentsRemoved(const TArray<int32>& AgentIds) override;
+	virtual void OnAgentsReset(const TArray<int32>& AgentIds) override;
+	//~ End ULearningAgentsManagerComponent Interface
 
 // ----- Observations -----
 public:
 
 	/**
-	 * Specifies the structure of the observations using the observation schema
-	 * 
-	 * @param OutObservationSchemaElement		Output Schema Element
-	 * @param InObservationSchema				Observation Schema
+	 * During this event, all observations should be added to this agent interactor.
+	 * @see LearningAgentsObservations.h for the list of available observations.
 	 */
-	UFUNCTION(BlueprintNativeEvent, Category = "LearningAgents", Meta = (ForceAsFunction))
-	void SpecifyAgentObservations(FLearningAgentsObservationSchemaElement& OutObservationSchemaElement, ULearningAgentsObservationSchema* InObservationSchema);
+	UFUNCTION(BlueprintNativeEvent, Category = "LearningAgents")
+	void SetupObservations();
 
 	/**
-	 * Gathers all the observations for the given agents. The structure of the Observation Elements output by this function should match that defined
-	 * by the schema.
-	 *
-	 * @param OutObservationObjectElements		Output Observation Element. This should be the same size as AgentIds.
-	 * @param InObservationObject				Observation Object.
-	 * @param AgentIds							Set of Agent Ids to gather observations for.
+	 * During this event, all observations should be set for each agent.
+	 * @param AgentIds The list of agent ids to set observations for.
+	 * @see LearningAgentsObservations.h for the list of available observations.
+	 * @see GetAgent to get the agent corresponding to each id.
 	 */
-	UFUNCTION(BlueprintNativeEvent, Category = "LearningAgents", Meta = (ForceAsFunction))
-	void GatherAgentObservations(TArray<FLearningAgentsObservationObjectElement>& OutObservationObjectElements, ULearningAgentsObservationObject* InObservationObject, const TArray<int32>& AgentIds);
+	UFUNCTION(BlueprintNativeEvent, Category = "LearningAgents")
+	void SetObservations(const TArray<int32>& AgentIds);
 
 	/**
-	 * This is a convenience function that can be used to implement `GatherAgentObservations` in terms of a callback which iterates over every agent
-	 * individually and gathers the observations.
-	 * 
-	 * @param OutObservationObjectElements		Output Observation Element. This should be the same size as AgentIds.
-	 * @param InObservationObject				Observation Object.
-	 * @param AgentIds							Set of Agent Ids to gather observations for.
-	 * @param Delegate							The Delegate used to gather each agent's observations.
+	 * Used by objects derived from ULearningAgentsObservation to add themselves to this agent interactor during their
+	 * creation. You shouldn't need to call this directly.
 	 */
-	UFUNCTION(BlueprintPure = false, Category = "LearningAgents")
-	void GatherAgentObservationsUsingDelegate(
-			TArray<FLearningAgentsObservationObjectElement>& OutObservationElements,
-			ULearningAgentsObservationObject* InObservationObject, 
-			const TArray<int32>& AgentIds,
-			const FGatherAgentObservationDelegate& Delegate);
+	void AddObservation(TObjectPtr<ULearningAgentsObservation> Object, const TSharedRef<UE::Learning::FFeatureObject>& Feature);
 
 // ----- Actions -----
 public:
 
 	/**
-	 * Specifies the structure of the actions using the action schema
-	 *
-	 * @param OutActionSchemaElement			Output Schema Element
-	 * @param InActionSchema					Action Schema
+	 * During this event, all actions should be added to this agent interactor.
+	 * @see LearningAgentsActions.h for the list of available actions.
 	 */
-	UFUNCTION(BlueprintNativeEvent, Category = "LearningAgents", Meta = (ForceAsFunction))
-	void SpecifyAgentActions(FLearningAgentsActionSchemaElement& OutActionSchemaElement, ULearningAgentsActionSchema* InActionSchema);
+	UFUNCTION(BlueprintNativeEvent, Category = "LearningAgents")
+	void SetupActions();
 
 	/**
-	 * Scatters all the actions for the given agents and applies their affect in the world. The structure of the Action Elements given as input to 
-	 * this function will match that defined by the schema.
-	 *
-	 * @param InActionObjectElements			Input Actions Element. This will be the same size as AgentIds.
-	 * @param InActionObject					Action Object.
-	 * @param AgentIds							Set of Agent Ids to scatter observations for.
+	 * During this event, you should retrieve the actions and apply them to your agents.
+	 * @param AgentIds The list of agent ids to get actions for.
+	 * @see LearningAgentsActions.h for the list of available actions.
+	 * @see GetAgent to get the agent corresponding to each id.
 	 */
-	UFUNCTION(BlueprintNativeEvent, Category = "LearningAgents", Meta = (ForceAsFunction))
-	void ScatterAgentActions(const TArray<FLearningAgentsActionObjectElement>& InActionObjectElements, const ULearningAgentsActionObject* InActionObject, const TArray<int32>& AgentIds);
+	UFUNCTION(BlueprintNativeEvent, Category = "LearningAgents")
+	void GetActions(const TArray<int32>& AgentIds);
 
 	/**
-	 * This is a convenience function that can be used to implement `ScatterAgentActions` in terms of a callback which iterates over every agent
-	 * individually and scatters the actions.
-	 *
-	 * @param InActionObjectElements			Input Actions Element. This will be the same size as AgentIds.
-	 * @param InActionObject					Action Object.
-	 * @param AgentIds							Set of Agent Ids to scatter observations for.
-	 * @param Delegate							The Delegate used to scatter each agent's actions.
+	 * Used by objects derived from ULearningAgentsAction to add themselves to this agent interactor during their creation.
+	 * You shouldn't need to call this directly.
 	 */
-	UFUNCTION(BlueprintPure = false, Category = "LearningAgents")
-	void ScatterAgentObservationsUsingDelegate(
-		const TArray<FLearningAgentsActionObjectElement>& InActionObjectElements, 
-		const ULearningAgentsActionObject* InActionObject,
-		const TArray<int32>& AgentIds,
-		const FScatterAgentActionDelegate& Delegate);
-
-// ----- Blueprint public interface -----
-public:
-
-	/** Gathers all the observations for all agents. This will call GatherAgentObservations. */
-	UFUNCTION(BlueprintCallable, Category = "LearningAgents")
-	void GatherObservations();
+	void AddAction(TObjectPtr<ULearningAgentsAction> Object, const TSharedRef<UE::Learning::FFeatureObject>& Feature);
 
-	/** Scatters all the actions for all agents. This will call ScatterAgentActions. */
-	UFUNCTION(BlueprintCallable, Category = "LearningAgents")
-	void ScatterActions();
-
-	/**
-	 * Get the current buffered observation vector for the given agent.
-	 * 
-	 * @param OutObservationVector				Output Observation Vector
-	 * @param OutObservationCompatibilityHash	Output Compatibility Hash used to identify which schema this vector is compatible with.
-	 * @param AgentId							Agent Id to look up the observation vector for.
-	 */
-	UFUNCTION(BlueprintCallable, Category = "LearningAgents", Meta = (AgentId = "-1"))
-	void GetObservationVector(TArray<float>& OutObservationVector, int32& OutObservationCompatibilityHash, const int32 AgentId);
+// ----- Encoding / Decoding -----
+public:
 
 	/**
-	 * Get the current buffered action vector for the given agent.
-	 *
-	 * @param OutActionVector					Output Action Vector
-	 * @param OutActionCompatibilityHash		Output Compatibility Hash used to identify which schema this vector is compatible with.
-	 * @param AgentId							Agent Id to look up the action vector for.
+	 * Call this function when it is time to gather all the observations for your agents. This can be done each frame or
+	 * you can consider wiring it up to some kind of meaningful event, e.g. a hypothetical "OnAITurnStarted" if you have
+	 * a turn-based game. This will call this agent interactor's SetObservations event.
 	 */
-	UFUNCTION(BlueprintCallable, Category = "LearningAgents", Meta = (AgentId = "-1"))
-	void GetActionVector(TArray<float>& OutActionVector, int32& OutActionCompatibilityHash, const int32 AgentId);
+	UFUNCTION(BlueprintCallable, Category = "LearningAgents")
+	void EncodeObservations();
 
 	/**
-	 * Sets the current buffered observation vector for the given agent.
-	 *
-	 * @param ObservationVector					Observation Vector
-	 * @param InObservationCompatibilityHash	Compatibility Hash used to identify which schema this vector is compatible with.
-	 * @param AgentId							Agent Id to set the observation vector for.
-	 * @param bIncrementIteration				If to increment the iteration number used to keep track of associated actions and observations.
+	 * Call this function when it is time for your agents to take their actions. You most likely want to call this after
+	 * your policy's EvaluatePolicy function to ensure that you are receiving the latest actions. This will call this
+	 * agent interactor's GetActions event.
 	 */
-	UFUNCTION(BlueprintCallable, Category = "LearningAgents", Meta = (AgentId = "-1"))
-	void SetObservationVector(const TArray<float>& ObservationVector, const int32 InObservationCompatibilityHash, const int32 AgentId, bool bIncrementIteration = true);
+	UFUNCTION(BlueprintCallable, Category = "LearningAgents")
+	void DecodeActions();
 
 	/**
-	 * Sets the current buffered action vector for the given agent.
+	 * Gets the observation vector used by a given agent. Should be called only after EncodeObservations.
 	 *
-	 * @param ActionVector						Action Vector
-	 * @param InActionCompatibilityHash			Compatibility Hash used to identify which schema this vector is compatible with.
-	 * @param AgentId							Agent Id to set the observation vector for.
-	 * @param bIncrementIteration				If to increment the iteration number used to keep track of associated actions and observations.
-	 */
-	UFUNCTION(BlueprintCallable, Category = "LearningAgents", Meta = (AgentId = "-1"))
-	void SetActionVector(const TArray<float>& ActionVector, const int32 InActionCompatibilityHash, const int32 AgentId, bool bIncrementIteration = true);
-
-	/**
-	 * Returns true if GatherObservations or SetObservationVector has been called and the observation vector already set for the given agent.
+	 * @param AgentId				The AgentId to look-up the observation vector for
+	 * @param OutObservationVector	The observation vector for the given agent
 	 */
 	UFUNCTION(BlueprintPure, Category = "LearningAgents", meta = (AgentId = "-1"))
-	bool HasObservationVector(const int32 AgentId) const;
+	void GetObservationVector(const int32 AgentId, TArray<float>& OutObservationVector) const;
 
 	/**
-	 * Returns true if DecodeAndSampleActions on the policy or SetActionVector has been called and the action vector already set for the given agent.
+	 * Gets the action vector used by a given agent. Should be called only after EncodeActions or EvaluatePolicy.
+	 *
+	 * @param AgentId				The AgentId to look-up the action vector for
+	 * @param OutActionVector		The action vector for the given agent
 	 */
 	UFUNCTION(BlueprintPure, Category = "LearningAgents", meta = (AgentId = "-1"))
-	bool HasActionVector(const int32 AgentId) const;
-
-	/** Gets the size of the observation vector used by this interactor. */
-	UFUNCTION(BlueprintPure, Category = "LearningAgents")
-	int32 GetObservationVectorSize() const;
-
-	/** Gets the size of the encoded observation vector used by this interactor. */
-	UFUNCTION(BlueprintPure, Category = "LearningAgents")
-	int32 GetObservationEncodedVectorSize() const;
-
-	/** Gets the size of the action vector used by this interactor. */
-	UFUNCTION(BlueprintPure, Category = "LearningAgents")
-	int32 GetActionVectorSize() const;
-
-	/** Gets the size of the action distribution vector used by this interactor. */
-	UFUNCTION(BlueprintPure, Category = "LearningAgents")
-	int32 GetActionDistributionVectorSize() const;
-
-	/** Gets the size of the encoded action vector used by this interactor. */
-	UFUNCTION(BlueprintPure, Category = "LearningAgents")
-	int32 GetActionEncodedVectorSize() const;
+	void GetActionVector(const int32 AgentId, TArray<float>& OutActionVector) const;
 
 // ----- Non-blueprint public interface -----
 public:
 
-	/** Encode Observations for a specific set of agents */
-	void GatherObservations(const UE::Learning::FIndexSet AgentSet, bool bIncrementIteration = true);
-
-	/** Scatter Actions for a specific set of agents */
-	void ScatterActions(const UE::Learning::FIndexSet AgentSet);
+	/** Get a reference to this agent interactor's observation feature. */
+	UE::Learning::FFeatureObject& GetObservationFeature() const;
 
-	/** Gets the internal observation schema object */
-	const UE::Learning::Observation::FSchema& GetObservationSchema() const;
+	/** Get a reference to this agent interactor's action feature. */
+	UE::Learning::FFeatureObject& GetActionFeature() const;
 
-	/** Gets the internal observation schema element */
-	UE::Learning::Observation::FSchemaElement GetObservationSchemaElement() const;
+	/** Get a const array view of this agent interactor's observation objects. */
+	TConstArrayView<ULearningAgentsObservation*> GetObservationObjects() const;
 
-	/** Gets the internal action schema object */
-	const UE::Learning::Action::FSchema& GetActionSchema() const;
+	/** Get a const array view of this agent interactor's action objects. */
+	TConstArrayView<ULearningAgentsAction*> GetActionObjects() const;
 
-	/** Gets the internal action schema element */
-	UE::Learning::Action::FSchemaElement GetActionSchemaElement() const;
+	/** Get the number of times observations have been encoded for all agents */
+	TLearningArrayView<1, uint64> GetObservationEncodingAgentIteration();
 
-private:
+	/** Get the number of times actions have been encoded for all agents */
+	TLearningArrayView<1, uint64> GetActionEncodingAgentIteration();
 
-	/** Observation Schema used by this interactor */
-	UPROPERTY(VisibleAnywhere, Transient, Category = "LearningAgents")
-	TObjectPtr<ULearningAgentsObservationSchema> ObservationSchema;
-
-	/** Observation Schema Element used by this interactor */
-	UPROPERTY(VisibleAnywhere, Transient, Category = "LearningAgents")
-	FLearningAgentsObservationSchemaElement ObservationSchemaElement;
-
-	/** Action Schema used by this interactor */
-	UPROPERTY(VisibleAnywhere, Transient, Category = "LearningAgents")
-	TObjectPtr<ULearningAgentsActionSchema> ActionSchema;
-
-	/** Action Schema Element used by this interactor */
-	UPROPERTY(VisibleAnywhere, Transient, Category = "LearningAgents")
-	FLearningAgentsActionSchemaElement ActionSchemaElement;
+	/** Encode Observations for a specific set of agents */
+	void EncodeObservations(const UE::Learning::FIndexSet AgentSet);
 
-	/** Observation Object used by this interactor */
-	UPROPERTY(VisibleAnywhere, Transient, Category = "LearningAgents")
-	TObjectPtr<ULearningAgentsObservationObject> ObservationObject;
+	/** Decode Actions for a specific set of agents */
+	void DecodeActions(const UE::Learning::FIndexSet AgentSet);
 
-	/** Observation Object Elements used by this interactor */
-	UPROPERTY(VisibleAnywhere, Transient, Category = "LearningAgents")
-	TArray<FLearningAgentsObservationObjectElement> ObservationObjectElements;
+private:
 
-	/** Action Object used by this interactor */
+	/** The list of current observation objects. */
 	UPROPERTY(VisibleAnywhere, Transient, Category = "LearningAgents")
-	TObjectPtr<ULearningAgentsActionObject> ActionObject;
+	TArray<TObjectPtr<ULearningAgentsObservation>> ObservationObjects;
 
-	/** Action Object Elements used by this interactor */
+	/** The list of current action objects. */
 	UPROPERTY(VisibleAnywhere, Transient, Category = "LearningAgents")
-	TArray<FLearningAgentsActionObjectElement> ActionObjectElements;
+	TArray<TObjectPtr<ULearningAgentsAction>> ActionObjects;
 
 // ----- Private Data -----
 private:
 
-	/** Buffer of Observation Vectors for each agent */
-	TLearningArray<2, float> ObservationVectors;
-
-	/** Buffer of Action Vectors for each agent */
-	TLearningArray<2, float> ActionVectors;
+	TArray<TSharedRef<UE::Learning::FFeatureObject>, TInlineAllocator<16>> ObservationFeatures;
+	TArray<TSharedRef<UE::Learning::FFeatureObject>, TInlineAllocator<16>> ActionFeatures;
 
-	/** Compatibility Hash for Observation Schema */
-	int32 ObservationCompatibilityHash = 0;
+	TSharedPtr<UE::Learning::FConcatenateFeature> Observations;
+	TSharedPtr<UE::Learning::FConcatenateFeature> Actions;
 
-	/** Compatibility Hash for Actiuon Schema */
-	int32 ActionCompatibilityHash = 0;
+// ----- Private Iteration Checks ----- 
+private:
 
-	/** Number of times observation vector has been set for all agents */
-	TLearningArray<1, uint64, TInlineAllocator<32>> ObservationVectorIteration;
+	/** Number of times observations have been encoded for all agents */
+	TLearningArray<1, uint64, TInlineAllocator<32>> ObservationEncodingAgentIteration;
 
-	/** Number of times action vector has been set for all agents */
-	TLearningArray<1, uint64, TInlineAllocator<32>> ActionVectorIteration;
+	/** Number of times actions have been encoded for all agents */
+	TLearningArray<1, uint64, TInlineAllocator<32>> ActionEncodingAgentIteration;
 
 	/** Temp buffers used to record the set of agents that are valid for encoding/decoding */
 	TArray<int32> ValidAgentIds;
 	UE::Learning::FIndexSet ValidAgentSet;
+	TBitArray<TInlineAllocator<32>> ValidAgentStatus;
 
 };
diff --git a/Engine/Plugins/Experimental/LearningAgents/Source/LearningAgents/Public/LearningAgentsManager.h b/Engine/Plugins/Experimental/LearningAgents/Source/LearningAgents/Public/LearningAgentsManager.h
--- a/Engine/Plugins/Experimental/LearningAgents/Source/LearningAgents/Public/LearningAgentsManager.h
+++ b/Engine/Plugins/Experimental/LearningAgents/Source/LearningAgents/Public/LearningAgentsManager.h
@@ -4,18 +4,22 @@
 
 #include "LearningArray.h"
 
+#include "GameFramework/Actor.h"
 #include "Containers/Array.h"
 #include "Containers/ArrayView.h"
 #include "Containers/ContainerAllocationPolicies.h"
 #include "Templates/SharedPointer.h"
-
 #include "UObject/Object.h"
 #include "UObject/ObjectPtr.h"
-#include "Components/ActorComponent.h"
 
 #include "LearningAgentsManager.generated.h"
 
-class ULearningAgentsManagerListener;
+namespace UE::Learning
+{
+	struct FArrayMap;
+}
+
+class ULearningAgentsManagerComponent;
 
 /**
  * The agent manager is responsible for tracking which game objects are agents. It's the central class around which
@@ -24,22 +28,20 @@ class ULearningAgentsManagerListener;
  * If you have multiple different types of objects you want controlled by Learning Agents, you should consider creating
  * one agent manager per object type, rather than trying to share an agent manager.
  */
-UCLASS(BlueprintType, Blueprintable, meta = (BlueprintSpawnableComponent))
-class LEARNINGAGENTS_API ULearningAgentsManager : public UActorComponent
+UCLASS(BlueprintType, Blueprintable, meta=(ChildCanTick))
+class LEARNINGAGENTS_API ALearningAgentsManager : public AActor
 {
 	GENERATED_BODY()
 
 public:
 
 	// These constructors/destructors are needed to make forward declarations happy
-	ULearningAgentsManager();
-	ULearningAgentsManager(FVTableHelper& Helper);
-	virtual ~ULearningAgentsManager();
+	ALearningAgentsManager();
+	ALearningAgentsManager(FVTableHelper& Helper);
+	virtual ~ALearningAgentsManager();
 
 	virtual void PostInitProperties() override;
 
-	virtual void TickComponent(float DeltaTime, enum ELevelTick TickType, FActorComponentTickFunction* ThisTickFunction) override;
-
 // ----- Agent Management -----
 public:
 
@@ -110,7 +112,6 @@ public:
 	/**
 	 * Gets the agent with the given id. Calling this from blueprint with the appropriate AgentClass will automatically
 	 * cast the object to the given type. If not in a blueprint, you should use one of the other GetAgent overloads.
-	 * 
 	 * @param AgentId The id of the agent to get.
 	 * @param AgentClass The class to cast the agent object to (in blueprint).
 	 * @return The agent object.
@@ -121,7 +122,6 @@ public:
 	/**
 	 * Gets the agents associated with a set of ids. Calling this from blueprint with the appropriate AgentClass will 
 	 * automatically cast the object to the given type.
-	 * 
 	 * @param AgentIds The ids of the agents to get.
 	 * @param AgentClass The class to cast the agent objects to (in blueprint).
 	 * @param OutAgents The output array of agent objects.
@@ -132,7 +132,6 @@ public:
 	/**
 	 * Gets all added agents. Calling this from blueprint with the appropriate AgentClass will automatically
 	 * cast the object to the given type.
-	 * 
 	 * @param AgentClass The class to cast the agent objects to (in blueprint).
 	 * @param OutAgents The output array of agent objects.
 	 * @param OutAgentIds The output array of agent ids.
@@ -142,7 +141,6 @@ public:
 
 	/**
 	 * Gets the agent id associated with a given agent.
-	 * 
 	 * @param Agent The agent object.
 	 * @return The agent id.
 	 */
@@ -151,7 +149,6 @@ public:
 
 	/**
 	 * Gets the agent ids associated with a set of agents.
-	 * 
 	 * @param OutAgentIds The ids of the agents.
 	 * @param InAgents The agent objects.
 	 */
@@ -160,7 +157,6 @@ public:
 
 	/**
 	 * Gets the number of agents added
-	 * 
 	 * @return The number of agents added.
 	 */
 	UFUNCTION(BlueprintPure, Category = "LearningAgents")
@@ -176,15 +172,22 @@ public:
 	UFUNCTION(BlueprintPure, Category = "LearningAgents", meta = (AgentId = "-1"))
 	bool HasAgent(const int32 AgentId) const;
 
+// ----- Tick Prerequisites -----
 public:
 
-	/** Adds a listener to be tracked by this manager. */
+	/**
+	 * Adds this manager as a tick prerequisite of the given actor objects.
+	 * @param InAgents The actor object which this manager will be added as a tick prerequisite to.
+	 */
 	UFUNCTION(BlueprintCallable, Category = "LearningAgents")
-	void AddListener(ULearningAgentsManagerListener* Listener);
+	void AddManagerAsTickPrerequisiteOfAgents(const TArray<AActor*>& InAgents);
 
-	/** Removes a listener from being tracked by this manager. */
+	/**
+	 * Adds the given actor objects as tick prerequisite of this manager.
+	 * @param InAgents The actor object which will be added as tick prerequisites to this manager.
+	 */
 	UFUNCTION(BlueprintCallable, Category = "LearningAgents")
-	void RemoveListener(ULearningAgentsManagerListener* Listener);
+	void AddAgentsAsTickPrerequisiteOfManager(const TArray<AActor*>& InAgents);
 
 
 // ----- Non-blueprint public interface -----
@@ -202,6 +205,9 @@ public:
 	/** Gets the set of agent ids currently added as an FIndexSet */
 	UE::Learning::FIndexSet GetAllAgentSet() const;
 
+	/** Get a const reference to this manager's underlying instance data. */
+	const TSharedPtr<UE::Learning::FArrayMap>& GetInstanceData() const;
+
 	/** Get a const array view of this manager's agent objects. */
 	TConstArrayView<TObjectPtr<UObject>> GetAgents() const;
 
@@ -215,16 +221,15 @@ private:
 	UPROPERTY(VisibleAnywhere, Transient, Category = "LearningAgents")
 	TArray<TObjectPtr<UObject>> Agents;
 
-	/** The list of current listeners. */
-	UPROPERTY(VisibleAnywhere, Transient, Category = "LearningAgents")
-	TArray<TObjectPtr<ULearningAgentsManagerListener>> Listeners;
-
 private:
 
+	/** Object containing all of the instance data required by all agents */
+	TSharedPtr<UE::Learning::FArrayMap> InstanceData;
+
 	/** Update the agent sets to keep them in sync with the id lists. */
 	void UpdateAgentSets();
 
-	/** Array of agent ids to be passed to events such as ULearningAgentsManagerListener::OnAgentAdded. */
+	/** Array of agent ids to be passed to events such as ULearningAgentsManagerComponent::OnAgentAdded. */
 	TArray<int32> OnEventAgentIds;
 
 	/** Array of agent ids currently in use and associated with each agent object. */
diff --git a/Engine/Plugins/Experimental/LearningAgents/Source/LearningAgents/Public/LearningAgentsManagerComponent.h b/Engine/Plugins/Experimental/LearningAgents/Source/LearningAgents/Public/LearningAgentsManagerComponent.h
new file mode 100644
--- /dev/null
+++ b/Engine/Plugins/Experimental/LearningAgents/Source/LearningAgents/Public/LearningAgentsManagerComponent.h
@@ -0,0 +1,159 @@
+// Copyright Epic Games, Inc. All Rights Reserved.
+
+#pragma once
+
+#include "Components/ActorComponent.h"
+
+#include "LearningArray.h"
+#include "Containers/Array.h"
+#include "UObject/Object.h"
+#include "UObject/ObjectPtr.h"
+
+#include "LearningAgentsManagerComponent.generated.h"
+
+class ALearningAgentsManager;
+class ULearningAgentsHelper;
+
+/**
+ * Base class for components which can be attached to an ALearningAgentsManager.
+ */
+UCLASS(Abstract, BlueprintType, Blueprintable)
+class LEARNINGAGENTS_API ULearningAgentsManagerComponent : public UActorComponent
+{
+	GENERATED_BODY()
+
+public:
+
+	// These constructors/destructors are needed to make forward declarations happy
+	ULearningAgentsManagerComponent(const FObjectInitializer& ObjectInitializer = FObjectInitializer::Get());
+	ULearningAgentsManagerComponent(FVTableHelper& Helper);
+	virtual ~ULearningAgentsManagerComponent();
+
+	virtual void PostInitProperties() override;
+
+	/**
+	 * During this event, any additional logic required for when agents are added can be executed.
+	 */
+	UFUNCTION(BlueprintNativeEvent, Category = "LearningAgents")
+	void AgentsAdded(const TArray<int32>& AgentIds);
+
+	/**
+	 * During this event, any additional logic required for when agents are removed can be executed.
+	 */
+	UFUNCTION(BlueprintNativeEvent, Category = "LearningAgents")
+	void AgentsRemoved(const TArray<int32>& AgentIds);
+
+	/**
+	 * During this event, any additional logic required for when agents are reset can be executed.
+	 */
+	UFUNCTION(BlueprintNativeEvent, Category = "LearningAgents")
+	void AgentsReset(const TArray<int32>& AgentIds);
+
+	/**
+	 * Called whenever agents are added to the parent ALearningAgentsManager object.
+	 * @param AgentIds Array of agent ids which have been added
+	 */
+	virtual void OnAgentsAdded(const TArray<int32>& AgentIds);
+
+	/**
+	 * Called whenever agents are removed from the parent ALearningAgentsManager object.
+	 * @param AgentIds Array of agent ids which have been removed
+	 */
+	virtual void OnAgentsRemoved(const TArray<int32>& AgentIds);
+
+	/**
+	 * Called whenever agents are reset on the parent ALearningAgentsManager object.
+	 * @param AgentIds Array of agent ids which have been reset
+	 */
+	virtual void OnAgentsReset(const TArray<int32>& AgentIds);
+
+	/** Returns true if this component has been setup. Otherwise, false. */
+	UFUNCTION(BlueprintPure, Category = "LearningAgents")
+	bool IsSetup() const;
+
+// ----- Blueprint Convenience Functions -----
+protected:
+
+	/**
+	 * Gets the agent with the given id from the manager. Calling this from blueprint with the appropriate AgentClass
+	 * will automatically cast the object to the given type. If not in a blueprint, you should call the manager's
+	 * GetAgent methods directly.
+	 * @param AgentId The id of the agent to get.
+	 * @param AgentClass The class to cast the agent object to (in blueprint).
+	 * @return The agent object.
+	 */
+	UFUNCTION(BlueprintPure, Category = "LearningAgents", meta = (AgentId = "-1", DeterminesOutputType = "AgentClass"))
+	UObject* GetAgent(const int32 AgentId, const TSubclassOf<UObject> AgentClass) const;
+
+	/**
+	 * Gets the agents associated with a set of ids from the manager. Calling this from blueprint with the appropriate 
+	 * AgentClass will automatically cast the object to the given type. If not in a blueprint, you should call the 
+	 * manager's GetAgents method directly.
+	 * @param AgentIds The ids of the agents to get.
+	 * @param AgentClass The class to cast the agent objects to (in blueprint).
+	 * @param OutAgents The output array of agent objects.
+	 */
+	UFUNCTION(BlueprintPure = false, Category = "LearningAgents", meta = (DeterminesOutputType = "AgentClass", DynamicOutputParam = "OutAgents"))
+	void GetAgents(const TArray<int32>& AgentIds, const TSubclassOf<UObject> AgentClass, TArray<UObject*>& OutAgents) const;
+
+	/**
+	 * Gets all added agents from the manager. Calling this from blueprint with the appropriate AgentClass will 
+	 * automatically cast the object to the given type.
+	 * @param AgentClass The class to cast the agent objects to (in blueprint).
+	 * @param OutAgents The output array of agent objects.
+	 * @param OutAgentIds The output array of agent ids.
+	 */
+	UFUNCTION(BlueprintPure = false, Category = "LearningAgents", meta = (DeterminesOutputType = "AgentClass", DynamicOutputParam = "OutAgents"))
+	void GetAllAgents(TArray<UObject*>& OutAgents, TArray<int32>& OutAgentIds, const TSubclassOf<UObject> AgentClass) const;
+
+	/**
+	 * Gets the agent manager associated with this component.
+	 * @param AgentClass The class to cast the agent manager to (in blueprint).
+	 */
+	UFUNCTION(BlueprintPure, Category = "LearningAgents", meta = (DeterminesOutputType = "AgentManagerClass"))
+	ALearningAgentsManager* GetAgentManager(const TSubclassOf<ALearningAgentsManager> AgentManagerClass) const;
+
+// ----- Non-blueprint public interface -----
+public:
+
+	/** Gets the agent corresponding to the given id. */
+	const UObject* GetAgent(const int32 AgentId) const;
+
+	/** Gets the agent corresponding to the given id. */
+	UObject* GetAgent(const int32 AgentId);
+
+	/** Checks if the component has the given agent id */
+	bool HasAgent(const int32 AgentId) const;
+
+	/** Gets the associated agent manager assuming it exists. */
+	const ALearningAgentsManager* GetAgentManager() const;
+
+	/** Gets the associated agent manager assuming it exists. */
+	ALearningAgentsManager* GetAgentManager();
+
+	/** Checks if the component has an agent manager. */
+	bool HasAgentManager() const;
+
+// ----- Helpers -----
+public:
+
+	/**
+	 * Used by objects derived from ULearningAgentsHelper to add themselves to this component during their creation.
+	 * You shouldn't need to call this directly.
+	 */
+	void AddHelper(TObjectPtr<ULearningAgentsHelper> Object);
+
+protected:
+
+	/** True if this component has been setup. Otherwise, false. */
+	UPROPERTY(VisibleAnywhere, Transient, Category = "LearningAgents")
+	bool bIsSetup = false;
+
+	/** The associated manager this component is attached to. */
+	UPROPERTY(VisibleAnywhere, Transient, Category = "LearningAgents")
+	TObjectPtr<ALearningAgentsManager> Manager;
+
+	/** The list of current helper objects. */
+	UPROPERTY(VisibleAnywhere, Transient, Category = "LearningAgents")
+	TArray<TObjectPtr<ULearningAgentsHelper>> HelperObjects;
+};
diff --git a/Engine/Plugins/Experimental/LearningAgents/Source/LearningAgents/Public/LearningAgentsNeuralNetwork.h b/Engine/Plugins/Experimental/LearningAgents/Source/LearningAgents/Public/LearningAgentsNeuralNetwork.h
--- a/Engine/Plugins/Experimental/LearningAgents/Source/LearningAgents/Public/LearningAgentsNeuralNetwork.h
+++ b/Engine/Plugins/Experimental/LearningAgents/Source/LearningAgents/Public/LearningAgentsNeuralNetwork.h
@@ -20,7 +20,7 @@ enum class ELearningAgentsActivationFunction : uint8
 	TanH	UMETA(DisplayName = "TanH"),
 };
 
-class ULearningNeuralNetworkData;
+class ULearningAgentsNeuralNetworkData;
 
 /** A neural network data asset. */
 UCLASS(BlueprintType)
@@ -34,13 +34,14 @@ public:
 	ULearningAgentsNeuralNetwork(FVTableHelper& Helper);
 	virtual ~ULearningAgentsNeuralNetwork();
 
-	/** Resets this network asset to be empty. */
+	/**
+	 * Resets this network asset to be empty.
+	 */
 	UFUNCTION(CallInEditor, Category = "LearningAgents")
 	void ResetNetwork();
 
 	/**
 	 * Load this network from a snapshot.
-	 * 
 	 * @param File The snapshot file.
 	 */
 	UFUNCTION(BlueprintCallable, Category = "LearningAgents", meta = (RelativePath))
@@ -48,26 +49,23 @@ public:
 
 	/**
 	 * Save this network into a snapshot.
-	 * 
 	 * @param File The snapshot file.
 	 */
 	UFUNCTION(BlueprintCallable, Category = "LearningAgents", meta = (RelativePath))
 	void SaveNetworkToSnapshot(const FFilePath& File);
 
 	/**
-	 * Copy another asset into this network.
-	 * 
+	 * Copy another asset's network weights into this network. Network sizes must match.
 	 * @param NeuralNetworkAsset The asset to load from.
 	 */
 	UFUNCTION(BlueprintCallable, Category = "LearningAgents")
 	void LoadNetworkFromAsset(ULearningAgentsNeuralNetwork* NeuralNetworkAsset);
 
 	/**
-	 * Copy this network into another asset.
-	 * 
+	 * Copy this network's weights to another asset. Network sizes must match.
 	 * @param NeuralNetworkAsset The asset to save to.
 	 */
-	UFUNCTION(BlueprintCallable, Category = "LearningAgents")
+	UFUNCTION(BlueprintCallable, Category = "LearningAgents", Meta = (DevelopmentOnly))
 	void SaveNetworkToAsset(ULearningAgentsNeuralNetwork* NeuralNetworkAsset);
 
 public:
@@ -78,6 +76,6 @@ public:
 public:
 
 	/** The internal Neural Network Data */
-	UPROPERTY(BlueprintReadOnly, VisibleAnywhere, Category = "LearningAgents");
-	TObjectPtr<ULearningNeuralNetworkData> NeuralNetworkData;
+	UPROPERTY();
+	TObjectPtr<ULearningAgentsNeuralNetworkData> NeuralNetworkData;
 };
