[Backout] - CL30154408

diff --git a/Engine/Plugins/Experimental/LearningAgents/Source/LearningAgents/Public/LearningAgentsPolicy.h b/Engine/Plugins/Experimental/LearningAgents/Source/LearningAgents/Public/LearningAgentsPolicy.h
--- a/Engine/Plugins/Experimental/LearningAgents/Source/LearningAgents/Public/LearningAgentsPolicy.h
+++ b/Engine/Plugins/Experimental/LearningAgents/Source/LearningAgents/Public/LearningAgentsPolicy.h
@@ -2,20 +2,23 @@
 
 #pragma once
 
-#include "LearningAgentsManagerListener.h"
+#include "LearningAgentsManagerComponent.h"
 
 #include "LearningAgentsNeuralNetwork.h" // Included for ELearningAgentsActivationFunction
+#include "LearningAgentsDebug.h"
 #include "LearningArray.h"
+#include "LearningArrayMap.h"
 #include "UObject/ObjectPtr.h"
 
 #include "LearningAgentsPolicy.generated.h"
 
 namespace UE::Learning
 {
-	struct FNeuralNetworkPolicy;
-	struct FNeuralNetworkFunction;
+	struct INeuralNetwork;
+	struct FNeuralNetworkPolicyFunction;
 }
 
+struct FFilePath;
 class ULearningAgentsNeuralNetwork;
 
 /** The configurable settings for a ULearningAgentsPolicy. */
@@ -26,6 +29,22 @@ struct LEARNINGAGENTS_API FLearningAgentsPolicySettings
 
 public:
 
+	/** Seed for the action noise used by the policy  */
+	UPROPERTY(EditAnywhere, Category = "LearningAgents", meta = (ClampMin = "0", UIMin = "0"))
+	int32 ActionNoiseSeed = 1234;
+
+	/** Minimum action noise used by the policy */
+	UPROPERTY(EditAnywhere, Category = "LearningAgents", meta = (ClampMin = "0.0", UIMin = "0.0"))
+	float ActionNoiseMin = 0.25f;
+
+	/** Maximum action noise used by the policy */
+	UPROPERTY(EditAnywhere, Category = "LearningAgents", meta = (ClampMin = "0.0", UIMin = "0.0"))
+	float ActionNoiseMax = 0.25f;
+
+	/** Scale of the action noise used by the policy should be 1.0 during training. */
+	UPROPERTY(EditAnywhere, Category = "LearningAgents", meta = (ClampMin = "0.0", UIMin = "0.0"))
+	float ActionNoiseScale = 1.0f;
+
 	/** Total layers for policy network including input, hidden, and output layers */
 	UPROPERTY(EditAnywhere, Category = "LearningAgents", meta = (ClampMin = "2", UIMin = "2"))
 	int32 LayerNum = 3;
@@ -34,147 +53,112 @@ public:
 	UPROPERTY(EditAnywhere, Category = "LearningAgents", meta = (ClampMin = "1", UIMin = "1"))
 	int32 HiddenLayerSize = 128;
 
-	/** If this network should use memory or not. Setting this to false is equivalent to setting MemoryStateSize to 0. */
-	UPROPERTY(EditAnywhere, Category = "LearningAgents")
-	bool bUseMemory = true;
-
 	/** Number of neurons in the memory state of the policy network */
-	UPROPERTY(EditAnywhere, Category = "LearningAgents", meta = (ClampMin = "0", UIMin = "0"))
-	int32 MemoryStateSize = 64;
-
-	/** Initial scale for encoded actions. Setting this <1 reduces the initial bias from the network random weights and tends to help training. */
-	UPROPERTY(EditAnywhere, Category = "LearningAgents", meta = (ClampMin = "0", UIMin = "1"))
-	float InitialEncodedActionScale = 0.1f;
+	UPROPERTY(EditAnywhere, Category = "LearningAgents", meta = (ClampMin = "1", UIMin = "1"))
+	int32 MemoryStateSize = 128;
 
 	/** Activation function to use on hidden layers of the policy network */
 	UPROPERTY(EditAnywhere, Category = "LearningAgents")
 	ELearningAgentsActivationFunction ActivationFunction = ELearningAgentsActivationFunction::ELU;
 };
 
-/** A policy that maps from observations to actions. */
+/** A policy that maps from observations to actions for the managed agents. */
 UCLASS(BlueprintType, Blueprintable, meta = (BlueprintSpawnableComponent))
-class LEARNINGAGENTS_API ULearningAgentsPolicy : public ULearningAgentsManagerListener
+class LEARNINGAGENTS_API ULearningAgentsPolicy : public ULearningAgentsManagerComponent
 {
 	GENERATED_BODY()
 
 public:
 
-	friend class ULearningAgentsCritic;
-	friend class ULearningAgentsTrainer;
-
 	// These constructors/destructors are needed to make forward declarations happy
 	ULearningAgentsPolicy();
 	ULearningAgentsPolicy(FVTableHelper& Helper);
 	virtual ~ULearningAgentsPolicy();
 
-	/**
-	 * Constructs this object to be used with the given agent interactor and policy settings.
-	 * 
-	 * @param InManager						The input Manager
-	 * @param InInteractor					The input Interactor component
-	 * @param Class							The policy class
-	 * @param Name							The policy name
-	 * @param EncoderNeuralNetworkAsset		Optional Encoder Network Asset to use. If not provided, asset is empty, or
-	 *										bReinitializeEncoderNetwork is set then a new neural network object will be
-	 *										created.
-	 * @param PolicyNeuralNetworkAsset		Optional Policy Network Asset to use. If not provided, asset is empty, or
-	 *										bReinitializePolicyNetwork is set then a new neural network object will be 
-	 *										created according to the given PolicySettings.
-	 * @param DecoderNeuralNetworkAsset		Optional Decoder Network Asset to use. If not provided, asset is empty, or
-	 *										bReinitializeDecoderNetwork is set then a new neural network object will be
-	 *										created.
-	 * @param bReinitializeEncoderNetwork	If to reinitialize the encoder network
-	 * @param bReinitializePolicyNetwork	If to reinitialize the policy network
-	 * @param bReinitializeDecoderNetwork	If to reinitialize the decoder network
-	 * @param PolicySettings				The policy settings to use on creation of a new policy network
-	 * @param Seed							Random seed to use for initializing network weights and policy sampling
-	 */
-	UFUNCTION(BlueprintCallable, Category = "LearningAgents", meta=(Class = "/Script/LearningAgents.LearningAgentsPolicy", DeterminesOutputType = "Class"))
-	static ULearningAgentsPolicy* MakePolicy(
-		ULearningAgentsManager* InManager,
-		ULearningAgentsInteractor* InInteractor,
-		TSubclassOf<ULearningAgentsPolicy> Class,
-		const FName Name = TEXT("Policy"),
-		ULearningAgentsNeuralNetwork* EncoderNeuralNetworkAsset = nullptr,
-		ULearningAgentsNeuralNetwork* PolicyNeuralNetworkAsset = nullptr,
-		ULearningAgentsNeuralNetwork* DecoderNeuralNetworkAsset = nullptr,
-		const bool bReinitializeEncoderNetwork = true,
-		const bool bReinitializePolicyNetwork = true,
-		const bool bReinitializeDecoderNetwork = true,
-		const FLearningAgentsPolicySettings& PolicySettings = FLearningAgentsPolicySettings(),
-		const int32 Seed = 1234);
-
 	/**
 	 * Initializes this object to be used with the given agent interactor and policy settings.
-	 * 
-	 * @param InManager						The input Manager
-	 * @param InInteractor					The input Interactor component
-	 * @param EncoderNeuralNetworkAsset		Optional Encoder Network Asset to use. If not provided, asset is empty, or
-	 *										bReinitializeEncoderNetwork is set then a new neural network object will be
-	 *										created.
-	 * @param PolicyNeuralNetworkAsset		Optional Policy Network Asset to use. If not provided, asset is empty, or
-	 *										bReinitializePolicyNetwork is set then a new neural network object will be
-	 *										created according to the given PolicySettings.
-	 * @param DecoderNeuralNetworkAsset		Optional Decoder Network Asset to use. If not provided, asset is empty, or
-	 *										bReinitializeDecoderNetwork is set then a new neural network object will be
-	 *										created.
-	 * @param bReinitializeEncoderNetwork	If to reinitialize the encoder network
-	 * @param bReinitializePolicyNetwork	If to reinitialize the policy network
-	 * @param bReinitializeDecoderNetwork	If to reinitialize the decoder network
-	 * @param PolicySettings				The policy settings to use on creation of a new policy network
-	 * @param Seed							Random seed to use for initializing network weights and policy sampling
+	 * @param InInteractor The input Interactor component
+	 * @param PolicySettings The policy settings to use
+	 * @param NeuralNetworkAsset Optional Network Asset to use. If provided must match the given PolicySettings. If not
+	 * provided or asset is empty then a new neural network object will be created according to the given 
+	 * PolicySettings and used.
 	 */
 	UFUNCTION(BlueprintCallable, Category = "LearningAgents")
-	void SetupPolicy(
-		ULearningAgentsManager* InManager,
-		ULearningAgentsInteractor* InInteractor,
-		ULearningAgentsNeuralNetwork* EncoderNeuralNetworkAsset = nullptr,
-		ULearningAgentsNeuralNetwork* PolicyNeuralNetworkAsset = nullptr,
-		ULearningAgentsNeuralNetwork* DecoderNeuralNetworkAsset = nullptr,
-		const bool bReinitializeEncoderNetwork = true,
-		const bool bReinitializePolicyNetwork = true,
-		const bool bReinitializeDecoderNetwork = true,
+	void SetupPolicy(ULearningAgentsInteractor* InInteractor, 
 		const FLearningAgentsPolicySettings& PolicySettings = FLearningAgentsPolicySettings(),
-		const int32 Seed = 1234);
+		ULearningAgentsNeuralNetwork* NeuralNetworkAsset = nullptr);
 
 public:
 
-	//~ Begin ULearningAgentsManagerListener Interface
-	virtual void OnAgentsAdded_Implementation(const TArray<int32>& AgentIds) override;
-	virtual void OnAgentsRemoved_Implementation(const TArray<int32>& AgentIds) override;
-	virtual void OnAgentsReset_Implementation(const TArray<int32>& AgentIds) override;
-	//~ End ULearningAgentsManagerListener Interface
+	//~ Begin ULearningAgentsManagerComponent Interface
+	virtual void OnAgentsAdded(const TArray<int32>& AgentIds) override;
+	virtual void OnAgentsRemoved(const TArray<int32>& AgentIds) override;
+	virtual void OnAgentsReset(const TArray<int32>& AgentIds) override;
+	//~ End ULearningAgentsManagerComponent Interface
 
-
-// ----- Blueprint public interface -----
+// ----- Load / Save -----
 public:
 
 	/**
-	 * Encodes the buffered observation vectors using the Encoder network. This should be called after GatherObservations and before EvaluatePolicy.
+	 * Load a snapshot's weights into this policy.
+	 * @param File The snapshot file.
+	 */
+	UFUNCTION(BlueprintCallable, Category = "LearningAgents", meta = (RelativePath))
+	void LoadPolicyFromSnapshot(const FFilePath& File);
+
+	/**
+	 * Save this policy's weights into a snapshot.
+	 * @param File The snapshot file.
+	 */
+	UFUNCTION(BlueprintCallable, BlueprintPure = false, Category = "LearningAgents", meta = (RelativePath))
+	void SavePolicyToSnapshot(const FFilePath& File) const;
+
+	/**
+	 * Use a ULearningAgentsNeuralNetwork asset directly for this critic rather than making a copy. If used
+	 * during training then this asset's weights will be updated as training progresses.
+	 * @param NeuralNetworkAsset The asset to use.
 	 */
 	UFUNCTION(BlueprintCallable, Category = "LearningAgents")
-	void EncodeObservations();
+	void UsePolicyFromAsset(ULearningAgentsNeuralNetwork* NeuralNetworkAsset);
 
 	/**
-	 * Calling this function will run the underlying neural network on the previously encoded observations to populate
-	 * the encoded actions. This should be called after EncodeObservations and before DecodeAndSampleActions.
+	 * Load a ULearningAgentsNeuralNetwork asset's weights into this policy.
+	 * @param NeuralNetworkAsset The asset to load from.
 	 */
 	UFUNCTION(BlueprintCallable, Category = "LearningAgents")
-	void EvaluatePolicy();
+	void LoadPolicyFromAsset(ULearningAgentsNeuralNetwork* NeuralNetworkAsset);
+
+	/**
+	 * Save this policy's weights to a ULearningAgentsNeuralNetwork asset.
+	 * @param NeuralNetworkAsset The asset to save to.
+	 */
+	UFUNCTION(BlueprintCallable, Category = "LearningAgents", Meta = (DevelopmentOnly))
+	void SavePolicyToAsset(ULearningAgentsNeuralNetwork* NeuralNetworkAsset);
+
+// ----- Evaluation -----
+public:
 
 	/**
-	 * Decodes and samples action vectors using the Decoder network. This should be called after EvaluatePolicy and before Scatter Actions.
-	 * 
-	 * @param ActionNoiseScale		Scale of the action noise to use during sampling. Set this to zero to always sample the mean (expected) action.
+	 * Calling this function will run the underlying neural network on the previously buffered observations to populate
+	 * the output action buffer. This should be called after the associated agent interactor's EncodeObservations and
+	 * before its DecodeActions.
 	 */
 	UFUNCTION(BlueprintCallable, Category = "LearningAgents")
-	void DecodeAndSampleActions(const float ActionNoiseScale = 1.0f);
+	void EvaluatePolicy();
 
 	/**
-	 * Calls GatherObservations, EncodeObservations, EvaluatePolicy, DecodeAndSampleActions, ScatterActions
+	 * Calls EncodeObservations, followed by EvaluatePolicy, followed by DecodeActions
 	 */
 	UFUNCTION(BlueprintCallable, Category = "LearningAgents")
-	void RunInference(const float ActionNoiseScale = 1.0f);
+	void RunInference();
+
+	/** Get the action noise scale used. */
+	UFUNCTION(BlueprintPure, Category = "LearningAgents")
+	float GetActionNoiseScale() const;
+
+	/** Set the action noise scale used. */
+	UFUNCTION(BlueprintCallable, Category = "LearningAgents")
+	void SetActionNoiseScale(const float ActionNoiseScale);
 
 	/**
 	 * Gets the current memory state for a given agent as represented by an abstract vector learned by the policy.
@@ -194,34 +178,32 @@ public:
 	UFUNCTION(BlueprintCallable, Category = "LearningAgents", Meta = (AgentId = -1))
 	void SetMemoryState(const int32 AgentId, const TArray<float>& InMemoryState);
 
-	/** Gets the size of the memory state */
-	UFUNCTION(BlueprintPure, Category = "LearningAgents")
-	int32 GetMemoryStateSize() const;
+// ----- Non-blueprint public interface -----
+public:
 
-	/** Gets the current Encoder Network Asset being used */
-	UFUNCTION(BlueprintPure, Category = "LearningAgents")
-	ULearningAgentsNeuralNetwork* GetEncoderNetworkAsset();
+	/** Gets the current Network Asset being used */
+	ULearningAgentsNeuralNetwork* GetNetworkAsset();
 
-	/** Gets the current Policy Network Asset being used */
-	UFUNCTION(BlueprintPure, Category = "LearningAgents")
-	ULearningAgentsNeuralNetwork* GetPolicyNetworkAsset();
+	/** Get a reference to this policy's neural network. */
+	UE::Learning::INeuralNetwork& GetPolicyNetwork();
 
-	/** Gets the current Encoder Network Asset being used */
-	UFUNCTION(BlueprintPure, Category = "LearningAgents")
-	ULearningAgentsNeuralNetwork* GetDecoderNetworkAsset();
+	/** Get a reference to this policy's policy function object. */
+	UE::Learning::FNeuralNetworkPolicyFunction& GetPolicyObject();
 
+	/** Gets a view of the memory state for the given agent before evaluation */
+	TLearningArrayView<2, const float> GetPreEvaluationMemoryStateView() const;
 
-// ----- Non-blueprint public interface -----
-public:
+	/** Gets a view of the memory state for the given agent before evaluation */
+	TLearningArrayView<2, float> GetPreEvaluationMemoryStateView();
 
-	/** Get a reference to this policy's encoder function object. */
-	UE::Learning::FNeuralNetworkFunction& GetEncoderObject();
+	/** Gets a view of the current memory state for the given agent */
+	TLearningArrayView<2, const float> GetMemoryStateView() const;
 
-	/** Get a reference to this policy's policy function object. */
-	UE::Learning::FNeuralNetworkPolicy& GetPolicyObject();
+	/** Gets a view of the current memory state for the given agent */
+	TLearningArrayView<2, float> GetMemoryStateView();
 
-	/** Get a reference to this policy's decoder function object. */
-	UE::Learning::FNeuralNetworkFunction& GetDecoderObject();
+	/** Gets the size of the memory state */
+	int32 GetMemoryStateSize() const;
 
 // ----- Private Data -----
 private:
@@ -230,58 +212,35 @@ private:
 	UPROPERTY(VisibleAnywhere, Transient, Category = "LearningAgents")
 	TObjectPtr<ULearningAgentsInteractor> Interactor;
 
-	/** The underlying encoder neural network. */
-	UPROPERTY(VisibleAnywhere, Transient, Category = "LearningAgents")
-	TObjectPtr<ULearningAgentsNeuralNetwork> EncoderNetwork;
-
-	/** The underlying policy neural network. */
+	/** The underlying neural network. */
 	UPROPERTY(VisibleAnywhere, Transient, Category = "LearningAgents")
-	TObjectPtr<ULearningAgentsNeuralNetwork> PolicyNetwork;
+	TObjectPtr<ULearningAgentsNeuralNetwork> Network;
 
-	/** The underlying decoder neural network. */
-	UPROPERTY(VisibleAnywhere, Transient, Category = "LearningAgents")
-	TObjectPtr<ULearningAgentsNeuralNetwork> DecoderNetwork;
+	/** Internal Policy Function Object */
+	TSharedPtr<UE::Learning::FNeuralNetworkPolicyFunction> PolicyObject;
 
 	/** The internal memory state of each agent before evaluation. */
-	TLearningArray<2, float> PreEvaluationMemoryState;
+	UE::Learning::TArrayMapHandle<2, float> PreEvaluationMemoryStateHandle;
 
 	/** The internal memory state of each agent. */
-	TLearningArray<2, float> MemoryState;
-
-	/** The global random seed state used to initialize network weights and during sampling */
-	uint32 GlobalSeed = 0;
-
-	/** The individual per-agent seeds used during sampling of actions. */
-	TLearningArray<1, uint32> Seeds;
-
-	/** Encoded observation vector buffers */
-	TLearningArray<2, float> ObservationVectorsEncoded;
-
-	/** Encoded action vector buffers */
-	TLearningArray<2, float> ActionVectorsEncoded;
-
-	/** Action distribution vector buffers */
-	TLearningArray<2, float> ActionDistributionVectors;
-
-	/** Internal Encoder Function Object */
-	TSharedPtr<UE::Learning::FNeuralNetworkFunction> EncoderObject;
+	UE::Learning::TArrayMapHandle<2, float> MemoryStateHandle;
 
-	/** Internal Policy Function Object */
-	TSharedPtr<UE::Learning::FNeuralNetworkPolicy> PolicyObject;
-
-	/** Internal Decoder Function Object */
-	TSharedPtr<UE::Learning::FNeuralNetworkFunction> DecoderObject;
-
-	/** Number of times encoded observation vector has been set for all agents */
-	TLearningArray<1, uint64, TInlineAllocator<32>> ObservationVectorEncodedIteration;
-
-	/** Number of times encoded action vector has been set for all agents */
-	TLearningArray<1, uint64, TInlineAllocator<32>> ActionVectorEncodedIteration;
+// ----- Private Iteration Checks ----- 
+private:
 
-	/** Number of times the memory state has been updated for all agents. */
-	TLearningArray<1, uint64, TInlineAllocator<32>> MemoryStateIteration;
+	/** Number of times policy has been evaluated for all agents. */
+	TLearningArray<1, uint64, TInlineAllocator<32>> PolicyAgentIteration;
 
 	/** Temp buffers used to record the set of agents that are valid for evaluation */
 	TArray<int32> ValidAgentIds;
 	UE::Learning::FIndexSet ValidAgentSet;
+
+private:
+#if UE_LEARNING_AGENTS_ENABLE_VISUAL_LOG
+	/** Color used to draw this action in the visual log */
+	FLinearColor VisualLogColor = FColor::Purple;
+
+	/** Describes this policy to the visual logger for debugging purposes. */
+	void VisualLog(const UE::Learning::FIndexSet AgentSet) const;
+#endif
 };
diff --git a/Engine/Plugins/Experimental/LearningAgents/Source/LearningAgentsEditor/LearningAgentsEditor.Build.cs b/Engine/Plugins/Experimental/LearningAgents/Source/LearningAgentsEditor/LearningAgentsEditor.Build.cs
new file mode 100644
--- /dev/null
+++ b/Engine/Plugins/Experimental/LearningAgents/Source/LearningAgentsEditor/LearningAgentsEditor.Build.cs
@@ -0,0 +1,45 @@
+// Copyright Epic Games, Inc. All Rights Reserved.
+
+using UnrealBuildTool;
+
+public class LearningAgentsEditor : ModuleRules
+{
+	public LearningAgentsEditor(ReadOnlyTargetRules Target) : base(Target)
+	{
+		PCHUsage = ModuleRules.PCHUsageMode.UseExplicitOrSharedPCHs;
+
+		PublicIncludePaths.AddRange(
+			new string[] {
+				// ... add public include paths required here ...
+			});
+		
+		PrivateIncludePaths.AddRange(
+			new string[] {
+				// ... add other private include paths required here ...
+			});
+		
+		PublicDependencyModuleNames.AddRange(
+			new string[]
+			{
+				"Core",
+				// ... add other public dependencies that you statically link with here ...
+			});
+		
+		PrivateDependencyModuleNames.AddRange(
+			new string[]
+			{
+				"CoreUObject",
+				"Engine",
+				"UnrealEd",
+				"LearningAgents",
+				"LearningAgentsTraining",
+				// ... add private dependencies that you statically link with here ...	
+			});
+		
+		DynamicallyLoadedModuleNames.AddRange(
+			new string[]
+			{
+				// ... add any modules that your module loads dynamically here ...
+			});
+	}
+}
diff --git a/Engine/Plugins/Experimental/LearningAgents/Source/LearningAgentsEditor/Private/LearningAgentsEditor.cpp b/Engine/Plugins/Experimental/LearningAgents/Source/LearningAgentsEditor/Private/LearningAgentsEditor.cpp
new file mode 100644
--- /dev/null
+++ b/Engine/Plugins/Experimental/LearningAgents/Source/LearningAgentsEditor/Private/LearningAgentsEditor.cpp
@@ -0,0 +1,34 @@
+// Copyright Epic Games, Inc. All Rights Reserved.
+
+#include "LearningAgentsEditor.h"
+
+#include "Kismet2/KismetEditorUtilities.h"
+#include "LearningAgentsInteractor.h"
+#include "LearningAgentsTrainer.h"
+#include "LearningAgentsController.h"
+
+void FLearningAgentsEditorModule::StartupModule()
+{
+	// ULearningAgentsInteractor Default Events
+	FKismetEditorUtilities::RegisterAutoGeneratedDefaultEvent(this, ULearningAgentsInteractor::StaticClass(), GET_FUNCTION_NAME_CHECKED(ULearningAgentsInteractor, SetupObservations));
+	FKismetEditorUtilities::RegisterAutoGeneratedDefaultEvent(this, ULearningAgentsInteractor::StaticClass(), GET_FUNCTION_NAME_CHECKED(ULearningAgentsInteractor, SetObservations));
+	FKismetEditorUtilities::RegisterAutoGeneratedDefaultEvent(this, ULearningAgentsInteractor::StaticClass(), GET_FUNCTION_NAME_CHECKED(ULearningAgentsInteractor, SetupActions));
+	FKismetEditorUtilities::RegisterAutoGeneratedDefaultEvent(this, ULearningAgentsInteractor::StaticClass(), GET_FUNCTION_NAME_CHECKED(ULearningAgentsInteractor, GetActions));
+	
+	// ULearningAgentsTrainer Default Events
+	FKismetEditorUtilities::RegisterAutoGeneratedDefaultEvent(this, ULearningAgentsTrainer::StaticClass(), GET_FUNCTION_NAME_CHECKED(ULearningAgentsTrainer, SetupRewards));
+	FKismetEditorUtilities::RegisterAutoGeneratedDefaultEvent(this, ULearningAgentsTrainer::StaticClass(), GET_FUNCTION_NAME_CHECKED(ULearningAgentsTrainer, SetRewards));
+	FKismetEditorUtilities::RegisterAutoGeneratedDefaultEvent(this, ULearningAgentsTrainer::StaticClass(), GET_FUNCTION_NAME_CHECKED(ULearningAgentsTrainer, SetupCompletions));
+	FKismetEditorUtilities::RegisterAutoGeneratedDefaultEvent(this, ULearningAgentsTrainer::StaticClass(), GET_FUNCTION_NAME_CHECKED(ULearningAgentsTrainer, SetCompletions));
+	FKismetEditorUtilities::RegisterAutoGeneratedDefaultEvent(this, ULearningAgentsTrainer::StaticClass(), GET_FUNCTION_NAME_CHECKED(ULearningAgentsTrainer, ResetEpisodes));
+
+	// ULearningAgentsController Default Events
+	FKismetEditorUtilities::RegisterAutoGeneratedDefaultEvent(this, ULearningAgentsController::StaticClass(), GET_FUNCTION_NAME_CHECKED(ULearningAgentsController, SetActions));
+}
+
+void FLearningAgentsEditorModule::ShutdownModule()
+{
+	FKismetEditorUtilities::UnregisterAutoBlueprintNodeCreation(this);
+}
+	
+IMPLEMENT_MODULE(FLearningAgentsEditorModule, LearningAgentsEditor)
diff --git a/Engine/Plugins/Experimental/LearningAgents/Source/LearningAgentsEditor/Public/LearningAgentsEditor.h b/Engine/Plugins/Experimental/LearningAgents/Source/LearningAgentsEditor/Public/LearningAgentsEditor.h
new file mode 100644
--- /dev/null
+++ b/Engine/Plugins/Experimental/LearningAgents/Source/LearningAgentsEditor/Public/LearningAgentsEditor.h
@@ -0,0 +1,12 @@
+// Copyright Epic Games, Inc. All Rights Reserved.
+
+#pragma once
+
+#include "Modules/ModuleManager.h"
+
+class FLearningAgentsEditorModule : public IModuleInterface
+{
+public:
+	virtual void StartupModule() override;
+	virtual void ShutdownModule() override;
+};
diff --git a/Engine/Plugins/Experimental/LearningAgents/Source/LearningAgentsTraining/Private/LearningAgentsCompletions.cpp b/Engine/Plugins/Experimental/LearningAgents/Source/LearningAgentsTraining/Private/LearningAgentsCompletions.cpp
--- a/Engine/Plugins/Experimental/LearningAgents/Source/LearningAgentsTraining/Private/LearningAgentsCompletions.cpp
+++ b/Engine/Plugins/Experimental/LearningAgents/Source/LearningAgentsTraining/Private/LearningAgentsCompletions.cpp
@@ -2,123 +2,363 @@
 
 #include "LearningAgentsCompletions.h"
 
+#include "LearningAgentsTrainer.h"
+#include "LearningAgentsManager.h"
+#include "LearningArray.h"
+#include "LearningArrayMap.h"
 #include "LearningCompletion.h"
+#include "LearningCompletionObject.h"
 #include "LearningLog.h"
 
-namespace UE::Learning::Agents
+#include "GameFramework/Actor.h"
+
+namespace UE::Learning::Agents::Completions::Private
 {
-	ELearningAgentsCompletion GetLearningAgentsCompletion(const ECompletionMode CompletionMode)
+	template<typename CompletionUObject, typename CompletionFObject, typename... InArgTypes>
+	CompletionUObject* AddCompletion(ULearningAgentsTrainer* InAgentTrainer, const FName Name, const TCHAR* FunctionName, InArgTypes&& ...Args)
 	{
-		switch (CompletionMode)
+		if (!InAgentTrainer)
 		{
-		case ECompletionMode::Running: return ELearningAgentsCompletion::Running;
-		case ECompletionMode::Terminated: return ELearningAgentsCompletion::Termination;
-		case ECompletionMode::Truncated: return ELearningAgentsCompletion::Truncation;
-		default:UE_LOG(LogLearning, Error, TEXT("Unknown Completion Mode.")); return ELearningAgentsCompletion::Running;
+			UE_LOG(LogLearning, Error, TEXT("%s: InAgentTrainer is nullptr."), FunctionName);
+			return nullptr;
 		}
-	}
 
-	ECompletionMode GetCompletionMode(const ELearningAgentsCompletion Completion)
-	{
-		switch (Completion)
+		if (!InAgentTrainer->HasAgentManager())
 		{
-		case ELearningAgentsCompletion::Running: return ECompletionMode::Running;
-		case ELearningAgentsCompletion::Termination: return ECompletionMode::Terminated;
-		case ELearningAgentsCompletion::Truncation: return ECompletionMode::Truncated;
-		default:UE_LOG(LogLearning, Error, TEXT("Unknown Completion.")); return ECompletionMode::Running;
+			UE_LOG(LogLearning, Error, TEXT("%s: Must be attached to a LearningAgentsManager Actor."), *InAgentTrainer->GetName());
+			return nullptr;
 		}
+
+		const FName UniqueName = MakeUniqueObjectName(InAgentTrainer, CompletionUObject::StaticClass(), Name, EUniqueObjectNameOptions::GloballyUnique);
+
+		CompletionUObject* Completion = NewObject<CompletionUObject>(InAgentTrainer, UniqueName);
+		Completion->Init(InAgentTrainer->GetAgentManager()->GetMaxAgentNum());
+		Completion->AgentTrainer = InAgentTrainer;
+		Completion->CompletionObject = MakeShared<CompletionFObject>(
+			Completion->GetFName(),
+			InAgentTrainer->GetAgentManager()->GetInstanceData().ToSharedRef(),
+			InAgentTrainer->GetAgentManager()->GetMaxAgentNum(),
+			Forward<InArgTypes>(Args)...);
+
+		InAgentTrainer->AddCompletion(Completion, Completion->CompletionObject.ToSharedRef());
+
+		return Completion;
 	}
 }
 
-bool ULearningAgentsCompletions::CompletionIsRunning(const ELearningAgentsCompletion Completion)
+//------------------------------------------------------------------
+
+void ULearningAgentsCompletion::Init(const int32 MaxAgentNum)
 {
-	return Completion == ELearningAgentsCompletion::Running;
+	AgentIteration.SetNumUninitialized({ MaxAgentNum });
+	UE::Learning::Array::Set<1, uint64>(AgentIteration, INDEX_NONE);
 }
 
-bool ULearningAgentsCompletions::CompletionIsCompleted(const ELearningAgentsCompletion Completion)
+void ULearningAgentsCompletion::OnAgentsAdded(const TArray<int32>& AgentIds)
 {
-	return Completion == ELearningAgentsCompletion::Truncation || Completion == ELearningAgentsCompletion::Termination;
+	UE::Learning::Array::Set<1, uint64>(AgentIteration, 0, AgentIds);
 }
 
-bool ULearningAgentsCompletions::CompletionIsTruncation(const ELearningAgentsCompletion Completion)
+void ULearningAgentsCompletion::OnAgentsRemoved(const TArray<int32>& AgentIds)
 {
-	return Completion == ELearningAgentsCompletion::Truncation;
+	UE::Learning::Array::Set<1, uint64>(AgentIteration, INDEX_NONE, AgentIds);
 }
 
-bool ULearningAgentsCompletions::CompletionIsTermination(const ELearningAgentsCompletion Completion)
+void ULearningAgentsCompletion::OnAgentsReset(const TArray<int32>& AgentIds)
 {
-	return Completion == ELearningAgentsCompletion::Termination;
+	UE::Learning::Array::Set<1, uint64>(AgentIteration, 0, AgentIds);
 }
 
-ELearningAgentsCompletion ULearningAgentsCompletions::CompletionOr(ELearningAgentsCompletion A, ELearningAgentsCompletion B)
+uint64 ULearningAgentsCompletion::GetAgentIteration(const int32 AgentId) const
 {
-	if ((A == ELearningAgentsCompletion::Running && B != ELearningAgentsCompletion::Running) ||
-		(A == ELearningAgentsCompletion::Truncation && B == ELearningAgentsCompletion::Termination))
-	{
-		return B;
-	}
-	else
+	return AgentIteration[AgentId];
+}
+
+//------------------------------------------------------------------
+
+UConditionalCompletion* UConditionalCompletion::AddConditionalCompletion(ULearningAgentsTrainer* InAgentTrainer, const FName Name, const ELearningAgentsCompletion InCompletionMode)
+{
+	return UE::Learning::Agents::Completions::Private::AddCompletion<UConditionalCompletion, UE::Learning::FConditionalCompletion>(InAgentTrainer, Name, TEXT("AddConditionalCompletion"), UE::Learning::Agents::GetCompletionMode(InCompletionMode));
+}
+
+void UConditionalCompletion::SetConditionalCompletion(const int32 AgentId, const bool bIsComplete)
+{
+	if (!AgentTrainer->HasAgent(AgentId))
 	{
-		return A;
+		UE_LOG(LogLearning, Error, TEXT("%s: AgentId %d not found in the agents set."), *GetName(), AgentId);
+		return;
 	}
+
+	CompletionObject->InstanceData->View(CompletionObject->ConditionHandle)[AgentId] = bIsComplete;
+	AgentIteration[AgentId]++;
 }
 
-ELearningAgentsCompletion ULearningAgentsCompletions::CompletionAnd(ELearningAgentsCompletion A, ELearningAgentsCompletion B)
+#if UE_LEARNING_AGENTS_ENABLE_VISUAL_LOG
+void UConditionalCompletion::VisualLog(const UE::Learning::FIndexSet Instances) const
 {
-	if (A == ELearningAgentsCompletion::Running ||
-		(A == ELearningAgentsCompletion::Truncation && B != ELearningAgentsCompletion::Running))
+	UE_LEARNING_TRACE_CPUPROFILER_EVENT_SCOPE(UConditionalCompletion::VisualLog);
+
+	const TLearningArrayView<1, const bool> ConditionView = CompletionObject->InstanceData->ConstView(CompletionObject->ConditionHandle);
+	const TLearningArrayView<1, const UE::Learning::ECompletionMode> CompletionView = CompletionObject->InstanceData->ConstView(CompletionObject->CompletionHandle);
+
+	for (const int32 Instance : Instances)
 	{
-		return A;
+		if (const AActor* Actor = Cast<AActor>(AgentTrainer->GetAgent(Instance)))
+		{
+			UE_LEARNING_AGENTS_VLOG_STRING(this, LogLearning, Display,
+				Actor->GetActorLocation(),
+				VisualLogColor.ToFColor(true),
+				TEXT("Agent %i\nCondition: %s\nCompletion: %s"),
+				Instance,
+				ConditionView[Instance] ? TEXT("true") : TEXT("false"),
+				UE::Learning::Completion::CompletionModeString(CompletionView[Instance]));
+		}
 	}
-	else
+}
+#endif
+
+//------------------------------------------------------------------
+
+UTimeElapsedCompletion* UTimeElapsedCompletion::AddTimeElapsedCompletion(ULearningAgentsTrainer* InAgentTrainer, const FName Name, const float Threshold, const ELearningAgentsCompletion InCompletionMode)
+{
+	return UE::Learning::Agents::Completions::Private::AddCompletion<UTimeElapsedCompletion, UE::Learning::FTimeElapsedCompletion>(InAgentTrainer, Name, TEXT("AddTimeElapsedCompletion"), Threshold, UE::Learning::Agents::GetCompletionMode(InCompletionMode));
+}
+
+void UTimeElapsedCompletion::SetTimeElapsedCompletion(const int32 AgentId, const float Time)
+{
+	if (!AgentTrainer->HasAgent(AgentId))
 	{
-		return B;
+		UE_LOG(LogLearning, Error, TEXT("%s: AgentId %d not found in the agents set."), *GetName(), AgentId);
+		return;
 	}
+
+	CompletionObject->InstanceData->View(CompletionObject->TimeHandle)[AgentId] = Time;
+	AgentIteration[AgentId]++;
 }
 
-ELearningAgentsCompletion ULearningAgentsCompletions::CompletionNot(const ELearningAgentsCompletion A, const ELearningAgentsCompletion NotRunningType)
+#if UE_LEARNING_AGENTS_ENABLE_VISUAL_LOG
+void UTimeElapsedCompletion::VisualLog(const UE::Learning::FIndexSet Instances) const
 {
-	return A == ELearningAgentsCompletion::Running ? NotRunningType : ELearningAgentsCompletion::Running;
+	UE_LEARNING_TRACE_CPUPROFILER_EVENT_SCOPE(UTimeElapsedCompletion::VisualLog);
+
+	const TLearningArrayView<1, const float> TimeView = CompletionObject->InstanceData->ConstView(CompletionObject->TimeHandle);
+	const TLearningArrayView<1, const float> ThresholdView = CompletionObject->InstanceData->ConstView(CompletionObject->ThresholdHandle);
+	const TLearningArrayView<1, const UE::Learning::ECompletionMode> CompletionView = CompletionObject->InstanceData->ConstView(CompletionObject->CompletionHandle);
+
+	for (const int32 Instance : Instances)
+	{
+		if (const AActor* Actor = Cast<AActor>(AgentTrainer->GetAgent(Instance)))
+		{
+			UE_LEARNING_AGENTS_VLOG_STRING(this, LogLearning, Display,
+				Actor->GetActorLocation(),
+				VisualLogColor.ToFColor(true),
+				TEXT("Agent %i\nThreshold: [%6.3f]\nTime: [%6.3f]\nCompletion: %s"),
+				Instance,
+				ThresholdView[Instance],
+				TimeView[Instance],
+				UE::Learning::Completion::CompletionModeString(CompletionView[Instance]));
+		}
+	}
 }
+#endif
+
+//------------------------------------------------------------------
 
-ELearningAgentsCompletion ULearningAgentsCompletions::CompletionOnCondition(const bool bCondition, const ELearningAgentsCompletion CompletionType)
+UPlanarPositionDifferenceCompletion* UPlanarPositionDifferenceCompletion::AddPlanarPositionDifferenceCompletion(
+	ULearningAgentsTrainer* InAgentTrainer,
+	const FName Name,
+	const float Threshold,
+	const ELearningAgentsCompletion InCompletionMode,
+	const FVector Axis0,
+	const FVector Axis1)
 {
-	return bCondition ? CompletionType : ELearningAgentsCompletion::Running;
+	return UE::Learning::Agents::Completions::Private::AddCompletion<UPlanarPositionDifferenceCompletion, UE::Learning::FPlanarPositionDifferenceCompletion>(
+		InAgentTrainer,
+		Name,
+		TEXT("AddPlanarPositionDifferenceCompletion"),
+		1,
+		Threshold,
+		UE::Learning::Agents::GetCompletionMode(InCompletionMode),
+		Axis0.GetSafeNormal(UE_SMALL_NUMBER, FVector::ForwardVector),
+		Axis1.GetSafeNormal(UE_SMALL_NUMBER, FVector::RightVector));
 }
 
-ELearningAgentsCompletion ULearningAgentsCompletions::CompletionOnTimeElapsed(const float Time, const float TimeThreshold, const ELearningAgentsCompletion CompletionType)
+void UPlanarPositionDifferenceCompletion::SetPlanarPositionDifferenceCompletion(const int32 AgentId, const FVector Position0, const FVector Position1)
 {
-	return CompletionOnCondition(Time > TimeThreshold, CompletionType);
+	if (!AgentTrainer->HasAgent(AgentId))
+	{
+		UE_LOG(LogLearning, Error, TEXT("%s: AgentId %d not found in the agents set."), *GetName(), AgentId);
+		return;
+	}
+
+	CompletionObject->InstanceData->View(CompletionObject->Position0Handle)[AgentId][0] = Position0;
+	CompletionObject->InstanceData->View(CompletionObject->Position1Handle)[AgentId][0] = Position1;
+	AgentIteration[AgentId]++;
 }
 
-ELearningAgentsCompletion ULearningAgentsCompletions::CompletionOnEpisodeStepsRecorded(const int32 EpisodeSteps, const int32 MaxEpisodeSteps, const ELearningAgentsCompletion CompletionType)
+#if UE_LEARNING_AGENTS_ENABLE_VISUAL_LOG
+void UPlanarPositionDifferenceCompletion::VisualLog(const UE::Learning::FIndexSet Instances) const
 {
-	return CompletionOnCondition(EpisodeSteps >= MaxEpisodeSteps, CompletionType);
+	UE_LEARNING_TRACE_CPUPROFILER_EVENT_SCOPE(UPlanarPositionDifferenceCompletion::VisualLog);
+
+	const TLearningArrayView<2, const FVector> Position0View = CompletionObject->InstanceData->ConstView(CompletionObject->Position0Handle);
+	const TLearningArrayView<2, const FVector> Position1View = CompletionObject->InstanceData->ConstView(CompletionObject->Position1Handle);
+	const TLearningArrayView<1, const float> ThresholdView = CompletionObject->InstanceData->ConstView(CompletionObject->ThresholdHandle);
+	const TLearningArrayView<1, const UE::Learning::ECompletionMode> CompletionView = CompletionObject->InstanceData->ConstView(CompletionObject->CompletionHandle);
+
+	for (const int32 Instance : Instances)
+	{
+		if (const AActor* Actor = Cast<AActor>(AgentTrainer->GetAgent(Instance)))
+		{
+			const FVector Position0 = Position0View[Instance][0];
+			const FVector Position1 = Position1View[Instance][0];
+
+			const FVector PlanarPosition0 = FVector(CompletionObject->Axis0.Dot(Position0), CompletionObject->Axis1.Dot(Position0), 0.0f);
+			const FVector PlanarPosition1 = FVector(CompletionObject->Axis0.Dot(Position1), CompletionObject->Axis1.Dot(Position1), 0.0f);
+
+			UE_LEARNING_AGENTS_VLOG_LOCATION(this, LogLearning, Display,
+				Position0,
+				10,
+				VisualLogColor.ToFColor(true),
+				TEXT("Position0: [% 6.1f % 6.1f % 6.1f]\nPlanar Position0: [% 6.1f % 6.1f]"),
+				Position0.X, Position0.Y, Position0.Z,
+				PlanarPosition0.X, PlanarPosition0.Y);
+
+			UE_LEARNING_AGENTS_VLOG_PLANE(this, LogLearning, Display,
+				Position0,
+				FQuat::Identity,
+				CompletionObject->Axis0,
+				CompletionObject->Axis1,
+				VisualLogColor.ToFColor(true),
+				TEXT(""));
+
+			UE_LEARNING_AGENTS_VLOG_LOCATION(this, LogLearning, Display,
+				Position1,
+				10,
+				VisualLogColor.ToFColor(true),
+				TEXT("Position1: [% 6.1f % 6.1f % 6.1f]\nPlanar Position1: [% 6.1f % 6.1f]"),
+				Position1.X, Position1.Y, Position1.Z,
+				PlanarPosition1.X, PlanarPosition1.Y);
+
+			UE_LEARNING_AGENTS_VLOG_PLANE(this, LogLearning, Display,
+				Position1,
+				FQuat::Identity,
+				CompletionObject->Axis0,
+				CompletionObject->Axis1,
+				VisualLogColor.ToFColor(true),
+				TEXT(""));
+
+			UE_LEARNING_AGENTS_VLOG_SEGMENT(this, LogLearning, Display,
+				Position0,
+				Position1,
+				VisualLogColor.ToFColor(true),
+				TEXT("Agent %i\nDistance: [% 6.3f]\nPlanar Distance: [% 6.3f]\nThreshold: [% 6.2f]\nCompletion: %s"),
+				Instance,
+				FVector::Distance(Position0, Position1),
+				FVector::Distance(PlanarPosition0, PlanarPosition1),
+				ThresholdView[Instance],
+				UE::Learning::Completion::CompletionModeString(CompletionView[Instance]));
+		}
+	}
 }
+#endif
+
 
-ELearningAgentsCompletion ULearningAgentsCompletions::CompletionOnTranslationDifferenceBelowThreshold(const FVector TranslationA, const FVector TranslationB, const float DistanceThreshold, const ELearningAgentsCompletion CompletionType)
+UPlanarPositionSimilarityCompletion* UPlanarPositionSimilarityCompletion::AddPlanarPositionSimilarityCompletion(
+	ULearningAgentsTrainer* InAgentTrainer,
+	const FName Name,
+	const float Threshold,
+	const ELearningAgentsCompletion InCompletionMode,
+	const FVector Axis0,
+	const FVector Axis1)
 {
-	return CompletionOnCondition(FVector::Distance(TranslationA, TranslationB) < DistanceThreshold, CompletionType);
+	return UE::Learning::Agents::Completions::Private::AddCompletion<UPlanarPositionSimilarityCompletion, UE::Learning::FPlanarPositionSimilarityCompletion>(
+		InAgentTrainer,
+		Name,
+		TEXT("AddPlanarPositionSimilarityCompletion"),
+		1,
+		Threshold,
+		UE::Learning::Agents::GetCompletionMode(InCompletionMode),
+		Axis0.GetSafeNormal(UE_SMALL_NUMBER, FVector::ForwardVector),
+		Axis1.GetSafeNormal(UE_SMALL_NUMBER, FVector::RightVector));
 }
 
-ELearningAgentsCompletion ULearningAgentsCompletions::CompletionOnTranslationDifferenceAboveThreshold(const FVector TranslationA, const FVector TranslationB, const float DistanceThreshold, const ELearningAgentsCompletion CompletionType)
+void UPlanarPositionSimilarityCompletion::SetPlanarPositionSimilarityCompletion(const int32 AgentId, const FVector Position0, const FVector Position1)
 {
-	return CompletionOnCondition(FVector::Distance(TranslationA, TranslationB) > DistanceThreshold, CompletionType);
+	if (!AgentTrainer->HasAgent(AgentId))
+	{
+		UE_LOG(LogLearning, Error, TEXT("%s: AgentId %d not found in the agents set."), *GetName(), AgentId);
+		return;
+	}
+
+	CompletionObject->InstanceData->View(CompletionObject->Position0Handle)[AgentId][0] = Position0;
+	CompletionObject->InstanceData->View(CompletionObject->Position1Handle)[AgentId][0] = Position1;
+	AgentIteration[AgentId]++;
 }
 
-ELearningAgentsCompletion ULearningAgentsCompletions::CompletionOnTranslationOutsideBounds(
-	const FVector Translation,
-	const FTransform BoundsTransform,
-	const FVector BoundsMins,
-	const FVector BoundsMaxs,
-	const ELearningAgentsCompletion CompletionType)
+#if UE_LEARNING_AGENTS_ENABLE_VISUAL_LOG
+void UPlanarPositionSimilarityCompletion::VisualLog(const UE::Learning::FIndexSet Instances) const
 {
-	const FVector LocalTranslation = BoundsTransform.InverseTransformPosition(Translation);
+	UE_LEARNING_TRACE_CPUPROFILER_EVENT_SCOPE(UPlanarPositionSimilarityCompletion::VisualLog);
 
-	return CompletionOnCondition(
-		LocalTranslation.X < BoundsMins.X || LocalTranslation.X > BoundsMaxs.X ||
-		LocalTranslation.Y < BoundsMins.Y || LocalTranslation.Y > BoundsMaxs.Y ||
-		LocalTranslation.Z < BoundsMins.Z || LocalTranslation.Z > BoundsMaxs.Z,
-		CompletionType);
-}
+	const TLearningArrayView<2, const FVector> Position0View = CompletionObject->InstanceData->ConstView(CompletionObject->Position0Handle);
+	const TLearningArrayView<2, const FVector> Position1View = CompletionObject->InstanceData->ConstView(CompletionObject->Position1Handle);
+	const TLearningArrayView<1, const float> ThresholdView = CompletionObject->InstanceData->ConstView(CompletionObject->ThresholdHandle);
+	const TLearningArrayView<1, const UE::Learning::ECompletionMode> CompletionView = CompletionObject->InstanceData->ConstView(CompletionObject->CompletionHandle);
+
+	for (const int32 Instance : Instances)
+	{
+		if (const AActor* Actor = Cast<AActor>(AgentTrainer->GetAgent(Instance)))
+		{
+			const FVector Position0 = Position0View[Instance][0];
+			const FVector Position1 = Position1View[Instance][0];
 
+			const FVector PlanarPosition0 = FVector(CompletionObject->Axis0.Dot(Position0), CompletionObject->Axis1.Dot(Position0), 0.0f);
+			const FVector PlanarPosition1 = FVector(CompletionObject->Axis0.Dot(Position1), CompletionObject->Axis1.Dot(Position1), 0.0f);
+
+			UE_LEARNING_AGENTS_VLOG_LOCATION(this, LogLearning, Display,
+				Position0,
+				10,
+				VisualLogColor.ToFColor(true),
+				TEXT("Position0: [% 6.1f % 6.1f % 6.1f]\nPlanar Position0: [% 6.1f % 6.1f]"),
+				Position0.X, Position0.Y, Position0.Z,
+				PlanarPosition0.X, PlanarPosition0.Y);
+
+			UE_LEARNING_AGENTS_VLOG_PLANE(this, LogLearning, Display,
+				Position0,
+				FQuat::Identity,
+				CompletionObject->Axis0,
+				CompletionObject->Axis1,
+				VisualLogColor.ToFColor(true),
+				TEXT(""));
+
+			UE_LEARNING_AGENTS_VLOG_LOCATION(this, LogLearning, Display,
+				Position1,
+				10,
+				VisualLogColor.ToFColor(true),
+				TEXT("Position1: [% 6.1f % 6.1f % 6.1f]\nPlanar Position1: [% 6.1f % 6.1f]"),
+				Position1.X, Position1.Y, Position1.Z,
+				PlanarPosition1.X, PlanarPosition1.Y);
+
+			UE_LEARNING_AGENTS_VLOG_PLANE(this, LogLearning, Display,
+				Position1,
+				FQuat::Identity,
+				CompletionObject->Axis0,
+				CompletionObject->Axis1,
+				VisualLogColor.ToFColor(true),
+				TEXT(""));
+
+			UE_LEARNING_AGENTS_VLOG_SEGMENT(this, LogLearning, Display,
+				Position0,
+				Position1,
+				VisualLogColor.ToFColor(true),
+				TEXT("Agent %i\nDistance: [% 6.3f]\nPlanar Distance: [% 6.3f]\nThreshold: [% 6.2f]\nCompletion: %s"),
+				Instance,
+				FVector::Distance(Position0, Position1),
+				FVector::Distance(PlanarPosition0, PlanarPosition1),
+				ThresholdView[Instance],
+				UE::Learning::Completion::CompletionModeString(CompletionView[Instance]));
+		}
+	}
+}
+#endif
diff --git a/Engine/Plugins/Experimental/LearningAgents/Source/LearningAgentsTraining/Private/LearningAgentsImitationTrainer.cpp b/Engine/Plugins/Experimental/LearningAgents/Source/LearningAgentsTraining/Private/LearningAgentsImitationTrainer.cpp
--- a/Engine/Plugins/Experimental/LearningAgents/Source/LearningAgentsTraining/Private/LearningAgentsImitationTrainer.cpp
+++ b/Engine/Plugins/Experimental/LearningAgents/Source/LearningAgentsTraining/Private/LearningAgentsImitationTrainer.cpp
@@ -2,13 +2,15 @@
 
 #include "LearningAgentsImitationTrainer.h"
 
-#include "LearningAgentsManager.h"
 #include "LearningAgentsInteractor.h"
+#include "LearningArrayMap.h"
 #include "LearningExperience.h"
+#include "LearningFeatureObject.h"
 #include "LearningLog.h"
 #include "LearningImitationTrainer.h"
 #include "LearningAgentsRecording.h"
 #include "LearningAgentsPolicy.h"
+#include "LearningNeuralNetworkObject.h"
 
 #include "GenericPlatform/GenericPlatformMisc.h"
 #include "HAL/FileManager.h"
@@ -18,73 +20,27 @@ ULearningAgentsImitationTrainer::ULearningAgentsImitationTrainer() : Super(FObje
 ULearningAgentsImitationTrainer::ULearningAgentsImitationTrainer(FVTableHelper& Helper) : Super(Helper) {}
 ULearningAgentsImitationTrainer::~ULearningAgentsImitationTrainer() = default;
 
-void ULearningAgentsImitationTrainer::BeginDestroy()
+void ULearningAgentsImitationTrainer::EndPlay(const EEndPlayReason::Type EndPlayReason)
 {
 	if (IsTraining())
 	{
 		EndTraining();
 	}
 
-	Super::BeginDestroy();
+	Super::EndPlay(EndPlayReason);
 }
 
-ULearningAgentsImitationTrainer* ULearningAgentsImitationTrainer::MakeImitationTrainer(
-	ULearningAgentsManager* InManager,
-	ULearningAgentsInteractor* InInteractor,
-	ULearningAgentsPolicy* InPolicy,
-	TSubclassOf<ULearningAgentsImitationTrainer> Class)
-{
-	if (!InManager)
-	{
-		UE_LOG(LogLearning, Error, TEXT("MakeImitationTrainer: InManager is nullptr."));
-		return nullptr;
-	}
-
-	if (!Class)
-	{
-		UE_LOG(LogLearning, Error, TEXT("MakeImitationTrainer: Class is nullptr."));
-		return nullptr;
-	}
-
-	const FName UniqueName = MakeUniqueObjectName(InManager, Class, TEXT("ImitationTrainer"), EUniqueObjectNameOptions::GloballyUnique);
-
-	ULearningAgentsImitationTrainer* ImitationTrainer = NewObject<ULearningAgentsImitationTrainer>(InManager, Class, UniqueName);
-	if (!ImitationTrainer) { return nullptr; }
-
-	ImitationTrainer->SetupImitationTrainer(
-		InManager,
-		InInteractor,
-		InPolicy);
-
-	return ImitationTrainer->IsSetup() ? ImitationTrainer : nullptr;
-}
-
-void ULearningAgentsImitationTrainer::SetupImitationTrainer(
-	ULearningAgentsManager* InManager,
-	ULearningAgentsInteractor* InInteractor,
-	ULearningAgentsPolicy* InPolicy)
+void ULearningAgentsImitationTrainer::BeginTraining(
+	ULearningAgentsPolicy* InPolicy, 
+	const ULearningAgentsRecording* Recording,
+	const FLearningAgentsImitationTrainerSettings& ImitationTrainerSettings,
+	const FLearningAgentsImitationTrainerTrainingSettings& ImitationTrainerTrainingSettings,
+	const FLearningAgentsTrainerPathSettings& ImitationTrainerPathSettings,
+	const bool bReinitializePolicyNetwork)
 {
-	if (IsSetup())
-	{
-		UE_LOG(LogLearning, Error, TEXT("%s: Setup already performed!"), *GetName());
-		return;
-	}
-
-	if (!InManager)
-	{
-		UE_LOG(LogLearning, Error, TEXT("%s: InManager is nullptr."), *GetName());
-		return;
-	}
-
-	if (!InInteractor)
-	{
-		UE_LOG(LogLearning, Error, TEXT("%s: InInteractor is nullptr."), *GetName());
-		return;
-	}
-
-	if (!InInteractor->IsSetup())
+	if (IsTraining())
 	{
-		UE_LOG(LogLearning, Error, TEXT("%s: %s's Setup must be run before it can be used."), *GetName(), *InInteractor->GetName());
+		UE_LOG(LogLearning, Error, TEXT("%s: Cannot begin training as we are already training!"), *GetName());
 		return;
 	}
 
@@ -100,33 +56,6 @@ void ULearningAgentsImitationTrainer::SetupImitationTrainer(
 		return;
 	}
 
-	Manager = InManager;
-	Interactor = InInteractor;
-	Policy = InPolicy;
-
-	bIsSetup = true;
-
-	Manager->AddListener(this);
-}
-
-void ULearningAgentsImitationTrainer::BeginTraining(
-	const ULearningAgentsRecording* Recording,
-	const FLearningAgentsImitationTrainerSettings& ImitationTrainerSettings,
-	const FLearningAgentsImitationTrainerTrainingSettings& ImitationTrainerTrainingSettings,
-	const FLearningAgentsTrainerPathSettings& ImitationTrainerPathSettings)
-{
-	if (!IsSetup())
-	{
-		UE_LOG(LogLearning, Error, TEXT("%s: Setup not complete."), *GetName());
-		return;
-	}
-
-	if (IsTraining())
-	{
-		UE_LOG(LogLearning, Error, TEXT("%s: Cannot begin training as we are already training!"), *GetName());
-		return;
-	}
-
 	if (!Recording)
 	{
 		UE_LOG(LogLearning, Error, TEXT("%s: Recording is nullptr."), *GetName());
@@ -139,6 +68,8 @@ void ULearningAgentsImitationTrainer::BeginTraining(
 		return;
 	}
 
+	Policy = InPolicy;
+
 	// Record Timeout Setting
 
 	TrainerTimeout = ImitationTrainerSettings.TrainerCommunicationTimeout;
@@ -172,9 +103,9 @@ void ULearningAgentsImitationTrainer::BeginTraining(
 
 	// Sizes
 
-	const int32 ObservationNum = Interactor->GetObservationVectorSize();
-	const int32 ActionNum = Interactor->GetActionVectorSize();
-	const int32 MemoryStateNum = Policy->GetMemoryStateSize();
+	const int32 ObservationNum = Policy->GetPolicyObject().ObservationNum;
+	const int32 ActionNum = Policy->GetPolicyObject().ActionNum;
+	const int32 MemoryStateNum = Policy->GetPolicyObject().MemoryStateNum;
 
 	// Get Number of Steps
 
@@ -206,11 +137,6 @@ void ULearningAgentsImitationTrainer::BeginTraining(
 
 	// Copy into Flat Arrays
 
-	TLearningArray<1, int32> RecordedEpisodeStarts;
-	TLearningArray<1, int32> RecordedEpisodeLengths;
-	TLearningArray<2, float> RecordedObservations;
-	TLearningArray<2, float> RecordedActions;
-
 	RecordedEpisodeStarts.SetNumUninitialized({ TotalEpisodeNum });
 	RecordedEpisodeLengths.SetNumUninitialized({ TotalEpisodeNum });
 	RecordedObservations.SetNumUninitialized({ TotalStepNum, ObservationNum });
@@ -223,13 +149,10 @@ void ULearningAgentsImitationTrainer::BeginTraining(
 		if (Record.ObservationDimNum != ObservationNum) { continue; }
 		if (Record.ActionDimNum != ActionNum) { continue; }
 
-		TLearningArrayView<2, const float> ObservationsView = TLearningArrayView<2, const float>(Record.ObservationData.GetData(), { Record.StepNum, Record.ObservationDimNum });
-		TLearningArrayView<2, const float> ActionsView = TLearningArrayView<2, const float>(Record.ActionData.GetData(), { Record.StepNum, Record.ActionDimNum });
-
 		RecordedEpisodeStarts[EpisodeIdx] = StepIdx;
 		RecordedEpisodeLengths[EpisodeIdx] = Record.StepNum;
-		UE::Learning::Array::Copy(RecordedObservations.Slice(StepIdx, Record.StepNum), ObservationsView);
-		UE::Learning::Array::Copy(RecordedActions.Slice(StepIdx, Record.StepNum), ActionsView);
+		UE::Learning::Array::Copy(RecordedObservations.Slice(StepIdx, Record.StepNum), Record.Observations);
+		UE::Learning::Array::Copy(RecordedActions.Slice(StepIdx, Record.StepNum), Record.Actions);
 		EpisodeIdx++;
 		StepIdx += Record.StepNum;
 	}
@@ -244,17 +167,18 @@ void ULearningAgentsImitationTrainer::BeginTraining(
 
 	UE::Learning::FImitationTrainerTrainingSettings ImitationTrainingSettings;
 	ImitationTrainingSettings.IterationNum = ImitationTrainerTrainingSettings.NumberOfIterations;
-	ImitationTrainingSettings.LearningRate = ImitationTrainerTrainingSettings.LearningRate;
+	ImitationTrainingSettings.LearningRatePolicy = ImitationTrainerTrainingSettings.LearningRate;
 	ImitationTrainingSettings.LearningRateDecay = ImitationTrainerTrainingSettings.LearningRateDecay;
 	ImitationTrainingSettings.WeightDecay = ImitationTrainerTrainingSettings.WeightDecay;
-	ImitationTrainingSettings.BatchSize = ImitationTrainerTrainingSettings.BatchSize;
-	ImitationTrainingSettings.Window = ImitationTrainerTrainingSettings.Window;
-	ImitationTrainingSettings.ActionRegularizationWeight = ImitationTrainerTrainingSettings.ActionRegularizationWeight;
-	ImitationTrainingSettings.ActionEntropyWeight = ImitationTrainerTrainingSettings.ActionEntropyWeight;
+	ImitationTrainingSettings.PolicyBatchSize = ImitationTrainerTrainingSettings.BatchSize;
 	ImitationTrainingSettings.Seed = ImitationTrainerTrainingSettings.RandomSeed;
 	ImitationTrainingSettings.Device = UE::Learning::Agents::GetTrainerDevice(ImitationTrainerTrainingSettings.Device);
 	ImitationTrainingSettings.bUseTensorboard = ImitationTrainerTrainingSettings.bUseTensorboard;
-	ImitationTrainingSettings.bSaveSnapshots = ImitationTrainerTrainingSettings.bSaveSnapshots;
+
+	const UE::Learning::EImitationTrainerFlags TrainerFlags = 
+		bReinitializePolicyNetwork ? 
+		UE::Learning::EImitationTrainerFlags::None : 
+		UE::Learning::EImitationTrainerFlags::UseInitialPolicyNetwork;
 
 	ImitationTrainer = MakeUnique<UE::Learning::FSharedMemoryImitationTrainer>(
 		GetName(),
@@ -267,44 +191,36 @@ void ULearningAgentsImitationTrainer::BeginTraining(
 		ObservationNum,
 		ActionNum,
 		MemoryStateNum,
-		*Policy->GetPolicyNetworkAsset()->NeuralNetworkData,
-		*Policy->GetEncoderNetworkAsset()->NeuralNetworkData,
-		*Policy->GetDecoderNetworkAsset()->NeuralNetworkData,
-		Interactor->GetObservationSchema(),
-		Interactor->GetObservationSchemaElement(),
-		Interactor->GetActionSchema(),
-		Interactor->GetActionSchemaElement(),
+		Policy->GetPolicyNetwork(),
 		ImitationTrainingSettings);
 
 	UE_LOG(LogLearning, Display, TEXT("%s: Sending / Receiving initial policy..."), *GetName());
 
 	UE::Learning::ETrainerResponse Response = UE::Learning::ETrainerResponse::Success;
 
-	Response = ImitationTrainer->SendPolicy(*Policy->GetPolicyNetworkAsset()->NeuralNetworkData, TrainerTimeout);
-	if (Response != UE::Learning::ETrainerResponse::Success)
-	{
-		UE_LOG(LogLearning, Error, TEXT("%s: Error sending policy to trainer: %s. Check log for additional errors."), *GetName(), UE::Learning::Trainer::GetResponseString(Response));
-		ImitationTrainer->Terminate();
-		bHasTrainingFailed = true;
-		return;
-	}
+	Response = ImitationTrainer->SendPolicy(Policy->GetPolicyNetwork(), TrainerTimeout);
 
-	Response = ImitationTrainer->SendEncoder(*Policy->GetEncoderNetworkAsset()->NeuralNetworkData, TrainerTimeout);
 	if (Response != UE::Learning::ETrainerResponse::Success)
 	{
-		UE_LOG(LogLearning, Error, TEXT("%s: Error sending encoder to trainer: %s. Check log for additional errors."), *GetName(), UE::Learning::Trainer::GetResponseString(Response));
+		UE_LOG(LogLearning, Error, TEXT("%s: Error sending policy to trainer: %s. Check log for errors."), *GetName(), UE::Learning::Trainer::GetResponseString(Response));
 		ImitationTrainer->Terminate();
 		bHasTrainingFailed = true;
 		return;
 	}
 
-	Response = ImitationTrainer->SendDecoder(*Policy->GetDecoderNetworkAsset()->NeuralNetworkData, TrainerTimeout);
-	if (Response != UE::Learning::ETrainerResponse::Success)
+	if (!bReinitializePolicyNetwork)
 	{
-		UE_LOG(LogLearning, Error, TEXT("%s: Error sending decoder to trainer: %s. Check log for additional errors."), *GetName(), UE::Learning::Trainer::GetResponseString(Response));
-		ImitationTrainer->Terminate();
-		bHasTrainingFailed = true;
-		return;
+		Response = ImitationTrainer->RecvPolicy(Policy->GetPolicyNetwork(), TrainerTimeout);
+
+		if (Response != UE::Learning::ETrainerResponse::Success)
+		{
+			UE_LOG(LogLearning, Error, TEXT("%s: Error receiving policy from trainer: %s. Check log for errors."), *GetName(), UE::Learning::Trainer::GetResponseString(Response));
+			bHasTrainingFailed = true;
+			ImitationTrainer->Terminate();
+			return;
+		}
+
+		Policy->GetNetworkAsset()->ForceMarkDirty();
 	}
 
 	UE_LOG(LogLearning, Display, TEXT("%s: Sending Experience..."), *GetName());
@@ -320,7 +236,7 @@ void ULearningAgentsImitationTrainer::BeginTraining(
 
 	if (Response != UE::Learning::ETrainerResponse::Success)
 	{
-		UE_LOG(LogLearning, Error, TEXT("%s: Error sending experience to trainer: %s. Check log for additional errors."), *GetName(), UE::Learning::Trainer::GetResponseString(Response));
+		UE_LOG(LogLearning, Error, TEXT("%s: Error sending experience to trainer: %s. Check log for errors."), *GetName(), UE::Learning::Trainer::GetResponseString(Response));
 		bHasTrainingFailed = true;
 		ImitationTrainer->Terminate();
 		return;
@@ -331,12 +247,6 @@ void ULearningAgentsImitationTrainer::BeginTraining(
 
 void ULearningAgentsImitationTrainer::DoneTraining()
 {
-	if (!IsSetup())
-	{
-		UE_LOG(LogLearning, Error, TEXT("%s: Setup not complete."), *GetName());
-		return;
-	}
-
 	if (IsTraining())
 	{
 		// Wait for Trainer to finish
@@ -351,12 +261,6 @@ void ULearningAgentsImitationTrainer::DoneTraining()
 
 void ULearningAgentsImitationTrainer::EndTraining()
 {
-	if (!IsSetup())
-	{
-		UE_LOG(LogLearning, Error, TEXT("%s: Setup not complete."), *GetName());
-		return;
-	}
-
 	if (IsTraining())
 	{
 		UE_LOG(LogLearning, Display, TEXT("%s: Stopping training..."), *GetName());
@@ -369,12 +273,6 @@ void ULearningAgentsImitationTrainer::IterateTraining()
 {
 	UE_LEARNING_TRACE_CPUPROFILER_EVENT_SCOPE(ULearningAgentsImitationTrainer::IterateTraining);
 
-	if (!IsSetup())
-	{
-		UE_LOG(LogLearning, Error, TEXT("%s: Setup not complete."), *GetName());
-		return;
-	}
-
 	if (!IsTraining())
 	{
 		UE_LOG(LogLearning, Error, TEXT("%s: Training not running."), *GetName());
@@ -383,7 +281,9 @@ void ULearningAgentsImitationTrainer::IterateTraining()
 
 	if (ImitationTrainer->HasPolicyOrCompleted())
 	{
-		UE::Learning::ETrainerResponse Response = ImitationTrainer->RecvPolicy(*Policy->GetPolicyNetworkAsset()->NeuralNetworkData, TrainerTimeout);
+		UE::Learning::ETrainerResponse Response = ImitationTrainer->RecvPolicy(Policy->GetPolicyNetwork(), TrainerTimeout);
+		Policy->GetNetworkAsset()->ForceMarkDirty();
+
 		if (Response == UE::Learning::ETrainerResponse::Completed)
 		{
 			UE_LOG(LogLearning, Display, TEXT("%s: Trainer completed training."), *GetName());
@@ -392,47 +292,22 @@ void ULearningAgentsImitationTrainer::IterateTraining()
 		}
 		else if (Response != UE::Learning::ETrainerResponse::Success)
 		{
-			UE_LOG(LogLearning, Error, TEXT("%s: Error receiving policy from trainer: %s. Check log for additional errors."), *GetName(), UE::Learning::Trainer::GetResponseString(Response));
-			bHasTrainingFailed = true;
-			EndTraining();
-			return;
-		}
-		Policy->GetPolicyNetworkAsset()->ForceMarkDirty();
-
-		Response = ImitationTrainer->RecvEncoder(*Policy->GetEncoderNetworkAsset()->NeuralNetworkData, TrainerTimeout);
-		if (Response != UE::Learning::ETrainerResponse::Success)
-		{
-			UE_LOG(LogLearning, Error, TEXT("%s: Error receiving encoder from trainer: %s. Check log for additional errors."), *GetName(), UE::Learning::Trainer::GetResponseString(Response));
+			UE_LOG(LogLearning, Error, TEXT("%s: Error receiving policy from trainer. Check log for errors."), *GetName());
 			bHasTrainingFailed = true;
 			EndTraining();
 			return;
 		}
-		Policy->GetEncoderNetworkAsset()->ForceMarkDirty();
-
-		Response = ImitationTrainer->RecvDecoder(*Policy->GetDecoderNetworkAsset()->NeuralNetworkData, TrainerTimeout);
-		if (Response != UE::Learning::ETrainerResponse::Success)
-		{
-			UE_LOG(LogLearning, Error, TEXT("%s: Error receiving decoder from trainer: %s. Check log for additional errors."), *GetName(), UE::Learning::Trainer::GetResponseString(Response));
-			bHasTrainingFailed = true;
-			EndTraining();
-			return;
-		}
-		Policy->GetDecoderNetworkAsset()->ForceMarkDirty();
 	}
 }
 
 void ULearningAgentsImitationTrainer::RunTraining(
+	ULearningAgentsPolicy* InPolicy,
 	const ULearningAgentsRecording* Recording,
 	const FLearningAgentsImitationTrainerSettings& ImitationTrainerSettings,
 	const FLearningAgentsImitationTrainerTrainingSettings& ImitationTrainerTrainingSettings,
-	const FLearningAgentsTrainerPathSettings& ImitationTrainerPathSettings)
+	const FLearningAgentsTrainerPathSettings& ImitationTrainerPathSettings,
+	const bool bReinitializePolicyNetwork)
 {
-	if (!IsSetup())
-	{
-		UE_LOG(LogLearning, Error, TEXT("%s: Setup not complete."), *GetName());
-		return;
-	}
-
 	if (bHasTrainingFailed)
 	{
 		UE_LOG(LogLearning, Error, TEXT("%s: Training has failed. Check log for errors."), *GetName());
@@ -443,10 +318,12 @@ void ULearningAgentsImitationTrainer::RunTraining(
 	if (!IsTraining())
 	{
 		BeginTraining(
+			InPolicy,
 			Recording,
 			ImitationTrainerSettings,
 			ImitationTrainerTrainingSettings,
-			ImitationTrainerPathSettings);
+			ImitationTrainerPathSettings,
+			bReinitializePolicyNetwork);
 
 		if (!IsTraining())
 		{
diff --git a/Engine/Plugins/Experimental/LearningAgents/Source/LearningAgentsTraining/Private/LearningAgentsRecorder.cpp b/Engine/Plugins/Experimental/LearningAgents/Source/LearningAgentsTraining/Private/LearningAgentsRecorder.cpp
--- a/Engine/Plugins/Experimental/LearningAgents/Source/LearningAgentsTraining/Private/LearningAgentsRecorder.cpp
+++ b/Engine/Plugins/Experimental/LearningAgents/Source/LearningAgentsTraining/Private/LearningAgentsRecorder.cpp
@@ -5,6 +5,8 @@
 #include "LearningAgentsManager.h"
 #include "LearningAgentsInteractor.h"
 #include "LearningAgentsRecording.h"
+#include "LearningAgentsHelpers.h"
+#include "LearningFeatureObject.h"
 #include "LearningLog.h"
 #include "LearningTrainer.h"
 
@@ -17,12 +19,6 @@ FLearningAgentsRecorderPathSettings::FLearningAgentsRecorderPathSettings()
 	IntermediateRelativePath.Path = FPaths::ProjectIntermediateDir();
 }
 
-ULearningAgentsRecorder::FAgentRecordBuffer::FAgentRecordBuffer() = default;
-
-ULearningAgentsRecorder::FAgentRecordBuffer::FAgentRecordBuffer(const int32 InObservationCompatibilityHash, const int32 InActionCompatibilityHash) 
-	: ObservationCompatibilityHash(InObservationCompatibilityHash)
-	, ActionCompatibilityHash(InActionCompatibilityHash) {};
-
 TLearningArrayView<1, float> ULearningAgentsRecorder::FAgentRecordBuffer::GetObservation(const int32 SampleIdx)
 {
 	return Observations[SampleIdx / ChunkSize][SampleIdx % ChunkSize];
@@ -77,18 +73,13 @@ void ULearningAgentsRecorder::FAgentRecordBuffer::CopyToRecord(FLearningAgentsRe
 	Record.StepNum = StepNum;
 	Record.ObservationDimNum = GetObservation(0).Num();
 	Record.ActionDimNum = GetAction(0).Num();
-	Record.ObservationCompatibilityHash = ObservationCompatibilityHash;
-	Record.ActionCompatibilityHash = ActionCompatibilityHash;
-	Record.ObservationData.SetNumUninitialized(StepNum * Record.ObservationDimNum);
-	Record.ActionData.SetNumUninitialized(StepNum * Record.ActionDimNum);
-
-	TLearningArrayView<2, float> ObservationsView = TLearningArrayView<2, float>(Record.ObservationData.GetData(), { Record.StepNum, Record.ObservationDimNum });
-	TLearningArrayView<2, float> ActionsView = TLearningArrayView<2, float>(Record.ActionData.GetData(), { Record.StepNum, Record.ActionDimNum });
+	Record.Observations.SetNumUninitialized({ StepNum, Observations[0].Num<1>() });
+	Record.Actions.SetNumUninitialized({ StepNum, Actions[0].Num<1>() });
 
 	for (int32 StepIdx = 0; StepIdx < StepNum; StepIdx++)
 	{
-		UE::Learning::Array::Copy(ObservationsView[StepIdx], GetObservation(StepIdx));
-		UE::Learning::Array::Copy(ActionsView[StepIdx], GetAction(StepIdx));
+		UE::Learning::Array::Copy(Record.Observations[StepIdx], GetObservation(StepIdx));
+		UE::Learning::Array::Copy(Record.Actions[StepIdx], GetAction(StepIdx));
 	}
 }
 
@@ -96,58 +87,20 @@ ULearningAgentsRecorder::ULearningAgentsRecorder() : Super(FObjectInitializer::G
 ULearningAgentsRecorder::ULearningAgentsRecorder(FVTableHelper& Helper) : Super(Helper) {}
 ULearningAgentsRecorder::~ULearningAgentsRecorder() = default;
 
-void ULearningAgentsRecorder::BeginDestroy()
+void ULearningAgentsRecorder::EndPlay(const EEndPlayReason::Type EndPlayReason)
 {
 	if (IsRecording())
 	{
 		EndRecording();
 	}
 
-	Super::BeginDestroy();
-}
-
-ULearningAgentsRecorder* ULearningAgentsRecorder::MakeRecorder(
-	ULearningAgentsManager* InManager,
-	ULearningAgentsInteractor* InInteractor,
-	TSubclassOf<ULearningAgentsRecorder> Class,
-	const FName Name,
-	const FLearningAgentsRecorderPathSettings& RecorderPathSettings,
-	ULearningAgentsRecording* RecordingAsset,
-	bool bReinitializeRecording)
-{
-	if (!InManager)
-	{
-		UE_LOG(LogLearning, Error, TEXT("MakeRecorder: InManager is nullptr."));
-		return nullptr;
-	}
-
-	if (!Class)
-	{
-		UE_LOG(LogLearning, Error, TEXT("MakeRecorder: Class is nullptr."));
-		return nullptr;
-	}
-
-	const FName UniqueName = MakeUniqueObjectName(InManager, Class, Name, EUniqueObjectNameOptions::GloballyUnique);
-
-	ULearningAgentsRecorder* Recorder = NewObject<ULearningAgentsRecorder>(InManager, Class, UniqueName);
-	if (!Recorder) { return nullptr; }
-
-	Recorder->SetupRecorder(
-		InManager,
-		InInteractor,
-		RecorderPathSettings,
-		RecordingAsset,
-		bReinitializeRecording);
-
-	return Recorder->IsSetup() ? Recorder : nullptr;
+	Super::EndPlay(EndPlayReason);
 }
 
 void ULearningAgentsRecorder::SetupRecorder(
-	ULearningAgentsManager* InManager,
 	ULearningAgentsInteractor* InInteractor,
 	const FLearningAgentsRecorderPathSettings& RecorderPathSettings,
-	ULearningAgentsRecording* RecordingAsset,
-	bool bReinitializeRecording)
+	ULearningAgentsRecording* RecordingAsset)
 {
 	if (IsSetup())
 	{
@@ -155,9 +108,9 @@ void ULearningAgentsRecorder::SetupRecorder(
 		return;
 	}
 
-	if (!InManager)
+	if (!Manager)
 	{
-		UE_LOG(LogLearning, Error, TEXT("%s: InManager is nullptr."), *GetName());
+		UE_LOG(LogLearning, Error, TEXT("%s: Must be attached to a LearningAgentsManager Actor."), *GetName());
 		return;
 	}
 
@@ -173,7 +126,6 @@ void ULearningAgentsRecorder::SetupRecorder(
 		return;
 	}
 
-	Manager = InManager;
 	Interactor = InInteractor;
 
 	if (RecordingAsset)
@@ -183,68 +135,61 @@ void ULearningAgentsRecorder::SetupRecorder(
 	else
 	{
 		const FName UniqueName = MakeUniqueObjectName(this, ULearningAgentsRecording::StaticClass(), TEXT("Recording"), EUniqueObjectNameOptions::GloballyUnique);
-		Recording = NewObject<ULearningAgentsRecording>(this, UniqueName);
-	}
 
-	if (bReinitializeRecording)
-	{
-		Recording->Records.Empty();
-		Recording->ForceMarkDirty();
+		Recording = NewObject<ULearningAgentsRecording>(this, UniqueName);
 	}
 
 	RecordingDirectory = RecorderPathSettings.IntermediateRelativePath.Path / TEXT("LearningAgents") / RecorderPathSettings.RecordingsSubdirectory;
 
-	// Find Compatibility Hash
-
-	const int32 ObservationCompatibilityHash = UE::Learning::Observation::GetSchemaObjectsCompatibilityHash(Interactor->GetObservationSchema(), Interactor->GetObservationSchemaElement());
-	const int32 ActionCompatibilityHash = UE::Learning::Action::GetSchemaObjectsCompatibilityHash(Interactor->GetActionSchema(), Interactor->GetActionSchemaElement());
-
-	// Create Record Buffers
-
-	RecordBuffers.Empty(Manager->GetMaxAgentNum());
-	for (int32 RecordBufferIdx = 0; RecordBufferIdx < Manager->GetMaxAgentNum(); RecordBufferIdx++)
-	{
-		RecordBuffers.Emplace(ObservationCompatibilityHash, ActionCompatibilityHash);
-	}
+	RecordBuffers.Empty();
+	RecordBuffers.SetNum(Manager->GetMaxAgentNum());
 
 	bIsSetup = true;
 
-	Manager->AddListener(this);
+	OnAgentsAdded(Manager->GetAllAgentIds());
 }
 
-void ULearningAgentsRecorder::OnAgentsRemoved_Implementation(const TArray<int32>& AgentIds)
+void ULearningAgentsRecorder::OnAgentsRemoved(const TArray<int32>& AgentIds)
 {
-	if (!IsSetup())
+	if (IsSetup())
 	{
-		UE_LOG(LogLearning, Error, TEXT("%s: Setup not complete."), *GetName());
-		return;
-	}
+		for (const int32 AgentId : AgentIds)
+		{
+			if (!RecordBuffers[AgentId].IsEmpty())
+			{
+				RecordBuffers[AgentId].CopyToRecord(Recording->Records.AddDefaulted_GetRef());
+				RecordBuffers[AgentId].Empty();
+			}
+		}
 
-	for (const int32 AgentId : AgentIds)
-	{
-		if (!RecordBuffers[AgentId].IsEmpty())
+		for (ULearningAgentsHelper* Helper : HelperObjects)
 		{
-			RecordBuffers[AgentId].CopyToRecord(Recording->Records.AddDefaulted_GetRef());
-			RecordBuffers[AgentId].Empty();
+			Helper->OnAgentsRemoved(AgentIds);
 		}
+
+		AgentsRemoved(AgentIds);
 	}
 }
 
-void ULearningAgentsRecorder::OnAgentsReset_Implementation(const TArray<int32>& AgentIds)
+void ULearningAgentsRecorder::OnAgentsReset(const TArray<int32>& AgentIds)
 {
-	if (!IsSetup())
+	if (IsSetup())
 	{
-		UE_LOG(LogLearning, Error, TEXT("%s: Setup not complete."), *GetName());
-		return;
-	}
+		for (const int32 AgentId : AgentIds)
+		{
+			if (!RecordBuffers[AgentId].IsEmpty())
+			{
+				RecordBuffers[AgentId].CopyToRecord(Recording->Records.AddDefaulted_GetRef());
+				RecordBuffers[AgentId].Empty();
+			}
+		}
 
-	for (const int32 AgentId : AgentIds)
-	{
-		if (!RecordBuffers[AgentId].IsEmpty())
+		for (ULearningAgentsHelper* Helper : HelperObjects)
 		{
-			RecordBuffers[AgentId].CopyToRecord(Recording->Records.AddDefaulted_GetRef());
-			RecordBuffers[AgentId].Empty();
+			Helper->OnAgentsReset(AgentIds);
 		}
+
+		AgentsReset(AgentIds);
 	}
 }
 
@@ -262,26 +207,25 @@ void ULearningAgentsRecorder::AddExperience()
 		return;
 	}
 
-	if (Manager->GetAgentNum() == 0)
-	{
-		UE_LOG(LogLearning, Warning, TEXT("%s: No agents added to Manager."), *GetName());
-	}
-
 	for (const int32 AgentId : Manager->GetAllAgentSet())
 	{
-		if (Interactor->ObservationVectorIteration[AgentId] == 0 || Interactor->ActionVectorIteration[AgentId] == 0)
+		if (Interactor->GetObservationEncodingAgentIteration()[AgentId] == 0 ||
+			Interactor->GetActionEncodingAgentIteration()[AgentId] == 0)
 		{
 			UE_LOG(LogLearning, Warning, TEXT("%s: Agent with id %i has not made observations and taken actions so experience will not be recorded for it."), *GetName(), AgentId);
 			continue;
 		}
 
-		if (Interactor->ObservationVectorIteration[AgentId] != Interactor->ActionVectorIteration[AgentId])
+		if (Interactor->GetObservationEncodingAgentIteration()[AgentId] !=
+			Interactor->GetActionEncodingAgentIteration()[AgentId])
 		{
 			UE_LOG(LogLearning, Warning, TEXT("%s: Agent with id %i does not have matching iteration numbers for observations and actions so experience will not be recorded for it."), *GetName(), AgentId);
 			continue;
 		}
 
-		RecordBuffers[AgentId].Push(Interactor->ObservationVectors[AgentId], Interactor->ActionVectors[AgentId]);
+		RecordBuffers[AgentId].Push(
+			Interactor->GetObservationFeature().FeatureBuffer()[AgentId],
+			Interactor->GetActionFeature().FeatureBuffer()[AgentId]);
 	}
 }
 
@@ -318,7 +262,96 @@ void ULearningAgentsRecorder::EndRecording()
 	bIsRecording = false;
 }
 
-void ULearningAgentsRecorder::BeginRecording()
+void ULearningAgentsRecorder::LoadRecordingFromFile(const FFilePath& File)
+{
+	if (!IsSetup())
+	{
+		UE_LOG(LogLearning, Error, TEXT("%s: Setup not complete."), *GetName());
+		return;
+	}
+
+	Recording->LoadRecordingFromFile(File);
+}
+
+void ULearningAgentsRecorder::SaveRecordingToFile(const FFilePath& File) const
+{
+	if (!IsSetup())
+	{
+		UE_LOG(LogLearning, Error, TEXT("%s: Setup not complete."), *GetName());
+		return;
+	}
+
+	Recording->SaveRecordingToFile(File);
+}
+
+void ULearningAgentsRecorder::AppendRecordingFromFile(const FFilePath& File)
+{
+	if (!IsSetup())
+	{
+		UE_LOG(LogLearning, Error, TEXT("%s: Setup not complete."), *GetName());
+		return;
+	}
+
+	Recording->AppendRecordingFromFile(File);
+}
+
+void ULearningAgentsRecorder::UseRecordingAsset(ULearningAgentsRecording* RecordingAsset)
+{
+	if (!IsSetup())
+	{
+		UE_LOG(LogLearning, Error, TEXT("%s: Setup not complete."), *GetName());
+		return;
+	}
+
+	if (!RecordingAsset)
+	{
+		UE_LOG(LogLearning, Error, TEXT("%s: Asset is nullptr."), *GetName());
+		return;
+	}
+
+	if (RecordingAsset == Recording)
+	{
+		UE_LOG(LogLearning, Error, TEXT("%s: Asset is same as the current recording."), *GetName());
+		return;
+	}
+
+	Recording = RecordingAsset;
+}
+
+void ULearningAgentsRecorder::LoadRecordingFromAsset(ULearningAgentsRecording* RecordingAsset)
+{
+	if (!IsSetup())
+	{
+		UE_LOG(LogLearning, Error, TEXT("%s: Setup not complete."), *GetName());
+		return;
+	}
+
+	Recording->LoadRecordingFromAsset(RecordingAsset);
+}
+
+void ULearningAgentsRecorder::SaveRecordingToAsset(ULearningAgentsRecording* RecordingAsset)
+{
+	if (!IsSetup())
+	{
+		UE_LOG(LogLearning, Error, TEXT("%s: Setup not complete."), *GetName());
+		return;
+	}
+
+	Recording->SaveRecordingToAsset(RecordingAsset);
+}
+
+void ULearningAgentsRecorder::AppendRecordingToAsset(ULearningAgentsRecording* RecordingAsset)
+{
+	if (!IsSetup())
+	{
+		UE_LOG(LogLearning, Error, TEXT("%s: Setup not complete."), *GetName());
+		return;
+	}
+
+	Recording->AppendRecordingToAsset(RecordingAsset);
+}
+
+void ULearningAgentsRecorder::BeginRecording(bool bReinitializeRecording)
 {
 	if (!IsSetup())
 	{
@@ -332,6 +365,12 @@ void ULearningAgentsRecorder::BeginRecording()
 		return;
 	}
 
+	if (bReinitializeRecording)
+	{
+		Recording->Records.Empty();
+		Recording->ForceMarkDirty();
+	}
+
 	for (const int32 AgentId : Manager->GetAllAgentSet())
 	{
 		RecordBuffers[AgentId].Empty();
@@ -340,7 +379,7 @@ void ULearningAgentsRecorder::BeginRecording()
 	bIsRecording = true;
 }
 
-const ULearningAgentsRecording* ULearningAgentsRecorder::GetRecordingAsset() const
+const ULearningAgentsRecording* ULearningAgentsRecorder::GetCurrentRecording() const
 {
 	if (!IsSetup())
 	{
diff --git a/Engine/Plugins/Experimental/LearningAgents/Source/LearningAgentsTraining/Private/LearningAgentsRecording.cpp b/Engine/Plugins/Experimental/LearningAgents/Source/LearningAgentsTraining/Private/LearningAgentsRecording.cpp
--- a/Engine/Plugins/Experimental/LearningAgents/Source/LearningAgentsTraining/Private/LearningAgentsRecording.cpp
+++ b/Engine/Plugins/Experimental/LearningAgents/Source/LearningAgentsTraining/Private/LearningAgentsRecording.cpp
@@ -7,6 +7,16 @@
 #include "UObject/Package.h"
 #include "Misc/FileHelper.h"
 
+bool FLearningAgentsRecord::Serialize(FArchive& Ar)
+{
+	Ar << StepNum;
+	Ar << ObservationDimNum;
+	Ar << ActionDimNum;
+	UE::Learning::Array::Serialize(Ar, Observations);
+	UE::Learning::Array::Serialize(Ar, Actions);
+	return true;
+}
+
 ULearningAgentsRecording::ULearningAgentsRecording() = default;
 ULearningAgentsRecording::ULearningAgentsRecording(FVTableHelper& Helper) : Super(Helper) {}
 ULearningAgentsRecording::~ULearningAgentsRecording() = default;
@@ -59,10 +69,8 @@ void ULearningAgentsRecording::LoadRecordingFromFile(const FFilePath& File)
 			UE::Learning::DeserializeFromBytes(Offset, RecordingData, Records[RecordIdx].StepNum);
 			UE::Learning::DeserializeFromBytes(Offset, RecordingData, Records[RecordIdx].ObservationDimNum);
 			UE::Learning::DeserializeFromBytes(Offset, RecordingData, Records[RecordIdx].ActionDimNum);
-			UE::Learning::DeserializeFromBytes(Offset, RecordingData, Records[RecordIdx].ObservationCompatibilityHash);
-			UE::Learning::DeserializeFromBytes(Offset, RecordingData, Records[RecordIdx].ActionCompatibilityHash);
-			UE::Learning::Array::DeserializeFromBytes<2, float>(Offset, RecordingData, Records[RecordIdx].ObservationData);
-			UE::Learning::Array::DeserializeFromBytes<2, float>(Offset, RecordingData, Records[RecordIdx].ActionData);
+			UE::Learning::Array::DeserializeFromBytes(Offset, RecordingData, Records[RecordIdx].Observations);
+			UE::Learning::Array::DeserializeFromBytes(Offset, RecordingData, Records[RecordIdx].Actions);
 		}
 
 		UE_LEARNING_CHECK(Offset == RecordingData.Num());
@@ -87,13 +95,11 @@ void ULearningAgentsRecording::SaveRecordingToFile(const FFilePath& File) const
 	for (int32 RecordIdx = 0; RecordIdx < Records.Num(); RecordIdx++)
 	{
 		TotalByteNum +=
-			sizeof(int32) + // StepNum
+			sizeof(int32) + // SampleNum
 			sizeof(int32) + // ObservationDimNum
 			sizeof(int32) + // ActionDimNum
-			sizeof(int32) + // ObservationCompatibilityHash
-			sizeof(int32) + // ActionCompatibilityHash
-			UE::Learning::Array::SerializationByteNum<2, float>({ Records[RecordIdx].StepNum, Records[RecordIdx].ObservationDimNum }) + // Observations
-			UE::Learning::Array::SerializationByteNum<2, float>({ Records[RecordIdx].StepNum, Records[RecordIdx].ActionDimNum });	   // Actions
+			UE::Learning::Array::SerializationByteNum<2, float>(Records[RecordIdx].Observations.Shape()) + // Observations
+			UE::Learning::Array::SerializationByteNum<2, float>(Records[RecordIdx].Actions.Shape());	   // Actions
 	}
 
 	RecordingData.SetNumUninitialized(TotalByteNum);
@@ -108,10 +114,8 @@ void ULearningAgentsRecording::SaveRecordingToFile(const FFilePath& File) const
 		UE::Learning::SerializeToBytes(Offset, RecordingData, Records[RecordIdx].StepNum);
 		UE::Learning::SerializeToBytes(Offset, RecordingData, Records[RecordIdx].ObservationDimNum);
 		UE::Learning::SerializeToBytes(Offset, RecordingData, Records[RecordIdx].ActionDimNum);
-		UE::Learning::SerializeToBytes(Offset, RecordingData, Records[RecordIdx].ObservationCompatibilityHash);
-		UE::Learning::SerializeToBytes(Offset, RecordingData, Records[RecordIdx].ActionCompatibilityHash);
-		UE::Learning::Array::SerializeToBytes<2, float>(Offset, RecordingData, { Records[RecordIdx].StepNum, Records[RecordIdx].ObservationDimNum }, Records[RecordIdx].ObservationData);
-		UE::Learning::Array::SerializeToBytes<2, float>(Offset, RecordingData, { Records[RecordIdx].StepNum, Records[RecordIdx].ActionDimNum }, Records[RecordIdx].ActionData);
+		UE::Learning::Array::SerializeToBytes(Offset, RecordingData, Records[RecordIdx].Observations);
+		UE::Learning::Array::SerializeToBytes(Offset, RecordingData, Records[RecordIdx].Actions);
 	}
 
 	UE_LEARNING_CHECK(Offset == RecordingData.Num());
@@ -185,67 +189,6 @@ void ULearningAgentsRecording::AppendRecordingToAsset(ULearningAgentsRecording*
 	RecordingAsset->ForceMarkDirty();
 }
 
-int32 ULearningAgentsRecording::GetRecordNum() const
-{
-	return Records.Num();
-}
-
-int32 ULearningAgentsRecording::GetRecordStepNum(const int32 Record) const
-{
-	if (Record < 0 || Record >= Records.Num())
-	{
-		UE_LOG(LogLearning, Error, TEXT("%s: Record out of range. Asked for record %i but recording only has %i records."), *GetName(), Record, Records.Num());
-		return 0;
-	}
-
-	return Records[Record].StepNum;
-}
-
-
-void ULearningAgentsRecording::GetObservationVector(TArray<float>& OutObservationVector, int32& OutObservationCompatibilityHash, const int32 Record, const int32 Step)
-{
-	if (Record < 0 || Record >= Records.Num())
-	{
-		UE_LOG(LogLearning, Error, TEXT("%s: Record out of range. Asked for record %i but recording only has %i records."), *GetName(), Record, Records.Num());
-		OutObservationVector.Empty();
-		OutObservationCompatibilityHash = 0;
-		return;
-	}
-
-	if (Step < 0 || Step >= Records[Record].StepNum)
-	{
-		UE_LOG(LogLearning, Error, TEXT("%s: Step out of range. Asked for step %i but recording only has %i steps."), *GetName(), Step, Records[Record].StepNum);
-		OutObservationVector.Empty();
-		OutObservationCompatibilityHash = 0;
-		return;
-	}
-
-	OutObservationVector = MakeArrayView(Records[Record].ObservationData).Slice(Records[Record].ObservationDimNum * Step, Records[Record].ObservationDimNum);
-	OutObservationCompatibilityHash = Records[Record].ObservationCompatibilityHash;
-}
-
-void ULearningAgentsRecording::GetActionVector(TArray<float>& OutActionVector, int32& OutActionCompatibilityHash, const int32 Record, const int32 Step)
-{
-	if (Record < 0 || Record >= Records.Num())
-	{
-		UE_LOG(LogLearning, Error, TEXT("%s: Record out of range. Asked for record %i but recording only has %i records."), *GetName(), Record, Records.Num());
-		OutActionVector.Empty();
-		OutActionCompatibilityHash = 0;
-		return;
-	}
-
-	if (Step < 0 || Step >= Records[Record].StepNum)
-	{
-		UE_LOG(LogLearning, Error, TEXT("%s: Step out of range. Asked for step %i but recording only has %i steps."), *GetName(), Step, Records[Record].StepNum);
-		OutActionVector.Empty();
-		OutActionCompatibilityHash = 0;
-		return;
-	}
-
-	OutActionVector = MakeArrayView(Records[Record].ActionData).Slice(Records[Record].ActionDimNum * Step, Records[Record].ActionDimNum);
-	OutActionCompatibilityHash = Records[Record].ActionCompatibilityHash;
-}
-
 void ULearningAgentsRecording::ForceMarkDirty()
 {
 	// Manually mark the package as dirty since just using `Modify` prevents 
