[Backout] - CL30154408

diff --git a/Engine/Plugins/Experimental/LearningAgents/Source/LearningAgentsTraining/Public/LearningAgentsImitationTrainer.h b/Engine/Plugins/Experimental/LearningAgents/Source/LearningAgentsTraining/Public/LearningAgentsImitationTrainer.h
--- a/Engine/Plugins/Experimental/LearningAgents/Source/LearningAgentsTraining/Public/LearningAgentsImitationTrainer.h
+++ b/Engine/Plugins/Experimental/LearningAgents/Source/LearningAgentsTraining/Public/LearningAgentsImitationTrainer.h
@@ -15,7 +15,6 @@ namespace UE::Learning
 	struct FSharedMemoryImitationTrainer;
 }
 
-class ULearningAgentsInteractor;
 class ULearningAgentsPolicy;
 class ULearningAgentsRecording;
 
@@ -46,7 +45,7 @@ public:
 
 	/** Learning rate of the policy network. Typical values are between 0.001 and 0.0001. */
 	UPROPERTY(EditAnywhere, Category = "LearningAgents", meta = (ClampMin = "0.0", ClampMax = "1.0", UIMin = "0.0", UIMax = "1.0"))
-	float LearningRate = 0.001f;
+	float LearningRate = 0.0001f;
 
 	/** Ratio by which to decay the learning rate every 1000 iterations. */
 	UPROPERTY(EditAnywhere, Category = "LearningAgents", meta = (ClampMin = "0.0", ClampMax = "1.0", UIMin = "0.0", UIMax = "1.0"))
@@ -66,27 +65,6 @@ public:
 	UPROPERTY(EditAnywhere, Category = "LearningAgents", meta = (ClampMin = "0", ClampMax = "4096", UIMin = "0", UIMax = "4096"))
 	uint32 BatchSize = 128;
 
-	/**
-	 * The window of observations and actions over which to do the training.Increasing this value will encourage the policy to use its memory 
-	 * effectively. Too large and training can become unstable. Given we don't know the memory state during imitation learning it is better this is 
-	 * slightly larger than when we are doing reinforcement learning.
-	 */
-	UPROPERTY(EditAnywhere, Category = "LearningAgents", meta = (ClampMin = "1", ClampMax = "128", UIMin = "1", UIMax = "512"))
-	uint32 Window = 64;
-
-	/**
-	 * Weight used to regularize actions. Larger values will encourage smaller actions but too large will cause actions to become always zero.
-	 */
-	UPROPERTY(EditAnywhere, Category = "LearningAgents", meta = (ClampMin = "0", UIMin = "0"))
-	float ActionRegularizationWeight = 0.001f;
-
-	/**
-	 * Weighting used for the entropy bonus. Larger values encourage larger action noise and therefore greater exploration but can make actions very 
-	 * noisy.
-	 */
-	UPROPERTY(EditAnywhere, Category = "LearningAgents", meta = (ClampMin = "0", UIMin = "0"))
-	float ActionEntropyWeight = 0.0f;
-
 	/** The seed used for any random sampling the trainer will perform, e.g. for weight initialization. */
 	UPROPERTY(EditAnywhere, Category = "LearningAgents", meta = (ClampMin = "0", UIMin = "0"))
 	int32 RandomSeed = 1234;
@@ -98,10 +76,6 @@ public:
 	/** If true, TensorBoard logs will be emitted to the intermediate directory. */
 	UPROPERTY(EditAnywhere, Category = "LearningAgents")
 	bool bUseTensorboard = false;
-
-	/** If true, snapshots of the trained networks will be emitted to the intermediate directory. */
-	UPROPERTY(EditAnywhere, Category = "LearningAgents")
-	bool bSaveSnapshots = false;
 };
 
 /**
@@ -114,7 +88,7 @@ public:
  * @see ULearningAgentsRecorder to understand how to make new recordings.
  */
 UCLASS(BlueprintType, Blueprintable)
-class LEARNINGAGENTSTRAINING_API ULearningAgentsImitationTrainer : public ULearningAgentsManagerListener
+class LEARNINGAGENTSTRAINING_API ULearningAgentsImitationTrainer : public UActorComponent
 {
 	GENERATED_BODY()
 
@@ -126,35 +100,8 @@ public:
 	ULearningAgentsImitationTrainer(FVTableHelper& Helper);
 	virtual ~ULearningAgentsImitationTrainer();
 
-	/** Will automatically call EndTraining if training is still in-progress when the object is destroyed. */
-	virtual void BeginDestroy() override;
-
-	/**
-	 * Constructs the imitation trainer and runs the setup functions.
-	 * 
-	 * @param InManager			The agent manager we are using.
-	 * @param InInteractor		The agent interactor we are recording with.
-	 * @param InPolicy			The policy we are using.
-	 * @param Class				The imitation trainer class
-	 */
-	UFUNCTION(BlueprintCallable, Category = "LearningAgents", meta = (Class = "/Script/LearningAgentsTraining.LearningAgentsImitationTrainer", DeterminesOutputType = "Class"))
-	static ULearningAgentsImitationTrainer* MakeImitationTrainer(
-		ULearningAgentsManager* InManager,
-		ULearningAgentsInteractor* InInteractor,
-		ULearningAgentsPolicy* InPolicy,
-		TSubclassOf<ULearningAgentsImitationTrainer> Class);
-
-	/**
-	 * Initializes the imitation trainer and runs the setup functions.
-	 * 
-	 * @param InManager			The agent manager we are using.
-	 * @param InInteractor		The agent interactor we are recording with.
-	 * @param InPolicy			The policy we are using.
-	 */
-	void SetupImitationTrainer(
-		ULearningAgentsManager* InManager,
-		ULearningAgentsInteractor* InInteractor,
-		ULearningAgentsPolicy* InPolicy);
+	/** Will automatically call EndTraining if training is still in-progress when play is ending. */
+	virtual void EndPlay(const EEndPlayReason::Type EndPlayReason) override;
 
 	/** Returns true if the trainer is currently training; Otherwise, false. */
 	UFUNCTION(BlueprintPure, Category = "LearningAgents")
@@ -171,18 +118,21 @@ public:
 
 	/**
 	 * Begins the training process with the provided settings.
-	 * 
-	 * @param Recording							The data to train on.
-	 * @param ImitationTrainerSettings			The settings for this trainer.
-	 * @param ImitationTrainerTrainingSettings	The training settings for this network.
-	 * @param ImitationTrainerPathSettings		The path settings used by the imitation trainer.
+	 * @param InPolicy The policy to train.
+	 * @param Recording The data to train on.
+	 * @param ImitationTrainerSettings The settings for this trainer.
+	 * @param ImitationTrainerTrainingSettings The training settings for this network.
+	 * @param ImitationTrainerPathSettings The path settings used by the imitation trainer.
+	 * @param bReinitializePolicyNetwork If true, reinitialize the policy. Set this to false if your policy is pre-trained.
 	 */
 	UFUNCTION(BlueprintCallable, Category = "LearningAgents")
 	void BeginTraining(
+		ULearningAgentsPolicy* InPolicy,
 		const ULearningAgentsRecording* Recording,
 		const FLearningAgentsImitationTrainerSettings& ImitationTrainerSettings = FLearningAgentsImitationTrainerSettings(),
 		const FLearningAgentsImitationTrainerTrainingSettings& ImitationTrainerTrainingSettings = FLearningAgentsImitationTrainerTrainingSettings(),
-		const FLearningAgentsTrainerPathSettings& ImitationTrainerPathSettings = FLearningAgentsTrainerPathSettings());
+		const FLearningAgentsTrainerPathSettings& ImitationTrainerPathSettings = FLearningAgentsTrainerPathSettings(),
+		const bool bReinitializePolicyNetwork = true);
 
 	/** Stops the training process. */
 	UFUNCTION(BlueprintCallable, Category = "LearningAgents")
@@ -195,26 +145,25 @@ public:
 	/**
 	 * Convenience function that runs a basic training loop. If training has not been started, it will start it. On 
 	 * each following call to this function, it will call IterateTraining.
-	 * 
-	 * @param Recording							The data to train on.
-	 * @param ImitationTrainerSettings			The settings for this trainer.
-	 * @param ImitationTrainerTrainingSettings	The training settings for this network.
-	 * @param ImitationTrainerPathSettings		The path settings used by the imitation trainer.
+	 * @param InPolicy The policy to train.
+	 * @param Recording The data to train on.
+	 * @param ImitationTrainerSettings The settings for this trainer.
+	 * @param ImitationTrainerTrainingSettings The training settings for this network.
+	 * @param ImitationTrainerPathSettings The path settings used by the imitation trainer.
+	 * @param bReinitializePolicyNetwork If true, reinitialize the policy. Set this to false if your policy is pre-trained.
 	 */
 	UFUNCTION(BlueprintCallable, Category = "LearningAgents")
 	void RunTraining(
+		ULearningAgentsPolicy* InPolicy,
 		const ULearningAgentsRecording* Recording,
 		const FLearningAgentsImitationTrainerSettings& ImitationTrainerSettings = FLearningAgentsImitationTrainerSettings(),
 		const FLearningAgentsImitationTrainerTrainingSettings& ImitationTrainerTrainingSettings = FLearningAgentsImitationTrainerTrainingSettings(),
-		const FLearningAgentsTrainerPathSettings& ImitationTrainerPathSettings = FLearningAgentsTrainerPathSettings());
+		const FLearningAgentsTrainerPathSettings& ImitationTrainerPathSettings = FLearningAgentsTrainerPathSettings(),
+		const bool bReinitializePolicyNetwork = true);
 
 // ----- Private Data -----
 private:
 
-	/** The interactor being trained. */
-	UPROPERTY(VisibleAnywhere, Transient, Category = "LearningAgents")
-	TObjectPtr<ULearningAgentsInteractor> Interactor;
-
 	/** The policy being trained. */
 	UPROPERTY(VisibleAnywhere, Transient, Category = "LearningAgents")
 	TObjectPtr<ULearningAgentsPolicy> Policy;
@@ -234,5 +183,10 @@ private:
 
 	void DoneTraining();
 
+	TLearningArray<1, int32> RecordedEpisodeStarts;
+	TLearningArray<1, int32> RecordedEpisodeLengths;
+	TLearningArray<2, float> RecordedObservations;
+	TLearningArray<2, float> RecordedActions;
+
 	TUniquePtr<UE::Learning::FSharedMemoryImitationTrainer> ImitationTrainer;
 };
diff --git a/Engine/Plugins/Experimental/LearningAgents/Source/LearningAgentsTraining/Public/LearningAgentsRecorder.h b/Engine/Plugins/Experimental/LearningAgents/Source/LearningAgentsTraining/Public/LearningAgentsRecorder.h
--- a/Engine/Plugins/Experimental/LearningAgents/Source/LearningAgentsTraining/Public/LearningAgentsRecorder.h
+++ b/Engine/Plugins/Experimental/LearningAgents/Source/LearningAgentsTraining/Public/LearningAgentsRecorder.h
@@ -2,7 +2,7 @@
 
 #pragma once
 
-#include "LearningAgentsManagerListener.h"
+#include "LearningAgentsManagerComponent.h"
 
 #include "LearningArray.h"
 #include "Containers/Map.h"
@@ -33,8 +33,8 @@ public:
 };
 
 /** A component that can be used to create recordings of training data for imitation learning. */
-UCLASS(BlueprintType, Blueprintable, meta = (BlueprintSpawnableComponent))
-class LEARNINGAGENTSTRAINING_API ULearningAgentsRecorder : public ULearningAgentsManagerListener
+UCLASS(BlueprintType, Blueprintable)
+class LEARNINGAGENTSTRAINING_API ULearningAgentsRecorder : public ULearningAgentsManagerComponent
 {
 	GENERATED_BODY()
 
@@ -46,64 +46,42 @@ public:
 	ULearningAgentsRecorder(FVTableHelper& Helper);
 	virtual ~ULearningAgentsRecorder();
 
-	/** Will automatically call EndRecording if recording is still in-progress when the object is destroyed. */
-	virtual void BeginDestroy() override;
-
-	/**
-	 * Constructs this object and runs the setup functions for the underlying data storage.
-	 * 
-	 * @param InManager					The agent manager we are using.
-	 * @param InInteractor				The agent interactor we are recording with.
-	 * @param Class						The recorder class
-	 * @param Name						The recorder class
-	 * @param RecorderPathSettings		The path settings used by the recorder.
-	 * @param RecordingAsset			Optional recording asset to use. If not provided or bReinitializeRecording is 
-	 *									set then a new recording object will be created.
-	 * @pram bReinitializeRecording		If to reinitialize the recording asset
-	 */
-	UFUNCTION(BlueprintCallable, Category = "LearningAgents", meta = (Class = "/Script/LearningAgentsTraining.LearningAgentsRecorder", DeterminesOutputType = "Class"))
-	static ULearningAgentsRecorder* MakeRecorder(
-		ULearningAgentsManager* InManager,
-		ULearningAgentsInteractor* InInteractor,
-		TSubclassOf<ULearningAgentsRecorder> Class,
-		const FName Name = TEXT("Recorder"),
-		const FLearningAgentsRecorderPathSettings& RecorderPathSettings = FLearningAgentsRecorderPathSettings(),
-		ULearningAgentsRecording* RecordingAsset = nullptr,
-		bool bReinitializeRecording = true);
+	/** Will automatically call EndRecording if recording is still in-progress when play is ending. */
+	virtual void EndPlay(const EEndPlayReason::Type EndPlayReason) override;
 
 	/**
 	 * Initializes this object and runs the setup functions for the underlying data storage.
-	 * 
-	 * @param InManager					The agent manager we are using.
-	 * @param InInteractor				The agent interactor we are recording with.
-	 * @param RecorderPathSettings		The path settings used by the recorder.
-	 * @param RecordingAsset			Optional recording asset to use. If not provided or bReinitializeRecording is
-	 *									set then a new recording object will be created.
-	 * @pram bReinitializeRecording		If to reinitialize the recording asset
+	 * @param InInteractor The agent interactor we are recording with.
+	 * @param RecorderPathSettings The path settings used by the recorder.
+	 * @param RecordingAsset Optional recording asset to use. If not provided then a new recording object will be 
+	 * created.
 	 */
 	UFUNCTION(BlueprintCallable, Category = "LearningAgents")
 	void SetupRecorder(
-		ULearningAgentsManager* InManager,
 		ULearningAgentsInteractor* InInteractor,
 		const FLearningAgentsRecorderPathSettings& RecorderPathSettings = FLearningAgentsRecorderPathSettings(),
-		ULearningAgentsRecording* RecordingAsset = nullptr,
-		bool bReinitializeRecording = true);
+		ULearningAgentsRecording* RecordingAsset = nullptr);
 
 public:
 
-	//~ Begin ULearningAgentsManagerListener Interface
-	virtual void OnAgentsRemoved_Implementation(const TArray<int32>& AgentIds) override;
-	virtual void OnAgentsReset_Implementation(const TArray<int32>& AgentIds) override;
-	//~ End ULearningAgentsManagerListener Interface
+	//~ Begin ULearningAgentsManagerComponent Interface
+	virtual void OnAgentsRemoved(const TArray<int32>& AgentIds) override;
+	virtual void OnAgentsReset(const TArray<int32>& AgentIds) override;
+	//~ End ULearningAgentsManagerComponent Interface
 
 // ----- Recording Process -----
 public:
 
-	/** Begins the recording of the observations and actions of each added agent. */
+	/**
+	 * Begins the recording of the observations and actions of each added agent.
+	 * @param bReinitializeRecording If to clear all records from the recording object in use at the start of recording.
+	 */
 	UFUNCTION(BlueprintCallable, Category = "LearningAgents")
-	void BeginRecording();
+	void BeginRecording(bool bReinitializeRecording = true);
 
-	/** Ends the recording of the observations and actions of each agent and stores them in the current recording object. */
+	/**
+	 * Ends the recording of the observations and actions of each agent and stores them in the current recording object.
+	 */
 	UFUNCTION(BlueprintCallable, Category = "LearningAgents")
 	void EndRecording();
 
@@ -112,16 +90,47 @@ public:
 	bool IsRecording() const;
 
 	/**
-	 * While recording, adds the current buffered observations and actions of the added agents to the internal buffer. Call this after 
-	 * GatherObservations and either EvaluateAgentController (if recording a human/AI demonstration) or DecodeAndSampleActions (if recording from 
-	 * another policy).
+	 * While recording, adds the current observations and actions of the added agents to the internal buffer. Call this 
+	 * after ULearningAgentsInteractor::EncodeObservations and either ULearningAgentsController::EncodeActions (if 
+	 * recording a human/AI demonstration) or ULearningAgentsInteractor::DecodeActions (if recording another policy).
 	 */
 	UFUNCTION(BlueprintCallable, Category = "LearningAgents")
 	void AddExperience();
 
 	/** Gets the current recording object. Note: this may be empty until EndRecording has been called. */
 	UFUNCTION(BlueprintPure, Category = "LearningAgents")
-	const ULearningAgentsRecording* GetRecordingAsset() const;
+	const ULearningAgentsRecording* GetCurrentRecording() const;
+
+// ----- Load / Save -----
+public:
+
+	/** Loads the current recording object from a file */
+	UFUNCTION(BlueprintCallable, Category = "LearningAgents", meta = (RelativePath))
+	void LoadRecordingFromFile(const FFilePath& File);
+
+	/** Saves the current recording object to a file */
+	UFUNCTION(BlueprintCallable, BlueprintPure = false, Category = "LearningAgents", meta = (RelativePath))
+	void SaveRecordingToFile(const FFilePath& File) const;
+
+	/** Append to the current recording object from a file. */
+	UFUNCTION(BlueprintCallable, Category = "LearningAgents", meta = (RelativePath))
+	void AppendRecordingFromFile(const FFilePath& File);
+
+	/** Uses the given recording asset. New recordings will be appended to this. */
+	UFUNCTION(BlueprintCallable, Category = "LearningAgents")
+	void UseRecordingAsset(ULearningAgentsRecording* RecordingAsset);
+
+	/** Loads the current recording object from the given recording asset */
+	UFUNCTION(BlueprintCallable, Category = "LearningAgents")
+	void LoadRecordingFromAsset(ULearningAgentsRecording* RecordingAsset);
+
+	/** Saves the current recording object to the given recording asset */
+	UFUNCTION(BlueprintCallable, Category = "LearningAgents", meta = (DevelopmentOnly))
+	void SaveRecordingToAsset(ULearningAgentsRecording* RecordingAsset);
+
+	/** Appends the current recording object to the given recording asset */
+	UFUNCTION(BlueprintCallable, Category = "LearningAgents", meta = (DevelopmentOnly))
+	void AppendRecordingToAsset(ULearningAgentsRecording* RecordingAsset);
 
 
 // ----- Private Data ----- 
@@ -148,13 +157,8 @@ private:
 	/** Basic structure used to buffer the observations and actions of each agent. */
 	struct FAgentRecordBuffer
 	{
-		FAgentRecordBuffer();
-		FAgentRecordBuffer(const int32 InObservationCompatibilityHash, const int32 InActionCompatibilityHash);
-
 		const int32 ChunkSize = 1024;
 		int32 StepNum = 0;
-		int32 ObservationCompatibilityHash = 0;
-		int32 ActionCompatibilityHash = 0;
 		TArray<TLearningArray<2, float>, TInlineAllocator<16>> Observations;
 		TArray<TLearningArray<2, float>, TInlineAllocator<16>> Actions;
 
diff --git a/Engine/Plugins/Experimental/LearningAgents/Source/LearningAgentsTraining/Public/LearningAgentsRecording.h b/Engine/Plugins/Experimental/LearningAgents/Source/LearningAgentsTraining/Public/LearningAgentsRecording.h
--- a/Engine/Plugins/Experimental/LearningAgents/Source/LearningAgentsTraining/Public/LearningAgentsRecording.h
+++ b/Engine/Plugins/Experimental/LearningAgents/Source/LearningAgentsTraining/Public/LearningAgentsRecording.h
@@ -30,26 +30,23 @@ public:
 	UPROPERTY(BlueprintReadOnly, VisibleAnywhere, Category = "LearningAgents")
 	int32 ActionDimNum = 0;
 
-	/** The compatibility hash for the recorded observation vectors */
-	UPROPERTY(BlueprintReadOnly, VisibleAnywhere, Category = "LearningAgents")
-	int32 ObservationCompatibilityHash = 0;
-
-	/** The compatibility hash for the recorded action vectors */
-	UPROPERTY(BlueprintReadOnly, VisibleAnywhere, Category = "LearningAgents")
-	int32 ActionCompatibilityHash = 0;
+	bool Serialize(FArchive& Ar);
 
-	/** Observation data */
-	UPROPERTY(BlueprintReadOnly, Category = "LearningAgents")
-	TArray<float> ObservationData;
-
-	/** Action data */
-	UPROPERTY(BlueprintReadOnly, Category = "LearningAgents")
-	TArray<float> ActionData;
+	TLearningArray<2, float> Observations;
+	TLearningArray<2, float> Actions;
+};
 
+template<>
+struct TStructOpsTypeTraits<FLearningAgentsRecord> : public TStructOpsTypeTraitsBase2<FLearningAgentsRecord>
+{
+	enum
+	{
+		WithSerializer = true,
+	};
 };
 
 /** Data asset representing an array of records. */
-UCLASS(BlueprintType, Blueprintable)
+UCLASS(BlueprintType)
 class LEARNINGAGENTSTRAINING_API ULearningAgentsRecording : public UDataAsset
 {
 	GENERATED_BODY()
@@ -79,43 +76,13 @@ public:
 	void LoadRecordingFromAsset(ULearningAgentsRecording* RecordingAsset);
 
 	/** Saves this recording to the given recording asset */
-	UFUNCTION(BlueprintCallable, Category = "LearningAgents")
+	UFUNCTION(BlueprintCallable, Category = "LearningAgents", meta = (DevelopmentOnly))
 	void SaveRecordingToAsset(ULearningAgentsRecording* RecordingAsset);
 
 	/** Appends this recording to the given recording asset */
-	UFUNCTION(BlueprintCallable, Category = "LearningAgents")
+	UFUNCTION(BlueprintCallable, Category = "LearningAgents", meta = (DevelopmentOnly))
 	void AppendRecordingToAsset(ULearningAgentsRecording* RecordingAsset);
 
-	/** Get the number of records */
-	UFUNCTION(BlueprintPure, Category = "LearningAgents")
-	int32 GetRecordNum() const;
-
-	/** Get the number of steps in a given record */
-	UFUNCTION(BlueprintPure, Category = "LearningAgents")
-	int32 GetRecordStepNum(const int32 Record) const;
-
-	/**
-	 * Get the Observation Vector associated with a particular step of a given recording
-	 * 
-	 * @param OutObservationVector				Output Observation Vector
-	 * @param OutObservationCompatibilityHash	Output Compatibility Hash for the given Observation Vector
-	 * @param Record							Index of the record in the array of records.
-	 * @param Step								Step of the recording
-	 */
-	UFUNCTION(BlueprintCallable, Category = "LearningAgents")
-	void GetObservationVector(TArray<float>& OutObservationVector, int32& OutObservationCompatibilityHash, const int32 Record, const int32 Step);
-
-	/**
-	 * Get the Action Vector associated with a particular step of a given recording
-	 *
-	 * @param OutActionVector					Output Action Vector
-	 * @param OutActionCompatibilityHash		Output Compatibility Hash for the given Action Vector
-	 * @param Record							Index of the record in the array of records.
-	 * @param Step								Step of the recording
-	 */
-	UFUNCTION(BlueprintCallable, Category = "LearningAgents")
-	void GetActionVector(TArray<float>& OutActionVector, int32& OutActionCompatibilityHash, const int32 Record, const int32 Step);
-
 public:
 
 	/** Marks this asset as modified even during PIE */
diff --git a/Engine/Plugins/Experimental/LearningAgents/Source/LearningAgentsTraining/Public/LearningAgentsTrainer.h b/Engine/Plugins/Experimental/LearningAgents/Source/LearningAgentsTraining/Public/LearningAgentsTrainer.h
--- a/Engine/Plugins/Experimental/LearningAgents/Source/LearningAgentsTraining/Public/LearningAgentsTrainer.h
+++ b/Engine/Plugins/Experimental/LearningAgents/Source/LearningAgentsTraining/Public/LearningAgentsTrainer.h
@@ -2,8 +2,8 @@
 
 #pragma once
 
-#include "LearningAgentsManagerListener.h"
-#include "LearningAgentsCompletions.h"
+#include "LearningAgentsManagerComponent.h"
+#include "LearningAgentsCritic.h" // Included for FLearningAgentsCriticSettings()
 
 #include "Templates/SharedPointer.h"
 #include "UObject/ObjectPtr.h"
@@ -13,21 +13,45 @@
 
 namespace UE::Learning
 {
+	struct FAnyCompletion;
+	struct FArrayMap;
 	struct FEpisodeBuffer;
 	struct FReplayBuffer;
 	struct FResetInstanceBuffer;
+	struct FRewardObject;
 	struct FSharedMemoryPPOTrainer;
+	struct FSumReward;
+	struct FCompletionObject;
+	enum class ECompletionMode : uint8;
 	enum class ETrainerDevice : uint8;
 }
 
 class ALearningAgentsManager;
 class ULearningAgentsInteractor;
+class ULearningAgentsCompletion;
+class ULearningAgentsReward;
 class ULearningAgentsPolicy;
 class ULearningAgentsCritic;
 
-DECLARE_DYNAMIC_DELEGATE_RetVal_OneParam(float, FGatherAgentRewardDelegate, const int32, AgentId);
-DECLARE_DYNAMIC_DELEGATE_RetVal_OneParam(ELearningAgentsCompletion, FGatherAgentCompletionDelegate, const int32, AgentId);
-DECLARE_DYNAMIC_DELEGATE_OneParam(FResetAgentEpisodeDelegate, const int32, AgentId);
+/** Completion modes for episodes. */
+UENUM(BlueprintType, Category = "LearningAgents", meta = (ScriptName = "LearningAgentsCompletionEnum"))
+enum class ELearningAgentsCompletion : uint8
+{
+	/** Episode ended early but was still in progress. Critic will be used to estimate final return. */
+	Truncation	UMETA(DisplayName = "Truncation"),
+
+	/** Episode ended early and zero reward was expected for all future steps. */
+	Termination	UMETA(DisplayName = "Termination"),
+};
+
+namespace UE::Learning::Agents
+{
+	/** Get the learning agents completion from the UE::Learning completion. */
+	LEARNINGAGENTSTRAINING_API ELearningAgentsCompletion GetLearningAgentsCompletion(const ECompletionMode CompletionMode);
+
+	/** Get the UE::Learning completion from the learning agents completion. */
+	LEARNINGAGENTSTRAINING_API ECompletionMode GetCompletionMode(const ELearningAgentsCompletion Completion);
+}
 
 /** The configurable settings for a ULearningAgentsTrainer. */
 USTRUCT(BlueprintType, Category = "LearningAgents")
@@ -37,26 +61,19 @@ struct LEARNINGAGENTSTRAINING_API FLearningAgentsTrainerSettings
 
 public:
 
-	/**
-	 * Maximum number of steps recorded in an episode before it is added to the replay buffer. This can generally be left at the default value and 
-	 * does not have a large impact on training.
-	 */
+	/** Completion type to use when the maximum number of steps for an episode is reached. */
+	UPROPERTY(EditAnywhere, Category = "LearningAgents")
+	ELearningAgentsCompletion MaxStepsCompletion = ELearningAgentsCompletion::Truncation;
+
+	/** Max number of steps to take while training before episode automatically completes. */
 	UPROPERTY(EditAnywhere, Category = "LearningAgents", meta = (ClampMin = "1", UIMin = "1"))
-	int32 MaxEpisodeStepNum = 512;
+	int32 MaxStepNum = 300;
 
-	/**
-	 * Maximum number of episodes to record before running a training iteration. An iteration of training will be run when either this or
-	 * MaximumRecordedEpisodesPerIteration is reached. Typical values for this should be around 1000. Setting this too small means there is not 
-	 * enough data each iteration for the system to train. Setting it too large means training will be very slow. 
-	 */
+	/** Maximum number of episodes to record before running a training iteration. */
 	UPROPERTY(EditAnywhere, Category = "LearningAgents", meta = (ClampMin = "1", UIMin = "1"))
 	int32 MaximumRecordedEpisodesPerIteration = 1000;
 
-	/**
-	 * Maximum number of steps to record before running a training iteration. An iteration of training will be run when either this or
-	 * MaximumRecordedEpisodesPerIteration is reached. Typical values for this should be around 10000. Setting this too small means there is not
-	 * enough data each iteration for the system to train. Setting it too large means training will be very slow.
-	 */
+	/** Maximum number of steps to record before running a training iteration. */
 	UPROPERTY(EditAnywhere, Category = "LearningAgents", meta = (ClampMin = "1", UIMin = "1"))
 	int32 MaximumRecordedStepsPerIteration = 10000;
 
@@ -156,6 +173,13 @@ public:
 	UPROPERTY(EditAnywhere, Category = "LearningAgents", meta = (ClampMin = "0.0", ClampMax = "1.0", UIMin = "0.0", UIMax = "1.0"))
 	float WeightDecay = 0.001f;
 
+	/**
+	 * The initial scaling for the weights of the output layer of the neural network. Typically, you would use this to
+	 * scale down the initial actions as it can stabilize the training and speed up convergence.
+	 */
+	UPROPERTY(EditAnywhere, Category = "LearningAgents", meta = (ClampMin = "0.0", ClampMax = "10.0", UIMin = "0.0", UIMax = "10.0"))
+	float InitialActionScale = 0.1f;
+
 	/**
 	 * Batch size to use for training the policy. Smaller values tend to produce better results at the cost of slowing down 
 	 * training. Large batch sizes are much more computationally efficient when training on the GPU.
@@ -203,7 +227,7 @@ public:
 	 * exploration but can make actions very noisy.
 	 */
 	UPROPERTY(EditAnywhere, Category = "LearningAgents", meta = (ClampMin = "0.0", ClampMax = "1.0", UIMin = "0.0", UIMax = "1.0"))
-	float ActionEntropyWeight = 0.0f;
+	float ActionEntropyWeight = 0.01f;
 
 	/**
 	 * This is used in the Generalized Advantage Estimation as what is essentially an exponential smoothing/decay. 
@@ -262,10 +286,6 @@ public:
 	/** If true, TensorBoard logs will be emitted to the intermediate directory. */
 	UPROPERTY(EditAnywhere, Category = "LearningAgents")
 	bool bUseTensorboard = false;
-
-	/** If true, snapshots of the trained networks will be emitted to the intermediate directory. */
-	UPROPERTY(EditAnywhere, Category = "LearningAgents")
-	bool bSaveSnapshots = false;
 };
 
 /** The path settings for the trainer. */
@@ -317,16 +337,21 @@ public:
 };
 
 /**
- * The ULearningAgentsTrainer is the core class for reinforcement learning training. It defines how agents rewards, completions, and resets are 
- * implemented and provides methods for orchestrating the training process.
+ * The ULearningAgentsTrainer is the core class for reinforcement learning training. It has a few responsibilities:
+ *   1) It keeps track of which agents are gathering training data.
+ *   2) It defines how those agents' rewards, completions, and resets are implemented.
+ *   3) It provides methods for orchestrating the training process.
  *
- * To use this class, you need to implement the GatherAgentRewards and GatherAgentCompletions functions, which will define the rewards and penalties 
- * the agent receives and what conditions cause an episode to end.
+ * To use this class, you need to implement the SetupRewards and SetupCompletions functions (as well as their
+ * corresponding SetRewards and SetCompletions functions), which will define the rewards and penalties the agent
+ * receives and what conditions cause an episode to end. Before you can begin training, you need to call SetupTrainer,
+ * which will initialize the underlying data structures, and you need to call AddAgent for each agent you want to gather
+ * training data from.
  *
  * @see ULearningAgentsInteractor to understand how observations and actions work.
  */
-UCLASS(Abstract, HideDropdown, BlueprintType, Blueprintable)
-class LEARNINGAGENTSTRAINING_API ULearningAgentsTrainer : public ULearningAgentsManagerListener
+UCLASS(Abstract, BlueprintType, Blueprintable)
+class LEARNINGAGENTSTRAINING_API ULearningAgentsTrainer : public ULearningAgentsManagerComponent
 {
 	GENERATED_BODY()
 
@@ -338,34 +363,11 @@ public:
 	ULearningAgentsTrainer(FVTableHelper& Helper);
 	virtual ~ULearningAgentsTrainer();
 
-	/** Will automatically call EndTraining if training is still in-progress when the object is destroyed. */
-	virtual void BeginDestroy() override;
-
-	/**
-	 * Constructs the trainer and runs the setup functions for rewards and completions.
-	 * 
-	 * @param InManager			The agent manager we are using.
-	 * @param InInteractor		The agent interactor we are training with.
-	 * @param InPolicy			The policy to be trained.
-	 * @param InCritic			The critic to be trained.
-	 * @param Class				The trainer class
-	 * @param Name				The trainer name
-	 * @param TrainerSettings	The trainer settings to use.
-	 */
-	UFUNCTION(BlueprintCallable, Category = "LearningAgents", meta = (DeterminesOutputType = "Class"))
-	static ULearningAgentsTrainer* MakeTrainer(
-		ULearningAgentsManager* InManager,
-		ULearningAgentsInteractor* InInteractor,
-		ULearningAgentsPolicy* InPolicy,
-		ULearningAgentsCritic* InCritic,
-		TSubclassOf<ULearningAgentsTrainer> Class,
-		const FName Name = TEXT("Trainer"),
-		const FLearningAgentsTrainerSettings& TrainerSettings = FLearningAgentsTrainerSettings());
+	/** Will automatically call EndTraining if training is still in-progress when play is ending. */
+	virtual void EndPlay(const EEndPlayReason::Type EndPlayReason) override;
 
 	/**
-	 * Initializes the trainer and runs the setup functions for rewards and completions.
-	 * 
-	 * @param InManager The agent manager we are using.
+	 * Initializes this object and runs the setup functions for rewards and completions.
 	 * @param InInteractor The agent interactor we are training with.
 	 * @param InPolicy The policy to be trained.
 	 * @param InCritic The critic to be trained.
@@ -373,7 +375,6 @@ public:
 	 */
 	UFUNCTION(BlueprintCallable, Category = "LearningAgents")
 	void SetupTrainer(
-		ULearningAgentsManager* InManager,
 		ULearningAgentsInteractor* InInteractor,
 		ULearningAgentsPolicy* InPolicy,
 		ULearningAgentsCritic* InCritic,
@@ -381,80 +382,72 @@ public:
 
 public: 
 
-	//~ Begin ULearningAgentsManagerListener Interface
-	virtual void OnAgentsAdded_Implementation(const TArray<int32>& AgentIds) override;
-	virtual void OnAgentsRemoved_Implementation(const TArray<int32>& AgentIds) override;
-	virtual void OnAgentsReset_Implementation(const TArray<int32>& AgentIds) override;
-	virtual void OnAgentsManagerTick_Implementation(const TArray<int32>& AgentIds, const float DeltaTime) override;
-	//~ End ULearningAgentsManagerListener Interface
+	//~ Begin ULearningAgentsManagerComponent Interface
+	virtual void OnAgentsAdded(const TArray<int32>& AgentIds) override;
+	virtual void OnAgentsRemoved(const TArray<int32>& AgentIds) override;
+	virtual void OnAgentsReset(const TArray<int32>& AgentIds) override;
+	//~ End ULearningAgentsManagerComponent Interface
 
 // ----- Rewards -----
 public:
+	
+	/**
+	 * During this event, all rewards/penalties should be added to this trainer.
+	 * @see LearningAgentsRewards.h for the list of available rewards/penalties.
+	 */
+	UFUNCTION(BlueprintNativeEvent, Category = "LearningAgents")
+	void SetupRewards();
 
 	/**
-	 * This callback gathers all the reward values for the given set of agents.
-	 *
-	 * @param OutRewards		Output rewards for each agent in AgentIds
-	 * @param AgentIds			Agents to gather rewards for.
+	 * During this event, all rewards/penalties should be set for each agent.
+	 * @param AgentIds The list of agent ids to set rewards/penalties for.
+	 * @see LearningAgentsRewards.h for the list of available rewards/penalties.
+	 * @see GetAgent to get the agent corresponding to each id.
 	 */
-	UFUNCTION(BlueprintNativeEvent, Category = "LearningAgents", Meta = (ForceAsFunction))
-	void GatherAgentRewards(TArray<float>& OutRewards, const TArray<int32>& AgentIds);
+	UFUNCTION(BlueprintNativeEvent, Category = "LearningAgents")
+	void SetRewards(const TArray<int32>& AgentIds);
 
 	/**
-	 * This is a convenience function that can be used to implement `GatherAgentRewards` in terms of a callback which iterates over every agent
-	 * individually and gathers the rewards.
-	 *
-	 * @param OutRewards		Output rewards for each agent in AgentIds
-	 * @param AgentIds			Agents to gather rewards for.
-	 * @param Delegate			The Delegate used to gather the rewards for each agent.
+	 * Used by objects derived from ULearningAgentsReward to add themselves to this trainer during their creation.
+	 * You shouldn't need to call this directly.
 	 */
-	UFUNCTION(BlueprintPure = false, Category = "LearningAgents")
-	void GatherAgentRewardsUsingDelegate(TArray<float>& OutRewards, const TArray<int32>& AgentIds, const FGatherAgentRewardDelegate& Delegate);
+	void AddReward(TObjectPtr<ULearningAgentsReward> Object, const TSharedRef<UE::Learning::FRewardObject>& Reward);
 
 // ----- Completions ----- 
 public:
 
 	/**
-	 * This callback gathers all the completions for the given set of agents.
-	 *
-	 * @param OutCompletions	Output completions for each agent in AgentIds
-	 * @param AgentIds			Agents to gather completions for.
+	 * During this event, all completions should be added to this trainer.
+	 * @see LearningAgentsCompletions.h for the list of available completions.
 	 */
-	UFUNCTION(BlueprintNativeEvent, Category = "LearningAgents", Meta = (ForceAsFunction))
-	void GatherAgentCompletions(TArray<ELearningAgentsCompletion>& OutCompletions, const TArray<int32>& AgentIds);
+	UFUNCTION(BlueprintNativeEvent, Category = "LearningAgents")
+	void SetupCompletions();
 
 	/**
-	 * This is a convenience function that can be used to implement `GatherAgentCompletions` in terms of a callback which iterates over every agent
-	 * individually and gathers the completions.
-	 *
-	 * @param OutCompletions	Output completions for each agent in AgentIds
-	 * @param AgentIds			Agents to gather completions for.
-	 * @param Delegate			The Delegate used to gather the completions for each agent.
+	 * During this event, all completions should be set for each agent.
+	 * @param AgentIds The list of agent ids to set completions for.
+	 * @see LearningAgentsCompletions.h for the list of available completions.
+	 * @see GetAgent to get the agent corresponding to each id.
 	 */
-	UFUNCTION(BlueprintPure = false, Category = "LearningAgents")
-	void GatherAgentCompletionsDelegate(TArray<ELearningAgentsCompletion>& OutCompletions, const TArray<int32>& AgentIds, const FGatherAgentCompletionDelegate& Delegate);
+	UFUNCTION(BlueprintNativeEvent, Category = "LearningAgents")
+	void SetCompletions(const TArray<int32>& AgentIds);
 
+	/**
+	 * Used by objects derived from ULearningAgentsCompletion to add themselves to this trainer during their creation.
+	 * You shouldn't need to call this directly.
+	 */
+	void AddCompletion(TObjectPtr<ULearningAgentsCompletion> Object, const TSharedRef<UE::Learning::FCompletionObject>& Completion);
 
 // ----- Resets ----- 
 public:
 
 	/**
 	 * During this event, all episodes should be reset for each agent.
-	 * 
-	 * @param AgentIds			The ids of the agents that need resetting.
+	 * @param AgentIds The ids of the agents that need resetting.
+	 * @see GetAgent to get the agent corresponding to each id.
 	 */
-	UFUNCTION(BlueprintNativeEvent, Category = "LearningAgents", Meta = (ForceAsFunction))
-	void ResetAgentEpisodes(const TArray<int32>& AgentIds);
-
-	/**
-	 * This is a convenience function that can be used to implement `ResetAgentEpisodes` in terms of a callback which iterates over every agent
-	 * individually and does the reset.
-	 *
-	 * @param AgentIds			The ids of the agents that need resetting.
-	 * @param Delegate			The Delegate used to reset the episode for each agent.
-	 */
-	UFUNCTION(BlueprintPure = false, Category = "LearningAgents")
-	void ResetAgentEpisodesDelegate(const TArray<int32>& AgentIds, const FResetAgentEpisodeDelegate& Delegate);
+	UFUNCTION(BlueprintNativeEvent, Category = "LearningAgents")
+	void ResetEpisodes(const TArray<int32>& AgentIds);
 
 // ----- Training Process -----
 public:
@@ -466,16 +459,20 @@ public:
 	/**
 	 * Begins the training process with the provided settings.
 	 * 
-	 * @param TrainerTrainingSettings	The settings for this training run.
-	 * @param TrainerGameSettings		The settings that will affect the game's simulation.
-	 * @param TrainerPathSettings		The path settings used by the trainer.
-	 * @param bResetAgentsOnBegin		If true, reset all agents at the beginning of training.
+	 * @param TrainerTrainingSettings The settings for this training run.
+	 * @param TrainerGameSettings The settings that will affect the game's simulation.
+	 * @param TrainerPathSettings The path settings used by the trainer.
+	 * @param bReinitializePolicyNetwork If true, reinitialize the policy. Set this to false if your policy is pre-trained, e.g. with imitation learning.
+	 * @param bReinitializeCriticNetwork If true, reinitialize the critic. Set this to false if your critic is pre-trained.
+	 * @param bResetAgentsOnBegin If true, reset all agents at the beginning of training.
 	 */
 	UFUNCTION(BlueprintCallable, Category = "LearningAgents")
 	void BeginTraining(
 		const FLearningAgentsTrainerTrainingSettings& TrainerTrainingSettings = FLearningAgentsTrainerTrainingSettings(),
 		const FLearningAgentsTrainerGameSettings& TrainerGameSettings = FLearningAgentsTrainerGameSettings(),
 		const FLearningAgentsTrainerPathSettings& TrainerPathSettings = FLearningAgentsTrainerPathSettings(),
+		const bool bReinitializePolicyNetwork = true,
+		const bool bReinitializeCriticNetwork = true,
 		const bool bResetAgentsOnBegin = true);
 
 	/** Stops the training process. */
@@ -488,7 +485,7 @@ public:
 	 * the next state before evaluating the rewards.
 	 */
 	UFUNCTION(BlueprintCallable, Category = "LearningAgents")
-	void GatherRewards();
+	void EvaluateRewards();
 
 	/**
 	 * Call this function when it is time to evaluate the completions for your agents. This should be done at the beginning
@@ -496,52 +493,40 @@ public:
 	 * the next state before evaluating the completions.
 	 */
 	UFUNCTION(BlueprintCallable, Category = "LearningAgents")
-	void GatherCompletions();
+	void EvaluateCompletions();
 
 	/**
 	 * Call this function at the end of each step of your training loop. This takes the current observations/actions/
 	 * rewards and moves them into the episode experience buffer. All agents with full episode buffers or those which
 	 * have been signaled complete will be reset. If enough experience is gathered, it will be sent to the training 
 	 * process and an iteration of training will be run and the updated policy will be synced back.
-	 *
-	 * @param bResetAgentsOnUpdate				If true, reset all agents whenever an updated policy is received.
 	 */
 	UFUNCTION(BlueprintCallable, Category = "LearningAgents")
-	void ProcessExperience(const bool bResetAgentsOnUpdate = true);
+	void ProcessExperience();
 
 	/**
 	 * Convenience function that runs a basic training loop. If training has not been started, it will start it, and 
-	 * then call RunInference. On each following call to this function, it will call GatherRewards, 
-	 * GatherCompletions, and ProcessExperience, followed by RunInference.
+	 * then call RunInference. On each following call to this function, it will call EvaluateRewards, 
+	 * EvaluateCompletions, and ProcessExperience, followed by RunInference.
 	 * 
-	 * @param TrainerTrainingSettings			The settings for this training run.
-	 * @param TrainerGameSettings				The settings that will affect the game's simulation.
-	 * @param TrainerPathSettings				The path settings used by the trainer.
-	 * @param bResetAgentsOnBegin				If true, reset all agents at the beginning of training.
-	 * @param bResetAgentsOnUpdate				If true, reset all agents whenever an updated policy is received.
+	 * @param TrainerTrainingSettings The settings for this training run.
+	 * @param TrainerGameSettings The settings that will affect the game's simulation.
+	 * @param TrainerPathSettings The path settings used by the trainer.
+	 * @param bReinitializePolicyNetwork If true, reinitialize the policy. Set this to false if your policy is pre-trained, e.g. with imitation learning.
+	 * @param bReinitializeCriticNetwork If true, reinitialize the critic. Set this to false if your critic is pre-trained.
+	 * @param bResetAgentsOnBegin If true, reset all agents at the beginning of training.
 	 */
 	UFUNCTION(BlueprintCallable, Category = "LearningAgents")
 	void RunTraining(
 		const FLearningAgentsTrainerTrainingSettings& TrainerTrainingSettings = FLearningAgentsTrainerTrainingSettings(),
 		const FLearningAgentsTrainerGameSettings& TrainerGameSettings = FLearningAgentsTrainerGameSettings(),
 		const FLearningAgentsTrainerPathSettings& TrainerPathSettings = FLearningAgentsTrainerPathSettings(),
-		const bool bResetAgentsOnBegin = true,
-		const bool bResetAgentsOnUpdate = true);
-
-	/**
-	 * Returns true if GatherRewards has been called and the reward already set for the given agent.
-	 */
-	UFUNCTION(BlueprintPure, Category = "LearningAgents", meta = (AgentId = "-1"))
-	bool HasReward(const int32 AgentId) const;
-
-	/**
-	 * Returns true if GatherCompletions has been called and the completion already set for the given agent.
-	 */
-	UFUNCTION(BlueprintPure, Category = "LearningAgents", meta = (AgentId = "-1"))
-	bool HasCompletion(const int32 AgentId) const;
+		const bool bReinitializePolicyNetwork = true,
+		const bool bReinitializeCriticNetwork = true,
+		const bool bResetAgentsOnBegin = true);
 
 	/**
-	 * Gets the current reward for an agent. Should be called only after GatherRewards.
+	 * Gets the current reward for an agent according to the critic. Should be called only after EvaluateRewards.
 	 *
 	 * @param AgentId	The AgentId to look-up the reward for
 	 * @returns			The reward
@@ -550,31 +535,15 @@ public:
 	float GetReward(const int32 AgentId) const;
 
 	/**
-	 * Gets the current completion for an agent. Should be called only after GatherCompletions.
-	 *
-	 * @param AgentId	The AgentId to look-up the completion for
-	 * @returns			The completion type
-	 */
-	UFUNCTION(BlueprintPure, Category = "LearningAgents", meta = (AgentId = "-1"))
-	ELearningAgentsCompletion GetCompletion(const int32 AgentId) const;
-
-	/**
-	 * Gets the current elapsed episode time for the given agent.
+	 * Gets if the agent will complete the episode or not according to the given set of completions. Should be called 
+	 * only after EvaluateCompletions.
 	 *
-	 * @param AgentId	The AgentId to look-up the episode time for
-	 * @returns			The elapsed episode time
+	 * @param AgentId		The AgentId to look-up the completion for
+	 * @param OutCompletion	The completion type if the agent will complete the episode
+	 * @returns				If the agent will complete the episode
 	 */
 	UFUNCTION(BlueprintPure, Category = "LearningAgents", meta = (AgentId = "-1"))
-	float GetEpisodeTime(const int32 AgentId) const;
-
-	/**
-	 * Gets the number of step recorded in an episode for the given agent.
-	 *
-	 * @param AgentId	The AgentId to look-up the number of recorded episode steps for
-	 * @returns			The number of recorded episode steps
-	 */
-	UFUNCTION(BlueprintPure, Category = "LearningAgents", meta = (AgentId = "-1"))
-	int32 GetEpisodeStepNum(const int32 AgentId) const;
+	bool IsCompleted(const int32 AgentId, ELearningAgentsCompletion& OutCompletion) const;
 
 	/**
 	 * Returns true if the trainer has failed to communicate with the external training process. This can be used in
@@ -584,7 +553,7 @@ public:
 	 */
 	UFUNCTION(BlueprintPure, Category = "LearningAgents")
 	bool HasTrainingFailed() const;
-
+	
 // ----- Private Data ----- 
 private:
 
@@ -611,47 +580,46 @@ private:
 	UPROPERTY(VisibleAnywhere, Transient, Category = "LearningAgents")
 	bool bHasTrainingFailed = false;
 
-	/** Callback Reward Output */
-	TArray<float> RewardBuffer;
-
-	/** Callback Completion Output */
-	TArray<ELearningAgentsCompletion> CompletionBuffer;
-
-	/** Reward Buffer */
-	TLearningArray<1, float> Rewards;
-	
-	/** Agent Completions Buffer */
-	TLearningArray<1, UE::Learning::ECompletionMode> AgentCompletions;
+	/** The list of current reward objects. */
+	UPROPERTY(VisibleAnywhere, Transient, Category = "LearningAgents")
+	TArray<TObjectPtr<ULearningAgentsReward>> RewardObjects;
 
-	/** Episode Completions Buffer */
-	TLearningArray<1, UE::Learning::ECompletionMode> EpisodeCompletions;
+	/** The list of current completion objects. */
+	UPROPERTY(VisibleAnywhere, Transient, Category = "LearningAgents")
+	TArray<TObjectPtr<ULearningAgentsCompletion>> CompletionObjects;
 
-	/** All Completions Buffer */
-	TLearningArray<1, UE::Learning::ECompletionMode> AllCompletions;
+	TArray<TSharedRef<UE::Learning::FRewardObject>, TInlineAllocator<16>> RewardFeatures;
+	TArray<TSharedRef<UE::Learning::FCompletionObject>, TInlineAllocator<16>> CompletionFeatures;
 
-	/** Agent episode times */
-	TLearningArray<1, float> EpisodeTimes;
+	TSharedPtr<UE::Learning::FSumReward> Rewards;
+	TSharedPtr<UE::Learning::FAnyCompletion> Completions;
 
 	TUniquePtr<UE::Learning::FEpisodeBuffer> EpisodeBuffer;
 	TUniquePtr<UE::Learning::FReplayBuffer> ReplayBuffer;
 	TUniquePtr<UE::Learning::FResetInstanceBuffer> ResetBuffer;
 	TUniquePtr<UE::Learning::FSharedMemoryPPOTrainer> Trainer;
 	
+	ELearningAgentsCompletion MaxStepsCompletion = ELearningAgentsCompletion::Truncation;
+
 	float TrainerTimeout = 10.0f;
 
 	void DoneTraining();
 
+// ----- Private Iteration Checks ----- 
+private:
+
 	/** Number of times rewards have been evaluated for all agents */
-	TLearningArray<1, uint64, TInlineAllocator<32>> RewardIteration;
+	TLearningArray<1, uint64, TInlineAllocator<32>> RewardEvaluatedAgentIteration;
 
 	/** Number of times completions have been evaluated for all agents */
-	TLearningArray<1, uint64, TInlineAllocator<32>> CompletionIteration;
+	TLearningArray<1, uint64, TInlineAllocator<32>> CompletionEvaluatedAgentIteration;
 
 	/** Temp buffers used to record the set of agents that are valid for training */
 	TArray<int32> ValidAgentIds;
 	TArray<int32> FinalValidAgentIds;
 	UE::Learning::FIndexSet ValidAgentSet;
 	UE::Learning::FIndexSet FinalValidAgentSet;
+	TBitArray<TInlineAllocator<32>> ValidAgentStatus;
 
 // ----- Private Recording of GameSettings ----- 
 private:
diff --git a/Engine/Plugins/Experimental/LearningAgents/Source/LearningTraining/Private/LearningExperience.cpp b/Engine/Plugins/Experimental/LearningAgents/Source/LearningTraining/Private/LearningExperience.cpp
--- a/Engine/Plugins/Experimental/LearningAgents/Source/LearningTraining/Private/LearningExperience.cpp
+++ b/Engine/Plugins/Experimental/LearningAgents/Source/LearningTraining/Private/LearningExperience.cpp
@@ -257,7 +257,18 @@ namespace UE::Learning
 
 	namespace Experience
 	{
-		void GatherExperienceUntilReplayBufferFull(
+		void GatherExperienceReset(
+			FReplayBuffer& ReplayBuffer,
+			FEpisodeBuffer& EpisodeBuffer,
+			FResetInstanceBuffer& ResetBuffer,
+			const FIndexSet Instances)
+		{
+			ReplayBuffer.Reset();
+			EpisodeBuffer.Reset(Instances);
+			ResetBuffer.SetResetInstances(Instances);
+		}
+
+		bool GatherExperienceIteration(
 			FReplayBuffer& ReplayBuffer,
 			FEpisodeBuffer& EpisodeBuffer,
 			FResetInstanceBuffer& ResetBuffer,
@@ -267,8 +278,7 @@ namespace UE::Learning
 			TLearningArrayView<2, float> MemoryStateVectorBuffer,
 			TLearningArrayView<1, float> RewardBuffer,
 			TLearningArrayView<1, ECompletionMode> CompletionBuffer,
-			TLearningArrayView<1, ECompletionMode> EpisodeCompletionBuffer,
-			TLearningArrayView<1, ECompletionMode> AllCompletionBuffer,
+			const ECompletionMode EpisodeEndCompletionMode,
 			const TFunctionRef<void(const FIndexSet Instances)> ResetFunction,
 			const TFunctionRef<void(const FIndexSet Instances)> ObservationFunction,
 			const TFunctionRef<void(const FIndexSet Instances)> PolicyFunction,
@@ -278,90 +288,137 @@ namespace UE::Learning
 			const TFunctionRef<void(const FIndexSet Instances)> CompletionFunction,
 			const FIndexSet Instances)
 		{
-			UE_LEARNING_TRACE_CPUPROFILER_EVENT_SCOPE(Learning::Experience::GatherExperienceUntilReplayBufferFull);
-
-			// Reset Everything
+			UE_LEARNING_TRACE_CPUPROFILER_EVENT_SCOPE(Learning::Experience::GatherExperienceIteration);
 
-			ReplayBuffer.Reset();
-			EpisodeBuffer.Reset(Instances);
-			ResetFunction(Instances);
-
-			while (true)
+			if (ResetBuffer.GetResetInstanceNum() > 0)
 			{
-				// Encode Observations
+				// Reset Environment
 
-				ObservationFunction(Instances);
+				ResetFunction(ResetBuffer.GetResetInstances());
+			}
 
-				// Evaluate Policy
+			// Encode Observations
 
-				PolicyFunction(Instances);
+			ObservationFunction(Instances);
+			Array::Check(ObservationVectorBuffer, Instances);
 
-				// Decode Actions
+			// Evaluate Policy
 
-				ActionFunction(Instances);
+			PolicyFunction(Instances);
 
-				// Update Environment
+			// Decode Actions
 
-				UpdateFunction(Instances);
+			Array::Check(ActionVectorBuffer, Instances);
+			ActionFunction(Instances);
 
-				// Compute Rewards
+			// Update Environment
 
-				RewardFunction(Instances);
+			UpdateFunction(Instances);
 
-				// Push to Experience Buffer
+			// Compute Rewards
 
-				EpisodeBuffer.Push(
-					ObservationVectorBuffer,
-					ActionVectorBuffer,
-					PreEvaluationMemoryStateVectorBuffer,
-					RewardBuffer,
-					Instances);
+			RewardFunction(Instances);
+			Array::Check(RewardBuffer, Instances);
 
-				// Evaluate Completions
+			// Push to Experience Buffer
 
-				CompletionFunction(Instances);
+			EpisodeBuffer.Push(
+				ObservationVectorBuffer,
+				ActionVectorBuffer,
+				PreEvaluationMemoryStateVectorBuffer,
+				RewardBuffer,
+				Instances);
 
-				Completion::EvaluateEndOfEpisodeCompletions(
-					EpisodeCompletionBuffer,
-					EpisodeBuffer.GetEpisodeStepNums(),
-					EpisodeBuffer.GetMaxStepNum(),
-					Instances);
+			// Evaluate Completions
 
-				for (const int32 Instance : Instances)
-				{
-					AllCompletionBuffer[Instance] = Completion::Or(CompletionBuffer[Instance], EpisodeCompletionBuffer[Instance]);
-				}
+			CompletionFunction(Instances);
 
-				ResetBuffer.SetResetInstancesFromCompletions(AllCompletionBuffer, Instances);
+			Completion::EvaluateEndOfEpisodeCompletions(
+				CompletionBuffer,
+				EpisodeBuffer.GetEpisodeStepNums(),
+				EpisodeBuffer.GetMaxStepNum(),
+				EpisodeEndCompletionMode,
+				Instances);
 
-				if (ResetBuffer.GetResetInstanceNum() == 0)
-				{
-					continue;
-				}
+			ResetBuffer.SetResetInstancesFromCompletions(CompletionBuffer, Instances);
 
-				// Evaluate Observations again for instances that are completed
+			// Evaluate Observations again for instances that are completed
 
+			if (ResetBuffer.GetResetInstanceNum() > 0)
+			{
 				ObservationFunction(ResetBuffer.GetResetInstances());
 
-				// Push completed instances to Replay Buffer and return if full
+				// Push to Replay Buffer
 
-				if (ReplayBuffer.AddEpisodes(
-					AllCompletionBuffer,
+				const bool bReplayBufferFull = ReplayBuffer.AddEpisodes(
+					CompletionBuffer,
 					ObservationVectorBuffer,
 					MemoryStateVectorBuffer,
 					EpisodeBuffer,
-					ResetBuffer.GetResetInstances()))
-				{
-					return;
-				}
+					ResetBuffer.GetResetInstances());
 
-				// Just reset Episode Buffer for instances who reached the maximum episode length
-				ResetBuffer.SetResetInstancesFromCompletions(EpisodeCompletionBuffer, Instances);
 				EpisodeBuffer.Reset(ResetBuffer.GetResetInstances());
 
-				// Call Reset Function for instances which signaled a completion
-				ResetBuffer.SetResetInstancesFromCompletions(CompletionBuffer, Instances);
-				ResetFunction(ResetBuffer.GetResetInstances());
+				return bReplayBufferFull;
+			}
+			else
+			{
+				return false;
+			}
+		}
+
+		void GatherExperienceUntilReplayBufferFull(
+			FReplayBuffer& ReplayBuffer,
+			FEpisodeBuffer& EpisodeBuffer,
+			FResetInstanceBuffer& ResetBuffer,
+			TLearningArrayView<2, float> ObservationVectorBuffer,
+			TLearningArrayView<2, float> ActionVectorBuffer,
+			TLearningArrayView<2, float> PreEvaluationMemoryStateVectorBuffer,
+			TLearningArrayView<2, float> MemoryStateVectorBuffer,
+			TLearningArrayView<1, float> RewardBuffer,
+			TLearningArrayView<1, ECompletionMode> CompletionBuffer,
+			const ECompletionMode EpisodeEndCompletionMode,
+			const TFunctionRef<void(const FIndexSet Instances)> ResetFunction,
+			const TFunctionRef<void(const FIndexSet Instances)> ObservationFunction,
+			const TFunctionRef<void(const FIndexSet Instances)> PolicyFunction,
+			const TFunctionRef<void(const FIndexSet Instances)> ActionFunction,
+			const TFunctionRef<void(const FIndexSet Instances)> UpdateFunction,
+			const TFunctionRef<void(const FIndexSet Instances)> RewardFunction,
+			const TFunctionRef<void(const FIndexSet Instances)> CompletionFunction,
+			const FIndexSet Instances)
+		{
+			UE_LEARNING_TRACE_CPUPROFILER_EVENT_SCOPE(Learning::Experience::GatherExperienceUntilReplayBufferFull);
+
+			bool bReplayBufferFull = false;
+
+			GatherExperienceReset(ReplayBuffer, EpisodeBuffer, ResetBuffer, Instances);
+
+			while (true)
+			{
+				bReplayBufferFull = GatherExperienceIteration(
+					ReplayBuffer,
+					EpisodeBuffer,
+					ResetBuffer,
+					ObservationVectorBuffer,
+					ActionVectorBuffer,
+					PreEvaluationMemoryStateVectorBuffer,
+					MemoryStateVectorBuffer,
+					RewardBuffer,
+					CompletionBuffer,
+					EpisodeEndCompletionMode,
+					ResetFunction,
+					ObservationFunction,
+					PolicyFunction,
+					ActionFunction,
+					UpdateFunction,
+					RewardFunction,
+					CompletionFunction,
+					Instances);
+
+				if (bReplayBufferFull)
+				{
+					break;
+				}
 			}
 		}
 
diff --git a/Engine/Plugins/Experimental/LearningAgents/Source/LearningTraining/Private/LearningImitationTrainer.cpp b/Engine/Plugins/Experimental/LearningAgents/Source/LearningTraining/Private/LearningImitationTrainer.cpp
--- a/Engine/Plugins/Experimental/LearningAgents/Source/LearningTraining/Private/LearningImitationTrainer.cpp
+++ b/Engine/Plugins/Experimental/LearningAgents/Source/LearningTraining/Private/LearningImitationTrainer.cpp
@@ -115,14 +115,10 @@ namespace UE::Learning
 		const int32 ObservationDimNum,
 		const int32 ActionDimNum,
 		const int32 MemoryStateDimNum,
-		const ULearningNeuralNetworkData& PolicyNetwork,
-		const ULearningNeuralNetworkData& EncoderNetwork,
-		const ULearningNeuralNetworkData& DecoderNetwork,
-		const Observation::FSchema& ObservationSchema,
-		const Observation::FSchemaElement& ObservationSchemaElement,
-		const Action::FSchema& ActionSchema,
-		const Action::FSchemaElement& ActionSchemaElement,
+		const INeuralNetwork& PolicyNetwork,
 		const FImitationTrainerTrainingSettings& TrainingSettings,
+		const FImitationTrainerNetworkSettings& NetworkSettings,
+		const EImitationTrainerFlags TrainerFlags,
 		const ESubprocessFlags TrainingProcessFlags,
 		const ELogSetting LogSettings)
 	{
@@ -132,9 +128,7 @@ namespace UE::Learning
 
 		// Allocate Shared Memory
 
-		Policy = SharedMemory::Allocate<1, uint8>({ PolicyNetwork.GetSnapshotByteNum() });
-		Encoder = SharedMemory::Allocate<1, uint8>({ EncoderNetwork.GetSnapshotByteNum() });
-		Decoder = SharedMemory::Allocate<1, uint8>({ DecoderNetwork.GetSnapshotByteNum() });
+		Policy = SharedMemory::Allocate<1, uint8>({ PolicyNetwork.GetSerializationByteNum() });
 		Controls = SharedMemory::Allocate<1, volatile int32>({ SharedMemoryTraining::GetControlNum() });
 		EpisodeStarts = SharedMemory::Allocate<1, int32>({ MaxEpisodeNum });
 		EpisodeLengths = SharedMemory::Allocate<1, int32>({ MaxEpisodeNum });
@@ -164,39 +158,38 @@ namespace UE::Learning
 		ConfigObject->SetStringField(TEXT("SitePackagesPath"), *FileManager.ConvertToAbsolutePathForExternalAppForRead(*SitePackagesPath));
 		ConfigObject->SetStringField(TEXT("IntermediatePath"), *FileManager.ConvertToAbsolutePathForExternalAppForRead(*IntermediatePath));
 
+		ConfigObject->SetStringField(TEXT("PolicyNetworkClass"), PolicyNetwork.GetPythonClassName());
+
 		ConfigObject->SetStringField(TEXT("PolicyGuid"), *Policy.Guid.ToString(EGuidFormats::DigitsWithHyphensInBraces));
-		ConfigObject->SetStringField(TEXT("EncoderGuid"), *Encoder.Guid.ToString(EGuidFormats::DigitsWithHyphensInBraces));
-		ConfigObject->SetStringField(TEXT("DecoderGuid"), *Decoder.Guid.ToString(EGuidFormats::DigitsWithHyphensInBraces));
 		ConfigObject->SetStringField(TEXT("ControlsGuid"), *Controls.Guid.ToString(EGuidFormats::DigitsWithHyphensInBraces));
 		ConfigObject->SetStringField(TEXT("EpisodeStartsGuid"), *EpisodeStarts.Guid.ToString(EGuidFormats::DigitsWithHyphensInBraces));
 		ConfigObject->SetStringField(TEXT("EpisodeLengthsGuid"), *EpisodeLengths.Guid.ToString(EGuidFormats::DigitsWithHyphensInBraces));
 		ConfigObject->SetStringField(TEXT("ObservationsGuid"), *Observations.Guid.ToString(EGuidFormats::DigitsWithHyphensInBraces));
 		ConfigObject->SetStringField(TEXT("ActionsGuid"), *Actions.Guid.ToString(EGuidFormats::DigitsWithHyphensInBraces));
 
-		ConfigObject->SetObjectField(TEXT("ObservationSchema"), Trainer::ConvertObservationSchemaToJSON(ObservationSchema, ObservationSchemaElement));
-		ConfigObject->SetObjectField(TEXT("ActionSchema"), Trainer::ConvertActionSchemaToJSON(ActionSchema, ActionSchemaElement));
 		ConfigObject->SetNumberField(TEXT("ObservationVectorDimensionNum"), ObservationDimNum);
 		ConfigObject->SetNumberField(TEXT("ActionVectorDimensionNum"), ActionDimNum);
 		ConfigObject->SetNumberField(TEXT("MemoryStateVectorDimensionNum"), MemoryStateDimNum);
 		ConfigObject->SetNumberField(TEXT("MaxEpisodeNum"), MaxEpisodeNum);
 		ConfigObject->SetNumberField(TEXT("MaxStepNum"), MaxStepNum);
 
-		ConfigObject->SetNumberField(TEXT("PolicyNetworkByteNum"), PolicyNetwork.GetSnapshotByteNum());
-		ConfigObject->SetNumberField(TEXT("EncoderNetworkByteNum"), EncoderNetwork.GetSnapshotByteNum());
-		ConfigObject->SetNumberField(TEXT("DecoderNetworkByteNum"), DecoderNetwork.GetSnapshotByteNum());
+		ConfigObject->SetNumberField(TEXT("PolicyNetworkByteNum"), PolicyNetwork.GetSerializationByteNum());
+		ConfigObject->SetNumberField(TEXT("PolicyActionNoiseMin"), NetworkSettings.PolicyActionNoiseMin);
+		ConfigObject->SetNumberField(TEXT("PolicyActionNoiseMax"), NetworkSettings.PolicyActionNoiseMax);
 
 		ConfigObject->SetNumberField(TEXT("IterationNum"), TrainingSettings.IterationNum);
-		ConfigObject->SetNumberField(TEXT("LearningRate"), TrainingSettings.LearningRate);
+		ConfigObject->SetNumberField(TEXT("LearningRatePolicy"), TrainingSettings.LearningRatePolicy);
 		ConfigObject->SetNumberField(TEXT("LearningRateDecay"), TrainingSettings.LearningRateDecay);
 		ConfigObject->SetNumberField(TEXT("WeightDecay"), TrainingSettings.WeightDecay);
-		ConfigObject->SetNumberField(TEXT("BatchSize"), TrainingSettings.BatchSize);
-		ConfigObject->SetNumberField(TEXT("Window"), TrainingSettings.Window);
-		ConfigObject->SetNumberField(TEXT("ActionRegularizationWeight"), TrainingSettings.ActionRegularizationWeight);
-		ConfigObject->SetNumberField(TEXT("ActionEntropyWeight"), TrainingSettings.ActionEntropyWeight);
+		ConfigObject->SetNumberField(TEXT("InitialActionScale"), TrainingSettings.InitialActionScale);
+		ConfigObject->SetNumberField(TEXT("InitialMemoryScale"), TrainingSettings.InitialMemoryScale);
+		ConfigObject->SetNumberField(TEXT("PolicyBatchSize"), TrainingSettings.PolicyBatchSize);
+		ConfigObject->SetNumberField(TEXT("PolicyWindow"), TrainingSettings.PolicyWindow);
 		ConfigObject->SetNumberField(TEXT("Seed"), TrainingSettings.Seed);
 		ConfigObject->SetStringField(TEXT("Device"), Trainer::GetDeviceString(TrainingSettings.Device));
 		ConfigObject->SetBoolField(TEXT("UseTensorBoard"), TrainingSettings.bUseTensorboard);
-		ConfigObject->SetBoolField(TEXT("SaveSnapshots"), TrainingSettings.bSaveSnapshots);
+
+		ConfigObject->SetBoolField(TEXT("UseInitialPolicyNetwork"), (bool)(TrainerFlags & EImitationTrainerFlags::UseInitialPolicyNetwork));
 
 		ConfigObject->SetBoolField(TEXT("LoggingEnabled"), LogSettings == ELogSetting::Silent ? false : true);
 
@@ -244,95 +237,29 @@ namespace UE::Learning
 	}
 
 	ETrainerResponse FSharedMemoryImitationTrainer::RecvPolicy(
-		ULearningNeuralNetworkData& OutNetwork,
+		INeuralNetwork& OutNetwork,
 		const float Timeout,
 		FRWLock* NetworkLock,
 		const ELogSetting LogSettings)
 	{
-		return SharedMemoryTraining::RecvNetwork(
+		return SharedMemoryTraining::RecvPolicy(
 			Controls.View,
 			OutNetwork,
-			SharedMemoryTraining::EControls::PolicySignal,
 			Policy.View,
 			Timeout,
 			NetworkLock,
 			LogSettings);
 	}
 
-	ETrainerResponse FSharedMemoryImitationTrainer::RecvEncoder(
-		ULearningNeuralNetworkData& OutNetwork,
-		const float Timeout,
-		FRWLock* NetworkLock,
-		const ELogSetting LogSettings)
-	{
-		return SharedMemoryTraining::RecvNetwork(
-			Controls.View,
-			OutNetwork,
-			SharedMemoryTraining::EControls::EncoderSignal,
-			Encoder.View,
-			Timeout,
-			NetworkLock,
-			LogSettings);
-	}
-
-	ETrainerResponse FSharedMemoryImitationTrainer::RecvDecoder(
-		ULearningNeuralNetworkData& OutNetwork,
-		const float Timeout,
-		FRWLock* NetworkLock,
-		const ELogSetting LogSettings)
-	{
-		return SharedMemoryTraining::RecvNetwork(
-			Controls.View,
-			OutNetwork,
-			SharedMemoryTraining::EControls::DecoderSignal,
-			Decoder.View,
-			Timeout,
-			NetworkLock,
-			LogSettings);
-	}
-
 	ETrainerResponse FSharedMemoryImitationTrainer::SendPolicy(
-		const ULearningNeuralNetworkData& Network,
+		const INeuralNetwork& Network,
 		const float Timeout,
 		FRWLock* NetworkLock,
 		const ELogSetting LogSettings)
 	{
-		return SharedMemoryTraining::SendNetwork(
+		return SharedMemoryTraining::SendPolicy(
 			Controls.View,
 			Policy.View,
-			SharedMemoryTraining::EControls::PolicySignal,
-			Network,
-			Timeout,
-			NetworkLock,
-			LogSettings);
-	}
-
-	ETrainerResponse FSharedMemoryImitationTrainer::SendEncoder(
-		const ULearningNeuralNetworkData& Network,
-		const float Timeout,
-		FRWLock* NetworkLock,
-		const ELogSetting LogSettings)
-	{
-		return SharedMemoryTraining::SendNetwork(
-			Controls.View,
-			Encoder.View,
-			SharedMemoryTraining::EControls::EncoderSignal,
-			Network,
-			Timeout,
-			NetworkLock,
-			LogSettings);
-	}
-
-	ETrainerResponse FSharedMemoryImitationTrainer::SendDecoder(
-		const ULearningNeuralNetworkData& Network,
-		const float Timeout,
-		FRWLock* NetworkLock,
-		const ELogSetting LogSettings)
-	{
-		return SharedMemoryTraining::SendNetwork(
-			Controls.View,
-			Decoder.View,
-			SharedMemoryTraining::EControls::DecoderSignal,
 			Network,
 			Timeout,
 			NetworkLock,
@@ -364,8 +291,6 @@ namespace UE::Learning
 	void FSharedMemoryImitationTrainer::Deallocate()
 	{
 		SharedMemory::Deallocate(Policy);
-		SharedMemory::Deallocate(Encoder);
-		SharedMemory::Deallocate(Decoder);
 		SharedMemory::Deallocate(Controls);
 		SharedMemory::Deallocate(Observations);
 		SharedMemory::Deallocate(Actions);
@@ -542,17 +467,13 @@ namespace UE::Learning
 		const int32 ObservationDimNum,
 		const int32 ActionDimNum,
 		const int32 MemoryStateDimNum,
-		const ULearningNeuralNetworkData& PolicyNetwork,
-		const ULearningNeuralNetworkData& EncoderNetwork,
-		const ULearningNeuralNetworkData& DecoderNetwork,
-		const Observation::FSchema& ObservationSchema,
-		const Observation::FSchemaElement& ObservationSchemaElement,
-		const Action::FSchema& ActionSchema,
-		const Action::FSchemaElement& ActionSchemaElement,
+		const INeuralNetwork& PolicyNetwork,
 		const TCHAR* IpAddress,
 		const uint32 Port,
 		const float Timeout,
-		const FImitationTrainerTrainingSettings& TrainingSettings)
+		const FImitationTrainerTrainingSettings& TrainingSettings,
+		const FImitationTrainerNetworkSettings& NetworkSettings,
+		const EImitationTrainerFlags TrainerFlags)
 	{
 		// Write Config
 
@@ -566,30 +487,31 @@ namespace UE::Learning
 		ConfigObject->SetStringField(TEXT("TrainerType"), TrainerType);
 		ConfigObject->SetStringField(TEXT("TimeStamp"), *TimeStamp);
 
-		ConfigObject->SetObjectField(TEXT("ObservationSchema"), Trainer::ConvertObservationSchemaToJSON(ObservationSchema, ObservationSchemaElement));
-		ConfigObject->SetObjectField(TEXT("ActionSchema"), Trainer::ConvertActionSchemaToJSON(ActionSchema, ActionSchemaElement));
+		ConfigObject->SetStringField(TEXT("PolicyNetworkClass"), PolicyNetwork.GetPythonClassName());
+
 		ConfigObject->SetNumberField(TEXT("ObservationVectorDimensionNum"), ObservationDimNum);
 		ConfigObject->SetNumberField(TEXT("ActionVectorDimensionNum"), ActionDimNum);
 		ConfigObject->SetNumberField(TEXT("MemoryStateVectorDimensionNum"), MemoryStateDimNum);
 		ConfigObject->SetNumberField(TEXT("MaxEpisodeNum"), MaxEpisodeNum);
 		ConfigObject->SetNumberField(TEXT("MaxStepNum"), MaxStepNum);
 
-		ConfigObject->SetNumberField(TEXT("PolicyNetworkByteNum"), PolicyNetwork.GetSnapshotByteNum());
-		ConfigObject->SetNumberField(TEXT("EncoderNetworkByteNum"), EncoderNetwork.GetSnapshotByteNum());
-		ConfigObject->SetNumberField(TEXT("DecoderNetworkByteNum"), DecoderNetwork.GetSnapshotByteNum());
+		ConfigObject->SetNumberField(TEXT("PolicyNetworkByteNum"), PolicyNetwork.GetSerializationByteNum());
+		ConfigObject->SetNumberField(TEXT("PolicyActionNoiseMin"), NetworkSettings.PolicyActionNoiseMin);
+		ConfigObject->SetNumberField(TEXT("PolicyActionNoiseMax"), NetworkSettings.PolicyActionNoiseMax);
 
 		ConfigObject->SetNumberField(TEXT("IterationNum"), TrainingSettings.IterationNum);
-		ConfigObject->SetNumberField(TEXT("LearningRate"), TrainingSettings.LearningRate);
+		ConfigObject->SetNumberField(TEXT("LearningRatePolicy"), TrainingSettings.LearningRatePolicy);
 		ConfigObject->SetNumberField(TEXT("LearningRateDecay"), TrainingSettings.LearningRateDecay);
 		ConfigObject->SetNumberField(TEXT("WeightDecay"), TrainingSettings.WeightDecay);
-		ConfigObject->SetNumberField(TEXT("BatchSize"), TrainingSettings.BatchSize);
-		ConfigObject->SetNumberField(TEXT("Window"), TrainingSettings.Window);
-		ConfigObject->SetNumberField(TEXT("ActionRegularizationWeight"), TrainingSettings.ActionRegularizationWeight);
-		ConfigObject->SetNumberField(TEXT("ActionEntropyWeight"), TrainingSettings.ActionEntropyWeight);
+		ConfigObject->SetNumberField(TEXT("InitialActionScale"), TrainingSettings.InitialActionScale);
+		ConfigObject->SetNumberField(TEXT("InitialMemoryScale"), TrainingSettings.InitialMemoryScale);
+		ConfigObject->SetNumberField(TEXT("PolicyBatchSize"), TrainingSettings.PolicyBatchSize);
+		ConfigObject->SetNumberField(TEXT("PolicyWindow"), TrainingSettings.PolicyWindow);
 		ConfigObject->SetNumberField(TEXT("Seed"), TrainingSettings.Seed);
 		ConfigObject->SetStringField(TEXT("Device"), Trainer::GetDeviceString(TrainingSettings.Device));
 		ConfigObject->SetBoolField(TEXT("UseTensorBoard"), TrainingSettings.bUseTensorboard);
-		ConfigObject->SetBoolField(TEXT("SaveSnapshots"), TrainingSettings.bSaveSnapshots);
+
+		ConfigObject->SetBoolField(TEXT("UseInitialPolicyNetwork"), (bool)(TrainerFlags & EImitationTrainerFlags::UseInitialPolicyNetwork));
 
 		FString JsonString;
 		TSharedRef<TJsonWriter<>> JsonWriter = TJsonWriterFactory<>::Create(&JsonString, 0);
@@ -597,9 +519,7 @@ namespace UE::Learning
 
 		// Allocate buffer to receive network data in
 
-		PolicyBuffer.SetNumUninitialized({ PolicyNetwork.GetSnapshotByteNum() });
-		EncoderBuffer.SetNumUninitialized({ EncoderNetwork.GetSnapshotByteNum() });
-		DecoderBuffer.SetNumUninitialized({ DecoderNetwork.GetSnapshotByteNum() });
+		NetworkBuffer.SetNumUninitialized({ PolicyNetwork.GetSerializationByteNum() });
 
 		// Create Socket
 
@@ -668,57 +588,21 @@ namespace UE::Learning
 	}
 
 	ETrainerResponse FSocketImitationTrainer::RecvPolicy(
-		ULearningNeuralNetworkData& OutNetwork,
-		const float Timeout,
-		FRWLock* NetworkLock,
-		const ELogSetting LogSettings)
-	{
-		return SocketTraining::RecvNetwork(*Socket, OutNetwork, PolicyBuffer, SocketTraining::ESignal::RecvPolicy, Timeout, NetworkLock, LogSettings);
-	}
-
-	ETrainerResponse FSocketImitationTrainer::RecvEncoder(
-		ULearningNeuralNetworkData& OutNetwork,
+		INeuralNetwork& OutNetwork,
 		const float Timeout,
 		FRWLock* NetworkLock,
 		const ELogSetting LogSettings)
 	{
-		return SocketTraining::RecvNetwork(*Socket, OutNetwork, EncoderBuffer, SocketTraining::ESignal::RecvEncoder, Timeout, NetworkLock, LogSettings);
-	}
-
-	ETrainerResponse FSocketImitationTrainer::RecvDecoder(
-		ULearningNeuralNetworkData& OutNetwork,
-		const float Timeout,
-		FRWLock* NetworkLock,
-		const ELogSetting LogSettings)
-	{
-		return SocketTraining::RecvNetwork(*Socket, OutNetwork, DecoderBuffer, SocketTraining::ESignal::RecvDecoder, Timeout, NetworkLock, LogSettings);
+		return SocketTraining::RecvPolicy(*Socket, OutNetwork, NetworkBuffer, Timeout, NetworkLock, LogSettings);
 	}
 
 	ETrainerResponse FSocketImitationTrainer::SendPolicy(
-		const ULearningNeuralNetworkData& Network,
-		const float Timeout,
-		FRWLock* NetworkLock,
-		const ELogSetting LogSettings)
-	{
-		return SocketTraining::SendNetwork(*Socket, PolicyBuffer, SocketTraining::ESignal::SendPolicy, Network, Timeout, NetworkLock, LogSettings);
-	}
-
-	ETrainerResponse FSocketImitationTrainer::SendEncoder(
-		const ULearningNeuralNetworkData& Network,
-		const float Timeout,
-		FRWLock* NetworkLock,
-		const ELogSetting LogSettings)
-	{
-		return SocketTraining::SendNetwork(*Socket, EncoderBuffer, SocketTraining::ESignal::SendEncoder, Network, Timeout, NetworkLock, LogSettings);
-	}
-
-	ETrainerResponse FSocketImitationTrainer::SendDecoder(
-		const ULearningNeuralNetworkData& Network,
+		const INeuralNetwork& Network,
 		const float Timeout,
 		FRWLock* NetworkLock,
 		const ELogSetting LogSettings)
 	{
-		return SocketTraining::SendNetwork(*Socket, DecoderBuffer, SocketTraining::ESignal::SendDecoder, Network, Timeout, NetworkLock, LogSettings);
+		return SocketTraining::SendPolicy(*Socket, NetworkBuffer, Network, Timeout, NetworkLock, LogSettings);
 	}
 
 	ETrainerResponse FSocketImitationTrainer::SendExperience(
@@ -743,20 +627,15 @@ namespace UE::Learning
 	{
 		ETrainerResponse Train(
 			IImitationTrainer& Trainer,
-			ULearningNeuralNetworkData& PolicyNetwork,
-			ULearningNeuralNetworkData& EncoderNetwork,
-			ULearningNeuralNetworkData& DecoderNetwork,
+			INeuralNetwork& Network,
 			const TLearningArrayView<1, const int32> EpisodeStartsExperience,
 			const TLearningArrayView<1, const int32> EpisodeLengthsExperience,
 			const TLearningArrayView<2, const float> ObservationsExperience,
 			const TLearningArrayView<2, const float> ActionsExperience,
+			const EImitationTrainerFlags TrainerFlags, 
 			TAtomic<bool>* bRequestTrainingStopSignal,
-			FRWLock* PolicyNetworkLock,
-			FRWLock* EncoderNetworkLock,
-			FRWLock* DecoderNetworkLock,
-			TAtomic<bool>* bPolicyNetworkUpdatedSignal,
-			TAtomic<bool>* bEncoderNetworkUpdatedSignal,
-			TAtomic<bool>* bDecoderNetworkUpdatedSignal,
+			FRWLock* NetworkLock,
+			TAtomic<bool>* bNetworkUpdatedSignal,
 			const ELogSetting LogSettings)
 		{
 			ETrainerResponse Response = ETrainerResponse::Success;
@@ -768,7 +647,7 @@ namespace UE::Learning
 				UE_LOG(LogLearning, Display, TEXT("Sending initial Policy..."));
 			}
 
-			Response = Trainer.SendPolicy(PolicyNetwork, Trainer::DefaultTimeout, PolicyNetworkLock);
+			Response = Trainer.SendPolicy(Network, 20.0f, NetworkLock);
 
 			if (Response != ETrainerResponse::Success)
 			{
@@ -781,44 +660,27 @@ namespace UE::Learning
 				return Response;
 			}
 
-			// Send initial Encoder
-
-			if (LogSettings != ELogSetting::Silent)
+			if (!(bool)(TrainerFlags & EImitationTrainerFlags::UseInitialPolicyNetwork))
 			{
-				UE_LOG(LogLearning, Display, TEXT("Sending initial Encoder..."));
-			}
-
-			Response = Trainer.SendEncoder(EncoderNetwork, Trainer::DefaultTimeout, EncoderNetworkLock);
+				// Receive initial Policy
 
-			if (Response != ETrainerResponse::Success)
-			{
 				if (LogSettings != ELogSetting::Silent)
 				{
-					UE_LOG(LogLearning, Error, TEXT("Error sending initial encoder from trainer: %s. Check log for errors."), Trainer::GetResponseString(Response));
+					UE_LOG(LogLearning, Display, TEXT("Receiving initial Policy..."));
 				}
 
-				Trainer.Terminate();
-				return Response;
-			}
-
-			// Send initial Decoder
-
-			if (LogSettings != ELogSetting::Silent)
-			{
-				UE_LOG(LogLearning, Display, TEXT("Sending initial Decoder..."));
-			}
+				Response = Trainer.RecvPolicy(Network, 20.0f, NetworkLock);
 
-			Response = Trainer.SendDecoder(DecoderNetwork, Trainer::DefaultTimeout, DecoderNetworkLock);
-
-			if (Response != ETrainerResponse::Success)
-			{
-				if (LogSettings != ELogSetting::Silent)
+				if (Response != ETrainerResponse::Success)
 				{
-					UE_LOG(LogLearning, Error, TEXT("Error sending initial decoder from trainer: %s. Check log for errors."), Trainer::GetResponseString(Response));
-				}
+					if (LogSettings != ELogSetting::Silent)
+					{
+						UE_LOG(LogLearning, Error, TEXT("Error receiving initial policy from trainer: %s. Check log for errors."), Trainer::GetResponseString(Response));
+					}
 
-				Trainer.Terminate();
-				return Response;
+					Trainer.Terminate();
+					return Response;
+				}
 			}
 
 			// Send Experience
@@ -828,7 +690,7 @@ namespace UE::Learning
 				EpisodeLengthsExperience,
 				ObservationsExperience, 
 				ActionsExperience, 
-				Trainer::DefaultTimeout);
+				10.0f);
 
 			if (Response != ETrainerResponse::Success)
 			{
@@ -872,7 +734,7 @@ namespace UE::Learning
 				
 				if (Trainer.HasPolicyOrCompleted())
 				{
-					Response = Trainer.RecvPolicy(PolicyNetwork, Trainer::DefaultTimeout, PolicyNetworkLock);
+					Response = Trainer.RecvPolicy(Network, 10.0f, NetworkLock);
 
 					if (Response == ETrainerResponse::Completed)
 					{
@@ -890,53 +752,17 @@ namespace UE::Learning
 						}
 						break;
 					}
-
-					if (bPolicyNetworkUpdatedSignal)
-					{
-						*bPolicyNetworkUpdatedSignal = true;
-					}
-
-					Response = Trainer.RecvEncoder(EncoderNetwork, Trainer::DefaultTimeout, EncoderNetworkLock);
-
-					if (Response != ETrainerResponse::Success)
-					{
-						if (LogSettings != ELogSetting::Silent)
-						{
-							UE_LOG(LogLearning, Error, TEXT("Error receiving encoder from trainer: %s. Check log for errors."), Trainer::GetResponseString(Response));
-						}
-						break;
-					}
-
-					if (bEncoderNetworkUpdatedSignal)
-					{
-						*bEncoderNetworkUpdatedSignal = true;
-					}
-
-					Response = Trainer.RecvDecoder(DecoderNetwork, Trainer::DefaultTimeout, DecoderNetworkLock);
-
-					if (Response != ETrainerResponse::Success)
-					{
-						if (LogSettings != ELogSetting::Silent)
-						{
-							UE_LOG(LogLearning, Error, TEXT("Error receiving decoder from trainer: %s. Check log for errors."), Trainer::GetResponseString(Response));
-						}
-						break;
-					}
-
-					if (bDecoderNetworkUpdatedSignal)
-					{
-						*bDecoderNetworkUpdatedSignal = true;
-					}
 				}
-				else
+
+				if (bNetworkUpdatedSignal)
 				{
-					FPlatformProcess::Sleep(0.001f);
+					*bNetworkUpdatedSignal = true;
 				}
 			}
 
 			// Allow some time for trainer to shut down gracefully before we kill it...
 
-			Response = Trainer.Wait(Trainer::DefaultTimeout);
+			Response = Trainer.Wait(5.0f);
 
 			if (Response != ETrainerResponse::Success)
 			{
