LearningAgents: Updated some default settings for training.

diff --git a/Engine/Plugins/Experimental/LearningAgents/Source/LearningAgentsTraining/Private/LearningAgentsTrainer.cpp b/Engine/Plugins/Experimental/LearningAgents/Source/LearningAgentsTraining/Private/LearningAgentsTrainer.cpp
--- a/Engine/Plugins/Experimental/LearningAgents/Source/LearningAgentsTraining/Private/LearningAgentsTrainer.cpp
+++ b/Engine/Plugins/Experimental/LearningAgents/Source/LearningAgentsTraining/Private/LearningAgentsTrainer.cpp
@@ -477,6 +477,8 @@ void ULearningAgentsTrainer::BeginTraining(
 	PPOTrainingSettings.PolicyBatchSize = TrainerTrainingSettings.PolicyBatchSize;
 	PPOTrainingSettings.CriticBatchSize = TrainerTrainingSettings.CriticBatchSize;
 	PPOTrainingSettings.PolicyWindow = TrainerTrainingSettings.PolicyWindowSize;
+	PPOTrainingSettings.IterationsPerGather = TrainerTrainingSettings.IterationsPerGather;
+	PPOTrainingSettings.CriticWarmupIterations = TrainerTrainingSettings.CriticWarmupIterations;
 	PPOTrainingSettings.EpsilonClip = TrainerTrainingSettings.EpsilonClip;
 	PPOTrainingSettings.ReturnRegularizationWeight = TrainerTrainingSettings.ReturnRegularizationWeight;
 	PPOTrainingSettings.ActionRegularizationWeight = TrainerTrainingSettings.ActionRegularizationWeight;
diff --git a/Engine/Plugins/Experimental/LearningAgents/Source/LearningAgentsTraining/Public/LearningAgentsImitationTrainer.h b/Engine/Plugins/Experimental/LearningAgents/Source/LearningAgentsTraining/Public/LearningAgentsImitationTrainer.h
--- a/Engine/Plugins/Experimental/LearningAgents/Source/LearningAgentsTraining/Public/LearningAgentsImitationTrainer.h
+++ b/Engine/Plugins/Experimental/LearningAgents/Source/LearningAgentsTraining/Public/LearningAgentsImitationTrainer.h
@@ -41,50 +41,50 @@ struct LEARNINGAGENTSTRAINING_API FLearningAgentsImitationTrainerTrainingSetting
 public:
 
 	/** The number of iterations to run before ending training. */
-	UPROPERTY(EditAnywhere, Category = "LearningAgents", meta = (ClampMin = "0", UIMin = "0"))
+	UPROPERTY(EditAnywhere, Category = "LearningAgents", meta = (ClampMin = "1", UIMin = "1"))
 	int32 NumberOfIterations = 1000000;
 
 	/** Learning rate of the policy network. Typical values are between 0.001 and 0.0001. */
-	UPROPERTY(EditAnywhere, Category = "LearningAgents", meta = (ClampMin = "0.0", ClampMax = "1.0", UIMin = "0.0", UIMax = "1.0"))
+	UPROPERTY(EditAnywhere, Category = "LearningAgents", meta = (ClampMin = "0.0", UIMin = "0.0", UIMax = "1.0"))
 	float LearningRate = 0.001f;
 
-	/** Ratio by which to decay the learning rate every 1000 iterations. */
+	/** Amount by which to multiply the learning rate every 1000 iterations. */
 	UPROPERTY(EditAnywhere, Category = "LearningAgents", meta = (ClampMin = "0.0", ClampMax = "1.0", UIMin = "0.0", UIMax = "1.0"))
-	float LearningRateDecay = 0.99f;
+	float LearningRateDecay = 1.0f;
 
 	/**
 	 * Amount of weight decay to apply to the network. Larger values encourage network weights to be smaller but too
 	 * large a value can cause the network weights to collapse to all zeros.
 	 */
-	UPROPERTY(EditAnywhere, Category = "LearningAgents", meta = (ClampMin = "0.0", ClampMax = "1.0", UIMin = "0.0", UIMax = "1.0"))
-	float WeightDecay = 0.001f;
+	UPROPERTY(EditAnywhere, Category = "LearningAgents", meta = (ClampMin = "0.0", UIMin = "0.0", UIMax = "1.0"))
+	float WeightDecay = 0.0001f;
 
 	/**
 	 * Batch size to use for training. Smaller values tend to produce better results at the cost of slowing down
 	 * training. Large batch sizes are much more computationally efficient when training on the GPU.
 	 */
-	UPROPERTY(EditAnywhere, Category = "LearningAgents", meta = (ClampMin = "0", ClampMax = "4096", UIMin = "0", UIMax = "4096"))
+	UPROPERTY(EditAnywhere, Category = "LearningAgents", meta = (ClampMin = "1", UIMin = "1", UIMax = "4096"))
 	uint32 BatchSize = 128;
 
 	/**
-	 * The window of observations and actions over which to do the training.Increasing this value will encourage the policy to use its memory 
+	 * The number of consecutive steps of observations and actions over which to train the policy. Increasing this value will encourage the policy to use its memory 
 	 * effectively. Too large and training can become unstable. Given we don't know the memory state during imitation learning it is better this is 
 	 * slightly larger than when we are doing reinforcement learning.
 	 */
-	UPROPERTY(EditAnywhere, Category = "LearningAgents", meta = (ClampMin = "1", ClampMax = "128", UIMin = "1", UIMax = "512"))
+	UPROPERTY(EditAnywhere, Category = "LearningAgents", meta = (ClampMin = "1", UIMin = "1", UIMax = "512"))
 	uint32 Window = 64;
 
 	/**
 	 * Weight used to regularize actions. Larger values will encourage smaller actions but too large will cause actions to become always zero.
 	 */
-	UPROPERTY(EditAnywhere, Category = "LearningAgents", meta = (ClampMin = "0", UIMin = "0"))
+	UPROPERTY(EditAnywhere, Category = "LearningAgents", meta = (ClampMin = "0.0", UIMin = "0.0"))
 	float ActionRegularizationWeight = 0.001f;
 
 	/**
 	 * Weighting used for the entropy bonus. Larger values encourage larger action noise and therefore greater exploration but can make actions very 
 	 * noisy.
 	 */
-	UPROPERTY(EditAnywhere, Category = "LearningAgents", meta = (ClampMin = "0", UIMin = "0"))
+	UPROPERTY(EditAnywhere, Category = "LearningAgents", meta = (ClampMin = "0.0", UIMin = "0.0"))
 	float ActionEntropyWeight = 0.0f;
 
 	/** The seed used for any random sampling the trainer will perform, e.g. for weight initialization. */
diff --git a/Engine/Plugins/Experimental/LearningAgents/Source/LearningAgentsTraining/Public/LearningAgentsTrainer.h b/Engine/Plugins/Experimental/LearningAgents/Source/LearningAgentsTraining/Public/LearningAgentsTrainer.h
--- a/Engine/Plugins/Experimental/LearningAgents/Source/LearningAgentsTraining/Public/LearningAgentsTrainer.h
+++ b/Engine/Plugins/Experimental/LearningAgents/Source/LearningAgentsTraining/Public/LearningAgentsTrainer.h
@@ -135,47 +135,59 @@ public:
 	int32 NumberOfIterations = 1000000;
 
 	/** Learning rate of the policy network. Typical values are between 0.001 and 0.0001. */
-	UPROPERTY(EditAnywhere, Category = "LearningAgents", meta = (ClampMin = "0.0", ClampMax = "1.0", UIMin = "0.0", UIMax = "1.0"))
+	UPROPERTY(EditAnywhere, Category = "LearningAgents", meta = (ClampMin = "0.0", UIMin = "0.0", UIMax = "1.0"))
 	float LearningRatePolicy = 0.0001f;
 
 	/**
 	 * Learning rate of the critic network. To avoid instability generally the critic should have a larger learning 
 	 * rate than the policy. Typically this can be set to 10x the rate of the policy.
 	 */
-	UPROPERTY(EditAnywhere, Category = "LearningAgents", meta = (ClampMin = "0.0", ClampMax = "1.0", UIMin = "0.0", UIMax = "1.0"))
+	UPROPERTY(EditAnywhere, Category = "LearningAgents", meta = (ClampMin = "0.0", UIMin = "0.0", UIMax = "1.0"))
 	float LearningRateCritic = 0.001f;
 
-	/** Ratio by which to decay the learning rate every 1000 iterations. */
+	/** Amount by which to multiply the learning rate every 1000 iterations. */
 	UPROPERTY(EditAnywhere, Category = "LearningAgents", meta = (ClampMin = "0.0", ClampMax = "1.0", UIMin = "0.0", UIMax = "1.0"))
-	float LearningRateDecay = 0.99f;
+	float LearningRateDecay = 1.0f;
 
 	/**
 	 * Amount of weight decay to apply to the network. Larger values encourage network weights to be smaller but too 
 	 * large a value can cause the network weights to collapse to all zeros.
 	 */
-	UPROPERTY(EditAnywhere, Category = "LearningAgents", meta = (ClampMin = "0.0", ClampMax = "1.0", UIMin = "0.0", UIMax = "1.0"))
-	float WeightDecay = 0.001f;
+	UPROPERTY(EditAnywhere, Category = "LearningAgents", meta = (ClampMin = "0.0", UIMin = "0.0", UIMax = "1.0"))
+	float WeightDecay = 0.0001f;
 
 	/**
-	 * Batch size to use for training the policy. Smaller values tend to produce better results at the cost of slowing down 
-	 * training. Large batch sizes are much more computationally efficient when training on the GPU.
+	 * Batch size to use for training the policy. Large batch sizes are much more computationally efficient when training on the GPU.
 	 */
-	UPROPERTY(EditAnywhere, Category = "LearningAgents", meta = (ClampMin = "1", ClampMax = "4096", UIMin = "1", UIMax = "4096"))
-	int32 PolicyBatchSize = 32;
+	UPROPERTY(EditAnywhere, Category = "LearningAgents", meta = (ClampMin = "1", UIMin = "1", UIMax = "4096"))
+	int32 PolicyBatchSize = 1024;
 
 	/**
-	 * Batch size to use for training the critic. Smaller values tend to produce better results at the cost of slowing down
-	 * training. Large batch sizes are much more computationally efficient when training on the GPU.
+	 * Batch size to use for training the critic. Large batch sizes are much more computationally efficient when training on the GPU.
 	 */
-	UPROPERTY(EditAnywhere, Category = "LearningAgents", meta = (ClampMin = "1", ClampMax = "4096", UIMin = "1", UIMax = "4096"))
-	int32 CriticBatchSize = 256;
+	UPROPERTY(EditAnywhere, Category = "LearningAgents", meta = (ClampMin = "1", UIMin = "1", UIMax = "4096"))
+	int32 CriticBatchSize = 4096;
 
 	/**
-	 * The size of the window of observations and actions over which to do the training of the policy. Increasing this value 
+	 * The number of consecutive steps of observations and actions over which to train the policy. Increasing this value 
 	 * will encourage the policy to use its memory effectively. Too large and training can become slow and unstable.
 	 */
-	UPROPERTY(EditAnywhere, Category = "LearningAgents", meta = (ClampMin = "1", ClampMax = "128", UIMin = "1", UIMax = "128"))
-	int32 PolicyWindowSize = 8;
+	UPROPERTY(EditAnywhere, Category = "LearningAgents", meta = (ClampMin = "1", UIMin = "1", UIMax = "128"))
+	int32 PolicyWindowSize = 16;
+
+	/**
+	 * Number of training iterations to perform per buffer of experience gathered. This should be large enough for
+	 * the critic and policy to be effectively updated, but too large and it will simply slow down training.
+	 */
+	UPROPERTY(EditAnywhere, Category = "LearningAgents", meta = (ClampMin = "1", UIMin = "1", UIMax = "1024"))
+	int32 IterationsPerGather = 32;
+
+	/**
+	 * Number of iterations of training to perform to warm - up the Critic. This helps speed up and stabilize training
+	 * at the beginning when the Critic may be producing predictions at the wrong order of magnitude.
+	 */
+	UPROPERTY(EditAnywhere, Category = "LearningAgents", meta = (ClampMin = "1", UIMin = "1", UIMax = "128"))
+	int32 CriticWarmupIterations = 8;
 
 	/**
 	 * Clipping ratio to apply to policy updates. Keeps the training "on-policy". Larger values may speed up training at 
@@ -188,29 +200,29 @@ public:
 	/**
 	 * Weight used to regularize returns. Encourages the critic not to over or under estimate returns.
 	 */
-	UPROPERTY(EditAnywhere, Category = "LearningAgents", meta = (ClampMin = "0.0", ClampMax = "1.0", UIMin = "0.0", UIMax = "1.0"))
+	UPROPERTY(EditAnywhere, Category = "LearningAgents", meta = (ClampMin = "0.0", UIMin = "0.0", UIMax = "1.0"))
 	float ReturnRegularizationWeight = 0.0001f;
 
 	/**
-	 * Weight used to regularize actions. Larger values will encourage smaller actions but too large will cause actions 
-	 * to become always zero.
+	 * Weight used to regularize actions. Larger values will encourage exploration and smaller actions, but too large will cause 
+	 * noisy actions centered around zero.
 	 */
-	UPROPERTY(EditAnywhere, Category = "LearningAgents", meta = (ClampMin = "0.0", ClampMax = "1.0", UIMin = "0.0", UIMax = "1.0"))
+	UPROPERTY(EditAnywhere, Category = "LearningAgents", meta = (ClampMin = "0.0", UIMin = "0.0", UIMax = "1.0"))
 	float ActionRegularizationWeight = 0.001f;
 
 	/**
 	 * Weighting used for the entropy bonus. Larger values encourage larger action noise and therefore greater 
 	 * exploration but can make actions very noisy.
 	 */
-	UPROPERTY(EditAnywhere, Category = "LearningAgents", meta = (ClampMin = "0.0", ClampMax = "1.0", UIMin = "0.0", UIMax = "1.0"))
+	UPROPERTY(EditAnywhere, Category = "LearningAgents", meta = (ClampMin = "0.0", UIMin = "0.0", UIMax = "1.0"))
 	float ActionEntropyWeight = 0.0f;
 
 	/**
-	 * This is used in the Generalized Advantage Estimation as what is essentially an exponential smoothing/decay. 
-	 * Typical values should be between 0.9 and 1.0.
+	 * This is used in the Generalized Advantage Estimation, where larger values will tend to assign more credit to recent actions. Typical
+	 * values should be between 0.9 and 1.0.
 	 */
 	UPROPERTY(EditAnywhere, Category = "LearningAgents", meta = (ClampMin = "0.0", ClampMax = "1.0", UIMin = "0.0", UIMax = "1.0"))
-	float GaeLambda = 0.9f;
+	float GaeLambda = 0.95f;
 
 	/** When true, advantages are normalized. This tends to makes training more robust to adjustments of the scale of rewards. */
 	UPROPERTY(EditAnywhere, Category = "LearningAgents")
@@ -220,14 +232,14 @@ public:
 	 * The minimum advantage to allow. Setting this below zero will encourage the policy to move away from bad actions, 
 	 * but can introduce instability.
 	 */
-	UPROPERTY(EditAnywhere, Category = "LearningAgents")
+	UPROPERTY(EditAnywhere, Category = "LearningAgents", meta = (UIMin = "-10.0", UIMax = "0.0"))
 	float MinimumAdvantage = 0.0f;
 
 	/**
 	 * The maximum advantage to allow. Making this smaller may increase training stability
 	 * at the cost of some training speed.
 	 */
-	UPROPERTY(EditAnywhere, Category = "LearningAgents")
+	UPROPERTY(EditAnywhere, Category = "LearningAgents", meta = (UIMin = "0.0", UIMax = "10.0"))
 	float MaximumAdvantage = 10.0f;
 
 	/**
diff --git a/Engine/Plugins/Experimental/LearningAgents/Source/LearningTraining/Public/LearningImitationTrainer.h b/Engine/Plugins/Experimental/LearningAgents/Source/LearningTraining/Public/LearningImitationTrainer.h
--- a/Engine/Plugins/Experimental/LearningAgents/Source/LearningTraining/Public/LearningImitationTrainer.h
+++ b/Engine/Plugins/Experimental/LearningAgents/Source/LearningTraining/Public/LearningImitationTrainer.h
@@ -39,18 +39,18 @@ namespace UE::Learning
 		// Learning rate of the training network. Typical values are between 0.001f and 0.0001f
 		float LearningRate = 0.001f;
 
-		// Ratio by which to decay the learning rate every 1000 iterations.
-		float LearningRateDecay = 0.99f;
+		// Amount by which to multiply the learning rate every 1000 iterations.
+		float LearningRateDecay = 1.0f;
 
 		// Amount of weight decay to apply to the network. Larger values encourage network 
 		// weights to be smaller.
-		float WeightDecay = 0.001f;
+		float WeightDecay = 0.0001f;
 
 		// Batch size to use for training. Smaller values tend to produce better results 
 		// at the cost of slowing down training.
 		uint32 BatchSize = 128;
 
-		// The window of observations and actions over which to do the training. Increasing this value 
+		// The number of consecutive steps of observations and actions over which to train the policy. Increasing this value 
 		// will encourage the policy to use its memory effectively. Too large and training can become unstable. Given
 		// we don't know the memory state during imitation learning it is better this is slightly larger than when we 
 		// are doing reinforcement learning.
diff --git a/Engine/Plugins/Experimental/LearningAgents/Source/LearningTraining/Public/LearningPPOTrainer.h b/Engine/Plugins/Experimental/LearningAgents/Source/LearningTraining/Public/LearningPPOTrainer.h
--- a/Engine/Plugins/Experimental/LearningAgents/Source/LearningTraining/Public/LearningPPOTrainer.h
+++ b/Engine/Plugins/Experimental/LearningAgents/Source/LearningTraining/Public/LearningPPOTrainer.h
@@ -51,34 +51,30 @@ namespace UE::Learning
 		// should have a larger learning rate than the policy.
 		float LearningRateCritic = 0.001f;
 
-		// Ratio by which to decay the learning rate every 1000 iterations.
-		float LearningRateDecay = 0.99f;
+		// Amount by which to multiply the learning rate every 1000 iterations.
+		float LearningRateDecay = 1.0f;
 
 		// Amount of weight decay to apply to the network. Larger values encourage network 
 		// weights to be smaller but too large a value can cause the network weights to collapse to all zeros.
-		float WeightDecay = 0.001f;
+		float WeightDecay = 0.0001f;
 
-		// Batch size to use for training the policy. Smaller values tend to produce better results 
-		// at the cost of slowing down training. Large batch sizes are much more computationally efficient 
-		// when training on the GPU.
-		uint32 PolicyBatchSize = 32;
+		// Batch size to use for training the policy. Large batch sizes are much more computationally efficient when training on the GPU.
+		uint32 PolicyBatchSize = 1024;
 
-		// Batch size to use for training the critic. Smaller values tend to produce better results 
-		// at the cost of slowing down training. Large batch sizes are much more computationally efficient 
-		// when training on the GPU.
-		uint32 CriticBatchSize = 256;
+		// Batch size to use for training the critic. Large batch sizes are much more computationally efficient when training on the GPU.
+		uint32 CriticBatchSize = 4096;
 
-		// The window of observations and actions over which to do the training of the policy. Increasing this value 
+		// The number of consecutive steps of observations and actions over which to train the policy. Increasing this value 
 		// will encourage the policy to use its memory effectively. Too large and training can become slow and unstable.
-		uint32 PolicyWindow = 8;
+		uint32 PolicyWindow = 16;
 
 		// Number of training iterations to perform per buffer of experience gathered. This should be large enough for
 		// the critic and policy to be effectively updated, but too large and it will simply slow down training.
-		uint32 IterationsPerGather = 250;
+		uint32 IterationsPerGather = 32;
 
 		// Number of iterations of training to perform to warm-up the Critic. This helps speed up and stabilize training
 		// at the beginning when the Critic may be producing predictions at the wrong order of magnitude.
-		uint32 CriticWarmupIterations = 10;
+		uint32 CriticWarmupIterations = 8;
 
 		// Clipping ratio to apply to policy updates. Keeps the training "on-policy". 
 		// Larger values may speed up training at the cost of stability. Conversely, too small 
@@ -91,17 +87,17 @@ namespace UE::Learning
 		// Weight for the loss used to train the policy via the PPO surrogate objective.
 		float ActionSurrogateWeight = 1.0f;
 
-		// Weight used to regularize actions. Larger values will encourage smaller actions but too large
-		// will cause actions to become always zero.
+		// Weight used to regularize actions.Larger values will encourage exploration and smaller actions, but too large will cause
+		// noisy actions centered around zero.
 		float ActionRegularizationWeight = 0.001f;
 
 		// Weighting used for the entropy bonus. Larger values encourage larger action 
 		// noise and therefore greater exploration but can make actions very noisy.
 		float ActionEntropyWeight = 0.0f;
 
-		// This is used in the Generalized Advantage Estimation as what is essentially 
-		// an exponential smoothing/decay. Typical values should be between 0.9 and 1.0.
-		float GaeLambda = 0.9f;
+		// This is used in the Generalized Advantage Estimation, where larger values will tend to assign more credit to recent actions. Typical 
+		// values should be between 0.9 and 1.0.
+		float GaeLambda = 0.95f;
 
 		// When true, advantages are normalized. This tends to makes training more robust to 
 		// adjustments of the scale of rewards. 
