LearningAgents: Refactoring of Neural Network interface to allow for evaluation of Neural Networks without locking when instances are already tightly packed.

diff --git a/Engine/Plugins/Experimental/LearningAgents/Source/Learning/Private/LearningCritic.cpp b/Engine/Plugins/Experimental/LearningAgents/Source/Learning/Private/LearningCritic.cpp
--- a/Engine/Plugins/Experimental/LearningAgents/Source/Learning/Private/LearningCritic.cpp
+++ b/Engine/Plugins/Experimental/LearningAgents/Source/Learning/Private/LearningCritic.cpp
@@ -10,19 +10,16 @@ namespace UE::Learning
 		const int32 InMemoryStateNum,
 		const TSharedPtr<FNeuralNetwork>& InNeuralNetwork,
 		const FNeuralNetworkInferenceSettings& InInferenceSettings)
-		: MaxInstanceNum(InMaxInstanceNum)
-		, ObservationEncodedNum(InObservationEncodedNum)
+		: ObservationEncodedNum(InObservationEncodedNum)
 		, MemoryStateNum(InMemoryStateNum)
-		, NeuralNetwork(InNeuralNetwork)
-		, InferenceSettings(InInferenceSettings)
 	{
-		UE_LEARNING_CHECK(NeuralNetwork->GetInputSize() == ObservationEncodedNum + MemoryStateNum);
-		UE_LEARNING_CHECK(NeuralNetwork->GetOutputSize() == 1);
+		UE_LEARNING_CHECK(InNeuralNetwork->GetInputSize() == ObservationEncodedNum + MemoryStateNum);
+		UE_LEARNING_CHECK(InNeuralNetwork->GetOutputSize() == 1);
 
-		Input.SetNumUninitialized({ MaxInstanceNum, ObservationEncodedNum + MemoryStateNum });
+		Input.SetNumUninitialized({ InMaxInstanceNum, ObservationEncodedNum + MemoryStateNum });
 		Array::Zero(Input);
 
-		NeuralNetworkInference = NeuralNetwork->CreateInferenceObject(MaxInstanceNum, InferenceSettings);
+		NeuralNetworkFunction = MakeShared<FNeuralNetworkFunction>(InMaxInstanceNum, InNeuralNetwork, InInferenceSettings);
 	}
 
 	void FNeuralNetworkCritic::Evaluate(
@@ -46,7 +43,7 @@ namespace UE::Learning
 
 		// Evaluate Network
 
-		NeuralNetworkInference->Evaluate(
+		NeuralNetworkFunction->Evaluate(
 			TLearningArrayView<2, float>(OutputReturns.GetData(), { OutputReturns.Num(), 1 }),
 			Input, 
 			Instances);
@@ -56,12 +53,6 @@ namespace UE::Learning
 
 	void FNeuralNetworkCritic::UpdateNeuralNetwork(const TSharedPtr<FNeuralNetwork>& NewNeuralNetwork)
 	{
-		if (NeuralNetwork != NewNeuralNetwork)
-		{
-			UE_LEARNING_CHECK(NewNeuralNetwork->GetInputSize() == ObservationEncodedNum + MemoryStateNum);
-			UE_LEARNING_CHECK(NewNeuralNetwork->GetOutputSize() == 1);
-			NeuralNetwork = NewNeuralNetwork;
-			NeuralNetworkInference = NeuralNetwork->CreateInferenceObject(MaxInstanceNum, InferenceSettings);
-		}
+		NeuralNetworkFunction->UpdateNeuralNetwork(NewNeuralNetwork);
 	}
 }
diff --git a/Engine/Plugins/Experimental/LearningAgents/Source/Learning/Private/LearningNeuralNetwork.cpp b/Engine/Plugins/Experimental/LearningAgents/Source/Learning/Private/LearningNeuralNetwork.cpp
--- a/Engine/Plugins/Experimental/LearningAgents/Source/Learning/Private/LearningNeuralNetwork.cpp
+++ b/Engine/Plugins/Experimental/LearningAgents/Source/Learning/Private/LearningNeuralNetwork.cpp
@@ -108,7 +108,7 @@ TSharedPtr<UE::Learning::FNeuralNetwork>& ULearningNeuralNetworkData::GetNetwork
 
 bool ULearningNeuralNetworkData::IsEmpty() const
 {
-	return Network->IsEmpty();
+	return !Network.IsValid() || Network->IsEmpty();
 }
 
 int32 ULearningNeuralNetworkData::GetInputSize() const
@@ -184,7 +184,7 @@ namespace UE::Learning
 
 	bool FNeuralNetwork::IsEmpty() const
 	{
-		return Model != nullptr;
+		return Model == nullptr;
 	}
 	
 	int32 FNeuralNetwork::GetInputSize() const
@@ -231,15 +231,15 @@ namespace UE::Learning
 
 	FNeuralNetworkInference::FNeuralNetworkInference(
 		UE::NNE::IModelCPU& InModel,
-		const int32 MaxBatchSize,
+		const int32 InMaxBatchSize,
 		const int32 InInputSize,
 		const int32 InOutputSize,
 		const FNeuralNetworkInferenceSettings& InSettings)
-		: Settings(InSettings)
+		: MaxBatchSize(InMaxBatchSize)
+		, InputSize(InInputSize)
+		, OutputSize(InOutputSize)
+		, Settings(InSettings)
 	{
-		InputBuffer.SetNumUninitialized({ MaxBatchSize, InInputSize });
-		OutputBuffer.SetNumUninitialized({ MaxBatchSize, InOutputSize });
-
 		ReloadModelInstances(InModel);
 	}
 
@@ -247,22 +247,18 @@ namespace UE::Learning
 	{
 		UE_LEARNING_TRACE_CPUPROFILER_EVENT_SCOPE(Learning::FNeuralNetworkInference::ReloadModelInstances);
 
-		const int32 MaxBatchSize = InputBuffer.Num<0>();
-		const int32 InputSize = InputBuffer.Num<1>();
-		const int32 OutputSize = OutputBuffer.Num<1>();
-
 		if (Settings.bParallelEvaluation)
 		{
 			const int32 IdealInferenceInstanceNum = ParallelForImpl::GetNumberOfThreadTasks(MaxBatchSize, Settings.MinParallelBatchSize, EParallelForFlags::None);
-			const int32 InferenceInstanceSliceLength = FMath::DivideAndRoundUp((int32)MaxBatchSize, IdealInferenceInstanceNum);
-			const int32 InferenceInstanceNum = FMath::DivideAndRoundUp((int32)MaxBatchSize, InferenceInstanceSliceLength);
+			const int32 InferenceInstanceBatchSize = FMath::DivideAndRoundUp((int32)MaxBatchSize, IdealInferenceInstanceNum);
+			const int32 InferenceInstanceNum = FMath::DivideAndRoundUp((int32)MaxBatchSize, InferenceInstanceBatchSize);
 
 			ModelInstances.Empty(InferenceInstanceNum);
 
 			for (int32 InferenceInstanceIdx = 0; InferenceInstanceIdx < InferenceInstanceNum; InferenceInstanceIdx++)
 			{
 				ModelInstances.Emplace(InModel.CreateModelInstanceCPU());
-				ModelInstances.Last()->SetInputTensorShapes({ NNE::FTensorShape::Make({(uint32)InferenceInstanceSliceLength, (uint32)InputSize}) });
+				ModelInstances.Last()->SetInputTensorShapes({ NNE::FTensorShape::Make({(uint32)InferenceInstanceBatchSize, (uint32)InputSize}) });
 			}
 		}
 		else
@@ -273,77 +269,69 @@ namespace UE::Learning
 		}
 	}
 
-	void FNeuralNetworkInference::Evaluate(
-		TLearningArrayView<2, float> Output,
-		const TLearningArrayView<2, const float> Input,
-		const FIndexSet Instances)
+	int32 FNeuralNetworkInference::GetMaxBatchSize() const
 	{
-		UE_LEARNING_TRACE_CPUPROFILER_EVENT_SCOPE(Learning::FNeuralNetworkInference::Evaluate);
+		return MaxBatchSize;
+	}
 
-		const int32 InstanceNum = Instances.Num();
-		const int32 InputNum = Input.Num<1>();
-		const int32 OutputNum = Output.Num<1>();
+	int32 FNeuralNetworkInference::GetInputSize() const
+	{
+		return InputSize;
+	}
 
-		if (InstanceNum == 0) { return; }
+	int32 FNeuralNetworkInference::GetOutputSize() const
+	{
+		return OutputSize;
+	}
 
-		// Evaluation is not thread-safe due to the scatter and gather we do into the InputBuffer and OutputBuffer
-		// so we take a lock here since the interface for this function makes it appear like this might be thread safe as 
-		// long as the Instances provided are not an overlapping set.
-		FScopeNullableWriteLock ScopeLock(&EvaluationLock);
+	void FNeuralNetworkInference::Evaluate(TLearningArrayView<2, float> Output, const TLearningArrayView<2, const float> Input)
+	{
+		UE_LEARNING_TRACE_CPUPROFILER_EVENT_SCOPE(Learning::FNeuralNetworkInference::Evaluate);
 
-		// Gather Inputs
+		UE_LEARNING_CHECK(Output.Num<0>() == Input.Num<0>());
+		UE_LEARNING_CHECK(Input.Num<1>() == InputSize);
+		UE_LEARNING_CHECK(Output.Num<1>() == OutputSize);
 
-		Array::Check(Input, Instances);
+		const int32 BatchSize = Output.Num<0>();
 
-		for (int32 InstanceIdx = 0; InstanceIdx < InstanceNum; InstanceIdx++)
-		{
-			Array::Copy(InputBuffer[InstanceIdx], Input[Instances[InstanceIdx]]);
-		}
+		UE_LEARNING_CHECK(BatchSize <= MaxBatchSize);
 
-		// Run Inference
+		if (BatchSize == 0) { return; }
 
-		Array::Check(InputBuffer.Slice(0, InstanceNum));
+		Array::Check(Input);
 
-		if (Settings.bParallelEvaluation && InstanceNum > Settings.MinParallelBatchSize)
+		if (Settings.bParallelEvaluation && BatchSize > Settings.MinParallelBatchSize)
 		{
-			const int32 IdealInferenceInstanceNum = ParallelForImpl::GetNumberOfThreadTasks(InstanceNum, Settings.MinParallelBatchSize, EParallelForFlags::None);
-			const int32 InferenceInstanceSliceLength = FMath::DivideAndRoundUp(InstanceNum, IdealInferenceInstanceNum);
-			const int32 InferenceInstanceNum = FMath::DivideAndRoundUp(InstanceNum, InferenceInstanceSliceLength);
+			const int32 IdealInferenceInstanceNum = ParallelForImpl::GetNumberOfThreadTasks(BatchSize, Settings.MinParallelBatchSize, EParallelForFlags::None);
+			const int32 InferenceInstanceBatchSize = FMath::DivideAndRoundUp(BatchSize, IdealInferenceInstanceNum);
+			const int32 InferenceInstanceNum = FMath::DivideAndRoundUp(BatchSize, InferenceInstanceBatchSize);
 
 			UE_LEARNING_CHECK(InferenceInstanceNum <= ModelInstances.Num());
 
-			ParallelFor(InferenceInstanceNum, [this, InstanceNum, InputNum, InferenceInstanceSliceLength](int32 InferenceInstanceIdx)
+			ParallelFor(InferenceInstanceNum, [this, Input, Output, BatchSize, InferenceInstanceBatchSize](int32 InferenceInstanceIdx)
 			{
-				const int32 StartIndex = InferenceInstanceIdx * InferenceInstanceSliceLength;
-				const int32 StopIndex = FMath::Min((InferenceInstanceIdx + 1) * InferenceInstanceSliceLength, InstanceNum);
+				const int32 StartIndex = InferenceInstanceIdx * InferenceInstanceBatchSize;
+				const int32 StopIndex = FMath::Min((InferenceInstanceIdx + 1) * InferenceInstanceBatchSize, BatchSize);
+				const int32 InstanceBatchSize = StopIndex - StartIndex;
 
-				TLearningArrayView<2, float> InputBufferSlice = InputBuffer.Slice(StartIndex, StopIndex - StartIndex);
-				TLearningArrayView<2, float> OutputBufferSlice = OutputBuffer.Slice(StartIndex, StopIndex - StartIndex);
+				const TLearningArrayView<2, const float> InputSlice = Input.Slice(StartIndex, InstanceBatchSize);
+				const TLearningArrayView<2, float> OutputSlice = Output.Slice(StartIndex, InstanceBatchSize);
 
-				ModelInstances[InferenceInstanceIdx]->SetInputTensorShapes({ NNE::FTensorShape::Make({(uint32)(StopIndex - StartIndex), (uint32)InputNum}) });
+				ModelInstances[InferenceInstanceIdx]->SetInputTensorShapes({ NNE::FTensorShape::Make({(uint32)InstanceBatchSize, (uint32)InputSize}) });
 				ModelInstances[InferenceInstanceIdx]->RunSync(
-					{ { (void*)InputBufferSlice.GetData(), InputBufferSlice.Num() * sizeof(float) } },
-					{ { (void*)OutputBufferSlice.GetData(), OutputBufferSlice.Num() * sizeof(float) } });
+					{ { (void*)InputSlice.GetData(), InputSlice.Num() * sizeof(float) } },
+					{ { (void*)OutputSlice.GetData(), OutputSlice.Num() * sizeof(float) } });
 			});
 		}
 		else
 		{
-			ModelInstances[0]->SetInputTensorShapes({ NNE::FTensorShape::Make({(uint32)InstanceNum, (uint32)InputNum}) });
+			ModelInstances[0]->SetInputTensorShapes({ NNE::FTensorShape::Make({(uint32)BatchSize, (uint32)InputSize}) });
 			ModelInstances[0]->RunSync(
-				{ { (void*)InputBuffer.GetData(), InputBuffer.Slice(0, InstanceNum).Num() * sizeof(float) } },
-				{ { (void*)OutputBuffer.GetData(), OutputBuffer.Slice(0, InstanceNum).Num() * sizeof(float) } });
+				{ { (void*)Input.GetData(), Input.Num() * sizeof(float) } },
+				{ { (void*)Output.GetData(), Output.Num() * sizeof(float) } });
 		}
 
-		Array::Check(OutputBuffer.Slice(0, InstanceNum));
-
-		// Scatter Outputs
-
-		for (int32 InstanceIdx = 0; InstanceIdx < InstanceNum; InstanceIdx++)
-		{
-			Array::Copy(Output[Instances[InstanceIdx]], OutputBuffer[InstanceIdx]);
-		}
-
-		Array::Check(Output, Instances);
+		Array::Check(Output);
 	}
 
 
@@ -356,6 +344,8 @@ namespace UE::Learning
 		, InferenceSettings(InInferenceSettings)
 	{
 		NeuralNetworkInference = NeuralNetwork->CreateInferenceObject(InMaxInstanceNum, InInferenceSettings);
+		InputBuffer.SetNumUninitialized({ MaxInstanceNum, InNeuralNetwork->GetInputSize() });
+		OutputBuffer.SetNumUninitialized({ MaxInstanceNum, InNeuralNetwork->GetOutputSize() });
 	}
 
 	void FNeuralNetworkFunction::Evaluate(
@@ -363,13 +353,44 @@ namespace UE::Learning
 		const TLearningArrayView<2, const float> Input,
 		const FIndexSet Instances)
 	{
-		NeuralNetworkInference->Evaluate(Output, Input, Instances);
+		// Evaluation is not thread-safe due to the scatter and gather we do into the InputBuffer and OutputBuffer
+		// so we take a lock here since the interface for this function makes it appear like this might be thread safe as 
+		// long as the Instances provided are not an overlapping set.
+		FScopeNullableWriteLock ScopeLock(&EvaluationLock);
+
+		const int32 InstanceNum = Instances.Num();
+
+		if (InstanceNum == 0) { return; }
+
+		// Gather Inputs
+
+		Array::Check(Input, Instances);
+
+		for (int32 InstanceIdx = 0; InstanceIdx < InstanceNum; InstanceIdx++)
+		{
+			Array::Copy(InputBuffer[InstanceIdx], Input[Instances[InstanceIdx]]);
+		}
+
+		// Evaluate Network
+
+		NeuralNetworkInference->Evaluate(OutputBuffer.Slice(0, InstanceNum), InputBuffer.Slice(0, InstanceNum));
+
+		// Scatter Outputs
+
+		for (int32 InstanceIdx = 0; InstanceIdx < InstanceNum; InstanceIdx++)
+		{
+			Array::Copy(Output[Instances[InstanceIdx]], OutputBuffer[InstanceIdx]);
+		}
+
+		Array::Check(Output, Instances);
 	}
 
 	void FNeuralNetworkFunction::UpdateNeuralNetwork(const TSharedPtr<FNeuralNetwork>& NewNeuralNetwork)
 	{
 		if (NeuralNetwork != NewNeuralNetwork)
 		{
+			UE_LEARNING_CHECK(NeuralNetwork->GetInputSize() == NewNeuralNetwork->GetInputSize());
+			UE_LEARNING_CHECK(NeuralNetwork->GetOutputSize() == NewNeuralNetwork->GetOutputSize());
 			NeuralNetwork = NewNeuralNetwork;
 			NeuralNetworkInference = NeuralNetwork->CreateInferenceObject(MaxInstanceNum, InferenceSettings);
 		}
diff --git a/Engine/Plugins/Experimental/LearningAgents/Source/Learning/Private/LearningPolicy.cpp b/Engine/Plugins/Experimental/LearningAgents/Source/Learning/Private/LearningPolicy.cpp
--- a/Engine/Plugins/Experimental/LearningAgents/Source/Learning/Private/LearningPolicy.cpp
+++ b/Engine/Plugins/Experimental/LearningAgents/Source/Learning/Private/LearningPolicy.cpp
@@ -11,23 +11,20 @@ namespace UE::Learning
 		const int32 InMemoryStateNum,
 		const TSharedPtr<FNeuralNetwork>& InNeuralNetwork,
 		const FNeuralNetworkInferenceSettings& InInferenceSettings)
-		: MaxInstanceNum(InMaxInstanceNum)
-		, ObservationEncodedNum(InObservationEncodedNum)
+		: ObservationEncodedNum(InObservationEncodedNum)
 		, ActionEncodedNum(InActionEncodedNum)
 		, MemoryStateNum(InMemoryStateNum)
-		, NeuralNetwork(InNeuralNetwork)
-		, InferenceSettings(InInferenceSettings)
 	{
-		UE_LEARNING_CHECK(NeuralNetwork->GetInputSize() == ObservationEncodedNum + MemoryStateNum);
-		UE_LEARNING_CHECK(NeuralNetwork->GetOutputSize() == ActionEncodedNum + MemoryStateNum);
+		UE_LEARNING_CHECK(InNeuralNetwork->GetInputSize() == ObservationEncodedNum + MemoryStateNum);
+		UE_LEARNING_CHECK(InNeuralNetwork->GetOutputSize() == ActionEncodedNum + MemoryStateNum);
 
-		Input.SetNumUninitialized({ MaxInstanceNum, ObservationEncodedNum + MemoryStateNum });
-		Output.SetNumUninitialized({ MaxInstanceNum, ActionEncodedNum + MemoryStateNum });
+		Input.SetNumUninitialized({ InMaxInstanceNum, ObservationEncodedNum + MemoryStateNum });
+		Output.SetNumUninitialized({ InMaxInstanceNum, ActionEncodedNum + MemoryStateNum });
 
 		Array::Zero(Input);
 		Array::Zero(Output);
 
-		NeuralNetworkInference = NeuralNetwork->CreateInferenceObject(MaxInstanceNum, InferenceSettings);
+		NeuralNetworkFunction = MakeShared<FNeuralNetworkFunction>(InMaxInstanceNum, InNeuralNetwork, InInferenceSettings);
 	}
 
 	void FNeuralNetworkPolicy::Evaluate(
@@ -52,7 +49,7 @@ namespace UE::Learning
 			Array::Copy(Input[InstanceIdx].Slice(ObservationEncodedNum, MemoryStateNum), InputMemoryState[InstanceIdx]);
 		}
 
-		NeuralNetworkInference->Evaluate(Output, Input, Instances);
+		NeuralNetworkFunction->Evaluate(Output, Input, Instances);
 
 		// Copy Out Memory State
 
@@ -68,12 +65,6 @@ namespace UE::Learning
 
 	void FNeuralNetworkPolicy::UpdateNeuralNetwork(const TSharedPtr<FNeuralNetwork>& NewNeuralNetwork)
 	{
-		if (NeuralNetwork != NewNeuralNetwork)
-		{
-			UE_LEARNING_CHECK(NewNeuralNetwork->GetInputSize() == ObservationEncodedNum + MemoryStateNum);
-			UE_LEARNING_CHECK(NewNeuralNetwork->GetOutputSize() == ActionEncodedNum + MemoryStateNum);
-			NeuralNetwork = NewNeuralNetwork;
-			NeuralNetworkInference = NeuralNetwork->CreateInferenceObject(MaxInstanceNum, InferenceSettings);
-		}
+		NeuralNetworkFunction->UpdateNeuralNetwork(NewNeuralNetwork);
 	}
 }
diff --git a/Engine/Plugins/Experimental/LearningAgents/Source/Learning/Public/LearningCritic.h b/Engine/Plugins/Experimental/LearningAgents/Source/Learning/Public/LearningCritic.h
--- a/Engine/Plugins/Experimental/LearningAgents/Source/Learning/Public/LearningCritic.h
+++ b/Engine/Plugins/Experimental/LearningAgents/Source/Learning/Public/LearningCritic.h
@@ -49,12 +49,9 @@ namespace UE::Learning
 
 	private:
 
-		int32 MaxInstanceNum = 0;
 		int32 ObservationEncodedNum = 0;
 		int32 MemoryStateNum = 0;
-		TSharedPtr<FNeuralNetwork> NeuralNetwork;
-		TSharedPtr<FNeuralNetworkInference> NeuralNetworkInference;
-		FNeuralNetworkInferenceSettings InferenceSettings;
+		TSharedPtr<FNeuralNetworkFunction> NeuralNetworkFunction;
 
 		TLearningArray<2, float> Input;
 	};
diff --git a/Engine/Plugins/Experimental/LearningAgents/Source/Learning/Public/LearningNeuralNetwork.h b/Engine/Plugins/Experimental/LearningAgents/Source/Learning/Public/LearningNeuralNetwork.h
--- a/Engine/Plugins/Experimental/LearningAgents/Source/Learning/Public/LearningNeuralNetwork.h
+++ b/Engine/Plugins/Experimental/LearningAgents/Source/Learning/Public/LearningNeuralNetwork.h
@@ -156,46 +156,43 @@ namespace UE::Learning
 	struct LEARNING_API FNeuralNetworkInference
 	{
 		/**
-		 * Constructs a new network inference object.
+		 * Constructs a new network inference object. Generally this should not be called directly and FNeuralNetwork::CreateInferenceObject should 
+		 * be used instead so that created instances are tracked.
 		 * 
 		 * @param InModel			NNE Model
-		 * @param MaxBatchSize		Maximum batch size
+		 * @param InMaxBatchSize	Maximum batch size
 		 * @param InInputSize		Network input size
 		 * @param InOutputSize		Network output size
 		 * @param InSettings		Inference settings
 		 */
 		FNeuralNetworkInference(
 			UE::NNE::IModelCPU& InModel,
-			const int32 MaxBatchSize,
+			const int32 InMaxBatchSize,
 			const int32 InInputSize,
 			const int32 InOutputSize,
 			const FNeuralNetworkInferenceSettings& InSettings = FNeuralNetworkInferenceSettings());
 
 		/**
-		 * Evaluate this network for the given instances.
-		 * 
-		 * Note: this function takes a lock, so although it can be called from multiple threads it will not benefit from multi-threading. This is 
-		 * because to use the NNE interface all the instances that need to be evaluated must be gathered and scattered into internal buffers. To
-		 * compensate for this, this call will do multi-threading internally based on the inference settings used.
+		 * Evaluate this network.
 		 * 
-		 * @param Output		The Output Buffer of shape (MaxBatchSize, InOutputSize)
-		 * @param Input			The Input Buffer of shape (MaxBatchSize, InInputSize)
-		 * @param Instances		The batch instances to evaluate
+		 * @param Output		The Output Buffer of shape (<= MaxBatchSize, OutputSize)
+		 * @param Input			The Input Buffer of shape (<= MaxBatchSize, InputSize)
 		 */
-		void Evaluate(
-			TLearningArrayView<2, float> Output,
-			const TLearningArrayView<2, const float> Input,
-			const FIndexSet Instances);
+		void Evaluate(TLearningArrayView<2, float> Output, const TLearningArrayView<2, const float> Input);
 
 		// This function will re-build the internal Model Instances used for multi-threading. It should be called whenever the given Model is updated. 
 		void ReloadModelInstances(NNE::IModelCPU& Model);
 
+		int32 GetMaxBatchSize() const;
+		int32 GetInputSize() const;
+		int32 GetOutputSize() const;
+
 	private:
 
+		int32 MaxBatchSize = 0;
+		int32 InputSize = 0;
+		int32 OutputSize = 0;
 		FNeuralNetworkInferenceSettings Settings;
-		FRWLock EvaluationLock;
-		TLearningArray<2, float> InputBuffer;
-		TLearningArray<2, float> OutputBuffer;
 		TArray<TSharedPtr<NNE::IModelInstanceCPU>, TInlineAllocator<64>> ModelInstances;
 	};
 
@@ -219,8 +216,12 @@ namespace UE::Learning
 		/**
 		 * Evaluate this network for the given instances.
 		 *
-		 * @param Output		The Output Buffer of shape (MaxInstanceNum, InOutputSize)
-		 * @param Input			The Input Buffer of shape (MaxInstanceNum, InInputSize)
+		 * Note: this function takes a lock, so although it can be called from multiple threads if the set of Instances don't overlap, it will not 
+		 * benefit from multi-threading. This is because to use the NNE interface all the instances that need to be evaluated must be gathered and 
+		 * scattered into internal buffers. 
+		 * 
+		 * @param Output		The Output Buffer of shape (<= MaxInstanceNum, OutputSize)
+		 * @param Input			The Input Buffer of shape (<= MaxInstanceNum, InputSize)
 		 * @param Instances		The instances to evaluate
 		 */
 		void Evaluate(
@@ -234,6 +235,9 @@ namespace UE::Learning
 	private:
 
 		int32 MaxInstanceNum = 0;
+		FRWLock EvaluationLock;
+		TLearningArray<2, float> InputBuffer;
+		TLearningArray<2, float> OutputBuffer;
 		TSharedPtr<FNeuralNetwork> NeuralNetwork;
 		TSharedPtr<FNeuralNetworkInference> NeuralNetworkInference;
 		FNeuralNetworkInferenceSettings InferenceSettings;
diff --git a/Engine/Plugins/Experimental/LearningAgents/Source/Learning/Public/LearningPolicy.h b/Engine/Plugins/Experimental/LearningAgents/Source/Learning/Public/LearningPolicy.h
--- a/Engine/Plugins/Experimental/LearningAgents/Source/Learning/Public/LearningPolicy.h
+++ b/Engine/Plugins/Experimental/LearningAgents/Source/Learning/Public/LearningPolicy.h
@@ -53,13 +53,10 @@ namespace UE::Learning
 
 	private:
 
-		int32 MaxInstanceNum = 0;
 		int32 ObservationEncodedNum = 0;
 		int32 ActionEncodedNum = 0;
 		int32 MemoryStateNum = 0;
-		TSharedPtr<FNeuralNetwork> NeuralNetwork;
-		TSharedPtr<FNeuralNetworkInference> NeuralNetworkInference;
-		FNeuralNetworkInferenceSettings InferenceSettings;
+		TSharedPtr<FNeuralNetworkFunction> NeuralNetworkFunction;
 
 		TLearningArray<2, float> Input;
 		TLearningArray<2, float> Output;
