LearningAgents: Updated to use new python site-packages location.

diff --git a/Engine/Plugins/Experimental/LearningAgents/Source/LearningAgentsTraining/Private/LearningAgentsImitationTrainer.cpp b/Engine/Plugins/Experimental/LearningAgents/Source/LearningAgentsTraining/Private/LearningAgentsImitationTrainer.cpp
--- a/Engine/Plugins/Experimental/LearningAgents/Source/LearningAgentsTraining/Private/LearningAgentsImitationTrainer.cpp
+++ b/Engine/Plugins/Experimental/LearningAgents/Source/LearningAgentsTraining/Private/LearningAgentsImitationTrainer.cpp
@@ -160,7 +160,7 @@ void ULearningAgentsImitationTrainer::BeginTraining(
 		return;
 	}
 
-	const FString SitePackagesPath = UE::Learning::Trainer::GetSitePackagesPath(ImitationTrainerPathSettings.GetEditorEnginePath());
+	const FString SitePackagesPath = UE::Learning::Trainer::GetSitePackagesPath(ImitationTrainerPathSettings.GetIntermediatePath());
 
 	if (!FPaths::DirectoryExists(SitePackagesPath))
 	{
diff --git a/Engine/Plugins/Experimental/LearningAgents/Source/LearningAgentsTraining/Private/LearningAgentsTrainer.cpp b/Engine/Plugins/Experimental/LearningAgents/Source/LearningAgentsTraining/Private/LearningAgentsTrainer.cpp
--- a/Engine/Plugins/Experimental/LearningAgents/Source/LearningAgentsTraining/Private/LearningAgentsTrainer.cpp
+++ b/Engine/Plugins/Experimental/LearningAgents/Source/LearningAgentsTraining/Private/LearningAgentsTrainer.cpp
@@ -405,7 +405,7 @@ void ULearningAgentsTrainer::BeginTraining(
 		return;
 	}
 
-	const FString SitePackagesPath = UE::Learning::Trainer::GetSitePackagesPath(TrainerPathSettings.GetEditorEnginePath());
+	const FString SitePackagesPath = UE::Learning::Trainer::GetSitePackagesPath(TrainerPathSettings.GetIntermediatePath());
 
 	if (!FPaths::DirectoryExists(SitePackagesPath))
 	{
diff --git a/Engine/Plugins/Experimental/LearningAgents/Source/LearningTraining/Private/LearningImitationTrainer.cpp b/Engine/Plugins/Experimental/LearningAgents/Source/LearningTraining/Private/LearningImitationTrainer.cpp
--- a/Engine/Plugins/Experimental/LearningAgents/Source/LearningTraining/Private/LearningImitationTrainer.cpp
+++ b/Engine/Plugins/Experimental/LearningAgents/Source/LearningTraining/Private/LearningImitationTrainer.cpp
@@ -47,9 +47,9 @@ int32 ULearningSocketImitationTrainerServerCommandlet::Main(const FString& Comma
 	const FString* LogSettingsParam = Params.Find(TEXT("LogSettings"));
 
 	const FString PythonExecutablePath = PythonExecutablePathParam ? *PythonExecutablePathParam : UE::Learning::Trainer::GetPythonExecutablePath(FPaths::EngineDir());
-	const FString SitePackagesPath = SitePackagesPathParam ? *SitePackagesPathParam : UE::Learning::Trainer::GetSitePackagesPath(FPaths::EngineDir());
+	const FString SitePackagesPath = SitePackagesPathParam ? *SitePackagesPathParam : UE::Learning::Trainer::GetSitePackagesPath(FPaths::ProjectIntermediateDir());
 	const FString PythonContentPath = PythonContentPathParam ? *PythonContentPathParam : UE::Learning::Trainer::GetPythonContentPath(FPaths::EngineDir());
-	const FString IntermediatePath = IntermediatePathParam ? *IntermediatePathParam : UE::Learning::Trainer::GetIntermediatePath(FPaths::EngineDir());
+	const FString IntermediatePath = IntermediatePathParam ? *IntermediatePathParam : UE::Learning::Trainer::GetIntermediatePath(FPaths::ProjectIntermediateDir());
 
 	const TCHAR* IpAddress = IpAddressParam ? *(*IpAddressParam) : UE::Learning::Trainer::DefaultIp;
 	const uint32 Port = PortParam ? FCString::Atoi(*(*PortParam)) : UE::Learning::Trainer::DefaultPort;
diff --git a/Engine/Plugins/Experimental/LearningAgents/Source/LearningTraining/Private/LearningPPOTrainer.cpp b/Engine/Plugins/Experimental/LearningAgents/Source/LearningTraining/Private/LearningPPOTrainer.cpp
--- a/Engine/Plugins/Experimental/LearningAgents/Source/LearningTraining/Private/LearningPPOTrainer.cpp
+++ b/Engine/Plugins/Experimental/LearningAgents/Source/LearningTraining/Private/LearningPPOTrainer.cpp
@@ -52,9 +52,9 @@ int32 ULearningSocketPPOTrainerServerCommandlet::Main(const FString& Commandline
 	const FString* LogSettingsParam = Params.Find(TEXT("LogSettings"));
 
 	const FString PythonExecutiblePath = PythonExecutiblePathParam ? *PythonExecutiblePathParam : UE::Learning::Trainer::GetPythonExecutablePath(FPaths::EngineDir());
-	const FString SitePackagesPath = SitePackagesPathParam ? *SitePackagesPathParam : UE::Learning::Trainer::GetSitePackagesPath(FPaths::EngineDir());
+	const FString SitePackagesPath = SitePackagesPathParam ? *SitePackagesPathParam : UE::Learning::Trainer::GetSitePackagesPath(FPaths::ProjectIntermediateDir());
 	const FString PythonContentPath = PythonContentPathParam ? *PythonContentPathParam : UE::Learning::Trainer::GetPythonContentPath(FPaths::EngineDir());
-	const FString IntermediatePath = IntermediatePathParam ? *IntermediatePathParam : UE::Learning::Trainer::GetIntermediatePath(FPaths::EngineDir());
+	const FString IntermediatePath = IntermediatePathParam ? *IntermediatePathParam : UE::Learning::Trainer::GetIntermediatePath(FPaths::ProjectIntermediateDir());
 
 	const TCHAR* IpAddress = IpAddressParam ? *(*IpAddressParam) : UE::Learning::Trainer::DefaultIp;
 	const uint32 Port = PortParam ? FCString::Atoi(*(*PortParam)) : UE::Learning::Trainer::DefaultPort;
diff --git a/Engine/Plugins/Experimental/LearningAgents/Source/LearningTraining/Private/LearningTrainer.cpp b/Engine/Plugins/Experimental/LearningAgents/Source/LearningTraining/Private/LearningTrainer.cpp
--- a/Engine/Plugins/Experimental/LearningAgents/Source/LearningTraining/Private/LearningTrainer.cpp
+++ b/Engine/Plugins/Experimental/LearningAgents/Source/LearningTraining/Private/LearningTrainer.cpp
@@ -299,9 +299,9 @@ namespace UE::Learning::Trainer
 		return EngineDir / TEXT("Binaries/ThirdParty/Python3") / FPlatformMisc::GetUBTPlatform() / (PLATFORM_WINDOWS ? TEXT("python.exe") : TEXT("bin/python"));
 	}
 
-	FString GetSitePackagesPath(const FString& EngineDir)
+	FString GetSitePackagesPath(const FString& IntermediateDir)
 	{
-		return EngineDir / TEXT("Plugins/Experimental/PythonFoundationPackages/Content/Python/Lib") / FPlatformMisc::GetUBTPlatform() / TEXT("site-packages");
+		return IntermediateDir / TEXT("PipInstall/Lib/site-packages");
 	}
 
 	FString GetPythonContentPath(const FString& EngineDir)
diff --git a/Engine/Plugins/Experimental/LearningAgents/Source/LearningTraining/Public/LearningTrainer.h b/Engine/Plugins/Experimental/LearningAgents/Source/LearningTraining/Public/LearningTrainer.h
--- a/Engine/Plugins/Experimental/LearningAgents/Source/LearningTraining/Public/LearningTrainer.h
+++ b/Engine/Plugins/Experimental/LearningAgents/Source/LearningTraining/Public/LearningTrainer.h
@@ -123,9 +123,9 @@ namespace UE::Learning
 		LEARNINGTRAINING_API FString GetPythonExecutablePath(const FString& EngineDir);
 
 		/**
-		* Gets the PythonFoundationPackages site-packages path from the engine directory.
+		* Gets the Project site-packages path from the intermediate directory.
 		*/
-		LEARNINGTRAINING_API FString GetSitePackagesPath(const FString& EngineDir);
+		LEARNINGTRAINING_API FString GetSitePackagesPath(const FString& IntermediateDir);
 
 		/**
 		* Gets the LearningAgents Content path from the engine directory.
