LearningAgents: Naming changes and switched imitation learning python script to use OrderedDict.

diff --git a/Engine/Plugins/Experimental/LearningAgents/Source/LearningAgents/Private/LearningAgentsActions.cpp b/Engine/Plugins/Experimental/LearningAgents/Source/LearningAgents/Private/LearningAgentsActions.cpp
--- a/Engine/Plugins/Experimental/LearningAgents/Source/LearningAgents/Private/LearningAgentsActions.cpp
+++ b/Engine/Plugins/Experimental/LearningAgents/Source/LearningAgents/Private/LearningAgentsActions.cpp
@@ -1043,7 +1043,7 @@ FLearningAgentsActionSchemaElement ULearningAgentsActionSchema::SpecifyFloatActi
 	return SpecifyContinuousAction(1, Name);
 }
 
-FLearningAgentsActionSchemaElement ULearningAgentsActionSchema::SpecifyTranslationAction(const FName Name)
+FLearningAgentsActionSchemaElement ULearningAgentsActionSchema::SpecifyLocationAction(const FName Name)
 {
 	return SpecifyContinuousAction(3, Name);
 }
@@ -1062,12 +1062,12 @@ FLearningAgentsActionSchemaElement ULearningAgentsActionSchema::SpecifyTransform
 {
 	return SpecifyStructActionFromArrayViews(
 		{
-			TEXT("Translation"),
+			TEXT("Location"),
 			TEXT("Rotation"),
 			TEXT("Scale")
 		},
 		{
-			SpecifyTranslationAction(),
+			SpecifyLocationAction(),
 			SpecifyRotationAction(),
 			SpecifyScaleAction()
 		}, 
@@ -1426,14 +1426,14 @@ FLearningAgentsActionObjectElement ULearningAgentsActionObject::MakeFloatAction(
 	return MakeContinuousActionFromArrayView({ Value / FMath::Max(FloatScale, UE_SMALL_NUMBER) }, Name);
 }
 
-FLearningAgentsActionObjectElement ULearningAgentsActionObject::MakeTranslationAction(const FVector Translation, const FTransform RelativeTransform, const float TranslationScale, const FName Name)
+FLearningAgentsActionObjectElement ULearningAgentsActionObject::MakeLocationAction(const FVector Location, const FTransform RelativeTransform, const float LocationScale, const FName Name)
 {
-	const FVector LocalTranslation = RelativeTransform.InverseTransformPosition(Translation);
+	const FVector LocalLocation = RelativeTransform.InverseTransformPosition(Location);
 
 	return MakeContinuousActionFromArrayView({
-		(float)LocalTranslation.X / FMath::Max(TranslationScale, UE_SMALL_NUMBER),
-		(float)LocalTranslation.Y / FMath::Max(TranslationScale, UE_SMALL_NUMBER),
-		(float)LocalTranslation.Z / FMath::Max(TranslationScale, UE_SMALL_NUMBER) }, Name);
+		(float)LocalLocation.X / FMath::Max(LocationScale, UE_SMALL_NUMBER),
+		(float)LocalLocation.Y / FMath::Max(LocationScale, UE_SMALL_NUMBER),
+		(float)LocalLocation.Z / FMath::Max(LocationScale, UE_SMALL_NUMBER) }, Name);
 }
 
 FLearningAgentsActionObjectElement ULearningAgentsActionObject::MakeRotationAction(const FRotator Rotation, const FRotator RelativeRotation, const float RotationScale, const FName Name)
@@ -1467,18 +1467,18 @@ FLearningAgentsActionObjectElement ULearningAgentsActionObject::MakeScaleAction(
 		}, Name);
 }
 
-FLearningAgentsActionObjectElement ULearningAgentsActionObject::MakeTransformAction(const FTransform Transform, const FTransform RelativeTransform, const float TranslationScale, const FName Name)
+FLearningAgentsActionObjectElement ULearningAgentsActionObject::MakeTransformAction(const FTransform Transform, const FTransform RelativeTransform, const float LocationScale, const FName Name)
 {
 	const FTransform LocalTransform = Transform * RelativeTransform.Inverse();
 
 	return MakeStructActionFromArrayViews(
 		{
-			TEXT("Translation"),
+			TEXT("Location"),
 			TEXT("Rotation"),
 			TEXT("Scale")
 		},
 		{
-			MakeTranslationAction(LocalTransform.GetTranslation(), FTransform::Identity, TranslationScale),
+			MakeLocationAction(LocalTransform.GetLocation(), FTransform::Identity, LocationScale),
 			MakeRotationActionFromQuat(LocalTransform.GetRotation(), FQuat::Identity),
 			MakeScaleAction(LocalTransform.GetScale3D(), FVector::OneVector)
 		},
@@ -2349,17 +2349,17 @@ bool ULearningAgentsActionObject::GetFloatAction(float& OutValue, const FLearnin
 	return true;
 }
 
-bool ULearningAgentsActionObject::GetTranslationAction(FVector& OutTranslation, const FLearningAgentsActionObjectElement Element, const FTransform RelativeTransform, const float TranslationScale, const FName Name) const
+bool ULearningAgentsActionObject::GetLocationAction(FVector& OutLocation, const FLearningAgentsActionObjectElement Element, const FTransform RelativeTransform, const float LocationScale, const FName Name) const
 {
 	TStaticArray<float, 3> OutValues;
 	if (!GetContinuousActionToArrayView(OutValues, Element, Name))
 	{
-		OutTranslation = FVector::ZeroVector;
+		OutLocation = FVector::ZeroVector;
 		return false;
 	}
 
-	const FVector LocalTranslation = TranslationScale * FVector(OutValues[0], OutValues[1], OutValues[2]);
-	OutTranslation = RelativeTransform.TransformPosition(LocalTranslation);
+	const FVector LocalLocation = LocationScale * FVector(OutValues[0], OutValues[1], OutValues[2]);
+	OutLocation = RelativeTransform.TransformPosition(LocalLocation);
 	return true;
 }
 
@@ -2404,7 +2404,7 @@ bool ULearningAgentsActionObject::GetScaleAction(FVector& OutScale, const FLearn
 	return true;
 }
 
-bool ULearningAgentsActionObject::GetTransformAction(FTransform& OutTransform, const FLearningAgentsActionObjectElement Element, const FTransform RelativeTransform, const float TranslationScale, const float RotationScale, const float ScaleScale, const FName Name) const
+bool ULearningAgentsActionObject::GetTransformAction(FTransform& OutTransform, const FLearningAgentsActionObjectElement Element, const FTransform RelativeTransform, const float LocationScale, const float RotationScale, const float ScaleScale, const FName Name) const
 {
 	TStaticArray<FName, 3> OutElementNames;
 	TStaticArray<FLearningAgentsActionObjectElement, 3> OutElements;
@@ -2414,8 +2414,8 @@ bool ULearningAgentsActionObject::GetTransformAction(FTransform& OutTransform, c
 		return false;
 	}
 
-	FVector OutTranslation;
-	if (!GetTranslationAction(OutTranslation, OutElements[MakeArrayView(OutElementNames).Find(TEXT("Translation"))], RelativeTransform, TranslationScale))
+	FVector OutLocation;
+	if (!GetLocationAction(OutLocation, OutElements[MakeArrayView(OutElementNames).Find(TEXT("Location"))], RelativeTransform, LocationScale))
 	{
 		OutTransform = FTransform::Identity;
 		return false;
@@ -2435,7 +2435,7 @@ bool ULearningAgentsActionObject::GetTransformAction(FTransform& OutTransform, c
 		return false;
 	}
 
-	OutTransform = FTransform(OutRotation, OutTranslation, OutScale);
+	OutTransform = FTransform(OutRotation, OutLocation, OutScale);
 	return true;
 }
 
diff --git a/Engine/Plugins/Experimental/LearningAgents/Source/LearningAgents/Private/LearningAgentsInteractor.cpp b/Engine/Plugins/Experimental/LearningAgents/Source/LearningAgents/Private/LearningAgentsInteractor.cpp
--- a/Engine/Plugins/Experimental/LearningAgents/Source/LearningAgents/Private/LearningAgentsInteractor.cpp
+++ b/Engine/Plugins/Experimental/LearningAgents/Source/LearningAgents/Private/LearningAgentsInteractor.cpp
@@ -228,7 +228,7 @@ void ULearningAgentsInteractor::ScatterAgentActions_Implementation(const TArray<
 	UE_LOG(LogLearning, Error, TEXT("%s: ScatterAgentActions function must be overridden!"), *GetName());
 }
 
-void  ULearningAgentsInteractor::ScatterAgentObservationsUsingDelegate(
+void  ULearningAgentsInteractor::ScatterAgentActionsUsingDelegate(
 	const TArray<FLearningAgentsActionObjectElement>& InActionObjectElements,
 	const ULearningAgentsActionObject* InActionObject,
 	const TArray<int32>& AgentIds,
diff --git a/Engine/Plugins/Experimental/LearningAgents/Source/LearningAgents/Private/LearningAgentsObservations.cpp b/Engine/Plugins/Experimental/LearningAgents/Source/LearningAgents/Private/LearningAgentsObservations.cpp
--- a/Engine/Plugins/Experimental/LearningAgents/Source/LearningAgents/Private/LearningAgentsObservations.cpp
+++ b/Engine/Plugins/Experimental/LearningAgents/Source/LearningAgents/Private/LearningAgentsObservations.cpp
@@ -919,7 +919,7 @@ FLearningAgentsObservationSchemaElement ULearningAgentsObservationSchema::Specif
 	return SpecifyContinuousObservation(1, Name);
 }
 
-FLearningAgentsObservationSchemaElement ULearningAgentsObservationSchema::SpecifyTranslationObservation(const FName Name)
+FLearningAgentsObservationSchemaElement ULearningAgentsObservationSchema::SpecifyLocationObservation(const FName Name)
 {
 	return SpecifyContinuousObservation(3, Name);
 }
@@ -938,12 +938,12 @@ FLearningAgentsObservationSchemaElement ULearningAgentsObservationSchema::Specif
 {
 	return SpecifyStructObservationFromArrayViews(
 		{
-			TEXT("Translation"),
+			TEXT("Location"),
 			TEXT("Rotation"),
 			TEXT("Scale")
 		},
 		{
-			SpecifyTranslationObservation(),
+			SpecifyLocationObservation(),
 			SpecifyRotationObservation(),
 			SpecifyScaleObservation()
 		}, 
@@ -967,7 +967,7 @@ FLearningAgentsObservationSchemaElement ULearningAgentsObservationSchema::Specif
 
 FLearningAgentsObservationSchemaElement ULearningAgentsObservationSchema::SpecifyLocationAlongSplineObservation(const FName Name)
 {
-	return SpecifyTranslationObservation(Name);
+	return SpecifyLocationObservation(Name);
 }
 
 FLearningAgentsObservationSchemaElement ULearningAgentsObservationSchema::SpecifyProportionAlongSplineObservation(const FName Name)
@@ -1500,14 +1500,14 @@ FLearningAgentsObservationObjectElement ULearningAgentsObservationObject::MakeFl
 	return MakeContinuousObservationFromArrayView({ Value / FMath::Max(FloatScale, UE_SMALL_NUMBER) }, Name);
 }
 
-FLearningAgentsObservationObjectElement ULearningAgentsObservationObject::MakeTranslationObservation(const FVector Translation, const FTransform RelativeTransform, const float TranslationScale, const FName Name)
+FLearningAgentsObservationObjectElement ULearningAgentsObservationObject::MakeLocationObservation(const FVector Location, const FTransform RelativeTransform, const float LocationScale, const FName Name)
 {
-	const FVector LocalTranslation = RelativeTransform.InverseTransformPosition(Translation);
+	const FVector LocalLocation = RelativeTransform.InverseTransformPosition(Location);
 
 	return MakeContinuousObservationFromArrayView({
-		(float)LocalTranslation.X / FMath::Max(TranslationScale, UE_SMALL_NUMBER),
-		(float)LocalTranslation.Y / FMath::Max(TranslationScale, UE_SMALL_NUMBER),
-		(float)LocalTranslation.Z / FMath::Max(TranslationScale, UE_SMALL_NUMBER),
+		(float)LocalLocation.X / FMath::Max(LocationScale, UE_SMALL_NUMBER),
+		(float)LocalLocation.Y / FMath::Max(LocationScale, UE_SMALL_NUMBER),
+		(float)LocalLocation.Z / FMath::Max(LocationScale, UE_SMALL_NUMBER),
 		}, Name);
 }
 
@@ -1545,18 +1545,18 @@ FLearningAgentsObservationObjectElement ULearningAgentsObservationObject::MakeSc
 		}, Name);
 }
 
-FLearningAgentsObservationObjectElement ULearningAgentsObservationObject::MakeTransformObservation(const FTransform Transform, const FTransform RelativeTransform, const float TranslationScale, const FName Name)
+FLearningAgentsObservationObjectElement ULearningAgentsObservationObject::MakeTransformObservation(const FTransform Transform, const FTransform RelativeTransform, const float LocationScale, const FName Name)
 {
 	const FTransform LocalTransform = Transform * RelativeTransform.Inverse();
 	
 	return MakeStructObservationFromArrayViews(
 		{
-			TEXT("Translation"),
+			TEXT("Location"),
 			TEXT("Rotation"),
 			TEXT("Scale")
 		},
 		{
-			MakeTranslationObservation(LocalTransform.GetTranslation(), FTransform::Identity, TranslationScale),
+			MakeLocationObservation(LocalTransform.GetLocation(), FTransform::Identity, LocationScale),
 			MakeRotationObservationFromQuat(LocalTransform.GetRotation(), FQuat::Identity),
 			MakeScaleObservation(LocalTransform.GetScale3D(), FVector::OneVector)
 		},
@@ -1600,7 +1600,7 @@ FLearningAgentsObservationObjectElement ULearningAgentsObservationObject::MakeDi
 		}, Name);
 }
 
-FLearningAgentsObservationObjectElement ULearningAgentsObservationObject::MakeLocationAlongSplineObservation(const USplineComponent* SplineComponent, const float DistanceAlongSpline, const FTransform RelativeTransform, const float TranslationScale, const FName Name)
+FLearningAgentsObservationObjectElement ULearningAgentsObservationObject::MakeLocationAlongSplineObservation(const USplineComponent* SplineComponent, const float DistanceAlongSpline, const FTransform RelativeTransform, const float LocationScale, const FName Name)
 {
 	if (!SplineComponent)
 	{
@@ -1608,7 +1608,7 @@ FLearningAgentsObservationObjectElement ULearningAgentsObservationObject::MakeLo
 		return FLearningAgentsObservationObjectElement();
 	}
 
-	return MakeTranslationObservation(SplineComponent->GetLocationAtDistanceAlongSpline(DistanceAlongSpline, ESplineCoordinateSpace::World), RelativeTransform, TranslationScale, Name);
+	return MakeLocationObservation(SplineComponent->GetLocationAtDistanceAlongSpline(DistanceAlongSpline, ESplineCoordinateSpace::World), RelativeTransform, LocationScale, Name);
 }
 
 FLearningAgentsObservationObjectElement ULearningAgentsObservationObject::MakeProportionAlongSplineObservation(const USplineComponent* SplineComponent, const float DistanceAlongSpline, const FName Name)
@@ -1645,7 +1645,7 @@ FLearningAgentsObservationObjectElement ULearningAgentsObservationObject::MakeDi
 	return MakeDirectionObservation(SplineComponent->GetDirectionAtDistanceAlongSpline(DistanceAlongSpline, ESplineCoordinateSpace::World), RelativeTransform, Name);
 }
 
-FLearningAgentsObservationObjectElement ULearningAgentsObservationObject::MakePropertiesAlongSplineObservation(const USplineComponent* SplineComponent, const float DistanceAlongSpline, const FTransform RelativeTransform, const float TranslationScale, const FName Name)
+FLearningAgentsObservationObjectElement ULearningAgentsObservationObject::MakePropertiesAlongSplineObservation(const USplineComponent* SplineComponent, const float DistanceAlongSpline, const FTransform RelativeTransform, const float LocationScale, const FName Name)
 {
 	if (!SplineComponent)
 	{
@@ -1660,7 +1660,7 @@ FLearningAgentsObservationObjectElement ULearningAgentsObservationObject::MakePr
 			TEXT("Direction")
 		},
 		{
-			MakeLocationAlongSplineObservation(SplineComponent, DistanceAlongSpline, RelativeTransform, TranslationScale),
+			MakeLocationAlongSplineObservation(SplineComponent, DistanceAlongSpline, RelativeTransform, LocationScale),
 			MakeProportionAlongSplineObservation(SplineComponent, DistanceAlongSpline),
 			MakeDirectionAlongSplineObservation(SplineComponent, DistanceAlongSpline, RelativeTransform),
 		}, Name);
@@ -2898,16 +2898,16 @@ bool ULearningAgentsObservationObject::GetFloatObservation(float& OutValue, cons
 	return true;
 }
 
-bool ULearningAgentsObservationObject::GetTranslationObservation(FVector& OutTranslation, const FLearningAgentsObservationObjectElement Element, const FTransform RelativeTransform, const float TranslationScale, const FName Name) const
+bool ULearningAgentsObservationObject::GetLocationObservation(FVector& OutLocation, const FLearningAgentsObservationObjectElement Element, const FTransform RelativeTransform, const float LocationScale, const FName Name) const
 {
 	TStaticArray<float, 3> OutValues;
 	if (!GetContinuousObservationToArrayView(OutValues, Element, Name))
 	{
-		OutTranslation = FVector::ZeroVector;
+		OutLocation = FVector::ZeroVector;
 		return false;
 	}
 
-	OutTranslation = RelativeTransform.TransformPosition(TranslationScale * FVector(OutValues[0], OutValues[1], OutValues[2]));
+	OutLocation = RelativeTransform.TransformPosition(LocationScale * FVector(OutValues[0], OutValues[1], OutValues[2]));
 	return true;
 }
 
@@ -2961,7 +2961,7 @@ bool ULearningAgentsObservationObject::GetScaleObservation(FVector& OutScale, co
 	return true;
 }
 
-bool ULearningAgentsObservationObject::GetTransformObservation(FTransform& OutTransform, const FLearningAgentsObservationObjectElement Element, const FTransform RelativeTransform, const float TranslationScale, const FName Name) const
+bool ULearningAgentsObservationObject::GetTransformObservation(FTransform& OutTransform, const FLearningAgentsObservationObjectElement Element, const FTransform RelativeTransform, const float LocationScale, const FName Name) const
 {
 	TStaticArray<FName, 3> OutElementNames;
 	TStaticArray<FLearningAgentsObservationObjectElement, 3> OutElements;
@@ -2971,9 +2971,9 @@ bool ULearningAgentsObservationObject::GetTransformObservation(FTransform& OutTr
 		return false;
 	}
 
-	const int32 TranslationElement = MakeArrayView(OutElementNames).Find(TEXT("Translation"));
-	FVector OutTranslation;
-	if (TranslationElement == INDEX_NONE || !GetTranslationObservation(OutTranslation, OutElements[TranslationElement], RelativeTransform, TranslationScale))
+	const int32 LocationElement = MakeArrayView(OutElementNames).Find(TEXT("Location"));
+	FVector OutLocation;
+	if (LocationElement == INDEX_NONE || !GetLocationObservation(OutLocation, OutElements[LocationElement], RelativeTransform, LocationScale))
 	{
 		OutTransform = FTransform::Identity;
 		return false;
@@ -2995,7 +2995,7 @@ bool ULearningAgentsObservationObject::GetTransformObservation(FTransform& OutTr
 		return false;
 	}
 
-	OutTransform = FTransform(OutRotation, OutTranslation, OutScale);
+	OutTransform = FTransform(OutRotation, OutLocation, OutScale);
 	return true;
 }
 
@@ -3050,9 +3050,9 @@ bool ULearningAgentsObservationObject::GetDirectionObservation(FVector& OutDirec
 	return true;
 }
 
-bool ULearningAgentsObservationObject::GetLocationAlongSplineObservation(FVector& OutLocation, const FLearningAgentsObservationObjectElement Element, const FTransform RelativeTransform, const float TranslationScale, const FName Name)
+bool ULearningAgentsObservationObject::GetLocationAlongSplineObservation(FVector& OutLocation, const FLearningAgentsObservationObjectElement Element, const FTransform RelativeTransform, const float LocationScale, const FName Name)
 {
-	return GetTranslationObservation(OutLocation, Element, RelativeTransform, TranslationScale, Name);
+	return GetLocationObservation(OutLocation, Element, RelativeTransform, LocationScale, Name);
 }
 
 bool ULearningAgentsObservationObject::GetProportionAlongSplineObservation(bool& bOutIsClosedLoop, float& OutAngle, float& OutPropotion, const FLearningAgentsObservationObjectElement Element, const FName Name)
@@ -3086,7 +3086,7 @@ bool ULearningAgentsObservationObject::GetDirectionAlongSplineObservation(FVecto
 	return GetDirectionObservation(OutDirection, Element, RelativeTransform, Name);
 }
 
-bool ULearningAgentsObservationObject::GetPropertiesAlongSplineObservation(FVector& OutLocation, bool& bOutIsClosedLoop, float& OutAngle, float& OutPropotion, FVector& OutDirection, const FLearningAgentsObservationObjectElement Element, const FTransform RelativeTransform, const float TranslationScale, const FName Name)
+bool ULearningAgentsObservationObject::GetPropertiesAlongSplineObservation(FVector& OutLocation, bool& bOutIsClosedLoop, float& OutAngle, float& OutPropotion, FVector& OutDirection, const FLearningAgentsObservationObjectElement Element, const FTransform RelativeTransform, const float LocationScale, const FName Name)
 {
 	TStaticArray<FName, 3> OutElementNames;
 	TStaticArray<FLearningAgentsObservationObjectElement, 3> OutElements;
@@ -3101,7 +3101,7 @@ bool ULearningAgentsObservationObject::GetPropertiesAlongSplineObservation(FVect
 	}
 
 	const int32 LocationElement = MakeArrayView(OutElementNames).Find(TEXT("Location"));
-	if (LocationElement == INDEX_NONE || !GetLocationAlongSplineObservation(OutLocation, OutElements[LocationElement], RelativeTransform, TranslationScale))
+	if (LocationElement == INDEX_NONE || !GetLocationAlongSplineObservation(OutLocation, OutElements[LocationElement], RelativeTransform, LocationScale))
 	{
 		OutLocation = FVector::ZeroVector;
 		bOutIsClosedLoop = false;
diff --git a/Engine/Plugins/Experimental/LearningAgents/Source/LearningAgents/Public/LearningAgentsActions.h b/Engine/Plugins/Experimental/LearningAgents/Source/LearningAgents/Public/LearningAgentsActions.h
--- a/Engine/Plugins/Experimental/LearningAgents/Source/LearningAgents/Public/LearningAgentsActions.h
+++ b/Engine/Plugins/Experimental/LearningAgents/Source/LearningAgents/Public/LearningAgentsActions.h
@@ -154,7 +154,7 @@ public:
 	FLearningAgentsActionSchemaElement SpecifyFloatAction(const FName Name = TEXT("Float"));
 
 	UFUNCTION(BlueprintPure, Category = "LearningAgents")
-	FLearningAgentsActionSchemaElement SpecifyTranslationAction(const FName Name = TEXT("Translation"));
+	FLearningAgentsActionSchemaElement SpecifyLocationAction(const FName Name = TEXT("Location"));
 
 	UFUNCTION(BlueprintPure, Category = "LearningAgents")
 	FLearningAgentsActionSchemaElement SpecifyRotationAction(const FName Name = TEXT("Rotation"));
@@ -276,7 +276,7 @@ public:
 	FLearningAgentsActionObjectElement MakeFloatAction(const float Value, const float FloatScale = 1.0f, const FName Name = TEXT("Float"));
 
 	UFUNCTION(BlueprintPure, Category = "LearningAgents")
-	FLearningAgentsActionObjectElement MakeTranslationAction(const FVector Translation, const FTransform RelativeTransform = FTransform(), const float TranslationScale = 100.0f, const FName Name = TEXT("Translation"));
+	FLearningAgentsActionObjectElement MakeLocationAction(const FVector Location, const FTransform RelativeTransform = FTransform(), const float LocationScale = 100.0f, const FName Name = TEXT("Location"));
 
 	UFUNCTION(BlueprintPure, Category = "LearningAgents")
 	FLearningAgentsActionObjectElement MakeRotationAction(const FRotator Rotation, const FRotator RelativeRotation = FRotator::ZeroRotator, const float RotationScale = 90.0f, const FName Name = TEXT("Rotation"));
@@ -286,7 +286,7 @@ public:
 	FLearningAgentsActionObjectElement MakeScaleAction(const FVector Scale, const FVector RelativeScale = FVector(1,1,1), const FName Name = TEXT("Scale"));
 
 	UFUNCTION(BlueprintPure, Category = "LearningAgents")
-	FLearningAgentsActionObjectElement MakeTransformAction(const FTransform Transform, const FTransform RelativeTransform = FTransform(), const float TranslationScale = 100.0f, const FName Name = TEXT("Transform"));
+	FLearningAgentsActionObjectElement MakeTransformAction(const FTransform Transform, const FTransform RelativeTransform = FTransform(), const float LocationScale = 100.0f, const FName Name = TEXT("Transform"));
 
 	UFUNCTION(BlueprintPure, Category = "LearningAgents")
 	FLearningAgentsActionObjectElement MakeAngleAction(const float Angle, const float RelativeAngle = 0.0f, const float AngleScale = 90.0f, const FName Name = TEXT("Angle"));
@@ -377,7 +377,7 @@ public:
 	UPARAM(DisplayName = "Success") bool GetFloatAction(float& OutValue, const FLearningAgentsActionObjectElement Element, const float FloatScale = 1.0f, const FName Name = TEXT("Float")) const;
 
 	UFUNCTION(BlueprintPure = false, Category = "LearningAgents")
-	UPARAM(DisplayName = "Success") bool GetTranslationAction(FVector& OutTranslation, const FLearningAgentsActionObjectElement Element, const FTransform RelativeTransform = FTransform(), const float TranslationScale = 100.0f, const FName Name = TEXT("Translation")) const;
+	UPARAM(DisplayName = "Success") bool GetLocationAction(FVector& OutLocation, const FLearningAgentsActionObjectElement Element, const FTransform RelativeTransform = FTransform(), const float LocationScale = 100.0f, const FName Name = TEXT("Location")) const;
 
 	UFUNCTION(BlueprintPure = false, Category = "LearningAgents")
 	UPARAM(DisplayName = "Success") bool GetRotationAction(FRotator& OutRotation, const FLearningAgentsActionObjectElement Element, const FRotator RelativeRotation = FRotator::ZeroRotator, const float RotationScale = 90.0f, const FName Name = TEXT("Rotation")) const;
@@ -387,7 +387,7 @@ public:
 	UPARAM(DisplayName = "Success") bool GetScaleAction(FVector& OutScale, const FLearningAgentsActionObjectElement Element, const FVector RelativeScale = FVector(1,1,1), const float Scale = 1.0f, const FName Name = TEXT("Scale")) const;
 
 	UFUNCTION(BlueprintPure = false, Category = "LearningAgents")
-	UPARAM(DisplayName = "Success") bool GetTransformAction(FTransform& OutTransform, const FLearningAgentsActionObjectElement Element, const FTransform RelativeTransform = FTransform(), const float TranslationScale = 100.0f, const float RotationScale = 1.0f, const float ScaleScale = 1.0f, const FName Name = TEXT("Transform")) const;
+	UPARAM(DisplayName = "Success") bool GetTransformAction(FTransform& OutTransform, const FLearningAgentsActionObjectElement Element, const FTransform RelativeTransform = FTransform(), const float LocationScale = 100.0f, const float RotationScale = 1.0f, const float ScaleScale = 1.0f, const FName Name = TEXT("Transform")) const;
 
 	UFUNCTION(BlueprintPure = false, Category = "LearningAgents")
 	UPARAM(DisplayName = "Success") bool GetAngleAction(float& OutAngle, const FLearningAgentsActionObjectElement Element, const float RelativeAngle = 0.0f, const float AngleScale = 90.0f, const FName Name = TEXT("Angle")) const;
diff --git a/Engine/Plugins/Experimental/LearningAgents/Source/LearningAgents/Public/LearningAgentsInteractor.h b/Engine/Plugins/Experimental/LearningAgents/Source/LearningAgents/Public/LearningAgentsInteractor.h
--- a/Engine/Plugins/Experimental/LearningAgents/Source/LearningAgents/Public/LearningAgentsInteractor.h
+++ b/Engine/Plugins/Experimental/LearningAgents/Source/LearningAgents/Public/LearningAgentsInteractor.h
@@ -143,11 +143,11 @@ public:
 	 *
 	 * @param InActionObjectElements			Input Actions Element. This will be the same size as AgentIds.
 	 * @param InActionObject					Action Object.
-	 * @param AgentIds							Set of Agent Ids to scatter observations for.
+	 * @param AgentIds							Set of Agent Ids to scatter actions for.
 	 * @param Delegate							The Delegate used to scatter each agent's actions.
 	 */
 	UFUNCTION(BlueprintPure = false, Category = "LearningAgents")
-	void ScatterAgentObservationsUsingDelegate(
+	void ScatterAgentActionsUsingDelegate(
 		const TArray<FLearningAgentsActionObjectElement>& InActionObjectElements, 
 		const ULearningAgentsActionObject* InActionObject,
 		const TArray<int32>& AgentIds,
diff --git a/Engine/Plugins/Experimental/LearningAgents/Source/LearningAgents/Public/LearningAgentsObservations.h b/Engine/Plugins/Experimental/LearningAgents/Source/LearningAgents/Public/LearningAgentsObservations.h
--- a/Engine/Plugins/Experimental/LearningAgents/Source/LearningAgents/Public/LearningAgentsObservations.h
+++ b/Engine/Plugins/Experimental/LearningAgents/Source/LearningAgents/Public/LearningAgentsObservations.h
@@ -178,7 +178,7 @@ public:
 	FLearningAgentsObservationSchemaElement SpecifyFloatObservation(const FName Name = TEXT("Float"));
 
 	UFUNCTION(BlueprintPure, Category = "LearningAgents")
-	FLearningAgentsObservationSchemaElement SpecifyTranslationObservation(const FName Name = TEXT("Translation"));
+	FLearningAgentsObservationSchemaElement SpecifyLocationObservation(const FName Name = TEXT("Location"));
 
 	UFUNCTION(BlueprintPure, Category = "LearningAgents")
 	FLearningAgentsObservationSchemaElement SpecifyRotationObservation(const FName Name = TEXT("Rotation"));
@@ -345,7 +345,7 @@ public:
 	FLearningAgentsObservationObjectElement MakeFloatObservation(const float Value, const float FloatScale = 1.0f, const FName Name = TEXT("Float"));
 
 	UFUNCTION(BlueprintPure, Category = "LearningAgents")
-	FLearningAgentsObservationObjectElement MakeTranslationObservation(const FVector Translation, const FTransform RelativeTransform = FTransform(), const float TranslationScale = 100.0f, const FName Name = TEXT("Translation"));
+	FLearningAgentsObservationObjectElement MakeLocationObservation(const FVector Location, const FTransform RelativeTransform = FTransform(), const float LocationScale = 100.0f, const FName Name = TEXT("Location"));
 
 	UFUNCTION(BlueprintPure, Category = "LearningAgents")
 	FLearningAgentsObservationObjectElement MakeRotationObservation(const FRotator Rotation, const FRotator RelativeRotation = FRotator::ZeroRotator, const FName Name = TEXT("Rotation"));
@@ -355,7 +355,7 @@ public:
 	FLearningAgentsObservationObjectElement MakeScaleObservation(const FVector Scale, const FVector RelativeScale = FVector(1,1,1), const FName Name = TEXT("Scale"));
 
 	UFUNCTION(BlueprintPure, Category = "LearningAgents")
-	FLearningAgentsObservationObjectElement MakeTransformObservation(const FTransform Transform, const FTransform RelativeTransform = FTransform(), const float TranslationScale = 100.0f, const FName Name = TEXT("Transform"));
+	FLearningAgentsObservationObjectElement MakeTransformObservation(const FTransform Transform, const FTransform RelativeTransform = FTransform(), const float LocationScale = 100.0f, const FName Name = TEXT("Transform"));
 
 	UFUNCTION(BlueprintPure, Category = "LearningAgents")
 	FLearningAgentsObservationObjectElement MakeAngleObservation(const float Angle, const float RelativeAngle = 0.0f, const FName Name = TEXT("Angle"));
@@ -372,7 +372,7 @@ public:
 	// Spline Observations
 
 	UFUNCTION(BlueprintPure, Category = "LearningAgents")
-	FLearningAgentsObservationObjectElement MakeLocationAlongSplineObservation(const USplineComponent* SplineComponent, const float DistanceAlongSpline, const FTransform RelativeTransform = FTransform(), const float TranslationScale = 100.0f, const FName Name = TEXT("LocationAlongSpline"));
+	FLearningAgentsObservationObjectElement MakeLocationAlongSplineObservation(const USplineComponent* SplineComponent, const float DistanceAlongSpline, const FTransform RelativeTransform = FTransform(), const float LocationScale = 100.0f, const FName Name = TEXT("LocationAlongSpline"));
 
 	UFUNCTION(BlueprintPure, Category = "LearningAgents")
 	FLearningAgentsObservationObjectElement MakeProportionAlongSplineObservation(const USplineComponent* SplineComponent, const float DistanceAlongSpline, const FName Name = TEXT("ProportionAlongSpline"));
@@ -381,7 +381,7 @@ public:
 	FLearningAgentsObservationObjectElement MakeDirectionAlongSplineObservation(const USplineComponent* SplineComponent, const float DistanceAlongSpline, const FTransform RelativeTransform = FTransform(), const FName Name = TEXT("DirectionAlongSpline"));
 
 	UFUNCTION(BlueprintPure, Category = "LearningAgents")
-	FLearningAgentsObservationObjectElement MakePropertiesAlongSplineObservation(const USplineComponent* SplineComponent, const float DistanceAlongSpline, const FTransform RelativeTransform = FTransform(), const float TranslationScale = 100.0f, const FName Name = TEXT("PropertiesAlongSpline"));
+	FLearningAgentsObservationObjectElement MakePropertiesAlongSplineObservation(const USplineComponent* SplineComponent, const float DistanceAlongSpline, const FTransform RelativeTransform = FTransform(), const float LocationScale = 100.0f, const FName Name = TEXT("PropertiesAlongSpline"));
 
 	// Ray Cast Observations
 
@@ -500,7 +500,7 @@ public:
 	UPARAM(DisplayName = "Success") bool GetFloatObservation(float& OutValue, const FLearningAgentsObservationObjectElement Element, const float FloatScale = 1.0f, const FName Name = TEXT("Float")) const;
 
 	UFUNCTION(BlueprintPure = false, Category = "LearningAgents")
-	UPARAM(DisplayName = "Success") bool GetTranslationObservation(FVector& OutTranslation, const FLearningAgentsObservationObjectElement Element, const FTransform RelativeTransform = FTransform(), const float TranslationScale = 100.0f, const FName Name = TEXT("Translation")) const;
+	UPARAM(DisplayName = "Success") bool GetLocationObservation(FVector& OutLocation, const FLearningAgentsObservationObjectElement Element, const FTransform RelativeTransform = FTransform(), const float LocationScale = 100.0f, const FName Name = TEXT("Location")) const;
 
 	UFUNCTION(BlueprintPure = false, Category = "LearningAgents")
 	UPARAM(DisplayName = "Success") bool GetRotationObservation(FRotator& OutRotation, const FLearningAgentsObservationObjectElement Element, const FRotator RelativeRotation = FRotator::ZeroRotator, const FName Name = TEXT("Rotation")) const;
@@ -510,7 +510,7 @@ public:
 	UPARAM(DisplayName = "Success") bool GetScaleObservation(FVector& OutScale, const FLearningAgentsObservationObjectElement Element, const FVector RelativeScale = FVector(1, 1, 1), const FName Name = TEXT("Scale")) const;
 
 	UFUNCTION(BlueprintPure = false, Category = "LearningAgents")
-	UPARAM(DisplayName = "Success") bool GetTransformObservation(FTransform& OutTransform, const FLearningAgentsObservationObjectElement Element, const FTransform RelativeTransform = FTransform(), const float TranslationScale = 100.0f, const FName Name = TEXT("Transform")) const;
+	UPARAM(DisplayName = "Success") bool GetTransformObservation(FTransform& OutTransform, const FLearningAgentsObservationObjectElement Element, const FTransform RelativeTransform = FTransform(), const float LocationScale = 100.0f, const FName Name = TEXT("Transform")) const;
 
 	UFUNCTION(BlueprintPure = false, Category = "LearningAgents")
 	UPARAM(DisplayName = "Success") bool GetAngleObservation(float& OutAngle, const FLearningAgentsObservationObjectElement Element, const float RelativeAngle = 0.0f, const FName Name = TEXT("Angle")) const;
@@ -527,7 +527,7 @@ public:
 	// Spline Observations
 
 	UFUNCTION(BlueprintPure = false, Category = "LearningAgents")
-	UPARAM(DisplayName = "Success") bool GetLocationAlongSplineObservation(FVector& OutLocation, const FLearningAgentsObservationObjectElement Element, const FTransform RelativeTransform = FTransform(), const float TranslationScale = 100.0f, const FName Name = TEXT("LocationAlongSpline"));
+	UPARAM(DisplayName = "Success") bool GetLocationAlongSplineObservation(FVector& OutLocation, const FLearningAgentsObservationObjectElement Element, const FTransform RelativeTransform = FTransform(), const float LocationScale = 100.0f, const FName Name = TEXT("LocationAlongSpline"));
 
 	UFUNCTION(BlueprintPure = false, Category = "LearningAgents")
 	UPARAM(DisplayName = "Success") bool GetProportionAlongSplineObservation(bool& bOutIsClosedLoop, float& OutAngle, float& OutPropotion, const FLearningAgentsObservationObjectElement Element, const FName Name = TEXT("ProportionAlongSpline"));
@@ -536,7 +536,7 @@ public:
 	UPARAM(DisplayName = "Success") bool GetDirectionAlongSplineObservation(FVector& OutDirection, const FLearningAgentsObservationObjectElement Element, const FTransform RelativeTransform = FTransform(), const FName Name = TEXT("DirectionAlongSpline"));
 
 	UFUNCTION(BlueprintPure = false, Category = "LearningAgents")
-	UPARAM(DisplayName = "Success") bool GetPropertiesAlongSplineObservation(FVector& OutLocation, bool& bOutIsClosedLoop, float& OutAngle, float& OutPropotion, FVector& OutDirection, const FLearningAgentsObservationObjectElement Element, const FTransform RelativeTransform = FTransform(), const float TranslationScale = 100.0f, const FName Name = TEXT("PropertiesAlongSpline"));
+	UPARAM(DisplayName = "Success") bool GetPropertiesAlongSplineObservation(FVector& OutLocation, bool& bOutIsClosedLoop, float& OutAngle, float& OutPropotion, FVector& OutDirection, const FLearningAgentsObservationObjectElement Element, const FTransform RelativeTransform = FTransform(), const float LocationScale = 100.0f, const FName Name = TEXT("PropertiesAlongSpline"));
 
 	// Ray Cast Observations
 
diff --git a/Engine/Plugins/Experimental/LearningAgents/Source/LearningAgentsTraining/Private/LearningAgentsCompletions.cpp b/Engine/Plugins/Experimental/LearningAgents/Source/LearningAgentsTraining/Private/LearningAgentsCompletions.cpp
--- a/Engine/Plugins/Experimental/LearningAgents/Source/LearningAgentsTraining/Private/LearningAgentsCompletions.cpp
+++ b/Engine/Plugins/Experimental/LearningAgents/Source/LearningAgentsTraining/Private/LearningAgentsCompletions.cpp
@@ -96,29 +96,29 @@ ELearningAgentsCompletion ULearningAgentsCompletions::CompletionOnEpisodeStepsRe
 	return CompletionOnCondition(EpisodeSteps >= MaxEpisodeSteps, CompletionType);
 }
 
-ELearningAgentsCompletion ULearningAgentsCompletions::CompletionOnTranslationDifferenceBelowThreshold(const FVector TranslationA, const FVector TranslationB, const float DistanceThreshold, const ELearningAgentsCompletion CompletionType)
+ELearningAgentsCompletion ULearningAgentsCompletions::CompletionOnLocationDifferenceBelowThreshold(const FVector LocationA, const FVector LocationB, const float DistanceThreshold, const ELearningAgentsCompletion CompletionType)
 {
-	return CompletionOnCondition(FVector::Distance(TranslationA, TranslationB) < DistanceThreshold, CompletionType);
+	return CompletionOnCondition(FVector::Distance(LocationA, LocationB) < DistanceThreshold, CompletionType);
 }
 
-ELearningAgentsCompletion ULearningAgentsCompletions::CompletionOnTranslationDifferenceAboveThreshold(const FVector TranslationA, const FVector TranslationB, const float DistanceThreshold, const ELearningAgentsCompletion CompletionType)
+ELearningAgentsCompletion ULearningAgentsCompletions::CompletionOnLocationDifferenceAboveThreshold(const FVector LocationA, const FVector LocationB, const float DistanceThreshold, const ELearningAgentsCompletion CompletionType)
 {
-	return CompletionOnCondition(FVector::Distance(TranslationA, TranslationB) > DistanceThreshold, CompletionType);
+	return CompletionOnCondition(FVector::Distance(LocationA, LocationB) > DistanceThreshold, CompletionType);
 }
 
-ELearningAgentsCompletion ULearningAgentsCompletions::CompletionOnTranslationOutsideBounds(
-	const FVector Translation,
+ELearningAgentsCompletion ULearningAgentsCompletions::CompletionOnLocationOutsideBounds(
+	const FVector Location,
 	const FTransform BoundsTransform,
 	const FVector BoundsMins,
 	const FVector BoundsMaxs,
 	const ELearningAgentsCompletion CompletionType)
 {
-	const FVector LocalTranslation = BoundsTransform.InverseTransformPosition(Translation);
+	const FVector LocalLocation = BoundsTransform.InverseTransformPosition(Location);
 
 	return CompletionOnCondition(
-		LocalTranslation.X < BoundsMins.X || LocalTranslation.X > BoundsMaxs.X ||
-		LocalTranslation.Y < BoundsMins.Y || LocalTranslation.Y > BoundsMaxs.Y ||
-		LocalTranslation.Z < BoundsMins.Z || LocalTranslation.Z > BoundsMaxs.Z,
+		LocalLocation.X < BoundsMins.X || LocalLocation.X > BoundsMaxs.X ||
+		LocalLocation.Y < BoundsMins.Y || LocalLocation.Y > BoundsMaxs.Y ||
+		LocalLocation.Z < BoundsMins.Z || LocalLocation.Z > BoundsMaxs.Z,
 		CompletionType);
 }
 
diff --git a/Engine/Plugins/Experimental/LearningAgents/Source/LearningAgentsTraining/Private/LearningAgentsRewards.cpp b/Engine/Plugins/Experimental/LearningAgents/Source/LearningAgentsTraining/Private/LearningAgentsRewards.cpp
--- a/Engine/Plugins/Experimental/LearningAgents/Source/LearningAgentsTraining/Private/LearningAgentsRewards.cpp
+++ b/Engine/Plugins/Experimental/LearningAgents/Source/LearningAgentsTraining/Private/LearningAgentsRewards.cpp
@@ -11,15 +11,15 @@ float ULearningAgentsRewards::RewardOnCondition(const bool bCondition, const flo
 	return bCondition ? RewardScale : 0.0f;
 }
 
-float ULearningAgentsRewards::RewardOnTranslationDifferenceBelowThreshold(const FVector TranslationA, const FVector TranslationB, const float DistanceThreshold, const float RewardScale)
+float ULearningAgentsRewards::RewardOnLocationDifferenceBelowThreshold(const FVector LocationA, const FVector LocationB, const float DistanceThreshold, const float RewardScale)
 {
-	return RewardOnCondition(FVector::Distance(TranslationA, TranslationB) < DistanceThreshold, RewardScale);
+	return RewardOnCondition(FVector::Distance(LocationA, LocationB) < DistanceThreshold, RewardScale);
 }
 
-float ULearningAgentsRewards::RewardFromTranslationSimilarity(const FVector TranslationA, const FVector TranslationB, const float TranslationScale, const float RewardScale)
+float ULearningAgentsRewards::RewardFromLocationSimilarity(const FVector LocationA, const FVector LocationB, const float LocationScale, const float RewardScale)
 {
-	const float TranslationDifference = FVector::Dist(TranslationA, TranslationB);
-	return RewardScale * FMath::InvExpApprox(FMath::Square(TranslationDifference / FMath::Max(TranslationScale, UE_SMALL_NUMBER)));
+	const float LocationDifference = FVector::Dist(LocationA, LocationB);
+	return RewardScale * FMath::InvExpApprox(FMath::Square(LocationDifference / FMath::Max(LocationScale, UE_SMALL_NUMBER)));
 }
 
 float ULearningAgentsRewards::RewardFromAngleSimilarity(const float AngleA, const float AngleB, const float AngleScale, const float RewardScale)
@@ -94,15 +94,15 @@ float ULearningAgentsRewards::PenaltyOnCondition(const bool bCondition, const fl
 	return bCondition ? -PenaltyScale : 0.0f;
 }
 
-float ULearningAgentsRewards::PenaltyOnTranslationDifferenceAboveThreshold(const FVector TranslationA, const FVector TranslationB, const float DistanceThreshold, const float PenaltyScale)
+float ULearningAgentsRewards::PenaltyOnLocationDifferenceAboveThreshold(const FVector LocationA, const FVector LocationB, const float DistanceThreshold, const float PenaltyScale)
 {
-	return PenaltyOnCondition(FVector::Distance(TranslationA, TranslationB) > DistanceThreshold, PenaltyScale);
+	return PenaltyOnCondition(FVector::Distance(LocationA, LocationB) > DistanceThreshold, PenaltyScale);
 }
 
-float ULearningAgentsRewards::PenaltyFromTranslationDifference(const FVector TranslationA, const FVector TranslationB, const float TranslationScale, const float PenaltyScale)
+float ULearningAgentsRewards::PenaltyFromLocationDifference(const FVector LocationA, const FVector LocationB, const float LocationScale, const float PenaltyScale)
 {
-	const float TranslationDifference = FVector::Dist(TranslationA, TranslationB);
-	return PenaltyScale  * (-TranslationDifference / FMath::Max(TranslationScale, UE_SMALL_NUMBER));
+	const float LocationDifference = FVector::Dist(LocationA, LocationB);
+	return PenaltyScale  * (-LocationDifference / FMath::Max(LocationScale, UE_SMALL_NUMBER));
 }
 
 float ULearningAgentsRewards::PenaltyFromAngleDifference(const float AngleA, const float AngleB, const float AngleScale, const float PenaltyScale)
diff --git a/Engine/Plugins/Experimental/LearningAgents/Source/LearningAgentsTraining/Public/LearningAgentsCompletions.h b/Engine/Plugins/Experimental/LearningAgents/Source/LearningAgentsTraining/Public/LearningAgentsCompletions.h
--- a/Engine/Plugins/Experimental/LearningAgents/Source/LearningAgentsTraining/Public/LearningAgentsCompletions.h
+++ b/Engine/Plugins/Experimental/LearningAgents/Source/LearningAgentsTraining/Public/LearningAgentsCompletions.h
@@ -72,14 +72,14 @@ public:
 	static ELearningAgentsCompletion CompletionOnEpisodeStepsRecorded(const int32 EpisodeSteps, const int32 MaxEpisodeSteps = 64, const ELearningAgentsCompletion CompletionType = ELearningAgentsCompletion::Truncation);
 
 	UFUNCTION(BlueprintPure, Category = "LearningAgents")
-	static ELearningAgentsCompletion CompletionOnTranslationDifferenceBelowThreshold(const FVector TranslationA, const FVector TranslationB, const float DistanceThreshold = 100.0f, const ELearningAgentsCompletion CompletionType = ELearningAgentsCompletion::Termination);
+	static ELearningAgentsCompletion CompletionOnLocationDifferenceBelowThreshold(const FVector LocationA, const FVector LocationB, const float DistanceThreshold = 100.0f, const ELearningAgentsCompletion CompletionType = ELearningAgentsCompletion::Termination);
 
 	UFUNCTION(BlueprintPure, Category = "LearningAgents")
-	static ELearningAgentsCompletion CompletionOnTranslationDifferenceAboveThreshold(const FVector TranslationA, const FVector TranslationB, const float DistanceThreshold = 100.0f, const ELearningAgentsCompletion CompletionType = ELearningAgentsCompletion::Termination);
+	static ELearningAgentsCompletion CompletionOnLocationDifferenceAboveThreshold(const FVector LocationA, const FVector LocationB, const float DistanceThreshold = 100.0f, const ELearningAgentsCompletion CompletionType = ELearningAgentsCompletion::Termination);
 
 	UFUNCTION(BlueprintPure, Category = "LearningAgents")
-	static ELearningAgentsCompletion CompletionOnTranslationOutsideBounds(
-		const FVector Translation, 
+	static ELearningAgentsCompletion CompletionOnLocationOutsideBounds(
+		const FVector Location, 
 		const FTransform BoundsTransform = FTransform(),
 		const FVector BoundsMins = FVector(-100.0f, -100.0f, -100.0f),
 		const FVector BoundsMaxs = FVector(+100.0f, +100.0f, +100.0f),
diff --git a/Engine/Plugins/Experimental/LearningAgents/Source/LearningAgentsTraining/Public/LearningAgentsRewards.h b/Engine/Plugins/Experimental/LearningAgents/Source/LearningAgentsTraining/Public/LearningAgentsRewards.h
--- a/Engine/Plugins/Experimental/LearningAgents/Source/LearningAgentsTraining/Public/LearningAgentsRewards.h
+++ b/Engine/Plugins/Experimental/LearningAgents/Source/LearningAgentsTraining/Public/LearningAgentsRewards.h
@@ -21,10 +21,10 @@ public:
 	static float RewardOnCondition(const bool bCondition, const float RewardScale = 1.0f);
 
 	UFUNCTION(BlueprintPure, Category = "LearningAgents")
-	static float RewardOnTranslationDifferenceBelowThreshold(const FVector TranslationA, const FVector TranslationB, const float DistanceThreshold = 100.0f, const float RewardScale = 1.0f);
+	static float RewardOnLocationDifferenceBelowThreshold(const FVector LocationA, const FVector LocationB, const float DistanceThreshold = 100.0f, const float RewardScale = 1.0f);
 
 	UFUNCTION(BlueprintPure, Category = "LearningAgents")
-	static float RewardFromTranslationSimilarity(const FVector TranslationA, const FVector TranslationB, const float TranslationScale = 100.0f, const float RewardScale = 1.0f);
+	static float RewardFromLocationSimilarity(const FVector LocationA, const FVector LocationB, const float LocationScale = 100.0f, const float RewardScale = 1.0f);
 
 	UFUNCTION(BlueprintPure, Category = "LearningAgents")
 	static float RewardFromRotationSimilarity(const FRotator RotationA, const FRotator RotationB, const float AngleScale = 90.0f, const float RewardScale = 1.0f);
@@ -47,10 +47,10 @@ public:
 	static float PenaltyOnCondition(const bool bCondition, const float PenaltyScale = 1.0f);
 
 	UFUNCTION(BlueprintPure, Category = "LearningAgents")
-	static float PenaltyOnTranslationDifferenceAboveThreshold(const FVector TranslationA, const FVector TranslationB, const float DistanceThreshold = 100.0f, const float PenaltyScale = 1.0f);
+	static float PenaltyOnLocationDifferenceAboveThreshold(const FVector LocationA, const FVector LocationB, const float DistanceThreshold = 100.0f, const float PenaltyScale = 1.0f);
 
 	UFUNCTION(BlueprintPure, Category = "LearningAgents")
-	static float PenaltyFromTranslationDifference(const FVector TranslationA, const FVector TranslationB, const float TranslationScale = 100.0f, const float PenaltyScale = 1.0f);
+	static float PenaltyFromLocationDifference(const FVector LocationA, const FVector LocationB, const float LocationScale = 100.0f, const float PenaltyScale = 1.0f);
 
 	UFUNCTION(BlueprintPure, Category = "LearningAgents")
 	static float PenaltyFromAngleDifference(const float AngleA, const float AngleB, const float AngleScale = 90.0f, const float PenaltyScale = 1.0f);
