LearningAgents: Settings structs no longer explicitly required to be provided.

diff --git a/Engine/Plugins/Experimental/LearningAgents/Source/LearningAgents/Public/LearningAgentsCritic.h b/Engine/Plugins/Experimental/LearningAgents/Source/LearningAgents/Public/LearningAgentsCritic.h
--- a/Engine/Plugins/Experimental/LearningAgents/Source/LearningAgents/Public/LearningAgentsCritic.h
+++ b/Engine/Plugins/Experimental/LearningAgents/Source/LearningAgents/Public/LearningAgentsCritic.h
@@ -68,7 +68,7 @@ public:
 	 * @param CriticSettings				The critic settings to use on creation of a new critic network
 	 * @param Seed							Random seed to use for initializing the critic weights
 	 */
-	UFUNCTION(BlueprintCallable, Category = "LearningAgents", meta = (Class = "/Script/LearningAgents.LearningAgentsCritic", DeterminesOutputType = "Class"))
+	UFUNCTION(BlueprintCallable, Category = "LearningAgents", meta = (Class = "/Script/LearningAgents.LearningAgentsCritic", DeterminesOutputType = "Class", AutoCreateRefTerm = "CriticSettings"))
 	static ULearningAgentsCritic* MakeCritic(
 		ULearningAgentsManager* InManager,
 		ULearningAgentsInteractor* InInteractor,
@@ -93,7 +93,7 @@ public:
 	 * @param CriticSettings				The critic settings to use on creation of a new critic network
 	 * @param Seed							Random seed to use for initializing the critic weights
 	 */
-	UFUNCTION(BlueprintCallable, Category = "LearningAgents")
+	UFUNCTION(BlueprintCallable, Category = "LearningAgents", meta = (AutoCreateRefTerm = "CriticSettings"))
 	void SetupCritic(
 		ULearningAgentsManager* InManager,
 		ULearningAgentsInteractor* InInteractor,
diff --git a/Engine/Plugins/Experimental/LearningAgents/Source/LearningAgents/Public/LearningAgentsPolicy.h b/Engine/Plugins/Experimental/LearningAgents/Source/LearningAgents/Public/LearningAgentsPolicy.h
--- a/Engine/Plugins/Experimental/LearningAgents/Source/LearningAgents/Public/LearningAgentsPolicy.h
+++ b/Engine/Plugins/Experimental/LearningAgents/Source/LearningAgents/Public/LearningAgentsPolicy.h
@@ -89,7 +89,7 @@ public:
 	 * @param PolicySettings				The policy settings to use on creation of a new policy network
 	 * @param Seed							Random seed to use for initializing network weights and policy sampling
 	 */
-	UFUNCTION(BlueprintCallable, Category = "LearningAgents", meta=(Class = "/Script/LearningAgents.LearningAgentsPolicy", DeterminesOutputType = "Class"))
+	UFUNCTION(BlueprintCallable, Category = "LearningAgents", meta=(Class = "/Script/LearningAgents.LearningAgentsPolicy", DeterminesOutputType = "Class", AutoCreateRefTerm = "PolicySettings"))
 	static ULearningAgentsPolicy* MakePolicy(
 		ULearningAgentsManager* InManager,
 		ULearningAgentsInteractor* InInteractor,
@@ -124,7 +124,7 @@ public:
 	 * @param PolicySettings				The policy settings to use on creation of a new policy network
 	 * @param Seed							Random seed to use for initializing network weights and policy sampling
 	 */
-	UFUNCTION(BlueprintCallable, Category = "LearningAgents")
+	UFUNCTION(BlueprintCallable, Category = "LearningAgents", meta = (AutoCreateRefTerm = "PolicySettings"))
 	void SetupPolicy(
 		ULearningAgentsManager* InManager,
 		ULearningAgentsInteractor* InInteractor,
diff --git a/Engine/Plugins/Experimental/LearningAgents/Source/LearningAgentsTraining/Public/LearningAgentsImitationTrainer.h b/Engine/Plugins/Experimental/LearningAgents/Source/LearningAgentsTraining/Public/LearningAgentsImitationTrainer.h
--- a/Engine/Plugins/Experimental/LearningAgents/Source/LearningAgentsTraining/Public/LearningAgentsImitationTrainer.h
+++ b/Engine/Plugins/Experimental/LearningAgents/Source/LearningAgentsTraining/Public/LearningAgentsImitationTrainer.h
@@ -177,7 +177,7 @@ public:
 	 * @param ImitationTrainerTrainingSettings	The training settings for this network.
 	 * @param ImitationTrainerPathSettings		The path settings used by the imitation trainer.
 	 */
-	UFUNCTION(BlueprintCallable, Category = "LearningAgents")
+	UFUNCTION(BlueprintCallable, Category = "LearningAgents", meta = (AutoCreateRefTerm = "ImitationTrainerSettings,ImitationTrainerTrainingSettings,ImitationTrainerPathSettings"))
 	void BeginTraining(
 		const ULearningAgentsRecording* Recording,
 		const FLearningAgentsImitationTrainerSettings& ImitationTrainerSettings = FLearningAgentsImitationTrainerSettings(),
@@ -201,7 +201,7 @@ public:
 	 * @param ImitationTrainerTrainingSettings	The training settings for this network.
 	 * @param ImitationTrainerPathSettings		The path settings used by the imitation trainer.
 	 */
-	UFUNCTION(BlueprintCallable, Category = "LearningAgents")
+	UFUNCTION(BlueprintCallable, Category = "LearningAgents", meta = (AutoCreateRefTerm = "ImitationTrainerSettings,ImitationTrainerTrainingSettings,ImitationTrainerPathSettings"))
 	void RunTraining(
 		const ULearningAgentsRecording* Recording,
 		const FLearningAgentsImitationTrainerSettings& ImitationTrainerSettings = FLearningAgentsImitationTrainerSettings(),
diff --git a/Engine/Plugins/Experimental/LearningAgents/Source/LearningAgentsTraining/Public/LearningAgentsRecorder.h b/Engine/Plugins/Experimental/LearningAgents/Source/LearningAgentsTraining/Public/LearningAgentsRecorder.h
--- a/Engine/Plugins/Experimental/LearningAgents/Source/LearningAgentsTraining/Public/LearningAgentsRecorder.h
+++ b/Engine/Plugins/Experimental/LearningAgents/Source/LearningAgentsTraining/Public/LearningAgentsRecorder.h
@@ -61,7 +61,7 @@ public:
 	 *									set then a new recording object will be created.
 	 * @pram bReinitializeRecording		If to reinitialize the recording asset
 	 */
-	UFUNCTION(BlueprintCallable, Category = "LearningAgents", meta = (Class = "/Script/LearningAgentsTraining.LearningAgentsRecorder", DeterminesOutputType = "Class"))
+	UFUNCTION(BlueprintCallable, Category = "LearningAgents", meta = (Class = "/Script/LearningAgentsTraining.LearningAgentsRecorder", DeterminesOutputType = "Class", AutoCreateRefTerm = "RecorderPathSettings"))
 	static ULearningAgentsRecorder* MakeRecorder(
 		ULearningAgentsManager* InManager,
 		ULearningAgentsInteractor* InInteractor,
@@ -81,7 +81,7 @@ public:
 	 *									set then a new recording object will be created.
 	 * @pram bReinitializeRecording		If to reinitialize the recording asset
 	 */
-	UFUNCTION(BlueprintCallable, Category = "LearningAgents")
+	UFUNCTION(BlueprintCallable, Category = "LearningAgents", meta = (AutoCreateRefTerm = "RecorderPathSettings"))
 	void SetupRecorder(
 		ULearningAgentsManager* InManager,
 		ULearningAgentsInteractor* InInteractor,
diff --git a/Engine/Plugins/Experimental/LearningAgents/Source/LearningAgentsTraining/Public/LearningAgentsTrainer.h b/Engine/Plugins/Experimental/LearningAgents/Source/LearningAgentsTraining/Public/LearningAgentsTrainer.h
--- a/Engine/Plugins/Experimental/LearningAgents/Source/LearningAgentsTraining/Public/LearningAgentsTrainer.h
+++ b/Engine/Plugins/Experimental/LearningAgents/Source/LearningAgentsTraining/Public/LearningAgentsTrainer.h
@@ -364,7 +364,7 @@ public:
 	 * @param Name				The trainer name
 	 * @param TrainerSettings	The trainer settings to use.
 	 */
-	UFUNCTION(BlueprintCallable, Category = "LearningAgents", meta = (DeterminesOutputType = "Class"))
+	UFUNCTION(BlueprintCallable, Category = "LearningAgents", meta = (DeterminesOutputType = "Class", AutoCreateRefTerm = "TrainerSettings"))
 	static ULearningAgentsTrainer* MakeTrainer(
 		ULearningAgentsManager* InManager,
 		ULearningAgentsInteractor* InInteractor,
@@ -383,7 +383,7 @@ public:
 	 * @param InCritic The critic to be trained.
 	 * @param TrainerSettings The trainer settings to use.
 	 */
-	UFUNCTION(BlueprintCallable, Category = "LearningAgents")
+	UFUNCTION(BlueprintCallable, Category = "LearningAgents", meta = (AutoCreateRefTerm = "TrainerSettings"))
 	void SetupTrainer(
 		ULearningAgentsManager* InManager,
 		ULearningAgentsInteractor* InInteractor,
@@ -483,7 +483,7 @@ public:
 	 * @param TrainerPathSettings		The path settings used by the trainer.
 	 * @param bResetAgentsOnBegin		If true, reset all agents at the beginning of training.
 	 */
-	UFUNCTION(BlueprintCallable, Category = "LearningAgents")
+	UFUNCTION(BlueprintCallable, Category = "LearningAgents", meta = (AutoCreateRefTerm = "TrainerTrainingSettings,TrainerGameSettings,TrainerPathSettings"))
 	void BeginTraining(
 		const FLearningAgentsTrainerTrainingSettings& TrainerTrainingSettings = FLearningAgentsTrainerTrainingSettings(),
 		const FLearningAgentsTrainerGameSettings& TrainerGameSettings = FLearningAgentsTrainerGameSettings(),
@@ -532,7 +532,7 @@ public:
 	 * @param bResetAgentsOnBegin				If true, reset all agents at the beginning of training.
 	 * @param bResetAgentsOnUpdate				If true, reset all agents whenever an updated policy is received.
 	 */
-	UFUNCTION(BlueprintCallable, Category = "LearningAgents")
+	UFUNCTION(BlueprintCallable, Category = "LearningAgents", meta = (AutoCreateRefTerm = "TrainerTrainingSettings,TrainerGameSettings,TrainerPathSettings"))
 	void RunTraining(
 		const FLearningAgentsTrainerTrainingSettings& TrainerTrainingSettings = FLearningAgentsTrainerTrainingSettings(),
 		const FLearningAgentsTrainerGameSettings& TrainerGameSettings = FLearningAgentsTrainerGameSettings(),
