Learning Agents: reintroduce visual debugger (WIP), spelling errors, noneditor intermediate relative path, other stuff

#rnx

#virtualized

diff --git a/Engine/Plugins/Experimental/LearningAgents/Source/LearningAgents/Private/LearningAgentsObservations.cpp b/Engine/Plugins/Experimental/LearningAgents/Source/LearningAgents/Private/LearningAgentsObservations.cpp
--- a/Engine/Plugins/Experimental/LearningAgents/Source/LearningAgents/Private/LearningAgentsObservations.cpp
+++ b/Engine/Plugins/Experimental/LearningAgents/Source/LearningAgents/Private/LearningAgentsObservations.cpp
@@ -11,6 +11,63 @@
 #include "Engine/World.h"
 #include "Components/SplineComponent.h"
 
+#include "VisualLogger/VisualLogger.h"
+
+#ifndef UE_LEARNING_AGENTS_ENABLE_VISUAL_LOG
+#define UE_LEARNING_AGENTS_ENABLE_VISUAL_LOG ENABLE_VISUAL_LOG
+#endif
+
+#if UE_LEARNING_AGENTS_ENABLE_VISUAL_LOG
+
+#define UE_LEARNING_AGENTS_VLOG_STRING(Owner, Category, Verbosity, Location, Color, Format, ...) \
+	UE_VLOG_LOCATION(Owner, Category, Verbosity, Location, 0.0f, Color, Format, ##__VA_ARGS__)
+
+#define UE_LEARNING_AGENTS_VLOG_LOCATION(Owner, Category, Verbosity, Location, Radius, Color, Format, ...) \
+	UE_VLOG_LOCATION(Owner, Category, Verbosity, Location, Radius, Color, Format, ##__VA_ARGS__)
+
+#define UE_LEARNING_AGENTS_VLOG_ARROW(Owner, Category, Verbosity, Start, End, Color, Format, ...) \
+	UE_VLOG_ARROW(Owner, Category, Verbosity, Start, End, Color, Format, ##__VA_ARGS__)
+
+#define UE_LEARNING_AGENTS_VLOG_SEGMENT(Owner, Category, Verbosity, Start, End, Color, Format, ...) \
+	UE_VLOG_SEGMENT(Owner, Category, Verbosity, Start, End, Color, Format, ##__VA_ARGS__)
+
+#define UE_LEARNING_AGENTS_VLOG_CIRCLE(Owner, Category, Verbosity, Center, UpAxis, Radius, Color, Format, ...) \
+	UE_VLOG_CIRCLE(Owner, Category, Verbosity, Center, UpAxis, Radius, Color, Format, ##__VA_ARGS__)
+
+#define UE_LEARNING_AGENTS_VLOG_OBOX(Owner, Category, Verbosity, Box, Matrix, Color, Format, ...) \
+	UE_VLOG_OBOX(Owner, Category, Verbosity, Box, Matrix, Color, Format, ##__VA_ARGS__)
+
+#define UE_LEARNING_AGENTS_VLOG_PLANE(Owner, Category, Verbosity, Location, Rotation, Axis0, Axis1, Color, Format, ...) \
+	UE_LEARNING_AGENTS_VLOG_OBOX(Owner, Category, Verbosity, FBox(FVector(-25.0f, -25.0f, -0.1f), FVector(25.0f, 25.0f, 0.1f)), UE::Learning::Agents::Debug::PlaneMatrix(Rotation, Location, Axis0, Axis1), Color, Format, ##__VA_ARGS__)
+
+#define UE_LEARNING_AGENTS_VLOG_MATRIX(Owner, Category, Verbosity, Matrix, Color, Format, ...) \
+	{ \
+		UE_LEARNING_AGENTS_VLOG_SEGMENT(Owner, Category, Verbosity, Matrix.TransformPosition(FVector::ZeroVector), Matrix.TransformPosition(15.0f * FVector::ForwardVector), FColor::Red, TEXT("")); \
+		UE_LEARNING_AGENTS_VLOG_SEGMENT(Owner, Category, Verbosity, Matrix.TransformPosition(FVector::ZeroVector), Matrix.TransformPosition(15.0f * FVector::RightVector), FColor::Green, TEXT("")); \
+		UE_LEARNING_AGENTS_VLOG_SEGMENT(Owner, Category, Verbosity, Matrix.TransformPosition(FVector::ZeroVector), Matrix.TransformPosition(15.0f * FVector::UpVector), FColor::Blue, TEXT("")); \
+		UE_LEARNING_AGENTS_VLOG_STRING(Owner, Category, Verbosity, Matrix.TransformPosition(FVector::ZeroVector) + FVector(0.0f, 0.0f, 20.0f), Color, Format, ##__VA_ARGS__); \
+	}
+
+#define UE_LEARNING_AGENTS_VLOG_TRANSFORM(Owner, Category, Verbosity, Location, Rotation, Color, Format, ...) \
+	{ \
+		UE_LEARNING_AGENTS_VLOG_SEGMENT(Owner, Category, Verbosity, Location, Location + 15.0f * Rotation.RotateVector(FVector::ForwardVector), FColor::Red, TEXT("")); \
+		UE_LEARNING_AGENTS_VLOG_SEGMENT(Owner, Category, Verbosity, Location, Location + 15.0f * Rotation.RotateVector(FVector::RightVector), FColor::Green, TEXT("")); \
+		UE_LEARNING_AGENTS_VLOG_SEGMENT(Owner, Category, Verbosity, Location, Location + 15.0f * Rotation.RotateVector(FVector::UpVector), FColor::Blue, TEXT("")); \
+		UE_LEARNING_AGENTS_VLOG_OBOX(Owner, Category, Verbosity, FBox(5.0f * FVector(-1, -1, -1), 5.0f * FVector(1, 1, 1)), FTransform(Rotation, Location, FVector::OneVector).ToMatrixNoScale(), Color, TEXT("")); \
+		UE_LEARNING_AGENTS_VLOG_STRING(Owner, Category, Verbosity, Location + FVector(0.0f, 0.0f, 20.0f), Color, Format, ##__VA_ARGS__); \
+	}
+
+#define UE_LEARNING_AGENTS_VLOG_ANGLE(Owner, Category, Verbosity, Angle, RelativeAngle, Location, Radius, Color, Format, ...) \
+	{ \
+		UE_LEARNING_AGENTS_VLOG_CIRCLE(Owner, Category, Verbosity, Location, FVector::UpVector, Radius, Color, TEXT("")); \
+		UE_LEARNING_AGENTS_VLOG_SEGMENT(Owner, Category, Verbosity, Location, Location + Radius * FVector(FMath::Sin(RelativeAngle), FMath::Cos(RelativeAngle), 0.0f), Color, TEXT("")); \
+		UE_LEARNING_AGENTS_VLOG_SEGMENT(Owner, Category, Verbosity, Location, Location + Radius * FVector(FMath::Sin(Angle), FMath::Cos(Angle), 0.0f), Color, TEXT("")); \
+		UE_LEARNING_AGENTS_VLOG_LOCATION(Owner, Category, Verbosity, Location + Radius * FVector(FMath::Sin(Angle), FMath::Cos(Angle), 0.0f), static_cast<uint16>(Radius / 20.0f), Color, Format, ##__VA_ARGS__); \
+	}
+
+#endif
+
+
 bool operator==(const FLearningAgentsObservationObjectElement& Lhs, const FLearningAgentsObservationObjectElement& Rhs)
 {
 	return Lhs.ObjectElement.Index == Rhs.ObjectElement.Index;
@@ -1134,7 +1191,7 @@ FLearningAgentsObservationObjectElement ULearningAgentsObservationObject::MakeSt
 		SubElements.Add(Element.Value);
 	}
 
-	return MakeStructObservationFromArrayViews(SubElementNames, SubElements);
+	return MakeStructObservationFromArrayViews(SubElementNames, SubElements, Name);
 }
 
 FLearningAgentsObservationObjectElement ULearningAgentsObservationObject::MakeStructObservationFromArrays(const TArray<FName>& ElementNames, const TArray<FLearningAgentsObservationObjectElement>& Elements, const FName Name)
@@ -1500,14 +1557,59 @@ FLearningAgentsObservationObjectElement ULearningAgentsObservationObject::MakeFl
 	return MakeContinuousObservationFromArrayView({ Value / FMath::Max(FloatScale, UE_SMALL_NUMBER) }, Name);
 }
 
-FLearningAgentsObservationObjectElement ULearningAgentsObservationObject::MakeLocationObservation(const FVector Location, const FTransform RelativeTransform, const float LocationScale, const FName Name)
+FLearningAgentsObservationObjectElement ULearningAgentsObservationObject::MakeLocationObservation(
+	const FVector Location,
+	const FTransform RelativeTransform,
+	const float LocationScale,
+	const FName Name,
+	const bool bVisualLoggerEnabled,
+	const FVector VisualLoggerLocation,
+	const FLinearColor VisualLoggerColor)
 {
 	const FVector LocalLocation = RelativeTransform.InverseTransformPosition(Location);
+	const FVector EncodedLocation = FVector(
+		LocalLocation.X / FMath::Max(LocationScale, UE_SMALL_NUMBER),
+		LocalLocation.Y / FMath::Max(LocationScale, UE_SMALL_NUMBER),
+		LocalLocation.Z / FMath::Max(LocationScale, UE_SMALL_NUMBER));
+
+#if UE_LEARNING_AGENTS_ENABLE_VISUAL_LOG
+	if (bVisualLoggerEnabled)
+	{
+		const ULearningAgentsObservationVisualLoggerObject* VisualLoggerObject = GetOrAddVisualLoggerObject(Name);
+
+		UE_LEARNING_AGENTS_VLOG_LOCATION(VisualLoggerObject, LogLearning, Display,
+			Location,
+			10,
+			VisualLoggerColor.ToFColor(true),
+			TEXT(""));
+
+		UE_LEARNING_AGENTS_VLOG_SEGMENT(VisualLoggerObject, LogLearning, Display,
+			RelativeTransform.GetTranslation(),
+			Location,
+			VisualLoggerColor.ToFColor(true),
+			TEXT(""));
+
+		UE_LEARNING_AGENTS_VLOG_TRANSFORM(VisualLoggerObject, LogLearning, Display,
+			RelativeTransform.GetTranslation(),
+			RelativeTransform.GetRotation(),
+			VisualLoggerColor.ToFColor(true),
+			TEXT(""));
+
+		UE_LEARNING_AGENTS_VLOG_STRING(VisualLoggerObject, LogLearning, Display, VisualLoggerLocation,
+			VisualLoggerColor.ToFColor(true),
+			TEXT("Name: %s\nLocation: [% 6.1f % 6.1f % 6.1f]\nLocal Location: [% 6.1f % 6.1f % 6.1f]\nScale: [% 6.2f]\nEncoded: [% 6.2f % 6.2f % 6.2f]"),
+			*Name.ToString(),
+			Location.X, Location.Y, Location.Z,
+			LocalLocation.X, LocalLocation.Y, LocalLocation.Z,
+			LocationScale,
+			EncodedLocation.X, EncodedLocation.Y, EncodedLocation.Z);
+	}
+#endif
 
 	return MakeContinuousObservationFromArrayView({
-		(float)LocalLocation.X / FMath::Max(LocationScale, UE_SMALL_NUMBER),
-		(float)LocalLocation.Y / FMath::Max(LocationScale, UE_SMALL_NUMBER),
-		(float)LocalLocation.Z / FMath::Max(LocationScale, UE_SMALL_NUMBER),
+		(float)EncodedLocation.X,
+		(float)EncodedLocation.Y,
+		(float)EncodedLocation.Z,
 		}, Name);
 }
 
@@ -1589,7 +1691,13 @@ FLearningAgentsObservationObjectElement ULearningAgentsObservationObject::MakeVe
 		}, Name);
 }
 
-FLearningAgentsObservationObjectElement ULearningAgentsObservationObject::MakeDirectionObservation(const FVector Direction, const FTransform RelativeTransform, const FName Name)
+FLearningAgentsObservationObjectElement ULearningAgentsObservationObject::MakeDirectionObservation(
+	const FVector Direction,
+	const FTransform RelativeTransform,
+	const FName Name,
+	const bool bVisualLoggerEnabled,
+	const FVector VisualLoggerLocation,
+	const FLinearColor VisualLoggerColor)
 {
 	const FVector LocalDirection = RelativeTransform.InverseTransformVectorNoScale(Direction).GetSafeNormal(UE_SMALL_NUMBER, FVector::ForwardVector);
 
@@ -1600,7 +1708,15 @@ FLearningAgentsObservationObjectElement ULearningAgentsObservationObject::MakeDi
 		}, Name);
 }
 
-FLearningAgentsObservationObjectElement ULearningAgentsObservationObject::MakeLocationAlongSplineObservation(const USplineComponent* SplineComponent, const float DistanceAlongSpline, const FTransform RelativeTransform, const float LocationScale, const FName Name)
+FLearningAgentsObservationObjectElement ULearningAgentsObservationObject::MakeLocationAlongSplineObservation(
+	const USplineComponent* SplineComponent,
+	const float DistanceAlongSpline,
+	const FTransform RelativeTransform,
+	const float LocationScale,
+	const FName Name,
+	const bool bVisualLoggerEnabled,
+	const FVector VisualLoggerLocation,
+	const FLinearColor VisualLoggerColor)
 {
 	if (!SplineComponent)
 	{
@@ -1608,7 +1724,7 @@ FLearningAgentsObservationObjectElement ULearningAgentsObservationObject::MakeLo
 		return FLearningAgentsObservationObjectElement();
 	}
 
-	return MakeLocationObservation(SplineComponent->GetLocationAtDistanceAlongSpline(DistanceAlongSpline, ESplineCoordinateSpace::World), RelativeTransform, LocationScale, Name);
+	return MakeLocationObservation(SplineComponent->GetLocationAtDistanceAlongSpline(DistanceAlongSpline, ESplineCoordinateSpace::World), RelativeTransform, LocationScale, Name, bVisualLoggerEnabled, VisualLoggerLocation, VisualLoggerColor);
 }
 
 FLearningAgentsObservationObjectElement ULearningAgentsObservationObject::MakeProportionAlongSplineObservation(const USplineComponent* SplineComponent, const float DistanceAlongSpline, const FName Name)
@@ -1634,7 +1750,14 @@ FLearningAgentsObservationObjectElement ULearningAgentsObservationObject::MakePr
 	}
 }
 
-FLearningAgentsObservationObjectElement ULearningAgentsObservationObject::MakeDirectionAlongSplineObservation(const USplineComponent* SplineComponent, const float DistanceAlongSpline, const FTransform RelativeTransform, const FName Name)
+FLearningAgentsObservationObjectElement ULearningAgentsObservationObject::MakeDirectionAlongSplineObservation(
+	const USplineComponent* SplineComponent,
+	const float DistanceAlongSpline,
+	const FTransform RelativeTransform,
+	const FName Name,
+	const bool bVisualLoggerEnabled,
+	const FVector VisualLoggerLocation,
+	const FLinearColor VisualLoggerColor)
 {
 	if (!SplineComponent)
 	{
@@ -1642,7 +1765,7 @@ FLearningAgentsObservationObjectElement ULearningAgentsObservationObject::MakeDi
 		return FLearningAgentsObservationObjectElement();
 	}
 
-	return MakeDirectionObservation(SplineComponent->GetDirectionAtDistanceAlongSpline(DistanceAlongSpline, ESplineCoordinateSpace::World), RelativeTransform, Name);
+	return MakeDirectionObservation(SplineComponent->GetDirectionAtDistanceAlongSpline(DistanceAlongSpline, ESplineCoordinateSpace::World), RelativeTransform, Name, bVisualLoggerEnabled, VisualLoggerLocation, VisualLoggerColor);
 }
 
 FLearningAgentsObservationObjectElement ULearningAgentsObservationObject::MakePropertiesAlongSplineObservation(const USplineComponent* SplineComponent, const float DistanceAlongSpline, const FTransform RelativeTransform, const float LocationScale, const FName Name)
@@ -3208,4 +3331,20 @@ bool ULearningAgentsObservationObject::GetProportionAlongRaysObservationToArrayV
 	}
 
 	return true;
-}
\ No newline at end of file
+}
+
+const ULearningAgentsObservationVisualLoggerObject* ULearningAgentsObservationObject::GetOrAddVisualLoggerObject(const FName Name)
+{
+	const TObjectPtr<const ULearningAgentsObservationVisualLoggerObject>* Value = VisualLoggerObjects.Find(Name);
+	if (Value)
+	{
+		return Value->Get();
+	}
+	else
+	{
+		const FName UniqueName = MakeUniqueObjectName(this, ULearningAgentsObservationVisualLoggerObject::StaticClass(), Name, EUniqueObjectNameOptions::GloballyUnique);
+		const ULearningAgentsObservationVisualLoggerObject* NewLoggerObject = NewObject<ULearningAgentsObservationVisualLoggerObject>(this, UniqueName);
+		VisualLoggerObjects.Add(Name, NewLoggerObject);
+		return NewLoggerObject;
+	}
+}
diff --git a/Engine/Plugins/Experimental/LearningAgents/Source/LearningAgents/Private/LearningAgentsPolicy.cpp b/Engine/Plugins/Experimental/LearningAgents/Source/LearningAgents/Private/LearningAgentsPolicy.cpp
--- a/Engine/Plugins/Experimental/LearningAgents/Source/LearningAgents/Private/LearningAgentsPolicy.cpp
+++ b/Engine/Plugins/Experimental/LearningAgents/Source/LearningAgents/Private/LearningAgentsPolicy.cpp
@@ -155,8 +155,8 @@ void ULearningAgentsPolicy::SetupPolicy(
 				EncoderNeuralNetworkAsset->NeuralNetworkData->GetOutputSize() != ObservationEncodedVectorSize)
 			{
 				UE_LOG(LogLearning, Error, TEXT("%s: Encoder Network Asset provided during Setup is incorrect size: Got inputs of size %i, expected %i. Got outputs of size %i, expected %i."), *GetName(),
-					DecoderNeuralNetworkAsset->NeuralNetworkData->GetInputSize(), ObservationVectorSize,
-					DecoderNeuralNetworkAsset->NeuralNetworkData->GetOutputSize(), ObservationEncodedVectorSize);
+					EncoderNeuralNetworkAsset->NeuralNetworkData->GetInputSize(), ObservationVectorSize,
+					EncoderNeuralNetworkAsset->NeuralNetworkData->GetOutputSize(), ObservationEncodedVectorSize);
 				return;
 			}
 		}
diff --git a/Engine/Plugins/Experimental/LearningAgents/Source/LearningAgents/Public/LearningAgentsObservations.h b/Engine/Plugins/Experimental/LearningAgents/Source/LearningAgents/Public/LearningAgentsObservations.h
--- a/Engine/Plugins/Experimental/LearningAgents/Source/LearningAgents/Public/LearningAgentsObservations.h
+++ b/Engine/Plugins/Experimental/LearningAgents/Source/LearningAgents/Public/LearningAgentsObservations.h
@@ -225,6 +225,12 @@ private:
 	UE::Learning::Observation::FSchema ObservationSchema;
 };
 
+UCLASS(BlueprintType)
+class LEARNINGAGENTS_API ULearningAgentsObservationVisualLoggerObject : public UObject
+{
+	GENERATED_BODY()
+};
+
 /**
  * Observation Object
  *
@@ -344,8 +350,15 @@ public:
 	UFUNCTION(BlueprintPure, Category = "LearningAgents")
 	FLearningAgentsObservationObjectElement MakeFloatObservation(const float Value, const float FloatScale = 1.0f, const FName Name = TEXT("Float"));
 
-	UFUNCTION(BlueprintPure, Category = "LearningAgents")
-	FLearningAgentsObservationObjectElement MakeLocationObservation(const FVector Location, const FTransform RelativeTransform = FTransform(), const float LocationScale = 100.0f, const FName Name = TEXT("Location"));
+	UFUNCTION(BlueprintPure, Category = "LearningAgents", meta = (AdvancedDisplay = 4))
+	FLearningAgentsObservationObjectElement MakeLocationObservation(
+		const FVector Location,
+		const FTransform RelativeTransform = FTransform(),
+		const float LocationScale = 100.0f,
+		const FName Name = TEXT("Location"),
+		const bool bVisualLoggerEnabled = false,
+		const FVector VisualLoggerLocation = FVector::ZeroVector,
+		const FLinearColor VisualLoggerColor = FLinearColor::Red);
 
 	UFUNCTION(BlueprintPure, Category = "LearningAgents")
 	FLearningAgentsObservationObjectElement MakeRotationObservation(const FRotator Rotation, const FRotator RelativeRotation = FRotator::ZeroRotator, const FName Name = TEXT("Rotation"));
@@ -366,19 +379,40 @@ public:
 	UFUNCTION(BlueprintPure, Category = "LearningAgents")
 	FLearningAgentsObservationObjectElement MakeVelocityObservation(const FVector Velocity, const FTransform RelativeTransform = FTransform(), const float VelocityScale = 200.0f, const FName Name = TEXT("Velocity"));
 
-	UFUNCTION(BlueprintPure, Category = "LearningAgents")
-	FLearningAgentsObservationObjectElement MakeDirectionObservation(const FVector Direction, const FTransform RelativeTransform = FTransform(), const FName Name = TEXT("Direction"));
+	UFUNCTION(BlueprintPure, Category = "LearningAgents", meta = (AdvancedDisplay = 3))
+	FLearningAgentsObservationObjectElement MakeDirectionObservation(
+		const FVector Direction,
+		const FTransform RelativeTransform = FTransform(),
+		const FName Name = TEXT("Direction"),
+		const bool bVisualLoggerEnabled = false,
+		const FVector VisualLoggerLocation = FVector::ZeroVector,
+		const FLinearColor VisualLoggerColor = FLinearColor::Red);
 
 	// Spline Observations
 
-	UFUNCTION(BlueprintPure, Category = "LearningAgents")
-	FLearningAgentsObservationObjectElement MakeLocationAlongSplineObservation(const USplineComponent* SplineComponent, const float DistanceAlongSpline, const FTransform RelativeTransform = FTransform(), const float LocationScale = 100.0f, const FName Name = TEXT("LocationAlongSpline"));
+	UFUNCTION(BlueprintPure, Category = "LearningAgents", meta = (AdvancedDisplay = 5))
+	FLearningAgentsObservationObjectElement MakeLocationAlongSplineObservation(
+		const USplineComponent* SplineComponent,
+		const float DistanceAlongSpline,
+		const FTransform RelativeTransform = FTransform(),
+		const float LocationScale = 100.0f,
+		const FName Name = TEXT("LocationAlongSpline"),
+		const bool bVisualLoggerEnabled = false,
+		const FVector VisualLoggerLocation = FVector::ZeroVector,
+		const FLinearColor VisualLoggerColor = FLinearColor::Red);
 
 	UFUNCTION(BlueprintPure, Category = "LearningAgents")
 	FLearningAgentsObservationObjectElement MakeProportionAlongSplineObservation(const USplineComponent* SplineComponent, const float DistanceAlongSpline, const FName Name = TEXT("ProportionAlongSpline"));
 
-	UFUNCTION(BlueprintPure, Category = "LearningAgents")
-	FLearningAgentsObservationObjectElement MakeDirectionAlongSplineObservation(const USplineComponent* SplineComponent, const float DistanceAlongSpline, const FTransform RelativeTransform = FTransform(), const FName Name = TEXT("DirectionAlongSpline"));
+	UFUNCTION(BlueprintPure, Category = "LearningAgents", meta = (AdvancedDisplay = 4))
+	FLearningAgentsObservationObjectElement MakeDirectionAlongSplineObservation(
+		const USplineComponent* SplineComponent,
+		const float DistanceAlongSpline,
+		const FTransform RelativeTransform = FTransform(),
+		const FName Name = TEXT("DirectionAlongSpline"),
+		const bool bVisualLoggerEnabled = false,
+		const FVector VisualLoggerLocation = FVector::ZeroVector,
+		const FLinearColor VisualLoggerColor = FLinearColor::Red);
 
 	UFUNCTION(BlueprintPure, Category = "LearningAgents")
 	FLearningAgentsObservationObjectElement MakePropertiesAlongSplineObservation(const USplineComponent* SplineComponent, const float DistanceAlongSpline, const FTransform RelativeTransform = FTransform(), const float LocationScale = 100.0f, const FName Name = TEXT("PropertiesAlongSpline"));
@@ -554,4 +588,9 @@ public:
 private:
 
 	UE::Learning::Observation::FObject ObservationObject;
+
+	const ULearningAgentsObservationVisualLoggerObject* GetOrAddVisualLoggerObject(const FName Name);
+
+	UPROPERTY()
+	TMap<FName, TObjectPtr<const ULearningAgentsObservationVisualLoggerObject>> VisualLoggerObjects;
 };
diff --git a/Engine/Plugins/Experimental/LearningAgents/Source/LearningAgentsTraining/Private/LearningAgentsTrainer.cpp b/Engine/Plugins/Experimental/LearningAgents/Source/LearningAgentsTraining/Private/LearningAgentsTrainer.cpp
--- a/Engine/Plugins/Experimental/LearningAgents/Source/LearningAgentsTraining/Private/LearningAgentsTrainer.cpp
+++ b/Engine/Plugins/Experimental/LearningAgents/Source/LearningAgentsTraining/Private/LearningAgentsTrainer.cpp
@@ -51,7 +51,7 @@ namespace UE::Learning::Agents
 FLearningAgentsTrainerPathSettings::FLearningAgentsTrainerPathSettings()
 {
 	EditorEngineRelativePath.Path = FPaths::EngineDir();
-	IntermediateRelativePath.Path = FPaths::ProjectIntermediateDir();
+	EditorIntermediateRelativePath.Path = FPaths::ProjectIntermediateDir();
 }
 
 FString FLearningAgentsTrainerPathSettings::GetEditorEnginePath() const
@@ -70,7 +70,16 @@ FString FLearningAgentsTrainerPathSettings::GetEditorEnginePath() const
 
 FString FLearningAgentsTrainerPathSettings::GetIntermediatePath() const
 {
-	return IntermediateRelativePath.Path;
+#if WITH_EDITOR
+	return EditorIntermediateRelativePath.Path;
+#else
+	if (NonEditorIntermediateRelativePath.IsEmpty())
+	{
+		UE_LOG(LogLearning, Warning, TEXT("GetIntermediatePath: NonEditorIntermediateRelativePath not set"));
+	}
+
+	return NonEditorIntermediateRelativePath;
+#endif
 }
 
 ULearningAgentsTrainer::ULearningAgentsTrainer() : Super(FObjectInitializer::Get()) {}
diff --git a/Engine/Plugins/Experimental/LearningAgents/Source/LearningAgentsTraining/Public/LearningAgentsTrainer.h b/Engine/Plugins/Experimental/LearningAgents/Source/LearningAgentsTraining/Public/LearningAgentsTrainer.h
--- a/Engine/Plugins/Experimental/LearningAgents/Source/LearningAgentsTraining/Public/LearningAgentsTrainer.h
+++ b/Engine/Plugins/Experimental/LearningAgents/Source/LearningAgentsTraining/Public/LearningAgentsTrainer.h
@@ -317,7 +317,11 @@ public:
 
 	/** The relative path to the Intermediate directory. Defaults to FPaths::ProjectIntermediateDir. */
 	UPROPERTY(EditAnywhere, Category = "LearningAgents", meta = (RelativePath))
-	FDirectoryPath IntermediateRelativePath;
+	FDirectoryPath EditorIntermediateRelativePath;
+
+	/** The relative path to the intermediate folder for non-editor builds. */
+	UPROPERTY(EditAnywhere, Category = "LearningAgents")
+	FString NonEditorIntermediateRelativePath;
 
 public:
 
