LearningAgents: Fixed issues with communication timeouts

diff --git a/Engine/Plugins/Experimental/LearningAgents/Source/LearningAgentsTraining/Public/LearningAgentsImitationTrainer.h b/Engine/Plugins/Experimental/LearningAgents/Source/LearningAgentsTraining/Public/LearningAgentsImitationTrainer.h
--- a/Engine/Plugins/Experimental/LearningAgents/Source/LearningAgentsTraining/Public/LearningAgentsImitationTrainer.h
+++ b/Engine/Plugins/Experimental/LearningAgents/Source/LearningAgentsTraining/Public/LearningAgentsImitationTrainer.h
@@ -29,7 +29,7 @@ public:
 
 	/** Time in seconds to wait for the training process before timing out. */
 	UPROPERTY(EditAnywhere, Category = "LearningAgents", meta = (ClampMin = "0.0", UIMin = "0.0"))
-	float TrainerCommunicationTimeout = 20.0f;
+	float TrainerCommunicationTimeout = 10.0f;
 };
 
 /** The configurable settings for the training process. */
diff --git a/Engine/Plugins/Experimental/LearningAgents/Source/LearningTraining/Private/LearningImitationTrainer.cpp b/Engine/Plugins/Experimental/LearningAgents/Source/LearningTraining/Private/LearningImitationTrainer.cpp
--- a/Engine/Plugins/Experimental/LearningAgents/Source/LearningTraining/Private/LearningImitationTrainer.cpp
+++ b/Engine/Plugins/Experimental/LearningAgents/Source/LearningTraining/Private/LearningImitationTrainer.cpp
@@ -250,6 +250,7 @@ namespace UE::Learning
 		const ELogSetting LogSettings)
 	{
 		return SharedMemoryTraining::RecvNetwork(
+			TrainingProcess.Get(),
 			Controls.View,
 			OutNetwork,
 			SharedMemoryTraining::EControls::PolicySignal,
@@ -266,6 +267,7 @@ namespace UE::Learning
 		const ELogSetting LogSettings)
 	{
 		return SharedMemoryTraining::RecvNetwork(
+			TrainingProcess.Get(),
 			Controls.View,
 			OutNetwork,
 			SharedMemoryTraining::EControls::EncoderSignal,
@@ -282,6 +284,7 @@ namespace UE::Learning
 		const ELogSetting LogSettings)
 	{
 		return SharedMemoryTraining::RecvNetwork(
+			TrainingProcess.Get(),
 			Controls.View,
 			OutNetwork,
 			SharedMemoryTraining::EControls::DecoderSignal,
@@ -298,6 +301,7 @@ namespace UE::Learning
 		const ELogSetting LogSettings)
 	{
 		return SharedMemoryTraining::SendNetwork(
+			TrainingProcess.Get(),
 			Controls.View,
 			Policy.View,
 			SharedMemoryTraining::EControls::PolicySignal,
@@ -314,6 +318,7 @@ namespace UE::Learning
 		const ELogSetting LogSettings)
 	{
 		return SharedMemoryTraining::SendNetwork(
+			TrainingProcess.Get(),
 			Controls.View,
 			Encoder.View,
 			SharedMemoryTraining::EControls::EncoderSignal,
@@ -330,6 +335,7 @@ namespace UE::Learning
 		const ELogSetting LogSettings)
 	{
 		return SharedMemoryTraining::SendNetwork(
+			TrainingProcess.Get(),
 			Controls.View,
 			Decoder.View,
 			SharedMemoryTraining::EControls::DecoderSignal,
@@ -348,6 +354,7 @@ namespace UE::Learning
 		const ELogSetting LogSettings)
 	{
 		return SharedMemoryTraining::SendExperience(
+			TrainingProcess.Get(),
 			EpisodeStarts.View,
 			EpisodeLengths.View,
 			Observations.View,
diff --git a/Engine/Plugins/Experimental/LearningAgents/Source/LearningTraining/Private/LearningPPOTrainer.cpp b/Engine/Plugins/Experimental/LearningAgents/Source/LearningTraining/Private/LearningPPOTrainer.cpp
--- a/Engine/Plugins/Experimental/LearningAgents/Source/LearningTraining/Private/LearningPPOTrainer.cpp
+++ b/Engine/Plugins/Experimental/LearningAgents/Source/LearningTraining/Private/LearningPPOTrainer.cpp
@@ -411,6 +411,7 @@ namespace UE::Learning
 		const ELogSetting LogSettings)
 	{
 		return SharedMemoryTraining::RecvNetwork(
+			TrainingProcess.Get(),
 			Controls.View[ProcessIdx],
 			OutNetwork,
 			SharedMemoryTraining::EControls::PolicySignal,
@@ -427,6 +428,7 @@ namespace UE::Learning
 		const ELogSetting LogSettings)
 	{
 		return SharedMemoryTraining::RecvNetwork(
+			TrainingProcess.Get(),
 			Controls.View[ProcessIdx],
 			OutNetwork,
 			SharedMemoryTraining::EControls::CriticSignal,
@@ -443,6 +445,7 @@ namespace UE::Learning
 		const ELogSetting LogSettings)
 	{
 		return SharedMemoryTraining::RecvNetwork(
+			TrainingProcess.Get(),
 			Controls.View[ProcessIdx],
 			OutNetwork,
 			SharedMemoryTraining::EControls::EncoderSignal,
@@ -459,6 +462,7 @@ namespace UE::Learning
 		const ELogSetting LogSettings)
 	{
 		return SharedMemoryTraining::RecvNetwork(
+			TrainingProcess.Get(),
 			Controls.View[ProcessIdx],
 			OutNetwork,
 			SharedMemoryTraining::EControls::DecoderSignal,
@@ -475,6 +479,7 @@ namespace UE::Learning
 		const ELogSetting LogSettings)
 	{
 		return SharedMemoryTraining::SendNetwork(
+			TrainingProcess.Get(),
 			Controls.View[ProcessIdx],
 			Policy.View,
 			SharedMemoryTraining::EControls::PolicySignal,
@@ -491,6 +496,7 @@ namespace UE::Learning
 		const ELogSetting LogSettings)
 	{
 		return SharedMemoryTraining::SendNetwork(
+			TrainingProcess.Get(),
 			Controls.View[ProcessIdx],
 			Critic.View,
 			SharedMemoryTraining::EControls::CriticSignal,
@@ -507,6 +513,7 @@ namespace UE::Learning
 		const ELogSetting LogSettings)
 	{
 		return SharedMemoryTraining::SendNetwork(
+			TrainingProcess.Get(),
 			Controls.View[ProcessIdx],
 			Encoder.View,
 			SharedMemoryTraining::EControls::EncoderSignal,
@@ -523,6 +530,7 @@ namespace UE::Learning
 		const ELogSetting LogSettings)
 	{
 		return SharedMemoryTraining::SendNetwork(
+			TrainingProcess.Get(),
 			Controls.View[ProcessIdx],
 			Decoder.View,
 			SharedMemoryTraining::EControls::DecoderSignal,
@@ -538,6 +546,7 @@ namespace UE::Learning
 		const ELogSetting LogSettings)
 	{
 		return SharedMemoryTraining::SendExperience(
+			TrainingProcess.Get(),
 			EpisodeStarts.View[ProcessIdx],
 			EpisodeLengths.View[ProcessIdx],
 			EpisodeCompletionModes.View[ProcessIdx],
diff --git a/Engine/Plugins/Experimental/LearningAgents/Source/LearningTraining/Private/LearningSharedMemoryTraining.cpp b/Engine/Plugins/Experimental/LearningAgents/Source/LearningTraining/Private/LearningSharedMemoryTraining.cpp
--- a/Engine/Plugins/Experimental/LearningAgents/Source/LearningTraining/Private/LearningSharedMemoryTraining.cpp
+++ b/Engine/Plugins/Experimental/LearningAgents/Source/LearningTraining/Private/LearningSharedMemoryTraining.cpp
@@ -7,6 +7,7 @@
 #include "LearningExperience.h"
 #include "LearningCompletion.h"
 
+#include "Misc/MonitoredProcess.h"
 #include "HAL/PlatformProcess.h"
 
 namespace UE::Learning::SharedMemoryTraining
@@ -28,6 +29,7 @@ namespace UE::Learning::SharedMemoryTraining
 	}
 
 	ETrainerResponse RecvNetwork(
+		FMonitoredProcess* Process, 
 		TLearningArrayView<1, volatile int32> Controls,
 		ULearningNeuralNetworkData& OutNetwork,
 		const EControls Signal,
@@ -50,13 +52,28 @@ namespace UE::Learning::SharedMemoryTraining
 				return ETrainerResponse::Completed;
 			}
 
-			FPlatformProcess::Sleep(SleepTime);
-			WaitTime += SleepTime;
+			// Check if the process has exited
+			if (!Process || !Process->Update() || Process->GetReturnCode() == 1)
+			{
+				return ETrainerResponse::Unexpected;
+			}
 
+			// Check if we've timed out
 			if (WaitTime > Timeout)
 			{
 				return ETrainerResponse::Timeout;
 			}
+
+			// Check if ping has been sent
+			if (Controls[(uint8)EControls::PingSignal])
+			{
+				Controls[(uint8)EControls::PingSignal] = false;
+				WaitTime = 0.0f;
+			}
+
+			// Sleep for some time
+			FPlatformProcess::Sleep(SleepTime);
+			WaitTime += SleepTime;
 		}
 
 		if (LogSettings != ELogSetting::Silent)
@@ -95,6 +112,7 @@ namespace UE::Learning::SharedMemoryTraining
 	}
 
 	ETrainerResponse SendNetwork(
+		FMonitoredProcess* Process,
 		TLearningArrayView<1, volatile int32> Controls,
 		TLearningArrayView<1, uint8> NetworkData,
 		const EControls Signal,
@@ -109,13 +127,28 @@ namespace UE::Learning::SharedMemoryTraining
 		// Wait until the policy is requested by the sub-process
 		while (!Controls[(uint8)Signal])
 		{
-			FPlatformProcess::Sleep(SleepTime);
-			WaitTime += SleepTime;
+			// Check if the process has exited
+			if (!Process || !Process->Update() || Process->GetReturnCode() == 1)
+			{
+				return ETrainerResponse::Unexpected;
+			}
 
+			// Check if we've timed out
 			if (WaitTime > Timeout)
 			{
 				return ETrainerResponse::Timeout;
 			}
+
+			// Check if ping has been sent
+			if (Controls[(uint8)EControls::PingSignal])
+			{
+				Controls[(uint8)EControls::PingSignal] = false;
+				WaitTime = 0.0f;
+			}
+
+			// Sleep for some time
+			FPlatformProcess::Sleep(SleepTime);
+			WaitTime += SleepTime;
 		}
 
 		if (LogSettings != ELogSetting::Silent)
@@ -146,6 +179,7 @@ namespace UE::Learning::SharedMemoryTraining
 	}
 
 	ETrainerResponse SendExperience(
+		FMonitoredProcess* Process,
 		TLearningArrayView<1, int32> EpisodeStarts,
 		TLearningArrayView<1, int32> EpisodeLengths,
 		TLearningArrayView<1, ECompletionMode> EpisodeCompletionModes,
@@ -166,13 +200,28 @@ namespace UE::Learning::SharedMemoryTraining
 		// Wait until the sub-process is done reading any experience
 		while (Controls[(uint8)EControls::ExperienceSignal])
 		{
-			FPlatformProcess::Sleep(SleepTime);
-			WaitTime += SleepTime;
+			// Check if the process has exited
+			if (!Process || !Process->Update() || Process->GetReturnCode() == 1)
+			{
+				return ETrainerResponse::Unexpected;
+			}
 
+			// Check if we've timed out
 			if (WaitTime > Timeout)
 			{
 				return ETrainerResponse::Timeout;
 			}
+
+			// Check if ping has been sent
+			if (Controls[(uint8)EControls::PingSignal])
+			{
+				Controls[(uint8)EControls::PingSignal] = false;
+				WaitTime = 0.0f;
+			}
+
+			// Sleep for some time
+			FPlatformProcess::Sleep(SleepTime);
+			WaitTime += SleepTime;
 		}
 
 		if (LogSettings != ELogSetting::Silent)
@@ -203,6 +252,7 @@ namespace UE::Learning::SharedMemoryTraining
 	}
 
 	ETrainerResponse SendExperience(
+		FMonitoredProcess* Process, 
 		TLearningArrayView<1, int32> EpisodeStarts,
 		TLearningArrayView<1, int32> EpisodeLengths,
 		TLearningArrayView<2, float> Observations,
@@ -221,13 +271,28 @@ namespace UE::Learning::SharedMemoryTraining
 		// Wait until the sub-process is done reading any experience
 		while (Controls[(uint8)EControls::ExperienceSignal])
 		{
-			FPlatformProcess::Sleep(SleepTime);
-			WaitTime += SleepTime;
+			// Check if the process has exited
+			if (!Process || !Process->Update() || Process->GetReturnCode() == 1)
+			{
+				return ETrainerResponse::Unexpected;
+			}
 
+			// Check if we've timed out
 			if (WaitTime > Timeout)
 			{
 				return ETrainerResponse::Timeout;
 			}
+
+			// Check if ping has been sent
+			if (Controls[(uint8)EControls::PingSignal])
+			{
+				Controls[(uint8)EControls::PingSignal] = false;
+				WaitTime = 0.0f;
+			}
+
+			// Sleep for some time
+			FPlatformProcess::Sleep(SleepTime);
+			WaitTime += SleepTime;
 		}
 
 		if (LogSettings != ELogSetting::Silent)
diff --git a/Engine/Plugins/Experimental/LearningAgents/Source/LearningTraining/Private/LearningSharedMemoryTraining.h b/Engine/Plugins/Experimental/LearningAgents/Source/LearningTraining/Private/LearningSharedMemoryTraining.h
--- a/Engine/Plugins/Experimental/LearningAgents/Source/LearningTraining/Private/LearningSharedMemoryTraining.h
+++ b/Engine/Plugins/Experimental/LearningAgents/Source/LearningTraining/Private/LearningSharedMemoryTraining.h
@@ -4,6 +4,7 @@
 
 #include "LearningTrainer.h"
 
+class FMonitoredProcess;
 class ULearningNeuralNetworkData;
 
 namespace UE::Learning
@@ -24,13 +25,15 @@ namespace UE::Learning
 			DecoderSignal			= 6,
 			CompleteSignal			= 7,
 			StopSignal				= 8,
+			PingSignal				= 9,
 
-			ControlNum				= 9,
+			ControlNum				= 10,
 		};
 
 		LEARNINGTRAINING_API uint8 GetControlNum();
 
 		LEARNINGTRAINING_API ETrainerResponse RecvNetwork(
+			FMonitoredProcess* Process,
 			TLearningArrayView<1, volatile int32> Controls,
 			ULearningNeuralNetworkData& OutNetwork,
 			const EControls Signal,
@@ -45,6 +48,7 @@ namespace UE::Learning
 		LEARNINGTRAINING_API bool HasPolicyOrCompleted(TLearningArrayView<1, volatile int32> Controls);
 
 		LEARNINGTRAINING_API ETrainerResponse SendNetwork(
+			FMonitoredProcess* Process,
 			TLearningArrayView<1, volatile int32> Controls,
 			TLearningArrayView<1, uint8> NetworkData,
 			const EControls Signal,
@@ -54,6 +58,7 @@ namespace UE::Learning
 			const ELogSetting LogSettings = Trainer::DefaultLogSettings);
 
 		LEARNINGTRAINING_API ETrainerResponse SendExperience(
+			FMonitoredProcess* Process,
 			TLearningArrayView<1, int32> EpisodeStarts,
 			TLearningArrayView<1, int32> EpisodeLengths,
 			TLearningArrayView<1, ECompletionMode> EpisodeCompletionModes,
@@ -69,6 +74,7 @@ namespace UE::Learning
 			const ELogSetting LogSettings = Trainer::DefaultLogSettings);
 
 		LEARNINGTRAINING_API ETrainerResponse SendExperience(
+			FMonitoredProcess* Process,
 			TLearningArrayView<1, int32> EpisodeStarts,
 			TLearningArrayView<1, int32> EpisodeLengths,
 			TLearningArrayView<2, float> Observations,
diff --git a/Engine/Plugins/Experimental/LearningAgents/Source/LearningTraining/Private/LearningSocketTraining.cpp b/Engine/Plugins/Experimental/LearningAgents/Source/LearningTraining/Private/LearningSocketTraining.cpp
--- a/Engine/Plugins/Experimental/LearningAgents/Source/LearningTraining/Private/LearningSocketTraining.cpp
+++ b/Engine/Plugins/Experimental/LearningAgents/Source/LearningTraining/Private/LearningSocketTraining.cpp
@@ -88,17 +88,27 @@ namespace UE::Learning::SocketTraining
 		}
 
 		uint8 Signal = (uint8)ESignal::Invalid;
-		ETrainerResponse Response = RecvWithTimeout(Socket, &Signal, 1, Timeout);
-		if (Response != ETrainerResponse::Success) { return Response; }
+		ETrainerResponse Response = ETrainerResponse::Unexpected;
 
-		if (Signal == (uint8)ESignal::RecvComplete)
+		while (true)
 		{
-			return ETrainerResponse::Completed;
-		}
+			Response = RecvWithTimeout(Socket, &Signal, 1, Timeout);
+			if (Response != ETrainerResponse::Success) { return Response; }
 
-		if (Signal != (uint8)NetworkSignal)
-		{
-			return ETrainerResponse::Unexpected;
+			if (Signal == (uint8)ESignal::RecvComplete)
+			{
+				return ETrainerResponse::Completed;
+			}
+
+			if (Signal == (uint8)ESignal::RecvPing)
+			{
+				continue;
+			}
+
+			if (Signal != (uint8)NetworkSignal)
+			{
+				return ETrainerResponse::Unexpected;
+			}
 		}
 
 		Response = RecvWithTimeout(Socket, OutNetworkBuffer.GetData(), OutNetworkBuffer.Num(), Timeout);
diff --git a/Engine/Plugins/Experimental/LearningAgents/Source/LearningTraining/Private/LearningSocketTraining.h b/Engine/Plugins/Experimental/LearningAgents/Source/LearningTraining/Private/LearningSocketTraining.h
--- a/Engine/Plugins/Experimental/LearningAgents/Source/LearningTraining/Private/LearningSocketTraining.h
+++ b/Engine/Plugins/Experimental/LearningAgents/Source/LearningTraining/Private/LearningSocketTraining.h
@@ -28,6 +28,7 @@ namespace UE::Learning
 			SendDecoder			= 10,
 			RecvComplete		= 11,
 			SendStop			= 12,
+			RecvPing			= 13,
 		};
 
 		LEARNINGTRAINING_API ETrainerResponse WaitForConnection(
diff --git a/Engine/Plugins/Experimental/LearningAgents/Source/LearningTraining/Public/LearningTrainer.h b/Engine/Plugins/Experimental/LearningAgents/Source/LearningTraining/Public/LearningTrainer.h
--- a/Engine/Plugins/Experimental/LearningAgents/Source/LearningTraining/Public/LearningTrainer.h
+++ b/Engine/Plugins/Experimental/LearningAgents/Source/LearningTraining/Public/LearningTrainer.h
@@ -73,7 +73,7 @@ namespace UE::Learning
 		/**
 		* Default Timeout to use during communication.
 		*/
-		static constexpr float DefaultTimeout = 30.0f;
+		static constexpr float DefaultTimeout = 10.0f;
 
 		/**
 		* Default Log Settings to use during communication.
