LearningAgents: Small tweaks to actions and observations.

diff --git a/Engine/Plugins/Experimental/LearningAgents/Source/LearningAgents/Private/LearningAgentsActions.cpp b/Engine/Plugins/Experimental/LearningAgents/Source/LearningAgents/Private/LearningAgentsActions.cpp
--- a/Engine/Plugins/Experimental/LearningAgents/Source/LearningAgents/Private/LearningAgentsActions.cpp
+++ b/Engine/Plugins/Experimental/LearningAgents/Source/LearningAgents/Private/LearningAgentsActions.cpp
@@ -1084,6 +1084,11 @@ FLearningAgentsActionSchemaElement ULearningAgentsActionSchema::SpecifyVelocityA
 	return SpecifyContinuousAction(3, Name);
 }
 
+FLearningAgentsActionSchemaElement ULearningAgentsActionSchema::SpecifyDirectionAction(const FName Name)
+{
+	return SpecifyContinuousAction(3, Name);
+}
+
 FLearningAgentsActionSchemaElement ULearningAgentsActionSchema::SpecifySpeedAction(const FName Name)
 {
 	return SpecifyContinuousAction(1, Name);
@@ -1507,6 +1512,17 @@ FLearningAgentsActionObjectElement ULearningAgentsActionObject::MakeVelocityActi
 		}, Name);
 }
 
+FLearningAgentsActionObjectElement ULearningAgentsActionObject::MakeDirectionAction(const FVector Direction, const FTransform RelativeTransform, const FName Name)
+{
+	const FVector LocalDirection = RelativeTransform.InverseTransformVectorNoScale(Direction).GetSafeNormal(UE_SMALL_NUMBER, FVector::ForwardVector);
+
+	return MakeContinuousActionFromArrayView({
+		(float)LocalDirection.X,
+		(float)LocalDirection.Y,
+		(float)LocalDirection.Z,
+		}, Name);
+}
+
 FLearningAgentsActionObjectElement ULearningAgentsActionObject::MakeSpeedAction(const float Speed, const float SpeedScale, const FName Name)
 {
 	return MakeContinuousActionFromArrayView({ Speed / FMath::Max(SpeedScale, UE_SMALL_NUMBER) });
@@ -2478,6 +2494,19 @@ bool ULearningAgentsActionObject::GetVelocityAction(FVector& OutVelocity, const
 	return true;
 }
 
+bool ULearningAgentsActionObject::GetDirectionAction(FVector& OutDirection, const FLearningAgentsActionObjectElement Element, const FTransform RelativeTransform, const FName Name) const
+{
+	TStaticArray<float, 3> OutValues;
+	if (!GetContinuousActionToArrayView(OutValues, Element, Name))
+	{
+		OutDirection = FVector::ForwardVector;
+		return false;
+	}
+
+	OutDirection = RelativeTransform.TransformVectorNoScale(FVector(OutValues[0], OutValues[1], OutValues[2]).GetSafeNormal(UE_SMALL_NUMBER, FVector::ForwardVector));
+	return true;
+}
+
 bool ULearningAgentsActionObject::GetSpeedAction(float& OutSpeed, const FLearningAgentsActionObjectElement Element, const float SpeedScale, const FName Name) const
 {
 	if (!GetContinuousActionToArrayView(MakeArrayView(&OutSpeed, 1), Element, Name))
diff --git a/Engine/Plugins/Experimental/LearningAgents/Source/LearningAgents/Private/LearningAgentsObservations.cpp b/Engine/Plugins/Experimental/LearningAgents/Source/LearningAgents/Private/LearningAgentsObservations.cpp
--- a/Engine/Plugins/Experimental/LearningAgents/Source/LearningAgents/Private/LearningAgentsObservations.cpp
+++ b/Engine/Plugins/Experimental/LearningAgents/Source/LearningAgents/Private/LearningAgentsObservations.cpp
@@ -515,6 +515,11 @@ FTransform ULearningAgentsObservationFunctions::ProjectTransformOntoGroundPlane(
 	return FTransform(FQuat::FindBetweenNormals(FVector::ForwardVector, Direction), Position, Transform.GetScale3D());
 }
 
+UEnum* ULearningAgentsObservationFunctions::FindEnumByName(const FString Name)
+{
+	return FindObject<UEnum>(nullptr, *Name);
+}
+
 bool ULearningAgentsObservationSchema::ValidateObjectMatchesSchema(
 	const FLearningAgentsObservationSchemaElement SchemaElement,
 	const ULearningAgentsObservationObject* Object,
@@ -564,6 +569,11 @@ FLearningAgentsObservationSchemaElement ULearningAgentsObservationSchema::Specif
 	return SpecifyContinuousObservation(Size, Name);
 }
 
+FLearningAgentsObservationSchemaElement ULearningAgentsObservationSchema::SpecifyCountObservation(const FName Name)
+{
+	return SpecifyContinuousObservation(1, Name);
+}
+
 FLearningAgentsObservationSchemaElement ULearningAgentsObservationSchema::SpecifyStructObservation(const TMap<FName, FLearningAgentsObservationSchemaElement>& Elements, const FName Name)
 {
 	if (Elements.Num() == 0)
@@ -1171,6 +1181,17 @@ FLearningAgentsObservationObjectElement ULearningAgentsObservationObject::MakeIn
 	return MakeContinuousObservationFromArrayView(Values, Name);
 }
 
+FLearningAgentsObservationObjectElement ULearningAgentsObservationObject::MakeCountObservation(const int32 Num, const int32 MaxNum, const FName Name)
+{
+	if (MaxNum == 0)
+	{
+		UE_LOG(LogLearning, Error, TEXT("%s: MaxNum must not be zero for Count Observation."), *GetName());
+		return FLearningAgentsObservationObjectElement();
+	}
+
+	return MakeContinuousObservationFromArrayView({ (float)Num / (float)MaxNum }, Name);
+}
+
 FLearningAgentsObservationObjectElement ULearningAgentsObservationObject::MakeStructObservation(const TMap<FName, FLearningAgentsObservationObjectElement>& Elements, const FName Name)
 {
 	if (Elements.Num() == 0)
@@ -2085,6 +2106,19 @@ bool ULearningAgentsObservationObject::GetIndexObservation(int32& OutIndex, cons
 	return true;
 }
 
+bool ULearningAgentsObservationObject::GetCountObservation(int32& OutNum, const FLearningAgentsObservationObjectElement Element, const int32 MaxNum, const FName Name) const
+{
+	float FloatNum = 0.0f;
+	if (!GetContinuousObservationToArrayView(MakeArrayView(&FloatNum, 1), Element, Name))
+	{
+		OutNum = -1;
+		return false;
+	}
+
+	OutNum = FMath::RoundToInt(FloatNum * MaxNum);
+	return true;
+}
+
 bool ULearningAgentsObservationObject::GetStructObservationNum(int32& OutNum, const FLearningAgentsObservationObjectElement Element, const FName Name) const
 {
 	if (!ObservationObject.IsValid(Element.ObjectElement))
diff --git a/Engine/Plugins/Experimental/LearningAgents/Source/LearningAgents/Public/LearningAgentsActions.h b/Engine/Plugins/Experimental/LearningAgents/Source/LearningAgents/Public/LearningAgentsActions.h
--- a/Engine/Plugins/Experimental/LearningAgents/Source/LearningAgents/Public/LearningAgentsActions.h
+++ b/Engine/Plugins/Experimental/LearningAgents/Source/LearningAgents/Public/LearningAgentsActions.h
@@ -91,11 +91,11 @@ public:
 
 	UFUNCTION(BlueprintPure, Category = "LearningAgents", meta=(AutoCreateRefTerm="PriorProbabilities"))
 	FLearningAgentsActionSchemaElement SpecifyExclusiveDiscreteAction(const int32 Size, const TArray<float>& PriorProbabilities, const FName Name = TEXT("DiscreteExclusive"));
-	FLearningAgentsActionSchemaElement SpecifyExclusiveDiscreteActionFromArrayView(const int32 Size, const TArrayView<const float> PriorProbabilities, const FName Name = TEXT("DiscreteExclusive"));
+	FLearningAgentsActionSchemaElement SpecifyExclusiveDiscreteActionFromArrayView(const int32 Size, const TArrayView<const float> PriorProbabilities = {}, const FName Name = TEXT("DiscreteExclusive"));
 
 	UFUNCTION(BlueprintPure, Category = "LearningAgents", meta = (AutoCreateRefTerm = "PriorProbabilities"))
 	FLearningAgentsActionSchemaElement SpecifyInclusiveDiscreteAction(const int32 Size, const TArray<float>& PriorProbabilities, const FName Name = TEXT("DiscreteInclusive"));
-	FLearningAgentsActionSchemaElement SpecifyInclusiveDiscreteActionFromArrayView(const int32 Size, const TArrayView<const float> PriorProbabilities, const FName Name = TEXT("DiscreteInclusive"));
+	FLearningAgentsActionSchemaElement SpecifyInclusiveDiscreteActionFromArrayView(const int32 Size, const TArrayView<const float> PriorProbabilities = {}, const FName Name = TEXT("DiscreteInclusive"));
 
 	UFUNCTION(BlueprintPure, Category = "LearningAgents")
 	FLearningAgentsActionSchemaElement SpecifyStructAction(const TMap<FName, FLearningAgentsActionSchemaElement>& Elements, const FName Name = TEXT("Struct"));
@@ -109,14 +109,14 @@ public:
 	
 	UFUNCTION(BlueprintPure, Category = "LearningAgents", meta = (AutoCreateRefTerm = "PriorProbabilities"))
 	FLearningAgentsActionSchemaElement SpecifyExclusiveUnionActionFromArrays(const TArray<FName>& ElementNames, const TArray<FLearningAgentsActionSchemaElement>& Elements, const TArray<float>& PriorProbabilities, const FName Name = TEXT("ExclusiveUnion"));
-	FLearningAgentsActionSchemaElement SpecifyExclusiveUnionActionFromArrayViews(const TArrayView<const FName> ElementNames, const TArrayView<const FLearningAgentsActionSchemaElement> Elements, const TArrayView<const float> PriorProbabilities, const FName Name = TEXT("ExclusiveUnion"));
+	FLearningAgentsActionSchemaElement SpecifyExclusiveUnionActionFromArrayViews(const TArrayView<const FName> ElementNames, const TArrayView<const FLearningAgentsActionSchemaElement> Elements, const TArrayView<const float> PriorProbabilities = {}, const FName Name = TEXT("ExclusiveUnion"));
 
 	UFUNCTION(BlueprintPure, Category = "LearningAgents", meta = (AutoCreateRefTerm = "PriorProbabilities"))
 	FLearningAgentsActionSchemaElement SpecifyInclusiveUnionAction(const TMap<FName, FLearningAgentsActionSchemaElement>& Elements, const TMap<FName, float>& PriorProbabilities, const FName Name = TEXT("InclusiveUnion"));
 	
 	UFUNCTION(BlueprintPure, Category = "LearningAgents", meta = (AutoCreateRefTerm = "PriorProbabilities"))
 	FLearningAgentsActionSchemaElement SpecifyInclusiveUnionActionFromArrays(const TArray<FName> ElementNames, const TArray<FLearningAgentsActionSchemaElement>& Elements, const TArray<float>& PriorProbabilities, const FName Name = TEXT("InclusiveUnion"));
-	FLearningAgentsActionSchemaElement SpecifyInclusiveUnionActionFromArrayViews(const TArrayView<const FName> ElementNames, const TArrayView<const FLearningAgentsActionSchemaElement> Elements, const TArrayView<const float> PriorProbabilities, const FName Name = TEXT("InclusiveUnion"));
+	FLearningAgentsActionSchemaElement SpecifyInclusiveUnionActionFromArrayViews(const TArrayView<const FName> ElementNames, const TArrayView<const FLearningAgentsActionSchemaElement> Elements, const TArrayView<const float> PriorProbabilities = {}, const FName Name = TEXT("InclusiveUnion"));
 
 	UFUNCTION(BlueprintPure, Category = "LearningAgents")
 	FLearningAgentsActionSchemaElement SpecifyStaticArrayAction(const FLearningAgentsActionSchemaElement Element, const int32 Num, const FName Name = TEXT("StaticArray"));
@@ -129,14 +129,14 @@ public:
 
 	UFUNCTION(BlueprintPure, Category = "LearningAgents", meta = (AutoCreateRefTerm = "PriorProbabilities"))
 	FLearningAgentsActionSchemaElement SpecifyEnumActionFromArray(const UEnum* Enum, const TArray<float>& PriorProbabilities, const FName Name = TEXT("Enum"));
-	FLearningAgentsActionSchemaElement SpecifyEnumActionFromArrayView(const UEnum* Enum, const TArrayView<const float> PriorProbabilities, const FName Name = TEXT("Enum"));
+	FLearningAgentsActionSchemaElement SpecifyEnumActionFromArrayView(const UEnum* Enum, const TArrayView<const float> PriorProbabilities = {}, const FName Name = TEXT("Enum"));
 
 	UFUNCTION(BlueprintPure, Category = "LearningAgents", meta = (AutoCreateRefTerm = "PriorProbabilities"))
 	FLearningAgentsActionSchemaElement SpecifyBitmaskAction(const UEnum* Enum, const TMap<uint8, float>& PriorProbabilities, const FName Name = TEXT("Bitmask"));
 
 	UFUNCTION(BlueprintPure, Category = "LearningAgents", meta = (AutoCreateRefTerm = "PriorProbabilities"))
 	FLearningAgentsActionSchemaElement SpecifyBitmaskActionFromArray(const UEnum* Enum, const TArray<float>& PriorProbabilities, const FName Name = TEXT("Bitmask"));
-	FLearningAgentsActionSchemaElement SpecifyBitmaskActionFromArrayView(const UEnum* Enum, const TArrayView<const float> PriorProbabilities, const FName Name = TEXT("Bitmask"));
+	FLearningAgentsActionSchemaElement SpecifyBitmaskActionFromArrayView(const UEnum* Enum, const TArrayView<const float> PriorProbabilities = {}, const FName Name = TEXT("Bitmask"));
 
 	UFUNCTION(BlueprintPure, Category = "LearningAgents")
 	FLearningAgentsActionSchemaElement SpecifyOptionalAction(const FLearningAgentsActionSchemaElement Element, const float PriorProbability = 0.5f, const FName Name = TEXT("Optional"));
@@ -171,6 +171,9 @@ public:
 	UFUNCTION(BlueprintPure, Category = "LearningAgents")
 	FLearningAgentsActionSchemaElement SpecifyVelocityAction(const FName Name = TEXT("Velocity"));
 
+	UFUNCTION(BlueprintPure, Category = "LearningAgents")
+	FLearningAgentsActionSchemaElement SpecifyDirectionAction(const FName Name = TEXT("Direction"));
+
 	UFUNCTION(BlueprintPure, Category = "LearningAgents")
 	FLearningAgentsActionSchemaElement SpecifySpeedAction(const FName Name = TEXT("Speed"));
 
@@ -297,6 +300,9 @@ public:
 	UFUNCTION(BlueprintPure, Category = "LearningAgents")
 	FLearningAgentsActionObjectElement MakeVelocityAction(const FVector Velocity, const FTransform RelativeTransform = FTransform(), const float VelocityScale = 200.0f, const FName Name = TEXT("Velocity"));
 
+	UFUNCTION(BlueprintPure, Category = "LearningAgents")
+	FLearningAgentsActionObjectElement MakeDirectionAction(const FVector Direction, const FTransform RelativeTransform = FTransform(), const FName Name = TEXT("Direction"));
+
 	UFUNCTION(BlueprintPure, Category = "LearningAgents")
 	FLearningAgentsActionObjectElement MakeSpeedAction(const float Speed, const float SpeedScale = 200.0f, const FName Name = TEXT("Speed"));
 
@@ -398,6 +404,9 @@ public:
 	UFUNCTION(BlueprintPure = false, Category = "LearningAgents")
 	UPARAM(DisplayName = "Success") bool GetVelocityAction(FVector& OutVelocity, const FLearningAgentsActionObjectElement Element, const FTransform RelativeTransform = FTransform(), const float VelocityScale = 200.0f, const FName Name = TEXT("Velocity")) const;
 
+	UFUNCTION(BlueprintPure = false, Category = "LearningAgents")
+	UPARAM(DisplayName = "Success") bool GetDirectionAction(FVector& OutDirection, const FLearningAgentsActionObjectElement Element, const FTransform RelativeTransform = FTransform(), const FName Name = TEXT("Direction")) const;
+
 	UFUNCTION(BlueprintPure = false, Category = "LearningAgents")
 	UPARAM(DisplayName = "Success") bool GetSpeedAction(float& OutSpeed, const FLearningAgentsActionObjectElement Element, const float SpeedScale = 200.0f, const FName Name = TEXT("Speed")) const;
 
diff --git a/Engine/Plugins/Experimental/LearningAgents/Source/LearningAgents/Public/LearningAgentsObservations.h b/Engine/Plugins/Experimental/LearningAgents/Source/LearningAgents/Public/LearningAgentsObservations.h
--- a/Engine/Plugins/Experimental/LearningAgents/Source/LearningAgents/Public/LearningAgentsObservations.h
+++ b/Engine/Plugins/Experimental/LearningAgents/Source/LearningAgents/Public/LearningAgentsObservations.h
@@ -68,8 +68,13 @@ class ULearningAgentsObservationFunctions : public UBlueprintFunctionLibrary
 	GENERATED_BODY()
 
 	/** Project a transform onto the ground plane, leaving just rotation around the vertical axis */
-	UFUNCTION(BlueprintPure, Category = "Learning Agents")
+	UFUNCTION(BlueprintPure, Category = "LearningAgents")
 	static FTransform ProjectTransformOntoGroundPlane(const FTransform Transform, const FVector LocalForwardVector = FVector::ForwardVector, const float GroundPlaneHeight = 0.0f);
+
+
+	/** Find an Enum type by Name. This can be used to find Enum types defined in C++. This call can be expensive so the result should be cached. */
+	UFUNCTION(BlueprintCallable, Category = "LearningAgents")
+	static UEnum* FindEnumByName(const FString Name);
 };
 
 /**
@@ -120,6 +125,9 @@ public:
 	UFUNCTION(BlueprintPure, Category = "LearningAgents")
 	FLearningAgentsObservationSchemaElement SpecifyIndexObservation(const int32 Size, const FName Name = TEXT("Index"));
 
+	UFUNCTION(BlueprintPure, Category = "LearningAgents")
+	FLearningAgentsObservationSchemaElement SpecifyCountObservation(const FName Name = TEXT("Count"));
+
 	UFUNCTION(BlueprintPure, Category = "LearningAgents")
 	FLearningAgentsObservationSchemaElement SpecifyStructObservation(const TMap<FName, FLearningAgentsObservationSchemaElement>& Elements, const FName Name = TEXT("Struct"));
 	
@@ -275,6 +283,9 @@ public:
 	UFUNCTION(BlueprintPure, Category = "LearningAgents")
 	FLearningAgentsObservationObjectElement MakeIndexObservation(const int32 Index, const int32 Size, const FName Name = TEXT("Index"));
 
+	UFUNCTION(BlueprintPure, Category = "LearningAgents")
+	FLearningAgentsObservationObjectElement MakeCountObservation(const int32 Num, const int32 MaxNum, const FName Name = TEXT("Count"));
+
 	UFUNCTION(BlueprintPure, Category = "LearningAgents")
 	FLearningAgentsObservationObjectElement MakeStructObservation(const TMap<FName, FLearningAgentsObservationObjectElement>& Elements, const FName Name = TEXT("Struct"));
 	
@@ -452,6 +463,9 @@ public:
 	UFUNCTION(BlueprintPure = false, Category = "LearningAgents")
 	UPARAM(DisplayName = "Success") bool GetIndexObservation(int32& OutIndex, const FLearningAgentsObservationObjectElement Element, const FName Name = TEXT("Index")) const;
 
+	UFUNCTION(BlueprintPure = false, Category = "LearningAgents")
+	UPARAM(DisplayName = "Success") bool GetCountObservation(int32& OutNum, const FLearningAgentsObservationObjectElement Element, const int32 MaxNum, const FName Name = TEXT("Count")) const;
+
 	UFUNCTION(BlueprintPure = false, Category = "LearningAgents")
 	UPARAM(DisplayName = "Success") bool GetStructObservationNum(int32& OutNum, const FLearningAgentsObservationObjectElement Element, const FName Name = TEXT("Struct")) const;
 
