LearningAgents: Fixed bug with socket based training.

diff --git a/Engine/Plugins/Experimental/LearningAgents/Source/LearningTraining/Private/LearningSocketTraining.cpp b/Engine/Plugins/Experimental/LearningAgents/Source/LearningTraining/Private/LearningSocketTraining.cpp
--- a/Engine/Plugins/Experimental/LearningAgents/Source/LearningTraining/Private/LearningSocketTraining.cpp
+++ b/Engine/Plugins/Experimental/LearningAgents/Source/LearningTraining/Private/LearningSocketTraining.cpp
@@ -82,6 +82,8 @@ namespace UE::Learning::SocketTraining
 		FRWLock* NetworkLock,
 		const ELogSetting LogSettings)
 	{
+		UE_LEARNING_CHECK(OutNetworkBuffer.Num() == OutNetwork.GetSnapshotByteNum());
+
 		if (LogSettings != ELogSetting::Silent)
 		{
 			UE_LOG(LogLearning, Display, TEXT("Pulling Network..."));
@@ -109,6 +111,8 @@ namespace UE::Learning::SocketTraining
 			{
 				return ETrainerResponse::Unexpected;
 			}
+
+			break;
 		}
 
 		Response = RecvWithTimeout(Socket, OutNetworkBuffer.GetData(), OutNetworkBuffer.Num(), Timeout);
