LearningAgents: Allowed users to specify per-agent callbacks in blueprints.

#virtualized

diff --git a/Engine/Plugins/Experimental/LearningAgents/Source/LearningAgents/Private/LearningAgentsController.cpp b/Engine/Plugins/Experimental/LearningAgents/Source/LearningAgents/Private/LearningAgentsController.cpp
--- a/Engine/Plugins/Experimental/LearningAgents/Source/LearningAgents/Private/LearningAgentsController.cpp
+++ b/Engine/Plugins/Experimental/LearningAgents/Source/LearningAgents/Private/LearningAgentsController.cpp
@@ -12,30 +12,23 @@ ULearningAgentsController::ULearningAgentsController(FVTableHelper& Helper) : Su
 ULearningAgentsController::~ULearningAgentsController() = default;
 
 void ULearningAgentsController::EvaluateAgentController_Implementation(
-	TArray<FLearningAgentsActionObjectElement>& OutActionObjectElements,
+	FLearningAgentsActionObjectElement& OutActionObjectElement,
 	ULearningAgentsActionObject* InActionObject,
-	const TArray<FLearningAgentsObservationObjectElement>& InObservationObjectElements,
+	const FLearningAgentsObservationObjectElement& InObservationObjectElement,
 	const ULearningAgentsObservationObject* InObservationObject,
-	const TArray<int32>& AgentIds)
+	const int32 AgentIds)
 {
 	UE_LOG(LogLearning, Error, TEXT("%s: EvaluateAgentController function must be overridden!"), *GetName());
-	OutActionObjectElements.Empty();
+	OutActionObjectElement = FLearningAgentsActionObjectElement();
 }
 
-void ULearningAgentsController::EvaluateAgentControllerUsingDelegate(
+void ULearningAgentsController::EvaluateAgentControllers_Implementation(
 	TArray<FLearningAgentsActionObjectElement>& OutActionObjectElements,
 	ULearningAgentsActionObject* InActionObject,
 	const TArray<FLearningAgentsObservationObjectElement>& InObservationObjectElements,
 	const ULearningAgentsObservationObject* InObservationObject,
-	const TArray<int32>& AgentIds,
-	const FEvaluateAgentControllerDelegate& Delegate)
+	const TArray<int32>& AgentIds)
 {
-	if (!Delegate.IsBound())
-	{
-		UE_LOG(LogLearning, Error, TEXT("%s: Delegate Not Bound."), *GetName());
-		return;
-	}
-
 	const int32 AgentNum = AgentIds.Num();
 
 	if (AgentNum != InObservationObjectElements.Num())
@@ -47,7 +40,9 @@ void ULearningAgentsController::EvaluateAgentControllerUsingDelegate(
 	OutActionObjectElements.Empty(AgentNum);
 	for (int32 AgentIdx = 0; AgentIdx < AgentNum; AgentIdx++)
 	{
-		OutActionObjectElements.Add(Delegate.Execute(InActionObject, InObservationObjectElements[AgentIdx], InObservationObject, AgentIds[AgentIdx]));
+		FLearningAgentsActionObjectElement OutActionObjectElement;
+		EvaluateAgentController(OutActionObjectElement, InActionObject, InObservationObjectElements[AgentIdx], InObservationObject, AgentIds[AgentIdx]);
+		OutActionObjectElements.Add(OutActionObjectElement);
 	}
 }
 
@@ -133,11 +128,11 @@ void ULearningAgentsController::EvaluateController()
 		return;
 	}
 
-	// Run EvaluateAgentController Callback
+	// Run EvaluateAgentControllers Callback
 
 	Interactor->ActionObject->GetActionObject().Reset();
 	Interactor->ActionObjectElements.Empty(Manager->GetMaxAgentNum());
-	EvaluateAgentController(
+	EvaluateAgentControllers(
 		Interactor->ActionObjectElements, 
 		Interactor->ActionObject, 
 		Interactor->ObservationObjectElements,
@@ -146,7 +141,7 @@ void ULearningAgentsController::EvaluateController()
 
 	if (Manager->GetAgentNum() != Interactor->ActionObjectElements.Num())
 	{
-		UE_LOG(LogLearning, Error, TEXT("%s: Not enough Action Objects added by EvaluateAgentController. Expected %i, Got %i."), *GetName(), Manager->GetAgentNum(), Interactor->ActionObjectElements.Num());
+		UE_LOG(LogLearning, Error, TEXT("%s: Not enough Action Objects added by EvaluateAgentControllers. Expected %i, Got %i."), *GetName(), Manager->GetAgentNum(), Interactor->ActionObjectElements.Num());
 		return;
 	}
 
@@ -194,5 +189,5 @@ void ULearningAgentsController::RunController()
 
 	Interactor->GatherObservations();
 	EvaluateController();
-	Interactor->ScatterActions();
+	Interactor->PerformActions();
 }
diff --git a/Engine/Plugins/Experimental/LearningAgents/Source/LearningAgents/Private/LearningAgentsCritic.cpp b/Engine/Plugins/Experimental/LearningAgents/Source/LearningAgents/Private/LearningAgentsCritic.cpp
--- a/Engine/Plugins/Experimental/LearningAgents/Source/LearningAgents/Private/LearningAgentsCritic.cpp
+++ b/Engine/Plugins/Experimental/LearningAgents/Source/LearningAgents/Private/LearningAgentsCritic.cpp
@@ -8,6 +8,7 @@
 #include "LearningNeuralNetwork.h"
 #include "LearningCritic.h"
 #include "LearningLog.h"
+#include "LearningRandom.h"
 
 #include "UObject/Package.h"
 #include "Misc/FileHelper.h"
@@ -173,7 +174,7 @@ void ULearningAgentsCritic::SetupCritic(
 			CriticNetwork->NeuralNetworkData = NewObject<ULearningNeuralNetworkData>(CriticNetwork);
 		}
 
-		UE::NNE::RuntimeBasic::FModelBuilder Builder;
+		UE::NNE::RuntimeBasic::FModelBuilder Builder(UE::Learning::Random::Int(Seed ^ 0x2610fc8f));
 
 		TArray<uint8> FileData;
 		uint32 CriticInputSize, CriticOutputSize;
diff --git a/Engine/Plugins/Experimental/LearningAgents/Source/LearningAgents/Private/LearningAgentsInteractor.cpp b/Engine/Plugins/Experimental/LearningAgents/Source/LearningAgents/Private/LearningAgentsInteractor.cpp
--- a/Engine/Plugins/Experimental/LearningAgents/Source/LearningAgents/Private/LearningAgentsInteractor.cpp
+++ b/Engine/Plugins/Experimental/LearningAgents/Source/LearningAgents/Private/LearningAgentsInteractor.cpp
@@ -65,7 +65,7 @@ void ULearningAgentsInteractor::SetupInteractor(ULearningAgentsManager* InManage
 	const FName ObservationSchemaUniqueName = MakeUniqueObjectName(this, ULearningAgentsObservationSchema::StaticClass(), TEXT("ObservationSchema"), EUniqueObjectNameOptions::GloballyUnique);
 
 	ObservationSchema = NewObject<ULearningAgentsObservationSchema>(this, ObservationSchemaUniqueName);
-	SpecifyAgentObservations(ObservationSchemaElement, ObservationSchema);
+	SpecifyAgentObservation(ObservationSchemaElement, ObservationSchema);
 
 	if (!ObservationSchema->GetObservationSchema().IsValid(ObservationSchemaElement.SchemaElement))
 	{
@@ -102,7 +102,7 @@ void ULearningAgentsInteractor::SetupInteractor(ULearningAgentsManager* InManage
 	const FName ActionSchemaUniqueName = MakeUniqueObjectName(this, ULearningAgentsActionSchema::StaticClass(), TEXT("ActionSchema"), EUniqueObjectNameOptions::GloballyUnique);
 
 	ActionSchema = NewObject<ULearningAgentsActionSchema>(this, ActionSchemaUniqueName);
-	SpecifyAgentActions(ActionSchemaElement, ActionSchema);
+	SpecifyAgentAction(ActionSchemaElement, ActionSchema);
 
 	if (!ActionSchema->GetActionSchema().IsValid(ActionSchemaElement.SchemaElement))
 	{
@@ -185,61 +185,42 @@ void ULearningAgentsInteractor::OnAgentsReset_Implementation(const TArray<int32>
 	UE::Learning::Array::Set<1, uint64>(ActionVectorIteration, 0, AgentIds);
 }
 
-void ULearningAgentsInteractor::SpecifyAgentObservations_Implementation(FLearningAgentsObservationSchemaElement& OutObservationSchemaElement, ULearningAgentsObservationSchema* InObservationSchema)
+void ULearningAgentsInteractor::SpecifyAgentObservation_Implementation(FLearningAgentsObservationSchemaElement& OutObservationSchemaElement, ULearningAgentsObservationSchema* InObservationSchema)
 {
-	UE_LOG(LogLearning, Error, TEXT("%s: SpecifyAgentObservations function must be overridden!"), *GetName());
+	UE_LOG(LogLearning, Error, TEXT("%s: SpecifyAgentObservation function must be overridden!"), *GetName());
 	OutObservationSchemaElement = FLearningAgentsObservationSchemaElement();
 }
 
-void ULearningAgentsInteractor::GatherAgentObservations_Implementation(TArray<FLearningAgentsObservationObjectElement>& OutObservationObjectElements, ULearningAgentsObservationObject* InObservationObject, const TArray<int32>& AgentIds)
+void ULearningAgentsInteractor::GatherAgentObservation_Implementation(FLearningAgentsObservationObjectElement& OutObservationObjectElement, ULearningAgentsObservationObject* InObservationObject, const int32 AgentId)
 {
-	UE_LOG(LogLearning, Error, TEXT("%s: GatherAgentObservations function must be overridden!"), *GetName());
-	OutObservationObjectElements.Empty();
+	UE_LOG(LogLearning, Error, TEXT("%s: GatherAgentObservation function must be overridden!"), *GetName());
+	OutObservationObjectElement = FLearningAgentsObservationObjectElement();
 }
 
-void ULearningAgentsInteractor::GatherAgentObservationsUsingDelegate(
-	TArray<FLearningAgentsObservationObjectElement>& OutObservationElements,
-	ULearningAgentsObservationObject* InObservationObject,
-	const TArray<int32>& AgentIds,
-	const FGatherAgentObservationDelegate& Delegate)
+void ULearningAgentsInteractor::GatherAgentObservations_Implementation(TArray<FLearningAgentsObservationObjectElement>& OutObservationObjectElements, ULearningAgentsObservationObject* InObservationObject, const TArray<int32>& AgentIds)
 {
-	if (!Delegate.IsBound())
-	{
-		UE_LOG(LogLearning, Error, TEXT("%s: Delegate Not Bound."), *GetName());
-		OutObservationElements.Empty();
-		return;
-	}
-
-	OutObservationElements.Empty(AgentIds.Num());
+	OutObservationObjectElements.Empty(AgentIds.Num());
 	for (const int32 AgentId : AgentIds)
 	{
-		OutObservationElements.Add(Delegate.Execute(InObservationObject, AgentId));
+		FLearningAgentsObservationObjectElement OutElement;
+		GatherAgentObservation(OutElement, InObservationObject, AgentId);
+		OutObservationObjectElements.Add(OutElement);
 	}
 }
 
-void ULearningAgentsInteractor::SpecifyAgentActions_Implementation(FLearningAgentsActionSchemaElement& OutActionSchemaElement, ULearningAgentsActionSchema* InActionSchema)
+void ULearningAgentsInteractor::SpecifyAgentAction_Implementation(FLearningAgentsActionSchemaElement& OutActionSchemaElement, ULearningAgentsActionSchema* InActionSchema)
 {
-	UE_LOG(LogLearning, Error, TEXT("%s: SpecifyAgentActions function must be overridden!"), *GetName());
+	UE_LOG(LogLearning, Error, TEXT("%s: SpecifyAgentAction function must be overridden!"), *GetName());
 	OutActionSchemaElement = FLearningAgentsActionSchemaElement();
 }
 
-void ULearningAgentsInteractor::ScatterAgentActions_Implementation(const TArray<FLearningAgentsActionObjectElement>& InActionObjectElements, const ULearningAgentsActionObject* InActionObject, const TArray<int32>& AgentIds)
+void ULearningAgentsInteractor::PerformAgentAction_Implementation(const FLearningAgentsActionObjectElement& InActionObjectElement, const ULearningAgentsActionObject* InActionObject, const int32 AgentId)
 {
-	UE_LOG(LogLearning, Error, TEXT("%s: ScatterAgentActions function must be overridden!"), *GetName());
+	UE_LOG(LogLearning, Error, TEXT("%s: PerformAgentAction function must be overridden!"), *GetName());
 }
 
-void  ULearningAgentsInteractor::ScatterAgentActionsUsingDelegate(
-	const TArray<FLearningAgentsActionObjectElement>& InActionObjectElements,
-	const ULearningAgentsActionObject* InActionObject,
-	const TArray<int32>& AgentIds,
-	const FScatterAgentActionDelegate& Delegate)
+void ULearningAgentsInteractor::PerformAgentActions_Implementation(const TArray<FLearningAgentsActionObjectElement>& InActionObjectElements, const ULearningAgentsActionObject* InActionObject, const TArray<int32>& AgentIds)
 {
-	if (!Delegate.IsBound())
-	{
-		UE_LOG(LogLearning, Error, TEXT("%s: Delegate Not Bound."), *GetName());
-		return;
-	}
-
 	const int32 AgentNum = AgentIds.Num();
 
 	if (AgentNum != InActionObjectElements.Num())
@@ -250,7 +231,7 @@ void  ULearningAgentsInteractor::ScatterAgentActionsUsingDelegate(
 
 	for (int32 AgentIdx = 0; AgentIdx < AgentNum; AgentIdx++)
 	{
-		Delegate.Execute(InActionObjectElements[AgentIdx], InActionObject, AgentIds[AgentIdx]);
+		PerformAgentAction(InActionObjectElements[AgentIdx], InActionObject, AgentIds[AgentIdx]);
 	}
 }
 
@@ -303,9 +284,9 @@ void ULearningAgentsInteractor::GatherObservations(const UE::Learning::FIndexSet
 	}
 }
 
-void ULearningAgentsInteractor::ScatterActions(const UE::Learning::FIndexSet AgentSet)
+void ULearningAgentsInteractor::PerformActions(const UE::Learning::FIndexSet AgentSet)
 {
-	UE_LEARNING_TRACE_CPUPROFILER_EVENT_SCOPE(ULearningAgentsInteractor::ScatterActions);
+	UE_LEARNING_TRACE_CPUPROFILER_EVENT_SCOPE(ULearningAgentsInteractor::PerformActions);
 
 	if (!IsSetup())
 	{
@@ -350,9 +331,9 @@ void ULearningAgentsInteractor::ScatterActions(const UE::Learning::FIndexSet Age
 		ActionObjectElements.Add(ActionObjectElement);
 	}
 
-	// Scatter Action Objects
+	// Perform Action Objects
 
-	ScatterAgentActions(ActionObjectElements, ActionObject, ValidAgentIds);
+	PerformAgentActions(ActionObjectElements, ActionObject, ValidAgentIds);
 }
 
 void ULearningAgentsInteractor::GatherObservations()
@@ -365,14 +346,14 @@ void ULearningAgentsInteractor::GatherObservations()
 	GatherObservations(Manager->GetAllAgentSet());
 }
 
-void ULearningAgentsInteractor::ScatterActions()
+void ULearningAgentsInteractor::PerformActions()
 {
 	if (Manager->GetAgentNum() == 0)
 	{
 		UE_LOG(LogLearning, Warning, TEXT("%s: No agents added to Manager."), *GetName());
 	}
 
-	ScatterActions(Manager->GetAllAgentSet());
+	PerformActions(Manager->GetAllAgentSet());
 }
 
 
diff --git a/Engine/Plugins/Experimental/LearningAgents/Source/LearningAgents/Private/LearningAgentsPolicy.cpp b/Engine/Plugins/Experimental/LearningAgents/Source/LearningAgents/Private/LearningAgentsPolicy.cpp
--- a/Engine/Plugins/Experimental/LearningAgents/Source/LearningAgents/Private/LearningAgentsPolicy.cpp
+++ b/Engine/Plugins/Experimental/LearningAgents/Source/LearningAgents/Private/LearningAgentsPolicy.cpp
@@ -183,7 +183,7 @@ void ULearningAgentsPolicy::SetupPolicy(
 			EncoderOutputSize,
 			Interactor->GetObservationSchema(),
 			Interactor->GetObservationSchemaElement(),
-			Seed);
+			UE::Learning::Random::Int(Seed ^ 0x658868dd));
 
 		UE_LEARNING_CHECK(EncoderInputSize == ObservationVectorSize);
 		UE_LEARNING_CHECK(EncoderOutputSize == ObservationEncodedVectorSize);
@@ -232,7 +232,7 @@ void ULearningAgentsPolicy::SetupPolicy(
 			PolicyNetwork->NeuralNetworkData = NewObject<ULearningNeuralNetworkData>(PolicyNetwork);
 		}
 
-		UE::NNE::RuntimeBasic::FModelBuilder Builder;
+		UE::NNE::RuntimeBasic::FModelBuilder Builder(UE::Learning::Random::Int(Seed ^ 0x69315bf9));
 
 		const int32 PolicyHiddenLayerSize = PolicySettings.HiddenLayerSize;
 		const int32 PolicyLayerNum = PolicySettings.LayerNum;
@@ -325,7 +325,7 @@ void ULearningAgentsPolicy::SetupPolicy(
 			DecoderOutputSize,
 			Interactor->GetActionSchema(),
 			Interactor->GetActionSchemaElement(),
-			Seed);
+			UE::Learning::Random::Int(Seed ^ 0xfa88bb7f));
 
 		UE_LEARNING_CHECK(DecoderInputSize == ActionEncodedVectorSize);
 		UE_LEARNING_CHECK(DecoderOutputSize == ActionDistributionVectorSize);
@@ -656,7 +656,7 @@ void ULearningAgentsPolicy::RunInference(const float ActionNoiseScale)
 	EncodeObservations();
 	EvaluatePolicy();
 	DecodeAndSampleActions(ActionNoiseScale);
-	Interactor->ScatterActions();
+	Interactor->PerformActions();
 }
 
 void ULearningAgentsPolicy::GetMemoryState(TArray<float>& OutMemoryState, const int32 AgentId) const
diff --git a/Engine/Plugins/Experimental/LearningAgents/Source/LearningAgents/Public/LearningAgentsController.h b/Engine/Plugins/Experimental/LearningAgents/Source/LearningAgents/Public/LearningAgentsController.h
--- a/Engine/Plugins/Experimental/LearningAgents/Source/LearningAgents/Public/LearningAgentsController.h
+++ b/Engine/Plugins/Experimental/LearningAgents/Source/LearningAgents/Public/LearningAgentsController.h
@@ -10,9 +10,6 @@
 
 #include "LearningAgentsController.generated.h"
 
-
-DECLARE_DYNAMIC_DELEGATE_RetVal_FourParams(FLearningAgentsActionObjectElement, FEvaluateAgentControllerDelegate, ULearningAgentsActionObject*, InActionObject, const FLearningAgentsObservationObjectElement, InObservationObjectElement, const ULearningAgentsObservationObject*, InObservationObject, const int32, AgentId);
-
 /**
  * A controller is an object that can be used to construct actions from observations - essentially a hand-made Policy. This can be useful for making 
  * a learning agents system that uses some other existing behavior, e.g. we may want to gather demonstrations from a human or AI behavior tree 
@@ -50,54 +47,50 @@ public:
 	void SetupController(ULearningAgentsManager* InManager, ULearningAgentsInteractor* InInteractor);
 
 	/**
-	 * This callback should be implemented by the Controller and should produce an array of action object elements, from an array of observation 
-	 * object elements.
-	 * 
-	 * @param OutActionObjectElements		Output Action Object Elements. This should be the same size as the input AgentIds and 
-	 *                                      InObservationObjectElements arrays.
-	 * @param InActionObject				Action object used to construct output elements.
-	 * @param InObservationObjectElements	Input observation object elements
-	 * @param InObservationObject			Input observation object
-	 * @param AgentIds						Agent ids associated with each observation
+	 * This callback should be overridden by the Controller and produces an Action Object Element from an Observation Object Element.
+	 *
+	 * @param OutActionObjectElement		Output Action Object Element.
+	 * @param InActionObject				Action object used to construct the output Action Object Element.
+	 * @param InObservationObjectElement	Input Observation Object Element.
+	 * @param InObservationObject			Input Observation Object.
+	 * @param AgentId						Agent id associated with the observation.
 	 */
 	UFUNCTION(BlueprintNativeEvent, Category = "LearningAgents", Meta = (ForceAsFunction))
 	void EvaluateAgentController(
-		TArray<FLearningAgentsActionObjectElement>& OutActionObjectElements, 
-		ULearningAgentsActionObject* InActionObject, 
-		const TArray<FLearningAgentsObservationObjectElement>& InObservationObjectElements,
+		FLearningAgentsActionObjectElement& OutActionObjectElement,
+		ULearningAgentsActionObject* InActionObject,
+		const FLearningAgentsObservationObjectElement& InObservationObjectElement,
 		const ULearningAgentsObservationObject* InObservationObject,
-		const TArray<int32>& AgentIds);
+		const int32 AgentId);
 
 	/**
-	 * This is a convenience function that can be used to implement `EvaluateAgentController` in terms of a callback which iterates over every agent
-	 * individually and computes the actions from the observations.
-	 *
-	 * @param OutActionObjectElements		Output Action Object Elements. This should be the same size as the input AgentIds and
+	 * This callback can be overridden by the Controller and produces an array of Action Object Elements, from an array of Observation 
+	 * Object Elements. By default this will call EvaluateAgentController for each agent.
+	 * 
+	 * @param OutActionObjectElements		Output Action Object Elements. This should be the same size as the input AgentIds and 
 	 *                                      InObservationObjectElements arrays.
 	 * @param InActionObject				Action object used to construct output elements.
-	 * @param InObservationObjectElements	Input observation object elements
-	 * @param InObservationObject			Input observation object
-	 * @param AgentIds						Agent ids associated with each observation
-	 * @param Delegate						The Delegate used to compute actions from observations.
+	 * @param InObservationObjectElements	Input Observation Object Elements.
+	 * @param InObservationObject			Input Observation Object.
+	 * @param AgentIds						Agent ids associated with each observation.
 	 */
-	UFUNCTION(BlueprintPure = false, Category = "LearningAgents")
-	void EvaluateAgentControllerUsingDelegate(
-		TArray<FLearningAgentsActionObjectElement>& OutActionObjectElements,
-		ULearningAgentsActionObject* InActionObject,
+	UFUNCTION(BlueprintNativeEvent, Category = "LearningAgents", Meta = (ForceAsFunction))
+	void EvaluateAgentControllers(
+		TArray<FLearningAgentsActionObjectElement>& OutActionObjectElements, 
+		ULearningAgentsActionObject* InActionObject, 
 		const TArray<FLearningAgentsObservationObjectElement>& InObservationObjectElements,
 		const ULearningAgentsObservationObject* InObservationObject,
-		const TArray<int32>& AgentIds,
-		const FEvaluateAgentControllerDelegate& Delegate);
+		const TArray<int32>& AgentIds);
 
 	/**
 	 * Call this function when it is time to evaluate the controller and produce the actions for the agents. This should be called after 
-	 * GatherObservations but before ScatterActions. This will call this controller's EvaluateAgentController event.
+	 * GatherObservations but before PerformActions. This will call this controller's EvaluateAgentController event.
 	 */
 	UFUNCTION(BlueprintCallable, Category = "LearningAgents")
 	void EvaluateController();
 
 	/**
-	 * Calls GatherObservations, followed by EvaluateController, followed by ScatterActions
+	 * Calls GatherObservations, followed by EvaluateController, followed by PerformActions
 	 */
 	UFUNCTION(BlueprintCallable, Category = "LearningAgents")
 	void RunController();
diff --git a/Engine/Plugins/Experimental/LearningAgents/Source/LearningAgents/Public/LearningAgentsInteractor.h b/Engine/Plugins/Experimental/LearningAgents/Source/LearningAgents/Public/LearningAgentsInteractor.h
--- a/Engine/Plugins/Experimental/LearningAgents/Source/LearningAgents/Public/LearningAgentsInteractor.h
+++ b/Engine/Plugins/Experimental/LearningAgents/Source/LearningAgents/Public/LearningAgentsInteractor.h
@@ -18,15 +18,12 @@
 
 class ULearningAgentsNeuralNetwork;
 
-DECLARE_DYNAMIC_DELEGATE_RetVal_TwoParams(FLearningAgentsObservationObjectElement, FGatherAgentObservationDelegate, ULearningAgentsObservationObject*, InObservationObject, const int32, AgentId);
-DECLARE_DYNAMIC_DELEGATE_ThreeParams(FScatterAgentActionDelegate, const FLearningAgentsActionObjectElement, InActionObjectElement, const ULearningAgentsActionObject*, InActionObject, const int32, AgentId);
-
 /**
  * ULearningAgentsInteractor defines how agents interact with the environment through their observations and actions.
  *
- * To use this class, you need to implement `SpecifyAgentObservations` and `SpecifyAgentActions`, which will define 
- * the structure of inputs and outputs to your policy. You also need to implement `GatherAgentObservations` and 
- * `ScatterAgentActions` which will dictate how those observations are gathered, and actions actuated in your
+ * To use this class, you need to implement `SpecifyAgentObservation` and `SpecifyAgentAction`, which will define 
+ * the structure of inputs and outputs to your policy. You also need to implement `GatherAgentObservation` and 
+ * `PerformAgentAction` which will dictate how those observations are gathered, and actions actuated in your
  * environment.
  */
 UCLASS(Abstract, HideDropdown, BlueprintType, Blueprintable)
@@ -79,79 +76,71 @@ public:
 public:
 
 	/**
-	 * Specifies the structure of the observations using the observation schema
+	 * This callback should be overridden by the Interactor and specifies the structure of the observations using the Observation Schema.
 	 * 
 	 * @param OutObservationSchemaElement		Output Schema Element
 	 * @param InObservationSchema				Observation Schema
 	 */
 	UFUNCTION(BlueprintNativeEvent, Category = "LearningAgents", Meta = (ForceAsFunction))
-	void SpecifyAgentObservations(FLearningAgentsObservationSchemaElement& OutObservationSchemaElement, ULearningAgentsObservationSchema* InObservationSchema);
+	void SpecifyAgentObservation(FLearningAgentsObservationSchemaElement& OutObservationSchemaElement, ULearningAgentsObservationSchema* InObservationSchema);
 
 	/**
-	 * Gathers all the observations for the given agents. The structure of the Observation Elements output by this function should match that defined
-	 * by the schema.
+	 * This callback should be overridden by the Interactor and gathers the observations for a single agent. The structure of the Observation Elements 
+	 * output by this function should match that defined by the Schema.
 	 *
-	 * @param OutObservationObjectElements		Output Observation Element. This should be the same size as AgentIds.
+	 * @param OutObservationObjectElement		Output Observation Element.
 	 * @param InObservationObject				Observation Object.
-	 * @param AgentIds							Set of Agent Ids to gather observations for.
+	 * @param AgentId							The Agent Id to gather observations for.
 	 */
 	UFUNCTION(BlueprintNativeEvent, Category = "LearningAgents", Meta = (ForceAsFunction))
-	void GatherAgentObservations(TArray<FLearningAgentsObservationObjectElement>& OutObservationObjectElements, ULearningAgentsObservationObject* InObservationObject, const TArray<int32>& AgentIds);
+	void GatherAgentObservation(FLearningAgentsObservationObjectElement& OutObservationObjectElement, ULearningAgentsObservationObject* InObservationObject, const int32 AgentId);
+
 
 	/**
-	 * This is a convenience function that can be used to implement `GatherAgentObservations` in terms of a callback which iterates over every agent
-	 * individually and gathers the observations.
-	 * 
-	 * @param OutObservationObjectElements		Output Observation Element. This should be the same size as AgentIds.
+	 * This callback can be overridden by the Interactor and gathers all the observations for the given agents. The structure of the Observation 
+	 * Elements output by this function should match that defined by the Schema. The default implementation calls GatherAgentObservation on each agent.
+	 *
+	 * @param OutObservationObjectElements		Output Observation Elements. This should be the same size as AgentIds.
 	 * @param InObservationObject				Observation Object.
 	 * @param AgentIds							Set of Agent Ids to gather observations for.
-	 * @param Delegate							The Delegate used to gather each agent's observations.
 	 */
-	UFUNCTION(BlueprintPure = false, Category = "LearningAgents")
-	void GatherAgentObservationsUsingDelegate(
-			TArray<FLearningAgentsObservationObjectElement>& OutObservationElements,
-			ULearningAgentsObservationObject* InObservationObject, 
-			const TArray<int32>& AgentIds,
-			const FGatherAgentObservationDelegate& Delegate);
+	UFUNCTION(BlueprintNativeEvent, Category = "LearningAgents", Meta = (ForceAsFunction))
+	void GatherAgentObservations(TArray<FLearningAgentsObservationObjectElement>& OutObservationObjectElements, ULearningAgentsObservationObject* InObservationObject, const TArray<int32>& AgentIds);
 
 // ----- Actions -----
 public:
 
 	/**
-	 * Specifies the structure of the actions using the action schema
+	 * This callback should be overridden by the Interactor and specifies the structure of the actions using the Action Schema.
 	 *
 	 * @param OutActionSchemaElement			Output Schema Element
 	 * @param InActionSchema					Action Schema
 	 */
 	UFUNCTION(BlueprintNativeEvent, Category = "LearningAgents", Meta = (ForceAsFunction))
-	void SpecifyAgentActions(FLearningAgentsActionSchemaElement& OutActionSchemaElement, ULearningAgentsActionSchema* InActionSchema);
+	void SpecifyAgentAction(FLearningAgentsActionSchemaElement& OutActionSchemaElement, ULearningAgentsActionSchema* InActionSchema);
 
 	/**
-	 * Scatters all the actions for the given agents and applies their affect in the world. The structure of the Action Elements given as input to 
-	 * this function will match that defined by the schema.
+	 * This callback should be overridden by the Interactor and performs the action for the given agent in the world. The 
+	 * structure of the Action Elements given as input to this function will match that defined by the Schema.
 	 *
-	 * @param InActionObjectElements			Input Actions Element. This will be the same size as AgentIds.
+	 * @param InActionObjectElement				Input Actions Element.
 	 * @param InActionObject					Action Object.
-	 * @param AgentIds							Set of Agent Ids to scatter observations for.
+	 * @param AgentId							Agent Id to perform actions for.
 	 */
 	UFUNCTION(BlueprintNativeEvent, Category = "LearningAgents", Meta = (ForceAsFunction))
-	void ScatterAgentActions(const TArray<FLearningAgentsActionObjectElement>& InActionObjectElements, const ULearningAgentsActionObject* InActionObject, const TArray<int32>& AgentIds);
+	void PerformAgentAction(const FLearningAgentsActionObjectElement& InActionObjectElement, const ULearningAgentsActionObject* InActionObject, const int32 AgentId);
 
 	/**
-	 * This is a convenience function that can be used to implement `ScatterAgentActions` in terms of a callback which iterates over every agent
-	 * individually and scatters the actions.
+	 * This callback can be overridden by the Interactor and performs all the actions for the given agents in the world. 
+	 * The structure of the Action Elements given as input to this function will match that defined by the Schema. The default implementation calls 
+	 * PerformAgentAction on each agent.
 	 *
 	 * @param InActionObjectElements			Input Actions Element. This will be the same size as AgentIds.
 	 * @param InActionObject					Action Object.
-	 * @param AgentIds							Set of Agent Ids to scatter actions for.
-	 * @param Delegate							The Delegate used to scatter each agent's actions.
+	 * @param AgentIds							Set of Agent Ids to perform actions for.
 	 */
-	UFUNCTION(BlueprintPure = false, Category = "LearningAgents")
-	void ScatterAgentActionsUsingDelegate(
-		const TArray<FLearningAgentsActionObjectElement>& InActionObjectElements, 
-		const ULearningAgentsActionObject* InActionObject,
-		const TArray<int32>& AgentIds,
-		const FScatterAgentActionDelegate& Delegate);
+	UFUNCTION(BlueprintNativeEvent, Category = "LearningAgents", Meta = (ForceAsFunction))
+	void PerformAgentActions(const TArray<FLearningAgentsActionObjectElement>& InActionObjectElements, const ULearningAgentsActionObject* InActionObject, const TArray<int32>& AgentIds);
 
 // ----- Blueprint public interface -----
 public:
@@ -160,9 +149,9 @@ public:
 	UFUNCTION(BlueprintCallable, Category = "LearningAgents")
 	void GatherObservations();
 
-	/** Scatters all the actions for all agents. This will call ScatterAgentActions. */
+	/** Performs all the actions for all agents. This will call PerformAgentActions. */
 	UFUNCTION(BlueprintCallable, Category = "LearningAgents")
-	void ScatterActions();
+	void PerformActions();
 
 	/**
 	 * Get the current buffered observation vector for the given agent.
@@ -244,8 +233,8 @@ public:
 	/** Encode Observations for a specific set of agents */
 	void GatherObservations(const UE::Learning::FIndexSet AgentSet, bool bIncrementIteration = true);
 
-	/** Scatter Actions for a specific set of agents */
-	void ScatterActions(const UE::Learning::FIndexSet AgentSet);
+	/** Perform Actions for a specific set of agents */
+	void PerformActions(const UE::Learning::FIndexSet AgentSet);
 
 	/** Gets the internal observation schema object */
 	const UE::Learning::Observation::FSchema& GetObservationSchema() const;
diff --git a/Engine/Plugins/Experimental/LearningAgents/Source/LearningAgents/Public/LearningAgentsPolicy.h b/Engine/Plugins/Experimental/LearningAgents/Source/LearningAgents/Public/LearningAgentsPolicy.h
--- a/Engine/Plugins/Experimental/LearningAgents/Source/LearningAgents/Public/LearningAgentsPolicy.h
+++ b/Engine/Plugins/Experimental/LearningAgents/Source/LearningAgents/Public/LearningAgentsPolicy.h
@@ -163,7 +163,7 @@ public:
 	void EvaluatePolicy();
 
 	/**
-	 * Decodes and samples action vectors using the Decoder network. This should be called after EvaluatePolicy and before Scatter Actions.
+	 * Decodes and samples action vectors using the Decoder network. This should be called after EvaluatePolicy and before PerformActions.
 	 * 
 	 * @param ActionNoiseScale		Scale of the action noise to use during sampling. Set this to zero to always sample the mean (expected) action.
 	 */
@@ -171,7 +171,7 @@ public:
 	void DecodeAndSampleActions(const float ActionNoiseScale = 1.0f);
 
 	/**
-	 * Calls GatherObservations, EncodeObservations, EvaluatePolicy, DecodeAndSampleActions, ScatterActions
+	 * Calls GatherObservations, EncodeObservations, EvaluatePolicy, DecodeAndSampleActions, PerformActions
 	 */
 	UFUNCTION(BlueprintCallable, Category = "LearningAgents")
 	void RunInference(const float ActionNoiseScale = 1.0f);
diff --git a/Engine/Plugins/Experimental/LearningAgents/Source/LearningAgentsTraining/Private/LearningAgentsTrainer.cpp b/Engine/Plugins/Experimental/LearningAgents/Source/LearningAgentsTraining/Private/LearningAgentsTrainer.cpp
--- a/Engine/Plugins/Experimental/LearningAgents/Source/LearningAgentsTraining/Private/LearningAgentsTrainer.cpp
+++ b/Engine/Plugins/Experimental/LearningAgents/Source/LearningAgentsTraining/Private/LearningAgentsTrainer.cpp
@@ -308,68 +308,52 @@ void ULearningAgentsTrainer::OnAgentsManagerTick_Implementation(const TArray<int
 	}
 }
 
-void ULearningAgentsTrainer::GatherAgentRewards_Implementation(TArray<float>& OutRewards, const TArray<int32>& AgentIds)
+void ULearningAgentsTrainer::GatherAgentReward_Implementation(float& OutReward, const int32 AgentId)
 {
-	UE_LOG(LogLearning, Error, TEXT("%s: GatherAgentRewards function must be overridden!"), *GetName());
-	OutRewards.Empty();
+	UE_LOG(LogLearning, Error, TEXT("%s: GatherAgentReward function must be overridden!"), *GetName());
+	OutReward = 0.0f;
 }
 
-void ULearningAgentsTrainer::GatherAgentRewardsUsingDelegate(TArray<float>& OutRewards, const TArray<int32>& AgentIds, const FGatherAgentRewardDelegate& Delegate)
+void ULearningAgentsTrainer::GatherAgentRewards_Implementation(TArray<float>& OutRewards, const TArray<int32>& AgentIds)
 {
-	if (!Delegate.IsBound())
-	{
-		UE_LOG(LogLearning, Error, TEXT("%s: Delegate Not Bound."), *GetName());
-		return;
-	}
 
-	const int32 AgentNum = AgentIds.Num();
-
-	OutRewards.Empty(AgentNum);
-	for (int32 AgentIdx = 0; AgentIdx < AgentNum; AgentIdx++)
+	OutRewards.Empty(AgentIds.Num());
+	for (const int32 AgentId : AgentIds)
 	{
-		OutRewards.Add(Delegate.Execute(AgentIds[AgentIdx]));
+		float OutReward = 0.0f;
+		GatherAgentReward(OutReward, AgentId);
+		OutRewards.Add(OutReward);
 	}
 }
 
-void ULearningAgentsTrainer::GatherAgentCompletions_Implementation(TArray<ELearningAgentsCompletion>& OutCompletions, const TArray<int32>& AgentIds)
+void ULearningAgentsTrainer::GatherAgentCompletion_Implementation(ELearningAgentsCompletion& OutCompletion, const int32 AgentId)
 {
-	UE_LOG(LogLearning, Error, TEXT("%s: GatherAgentCompletions function must be overridden!"), *GetName());
-	OutCompletions.Empty();
+	UE_LOG(LogLearning, Error, TEXT("%s: GatherAgentCompletion function must be overridden!"), *GetName());
+	OutCompletion = ELearningAgentsCompletion::Running;
 }
 
-void ULearningAgentsTrainer::GatherAgentCompletionsDelegate(TArray<ELearningAgentsCompletion>& OutCompletions, const TArray<int32>& AgentIds, const FGatherAgentCompletionDelegate& Delegate)
+void ULearningAgentsTrainer::GatherAgentCompletions_Implementation(TArray<ELearningAgentsCompletion>& OutCompletions, const TArray<int32>& AgentIds)
 {
-	if (!Delegate.IsBound())
-	{
-		UE_LOG(LogLearning, Error, TEXT("%s: Delegate Not Bound."), *GetName());
-		return;
-	}
 
-	const int32 AgentNum = AgentIds.Num();
-
-	OutCompletions.Empty(AgentNum);
-	for (int32 AgentIdx = 0; AgentIdx < AgentNum; AgentIdx++)
+	OutCompletions.Empty(AgentIds.Num());
+	for (const int32 AgentId : AgentIds)
 	{
-		OutCompletions.Add(Delegate.Execute(AgentIds[AgentIdx]));
+		ELearningAgentsCompletion OutCompletion = ELearningAgentsCompletion::Running;
+		GatherAgentCompletion(OutCompletion, AgentId);
+		OutCompletions.Add(OutCompletion);
 	}
 }
 
-void ULearningAgentsTrainer::ResetAgentEpisodes_Implementation(const TArray<int32>& AgentId)
+void ULearningAgentsTrainer::ResetAgentEpisode_Implementation(const int32 AgentId)
 {
-	UE_LOG(LogLearning, Error, TEXT("%s: ResetAgentEpisodes function must be overridden!"), *GetName());
+	UE_LOG(LogLearning, Error, TEXT("%s: ResetAgentEpisode function must be overridden!"), *GetName());
 }
 
-void ULearningAgentsTrainer::ResetAgentEpisodesDelegate(const TArray<int32>& AgentIds, const FResetAgentEpisodeDelegate& Delegate)
+void ULearningAgentsTrainer::ResetAgentEpisodes_Implementation(const TArray<int32>& AgentIds)
 {
-	if (!Delegate.IsBound())
-	{
-		UE_LOG(LogLearning, Error, TEXT("%s: Delegate Not Bound."), *GetName());
-		return;
-	}
-
 	for (const int32 AgentId : AgentIds)
 	{
-		Delegate.Execute(AgentId);
+		ResetAgentEpisode(AgentId);
 	}
 }
 
diff --git a/Engine/Plugins/Experimental/LearningAgents/Source/LearningAgentsTraining/Public/LearningAgentsTrainer.h b/Engine/Plugins/Experimental/LearningAgents/Source/LearningAgentsTraining/Public/LearningAgentsTrainer.h
--- a/Engine/Plugins/Experimental/LearningAgents/Source/LearningAgentsTraining/Public/LearningAgentsTrainer.h
+++ b/Engine/Plugins/Experimental/LearningAgents/Source/LearningAgentsTraining/Public/LearningAgentsTrainer.h
@@ -25,10 +25,6 @@ class ULearningAgentsInteractor;
 class ULearningAgentsPolicy;
 class ULearningAgentsCritic;
 
-DECLARE_DYNAMIC_DELEGATE_RetVal_OneParam(float, FGatherAgentRewardDelegate, const int32, AgentId);
-DECLARE_DYNAMIC_DELEGATE_RetVal_OneParam(ELearningAgentsCompletion, FGatherAgentCompletionDelegate, const int32, AgentId);
-DECLARE_DYNAMIC_DELEGATE_OneParam(FResetAgentEpisodeDelegate, const int32, AgentId);
-
 /** The configurable settings for a ULearningAgentsTrainer. */
 USTRUCT(BlueprintType, Category = "LearningAgents")
 struct LEARNINGAGENTSTRAINING_API FLearningAgentsTrainerSettings
@@ -62,7 +58,7 @@ public:
 
 	/** Time in seconds to wait for the training subprocess before timing out. */
 	UPROPERTY(EditAnywhere, Category = "LearningAgents", meta = (ClampMin = "0.0", UIMin = "0.0"))
-	float TrainerCommunicationTimeout = 20.0f;
+	float TrainerCommunicationTimeout = 10.0f;
 };
 
 /**
@@ -408,69 +404,66 @@ public:
 public:
 
 	/**
-	 * This callback gathers all the reward values for the given set of agents.
+	 * This callback should be overridden by the Trainer and gathers the reward value for the given agent.
 	 *
-	 * @param OutRewards		Output rewards for each agent in AgentIds
-	 * @param AgentIds			Agents to gather rewards for.
+	 * @param OutReward			Output reward for the given agent.
+	 * @param AgentId			Agent id to gather reward for.
 	 */
 	UFUNCTION(BlueprintNativeEvent, Category = "LearningAgents", Meta = (ForceAsFunction))
-	void GatherAgentRewards(TArray<float>& OutRewards, const TArray<int32>& AgentIds);
+	void GatherAgentReward(float& OutReward, const int32 AgentId);
+
 
 	/**
-	 * This is a convenience function that can be used to implement `GatherAgentRewards` in terms of a callback which iterates over every agent
-	 * individually and gathers the rewards.
+	 * This callback can be overridden by the Trainer and gathers all the reward values for the given set of agents. By default this will call 
+	 * GatherAgentReward on each agent.
 	 *
 	 * @param OutRewards		Output rewards for each agent in AgentIds
 	 * @param AgentIds			Agents to gather rewards for.
-	 * @param Delegate			The Delegate used to gather the rewards for each agent.
 	 */
-	UFUNCTION(BlueprintPure = false, Category = "LearningAgents")
-	void GatherAgentRewardsUsingDelegate(TArray<float>& OutRewards, const TArray<int32>& AgentIds, const FGatherAgentRewardDelegate& Delegate);
+	UFUNCTION(BlueprintNativeEvent, Category = "LearningAgents", Meta = (ForceAsFunction))
+	void GatherAgentRewards(TArray<float>& OutRewards, const TArray<int32>& AgentIds);
 
 // ----- Completions ----- 
 public:
 
 	/**
-	 * This callback gathers all the completions for the given set of agents.
+	 * This callback should be overridden by the Trainer and gathers the completion for a given agent.
 	 *
-	 * @param OutCompletions	Output completions for each agent in AgentIds
-	 * @param AgentIds			Agents to gather completions for.
+	 * @param OutCompletion		Output completion for the given agent.
+	 * @param AgentId			Agent id to gather completion for.
 	 */
 	UFUNCTION(BlueprintNativeEvent, Category = "LearningAgents", Meta = (ForceAsFunction))
-	void GatherAgentCompletions(TArray<ELearningAgentsCompletion>& OutCompletions, const TArray<int32>& AgentIds);
+	void GatherAgentCompletion(ELearningAgentsCompletion& OutCompletion, const int32 AgentId);
 
 	/**
-	 * This is a convenience function that can be used to implement `GatherAgentCompletions` in terms of a callback which iterates over every agent
-	 * individually and gathers the completions.
+	 * This callback can be overridden by the Trainer and gathers all the completions for the given set of agents. By default this will call 
+	 * GatherAgentCompletion on each agent.
 	 *
 	 * @param OutCompletions	Output completions for each agent in AgentIds
 	 * @param AgentIds			Agents to gather completions for.
-	 * @param Delegate			The Delegate used to gather the completions for each agent.
 	 */
-	UFUNCTION(BlueprintPure = false, Category = "LearningAgents")
-	void GatherAgentCompletionsDelegate(TArray<ELearningAgentsCompletion>& OutCompletions, const TArray<int32>& AgentIds, const FGatherAgentCompletionDelegate& Delegate);
-
+	UFUNCTION(BlueprintNativeEvent, Category = "LearningAgents", Meta = (ForceAsFunction))
+	void GatherAgentCompletions(TArray<ELearningAgentsCompletion>& OutCompletions, const TArray<int32>& AgentIds);
 
 // ----- Resets ----- 
 public:
 
 	/**
-	 * During this event, all episodes should be reset for each agent.
-	 * 
-	 * @param AgentIds			The ids of the agents that need resetting.
+	 * This callback should be overridden by the Trainer and resets the episode for the given agent.
+	 *
+	 * @param AgentId			The id of the agent that need resetting.
 	 */
 	UFUNCTION(BlueprintNativeEvent, Category = "LearningAgents", Meta = (ForceAsFunction))
-	void ResetAgentEpisodes(const TArray<int32>& AgentIds);
+	void ResetAgentEpisode(const int32 AgentId);
 
 	/**
-	 * This is a convenience function that can be used to implement `ResetAgentEpisodes` in terms of a callback which iterates over every agent
-	 * individually and does the reset.
-	 *
+	 * This callback can be overridden by the Trainer and resets all episodes for each agent in the given set. By default this will call 
+	 * ResetAgentEpisode on each agent.
+	 * 
 	 * @param AgentIds			The ids of the agents that need resetting.
-	 * @param Delegate			The Delegate used to reset the episode for each agent.
 	 */
-	UFUNCTION(BlueprintPure = false, Category = "LearningAgents")
-	void ResetAgentEpisodesDelegate(const TArray<int32>& AgentIds, const FResetAgentEpisodeDelegate& Delegate);
+	UFUNCTION(BlueprintNativeEvent, Category = "LearningAgents", Meta = (ForceAsFunction))
+	void ResetAgentEpisodes(const TArray<int32>& AgentIds);
 
 // ----- Training Process -----
 public:
