Learning Agents: Add some obs visual logger and code comments. Make MaxAgentNum Protected. Add QA Test Map. Bump Plugin version to 0.2

#rnx

#virtualized

diff --git a/Engine/Plugins/Experimental/LearningAgents/LearningAgents.uplugin b/Engine/Plugins/Experimental/LearningAgents/LearningAgents.uplugin
--- a/Engine/Plugins/Experimental/LearningAgents/LearningAgents.uplugin
+++ b/Engine/Plugins/Experimental/LearningAgents/LearningAgents.uplugin
@@ -1,7 +1,7 @@
 {
 	"FileVersion": 3,
-	"Version": 1,
-	"VersionName": "0.1",
+	"Version": 2,
+	"VersionName": "0.2",
 	"FriendlyName": "Learning Agents",
 	"Description": "Learning Agents is a machine learning library for AI character control in games. It simplifies the use of reinforcement and imitation learning in Unreal.",
 	"Category": "AI",
diff --git a/Engine/Plugins/Experimental/LearningAgents/Source/LearningAgents/Private/LearningAgentsObservations.cpp b/Engine/Plugins/Experimental/LearningAgents/Source/LearningAgents/Private/LearningAgentsObservations.cpp
--- a/Engine/Plugins/Experimental/LearningAgents/Source/LearningAgents/Private/LearningAgentsObservations.cpp
+++ b/Engine/Plugins/Experimental/LearningAgents/Source/LearningAgents/Private/LearningAgentsObservations.cpp
@@ -1568,14 +1568,58 @@ FLearningAgentsObservationObjectElement ULearningAgentsObservationObject::MakeEn
 	return { ObservationObject.CreateEncoding({ Element.ObjectElement }, Name) };
 }
 
-FLearningAgentsObservationObjectElement ULearningAgentsObservationObject::MakeBoolObservation(const bool bValue, const FName Name)
+FLearningAgentsObservationObjectElement ULearningAgentsObservationObject::MakeBoolObservation(
+	const bool bValue,
+	const FName Name,
+	const bool bVisualLoggerEnabled,
+	const FVector VisualLoggerLocation,
+	const FLinearColor VisualLoggerColor)
 {
-	return MakeContinuousObservationFromArrayView({ bValue ? 1.0f : -1.0f }, Name);
+	const float EncodedBool = bValue ? 1.0f : -1.0f;
+
+#if UE_LEARNING_AGENTS_ENABLE_VISUAL_LOG
+	if (bVisualLoggerEnabled)
+	{
+		const ULearningAgentsObservationVisualLoggerObject* VisualLoggerObject = GetOrAddVisualLoggerObject(Name);
+
+		UE_LEARNING_AGENTS_VLOG_STRING(VisualLoggerObject, LogLearning, Display, VisualLoggerLocation,
+			VisualLoggerColor.ToFColor(true),
+			TEXT("Name: %s\nValue: [%s]\nEncoded: [% 6.2f]"),
+			*Name.ToString(),
+			bValue ? TEXT("true") : TEXT("false"),
+			EncodedBool);
+	}
+#endif
+
+	return MakeContinuousObservationFromArrayView({ EncodedBool }, Name);
 }
 
-FLearningAgentsObservationObjectElement ULearningAgentsObservationObject::MakeFloatObservation(const float Value, const float FloatScale, const FName Name)
+FLearningAgentsObservationObjectElement ULearningAgentsObservationObject::MakeFloatObservation(
+	const float Value,
+	const float FloatScale,
+	const FName Name,
+	const bool bVisualLoggerEnabled,
+	const FVector VisualLoggerLocation,
+	const FLinearColor VisualLoggerColor)
 {
-	return MakeContinuousObservationFromArrayView({ Value / FMath::Max(FloatScale, UE_SMALL_NUMBER) }, Name);
+	const float EncodedValue = Value / FMath::Max(FloatScale, UE_SMALL_NUMBER);
+
+#if UE_LEARNING_AGENTS_ENABLE_VISUAL_LOG
+	if (bVisualLoggerEnabled)
+	{
+		const ULearningAgentsObservationVisualLoggerObject* VisualLoggerObject = GetOrAddVisualLoggerObject(Name);
+
+		UE_LEARNING_AGENTS_VLOG_STRING(VisualLoggerObject, LogLearning, Display, VisualLoggerLocation,
+			VisualLoggerColor.ToFColor(true),
+			TEXT("Name: %s\nValue: [% 6.1f]\nScale: [% 6.2f]\nEncoded: [% 6.2f]"),
+			*Name.ToString(),
+			Value,
+			FloatScale,
+			EncodedValue);
+	}
+#endif
+
+	return MakeContinuousObservationFromArrayView({ EncodedValue }, Name);
 }
 
 FLearningAgentsObservationObjectElement ULearningAgentsObservationObject::MakeLocationObservation(
@@ -1634,17 +1678,55 @@ FLearningAgentsObservationObjectElement ULearningAgentsObservationObject::MakeLo
 		}, Name);
 }
 
-FLearningAgentsObservationObjectElement ULearningAgentsObservationObject::MakeRotationObservation(const FRotator Rotation, const FRotator RelativeRotation, const FName Name)
+FLearningAgentsObservationObjectElement ULearningAgentsObservationObject::MakeRotationObservation(
+	const FRotator Rotation,
+	const FRotator RelativeRotation,
+	const FName Name,
+	const bool bVisualLoggerEnabled,
+	const FVector VisualLoggerLocation,
+	const FLinearColor VisualLoggerColor)
 {
-	return MakeRotationObservationFromQuat(FQuat::MakeFromRotator(Rotation), FQuat::MakeFromRotator(RelativeRotation), Name);
+	// Visual logging is handled in the MakeRotationObservationFromQuat
+
+	return MakeRotationObservationFromQuat(
+		FQuat::MakeFromRotator(Rotation), FQuat::MakeFromRotator(RelativeRotation), Name,
+		bVisualLoggerEnabled, VisualLoggerLocation,	VisualLoggerColor);
 }
 
-FLearningAgentsObservationObjectElement ULearningAgentsObservationObject::MakeRotationObservationFromQuat(const FQuat Rotation, const FQuat RelativeRotation, const FName Name)
+FLearningAgentsObservationObjectElement ULearningAgentsObservationObject::MakeRotationObservationFromQuat(
+	const FQuat Rotation,
+	const FQuat RelativeRotation,
+	const FName Name,
+	const bool bVisualLoggerEnabled,
+	const FVector VisualLoggerLocation,
+	const FLinearColor VisualLoggerColor)
 {
 	const FQuat LocalRotation = RelativeRotation.Inverse() * Rotation;
 	const FVector LocalAxisForward = LocalRotation.GetForwardVector();
 	const FVector LocalAxisRight = LocalRotation.GetRightVector();
 
+#if UE_LEARNING_AGENTS_ENABLE_VISUAL_LOG
+	if (bVisualLoggerEnabled)
+	{
+		const ULearningAgentsObservationVisualLoggerObject* VisualLoggerObject = GetOrAddVisualLoggerObject(Name);
+
+		UE_LEARNING_AGENTS_VLOG_TRANSFORM(VisualLoggerObject, LogLearning, Display,
+			VisualLoggerLocation,
+			LocalRotation.Rotator(),
+			VisualLoggerColor.ToFColor(true),
+			TEXT(""));
+
+		UE_LEARNING_AGENTS_VLOG_STRING(VisualLoggerObject, LogLearning, Display, VisualLoggerLocation,
+			VisualLoggerColor.ToFColor(true),
+			TEXT("Name: %s\nRotation: [% 6.1f % 6.1f % 6.1f % 6.1f]\nLocal Rotation: [% 6.1f % 6.1f % 6.1f % 6.1f]\nEncoded Forward: [% 6.2f % 6.2f % 6.2f]\nEncoded Right: [% 6.2f % 6.2f % 6.2f]"),
+			*Name.ToString(),
+			Rotation.W, Rotation.X, Rotation.Y, Rotation.Z,
+			LocalRotation.W, LocalRotation.X, LocalRotation.Y, LocalRotation.Z,
+			LocalAxisForward.X, LocalAxisForward.Y, LocalAxisForward.Z,
+			LocalAxisRight.X, LocalAxisRight.Y, LocalAxisRight.Z);
+	}
+#endif
+
 	return MakeContinuousObservationFromArrayView({
 		(float)LocalAxisForward.X,
 		(float)LocalAxisForward.Y,
@@ -1655,12 +1737,32 @@ FLearningAgentsObservationObjectElement ULearningAgentsObservationObject::MakeRo
 		}, Name);
 }
 
-FLearningAgentsObservationObjectElement ULearningAgentsObservationObject::MakeScaleObservation(const FVector Scale, const FVector RelativeScale, const FName Name)
+FLearningAgentsObservationObjectElement ULearningAgentsObservationObject::MakeScaleObservation(
+	const FVector Scale,
+	const FVector RelativeScale,
+	const FName Name,
+	const bool bVisualLoggerEnabled,
+	const FVector VisualLoggerLocation,
+	const FLinearColor VisualLoggerColor)
 {
 	const FVector LocalLogScale = 
 		UE::Learning::Agents::Observation::Private::VectorLogSafe(Scale) - 
 		UE::Learning::Agents::Observation::Private::VectorLogSafe(RelativeScale);
 
+#if UE_LEARNING_AGENTS_ENABLE_VISUAL_LOG
+	if (bVisualLoggerEnabled)
+	{
+		const ULearningAgentsObservationVisualLoggerObject* VisualLoggerObject = GetOrAddVisualLoggerObject(Name);
+
+		UE_LEARNING_AGENTS_VLOG_STRING(VisualLoggerObject, LogLearning, Display, VisualLoggerLocation,
+			VisualLoggerColor.ToFColor(true),
+			TEXT("Name: %s\nScale: [% 6.1f % 6.1f % 6.1f]\nEncoded: [% 6.2f % 6.2f % 6.2f]"),
+			*Name.ToString(),
+			Scale.X, Scale.Y, Scale.Z,
+			LocalLogScale.X, LocalLogScale.Y, LocalLogScale.Z);
+	}
+#endif
+
 	return MakeContinuousObservationFromArrayView({
 		(float)LocalLogScale.X,
 		(float)LocalLogScale.Y,
@@ -1668,10 +1770,50 @@ FLearningAgentsObservationObjectElement ULearningAgentsObservationObject::MakeSc
 		}, Name);
 }
 
-FLearningAgentsObservationObjectElement ULearningAgentsObservationObject::MakeTransformObservation(const FTransform Transform, const FTransform RelativeTransform, const float LocationScale, const FName Name)
+FLearningAgentsObservationObjectElement ULearningAgentsObservationObject::MakeTransformObservation(
+	const FTransform Transform,
+	const FTransform RelativeTransform,
+	const float LocationScale,
+	const FName Name,
+	const bool bVisualLoggerEnabled,
+	const FVector VisualLoggerLocation,
+	const FLinearColor VisualLoggerColor)
 {
 	const FTransform LocalTransform = Transform * RelativeTransform.Inverse();
 	
+#if UE_LEARNING_AGENTS_ENABLE_VISUAL_LOG
+	if (bVisualLoggerEnabled)
+	{
+		const ULearningAgentsObservationVisualLoggerObject* VisualLoggerObject = GetOrAddVisualLoggerObject(Name);
+
+		const FVector LocalLocation = LocalTransform.GetLocation();
+		const FRotator LocalRotation = LocalTransform.Rotator();
+		const FVector LocalScale = LocalTransform.GetScale3D();
+
+		UE_LEARNING_AGENTS_VLOG_TRANSFORM(VisualLoggerObject, LogLearning, Display,
+			LocalLocation,
+			LocalRotation,
+			VisualLoggerColor.ToFColor(true),
+			TEXT(""));
+
+		const FVector Location = Transform.GetLocation();
+		const FRotator Rotation = Transform.Rotator();
+		const FVector Scale = Transform.GetScale3D();
+
+		UE_LEARNING_AGENTS_VLOG_STRING(VisualLoggerObject, LogLearning, Display, VisualLoggerLocation,
+			VisualLoggerColor.ToFColor(true),
+			TEXT("Name: %s\nLocation: [% 6.1f % 6.1f % 6.1f]\nLocal Location: [% 6.1f % 6.1f % 6.1f]\nRotation: [% 6.1f % 6.1f % 6.1f]\nLocal Rotation: [% 6.1f % 6.1f % 6.1f]\nScale: [% 6.1f % 6.1f % 6.1f]\nLocal Scale: [% 6.1f % 6.1f % 6.1f]\nLocation Scale: [% 6.1f]"),
+			*Name.ToString(),
+			Location.X, Location.Y, Location.Z,
+			LocalLocation.X, LocalLocation.Y, LocalLocation.Z,
+			Rotation.Roll, Rotation.Pitch, Rotation.Yaw,
+			LocalRotation.Roll, LocalRotation.Pitch, LocalRotation.Yaw,
+			Scale.X, Scale.Y, Scale.Z,
+			LocalScale.X, LocalScale.Y, LocalScale.Z,
+			LocationScale);
+	}
+#endif
+
 	return MakeStructObservationFromArrayViews(
 		{
 			TEXT("Location"),
@@ -1718,10 +1860,38 @@ FLearningAgentsObservationObjectElement ULearningAgentsObservationObject::MakeDi
 	const FName Name,
 	const bool bVisualLoggerEnabled,
 	const FVector VisualLoggerLocation,
+	const float VisualLoggerArrowLength,
 	const FLinearColor VisualLoggerColor)
 {
 	const FVector LocalDirection = RelativeTransform.InverseTransformVectorNoScale(Direction).GetSafeNormal(UE_SMALL_NUMBER, FVector::ForwardVector);
 
+#if UE_LEARNING_AGENTS_ENABLE_VISUAL_LOG
+	if (bVisualLoggerEnabled)
+	{
+		const ULearningAgentsObservationVisualLoggerObject* VisualLoggerObject = GetOrAddVisualLoggerObject(Name);
+
+		UE_LEARNING_AGENTS_VLOG_ARROW(VisualLoggerObject, LogLearning, Display,
+			RelativeTransform.GetTranslation(),
+			RelativeTransform.GetTranslation() + VisualLoggerArrowLength * Direction,
+			VisualLoggerColor.ToFColor(true),
+			TEXT(""));
+
+		UE_LEARNING_AGENTS_VLOG_TRANSFORM(VisualLoggerObject, LogLearning, Display,
+			RelativeTransform.GetTranslation(),
+			RelativeTransform.GetRotation(),
+			VisualLoggerColor.ToFColor(true),
+			TEXT(""));
+
+		UE_LEARNING_AGENTS_VLOG_STRING(VisualLoggerObject, LogLearning, Display, VisualLoggerLocation,
+			VisualLoggerColor.ToFColor(true),
+			TEXT("Name: %s\nDirection: [% 6.1f % 6.1f % 6.1f]\nLocal Direction: [% 6.1f % 6.1f % 6.1f]\nEncoded: [% 6.2f % 6.2f % 6.2f]"),
+			*Name.ToString(),
+			Direction.X, Direction.Y, Direction.Z,
+			LocalDirection.X, LocalDirection.Y, LocalDirection.Z,
+			LocalDirection.X, LocalDirection.Y, LocalDirection.Z); // We don't have a scaling so the encoded value is the same
+	}
+#endif
+
 	return MakeContinuousObservationFromArrayView({
 		(float)LocalDirection.X,
 		(float)LocalDirection.Y,
@@ -1778,6 +1948,7 @@ FLearningAgentsObservationObjectElement ULearningAgentsObservationObject::MakeDi
 	const FName Name,
 	const bool bVisualLoggerEnabled,
 	const FVector VisualLoggerLocation,
+	const float VisualLoggerArrowLength,
 	const FLinearColor VisualLoggerColor)
 {
 	if (!SplineComponent)
@@ -1786,7 +1957,7 @@ FLearningAgentsObservationObjectElement ULearningAgentsObservationObject::MakeDi
 		return FLearningAgentsObservationObjectElement();
 	}
 
-	return MakeDirectionObservation(SplineComponent->GetDirectionAtDistanceAlongSpline(DistanceAlongSpline, ESplineCoordinateSpace::World), RelativeTransform, Name, bVisualLoggerEnabled, VisualLoggerLocation, VisualLoggerColor);
+	return MakeDirectionObservation(SplineComponent->GetDirectionAtDistanceAlongSpline(DistanceAlongSpline, ESplineCoordinateSpace::World), RelativeTransform, Name, bVisualLoggerEnabled, VisualLoggerLocation, VisualLoggerArrowLength, VisualLoggerColor);
 }
 
 FLearningAgentsObservationObjectElement ULearningAgentsObservationObject::MakePropertiesAlongSplineObservation(const USplineComponent* SplineComponent, const float DistanceAlongSpline, const FTransform RelativeTransform, const float LocationScale, const FName Name)
diff --git a/Engine/Plugins/Experimental/LearningAgents/Source/LearningAgents/Public/LearningAgentsManager.h b/Engine/Plugins/Experimental/LearningAgents/Source/LearningAgents/Public/LearningAgentsManager.h
--- a/Engine/Plugins/Experimental/LearningAgents/Source/LearningAgents/Public/LearningAgentsManager.h
+++ b/Engine/Plugins/Experimental/LearningAgents/Source/LearningAgents/Public/LearningAgentsManager.h
@@ -205,12 +205,14 @@ public:
 	/** Get a const array view of this manager's agent objects. */
 	TConstArrayView<TObjectPtr<UObject>> GetAgents() const;
 
-private:
+protected:
 
 	/** Maximum number of agents. Used to preallocate internal buffers. Setting this higher will allow more agents but use up more memory. */
 	UPROPERTY(EditDefaultsOnly, Category = "LearningAgents", meta = (ClampMin = "1", UIMin = "1"))
 	int32 MaxAgentNum = 1;
 
+private:
+
 	/** The list of current agents. */
 	UPROPERTY(VisibleAnywhere, Transient, Category = "LearningAgents")
 	TArray<TObjectPtr<UObject>> Agents;
diff --git a/Engine/Plugins/Experimental/LearningAgents/Source/LearningAgents/Public/LearningAgentsObservations.h b/Engine/Plugins/Experimental/LearningAgents/Source/LearningAgents/Public/LearningAgentsObservations.h
--- a/Engine/Plugins/Experimental/LearningAgents/Source/LearningAgents/Public/LearningAgentsObservations.h
+++ b/Engine/Plugins/Experimental/LearningAgents/Source/LearningAgents/Public/LearningAgentsObservations.h
@@ -179,12 +179,31 @@ public:
 	UFUNCTION(BlueprintPure, Category = "LearningAgents")
 	FLearningAgentsObservationSchemaElement SpecifyEncodingObservation(const FLearningAgentsObservationSchemaElement Element, const int32 EncodingSize = 128, const FName Name = TEXT("Encoding"));
 
+	/**
+	 * Specifies a new bool observation. A true or false observation.
+	 *
+	 * @param Name The name of this new observation. Used during observation object validation and debugging.
+	 * @return The newly created observation schema element.
+	 */
 	UFUNCTION(BlueprintPure, Category = "LearningAgents")
 	FLearningAgentsObservationSchemaElement SpecifyBoolObservation(const FName Name = TEXT("Bool"));
 
+	/**
+	 * Specifies a new float observation. A simple observation which can be used as a catch-all for situations where a 
+	 * type-specific observation does not exist.
+	 * 
+	 * @param Name The name of this new observation. Used during observation object validation and debugging.
+	 * @return The newly created observation schema element.
+	 */
 	UFUNCTION(BlueprintPure, Category = "LearningAgents")
 	FLearningAgentsObservationSchemaElement SpecifyFloatObservation(const FName Name = TEXT("Float"));
 
+	/**
+	 * Specifies a new location observation. Allows an agent to observe the location of some entity.
+	 *
+	 * @param Name The name of this new observation. Used during observation object validation and debugging.
+	 * @return The newly created observation schema element.
+	 */
 	UFUNCTION(BlueprintPure, Category = "LearningAgents")
 	FLearningAgentsObservationSchemaElement SpecifyLocationObservation(const FName Name = TEXT("Location"));
 
@@ -355,12 +374,56 @@ public:
 	UFUNCTION(BlueprintPure, Category = "LearningAgents")
 	FLearningAgentsObservationObjectElement MakeEncodingObservation(const FLearningAgentsObservationObjectElement Element, const FName Name = TEXT("Encoding"));
 
-	UFUNCTION(BlueprintPure, Category = "LearningAgents")
-	FLearningAgentsObservationObjectElement MakeBoolObservation(const bool bValue, const FName Name = TEXT("Bool"));
+	/**
+	 * Make a new bool observation.
+	 *
+	 * @param Value The new value of this observation.
+	 * @param Name The name of the corresponding observation. Must match the name given during Specify.
+	 * @param bVisualLoggerEnabled When true, debug data will be sent to the visual logger.
+	 * @param VisualLoggerLocation A location for the visual logger information in the world.
+	 * @param VisualLoggerColor The color for the visual logger display.
+	 * @return The newly created observation object element.
+	 */
+	UFUNCTION(BlueprintPure, Category = "LearningAgents", meta = (AdvancedDisplay = 2))
+	FLearningAgentsObservationObjectElement MakeBoolObservation(
+		const bool bValue,
+		const FName Name = TEXT("Bool"),
+		const bool bVisualLoggerEnabled = false,
+		const FVector VisualLoggerLocation = FVector::ZeroVector,
+		const FLinearColor VisualLoggerColor = FLinearColor::Red);
 
-	UFUNCTION(BlueprintPure, Category = "LearningAgents")
-	FLearningAgentsObservationObjectElement MakeFloatObservation(const float Value, const float FloatScale = 1.0f, const FName Name = TEXT("Float"));
+	/**
+	 * Make a new float observation.
+	 *
+	 * @param Value The new value of this observation.
+	 * @param FloatScale Used to normalize the data for this observation.
+	 * @param Name The name of the corresponding observation. Must match the name given during Specify.
+	 * @param bVisualLoggerEnabled When true, debug data will be sent to the visual logger.
+	 * @param VisualLoggerLocation A location for the visual logger information in the world.
+	 * @param VisualLoggerColor The color for the visual logger display.
+	 * @return The newly created observation object element.
+	 */
+	UFUNCTION(BlueprintPure, Category = "LearningAgents", meta = (AdvancedDisplay = 3))
+	FLearningAgentsObservationObjectElement MakeFloatObservation(
+		const float Value,
+		const float FloatScale = 1.0f,
+		const FName Name = TEXT("Float"),
+		const bool bVisualLoggerEnabled = false,
+		const FVector VisualLoggerLocation = FVector::ZeroVector,
+		const FLinearColor VisualLoggerColor = FLinearColor::Red);
 
+	/**
+	 * Make a new location observation.
+	 *
+	 * @param Location The location of interest to the agent.
+	 * @param RelativeTransform The transform the provided location should be encoded relative to.
+	 * @param LocationScale Used to normalize the data for this observation.
+	 * @param Name The name of the corresponding observation. Must match the name given during Specify.
+	 * @param bVisualLoggerEnabled When true, debug data will be sent to the visual logger.
+	 * @param VisualLoggerLocation A location for the visual logger information in the world.
+	 * @param VisualLoggerColor The color for the visual logger display.
+	 * @return The newly created observation object element.
+	 */
 	UFUNCTION(BlueprintPure, Category = "LearningAgents", meta = (AdvancedDisplay = 4))
 	FLearningAgentsObservationObjectElement MakeLocationObservation(
 		const FVector Location,
@@ -371,15 +434,87 @@ public:
 		const FVector VisualLoggerLocation = FVector::ZeroVector,
 		const FLinearColor VisualLoggerColor = FLinearColor::Red);
 
-	UFUNCTION(BlueprintPure, Category = "LearningAgents")
-	FLearningAgentsObservationObjectElement MakeRotationObservation(const FRotator Rotation, const FRotator RelativeRotation = FRotator::ZeroRotator, const FName Name = TEXT("Rotation"));
-	FLearningAgentsObservationObjectElement MakeRotationObservationFromQuat(const FQuat Rotation, const FQuat RelativeRotation = FQuat::Identity, const FName Name = TEXT("Rotation"));
+	/**
+	 * Make a new rotation observation.
+	 *
+	 * @param Rotation The rotation of interest to the agent.
+	 * @param RelativeRotation The rotation the provided rotation should be encoded relative to.
+	 * @param Name The name of the corresponding observation. Must match the name given during Specify.
+	 * @param bVisualLoggerEnabled When true, debug data will be sent to the visual logger.
+	 * @param VisualLoggerLocation A location for the visual logger information in the world.
+	 * @param VisualLoggerColor The color for the visual logger display.
+	 * @return The newly created observation object element.
+	 */
+	UFUNCTION(BlueprintPure, Category = "LearningAgents", meta = (AdvancedDisplay = 3))
+	FLearningAgentsObservationObjectElement MakeRotationObservation(
+		const FRotator Rotation,
+		const FRotator RelativeRotation = FRotator::ZeroRotator,
+		const FName Name = TEXT("Rotation"),
+		const bool bVisualLoggerEnabled = false,
+		const FVector VisualLoggerLocation = FVector::ZeroVector,
+		const FLinearColor VisualLoggerColor = FLinearColor::Red);
 
-	UFUNCTION(BlueprintPure, Category = "LearningAgents")
-	FLearningAgentsObservationObjectElement MakeScaleObservation(const FVector Scale, const FVector RelativeScale = FVector(1,1,1), const FName Name = TEXT("Scale"));
+	/**
+	 * Make a new rotation observation from a quaternion.
+	 *
+	 * @param Rotation The rotation of interest to the agent.
+	 * @param RelativeRotation The rotation the provided rotation should be encoded relative to.
+	 * @param Name The name of the corresponding observation. Must match the name given during Specify.
+	 * @param bVisualLoggerEnabled When true, debug data will be sent to the visual logger.
+	 * @param VisualLoggerLocation A location for the visual logger information in the world.
+	 * @param VisualLoggerColor The color for the visual logger display.
+	 * @return The newly created observation object element.
+	 */
+	UFUNCTION(BlueprintPure, Category = "LearningAgents", meta = (AdvancedDisplay = 3))
+	FLearningAgentsObservationObjectElement MakeRotationObservationFromQuat(
+		const FQuat Rotation,
+		const FQuat RelativeRotation,
+		const FName Name = TEXT("Rotation"),
+		const bool bVisualLoggerEnabled = false,
+		const FVector VisualLoggerLocation = FVector::ZeroVector,
+		const FLinearColor VisualLoggerColor = FLinearColor::Red);
 
-	UFUNCTION(BlueprintPure, Category = "LearningAgents")
-	FLearningAgentsObservationObjectElement MakeTransformObservation(const FTransform Transform, const FTransform RelativeTransform = FTransform(), const float LocationScale = 100.0f, const FName Name = TEXT("Transform"));
+	/**
+	 * Make a new scale observation. Negative scales are not supported by this observation type.
+	 *
+	 * @param Scale The scale of interest to the agent.
+	 * @param RelativeScale The scale the provided scale should be encoded relative to.
+	 * @param Name The name of the corresponding observation. Must match the name given during Specify.
+	 * @param bVisualLoggerEnabled When true, debug data will be sent to the visual logger.
+	 * @param VisualLoggerLocation A location for the visual logger information in the world.
+	 * @param VisualLoggerColor The color for the visual logger display.
+	 * @return The newly created observation object element.
+	 */
+	UFUNCTION(BlueprintPure, Category = "LearningAgents", meta = (AdvancedDisplay = 3))
+	FLearningAgentsObservationObjectElement MakeScaleObservation(
+		const FVector Scale,
+		const FVector RelativeScale = FVector(1,1,1),
+		const FName Name = TEXT("Scale"),
+		const bool bVisualLoggerEnabled = false,
+		const FVector VisualLoggerLocation = FVector::ZeroVector,
+		const FLinearColor VisualLoggerColor = FLinearColor::Red);
+
+	/**
+	 * Make a new transform observation.
+	 *
+	 * @param Transform The transform of interest to the agent.
+	 * @param RelativeTransform The transform the provided transform should be encoded relative to.
+	 * @param LocationScale Used to normalize the transform's location for this observation.
+	 * @param Name The name of the corresponding observation. Must match the name given during Specify.
+	 * @param bVisualLoggerEnabled When true, debug data will be sent to the visual logger.
+	 * @param VisualLoggerLocation A location for the visual logger information in the world.
+	 * @param VisualLoggerColor The color for the visual logger display.
+	 * @return The newly created observation object element.
+	 */
+	UFUNCTION(BlueprintPure, Category = "LearningAgents", meta = (AdvancedDisplay = 4))
+	FLearningAgentsObservationObjectElement MakeTransformObservation(
+		const FTransform Transform,
+		const FTransform RelativeTransform = FTransform(),
+		const float LocationScale = 100.0f,
+		const FName Name = TEXT("Transform"),
+		const bool bVisualLoggerEnabled = false,
+		const FVector VisualLoggerLocation = FVector::ZeroVector,
+		const FLinearColor VisualLoggerColor = FLinearColor::Red);
 
 	UFUNCTION(BlueprintPure, Category = "LearningAgents")
 	FLearningAgentsObservationObjectElement MakeAngleObservation(const float Angle, const float RelativeAngle = 0.0f, const FName Name = TEXT("Angle"));
@@ -397,6 +532,7 @@ public:
 		const FName Name = TEXT("Direction"),
 		const bool bVisualLoggerEnabled = false,
 		const FVector VisualLoggerLocation = FVector::ZeroVector,
+		const float VisualLoggerArrowLength = 100.0f,
 		const FLinearColor VisualLoggerColor = FLinearColor::Red);
 
 	// Spline Observations
@@ -423,6 +559,7 @@ public:
 		const FName Name = TEXT("DirectionAlongSpline"),
 		const bool bVisualLoggerEnabled = false,
 		const FVector VisualLoggerLocation = FVector::ZeroVector,
+		const float VisualLoggerArrowLength = 100.0f,
 		const FLinearColor VisualLoggerColor = FLinearColor::Red);
 
 	UFUNCTION(BlueprintPure, Category = "LearningAgents")
