Learning Agents: Expose Grad Norm Clipping as a training setting

#rnx

diff --git a/Engine/Plugins/Experimental/LearningAgents/Source/LearningAgentsTraining/Private/LearningAgentsTrainer.cpp b/Engine/Plugins/Experimental/LearningAgents/Source/LearningAgentsTraining/Private/LearningAgentsTrainer.cpp
--- a/Engine/Plugins/Experimental/LearningAgents/Source/LearningAgentsTraining/Private/LearningAgentsTrainer.cpp
+++ b/Engine/Plugins/Experimental/LearningAgents/Source/LearningAgentsTraining/Private/LearningAgentsTrainer.cpp
@@ -480,6 +480,8 @@ void ULearningAgentsTrainer::BeginTraining(
 	PPOTrainingSettings.bAdvantageNormalization = TrainerTrainingSettings.bAdvantageNormalization;
 	PPOTrainingSettings.AdvantageMin = TrainerTrainingSettings.MinimumAdvantage;
 	PPOTrainingSettings.AdvantageMax = TrainerTrainingSettings.MaximumAdvantage;
+	PPOTrainingSettings.bUseGradNormMaxClipping = TrainerTrainingSettings.bUseGradNormMaxClipping;
+	PPOTrainingSettings.GradNormMax = TrainerTrainingSettings.GradNormMax;
 	PPOTrainingSettings.TrimEpisodeStartStepNum = TrainerTrainingSettings.NumberOfStepsToTrimAtStartOfEpisode;
 	PPOTrainingSettings.TrimEpisodeEndStepNum = TrainerTrainingSettings.NumberOfStepsToTrimAtEndOfEpisode;
 	PPOTrainingSettings.Seed = TrainerTrainingSettings.RandomSeed;
diff --git a/Engine/Plugins/Experimental/LearningAgents/Source/LearningAgentsTraining/Public/LearningAgentsTrainer.h b/Engine/Plugins/Experimental/LearningAgents/Source/LearningAgentsTraining/Public/LearningAgentsTrainer.h
--- a/Engine/Plugins/Experimental/LearningAgents/Source/LearningAgentsTraining/Public/LearningAgentsTrainer.h
+++ b/Engine/Plugins/Experimental/LearningAgents/Source/LearningAgentsTraining/Public/LearningAgentsTrainer.h
@@ -238,6 +238,22 @@ public:
 	UPROPERTY(EditAnywhere, Category = "LearningAgents", meta = (UIMin = "0.0", UIMax = "10.0"))
 	float MaximumAdvantage = 10.0f;
 
+	/**
+	 * When true, gradient norm max clipping will be used on the policy, critic, encoder, and decoder. Set this as True if
+	 * training is unstable (and adjust GradNormMax) or leave as False if unused.
+	 */
+	UPROPERTY(EditAnywhere, Category = "LearningAgents")
+	bool bUseGradNormMaxClipping = false;
+
+	/**
+	 * The maximum gradient norm to clip updates to. Only used when bUseGradNormMaxClipping is set to true.
+	 * 
+	 * This needs to be carefully chosen based on the size of your gradients during training. Setting too low can make it
+	 * difficult to learn an optimal policy, and too high will have no impact.
+	 */
+	UPROPERTY(EditAnywhere, Category = "LearningAgents", meta = (UIMin = "0.0", UIMax = "10.0"))
+	float GradNormMax = 0.5f;
+
 	/**
 	 * The number of steps to trim from the start of the episode, e.g. can be useful if some things are still getting
 	 * setup at the start of the episode and you don't want them used for training.
diff --git a/Engine/Plugins/Experimental/LearningAgents/Source/LearningTraining/Private/LearningPPOTrainer.cpp b/Engine/Plugins/Experimental/LearningAgents/Source/LearningTraining/Private/LearningPPOTrainer.cpp
--- a/Engine/Plugins/Experimental/LearningAgents/Source/LearningTraining/Private/LearningPPOTrainer.cpp
+++ b/Engine/Plugins/Experimental/LearningAgents/Source/LearningTraining/Private/LearningPPOTrainer.cpp
@@ -285,6 +285,8 @@ namespace UE::Learning
 			ConfigObject->SetBoolField(TEXT("AdvantageNormalization"), TrainingSettings.bAdvantageNormalization);
 			ConfigObject->SetNumberField(TEXT("AdvantageMin"), TrainingSettings.AdvantageMin);
 			ConfigObject->SetNumberField(TEXT("AdvantageMax"), TrainingSettings.AdvantageMax);
+			ConfigObject->SetBoolField(TEXT("UseGradNormMaxClipping"), TrainingSettings.bUseGradNormMaxClipping);
+			ConfigObject->SetNumberField(TEXT("GradNormMax"), TrainingSettings.GradNormMax);
 			ConfigObject->SetNumberField(TEXT("TrimEpisodeStartStepNum"), TrainingSettings.TrimEpisodeStartStepNum);
 			ConfigObject->SetNumberField(TEXT("TrimEpisodeEndStepNum"), TrainingSettings.TrimEpisodeEndStepNum);
 			ConfigObject->SetNumberField(TEXT("Seed"), TrainingSettings.Seed);
@@ -794,6 +796,8 @@ namespace UE::Learning
 		ConfigObject->SetBoolField(TEXT("AdvantageNormalization"), TrainingSettings.bAdvantageNormalization);
 		ConfigObject->SetNumberField(TEXT("AdvantageMin"), TrainingSettings.AdvantageMin);
 		ConfigObject->SetNumberField(TEXT("AdvantageMax"), TrainingSettings.AdvantageMax);
+		ConfigObject->SetBoolField(TEXT("UseGradNormMaxClipping"), TrainingSettings.bUseGradNormMaxClipping);
+		ConfigObject->SetNumberField(TEXT("GradNormMax"), TrainingSettings.GradNormMax);
 		ConfigObject->SetNumberField(TEXT("TrimEpisodeStartStepNum"), TrainingSettings.TrimEpisodeStartStepNum);
 		ConfigObject->SetNumberField(TEXT("TrimEpisodeEndStepNum"), TrainingSettings.TrimEpisodeEndStepNum);
 		ConfigObject->SetNumberField(TEXT("Seed"), TrainingSettings.Seed);
diff --git a/Engine/Plugins/Experimental/LearningAgents/Source/LearningTraining/Public/LearningPPOTrainer.h b/Engine/Plugins/Experimental/LearningAgents/Source/LearningTraining/Public/LearningPPOTrainer.h
--- a/Engine/Plugins/Experimental/LearningAgents/Source/LearningTraining/Public/LearningPPOTrainer.h
+++ b/Engine/Plugins/Experimental/LearningAgents/Source/LearningTraining/Public/LearningPPOTrainer.h
@@ -111,6 +111,12 @@ namespace UE::Learning
 		// at the cost of some training speed.
 		float AdvantageMax = 10.0f;
 
+		// If true, uses gradient norm max clipping. Set this as True if training is unstable or leave as False if unused.
+		bool bUseGradNormMaxClipping = false;
+
+		// The maximum gradient norm to clip updates to.
+		float GradNormMax = 0.5f;
+
 		// Number of steps to trim from the start of each episode during training. This can
 		// be useful if some reset process is taking several steps or you know your starting
 		// states are not entirely valid for example.
