LearningAgents: Added LayerNum and Activation to Encoding Observations and Actions

diff --git a/Engine/Plugins/Experimental/LearningAgents/Source/Learning/Private/LearningAction.cpp b/Engine/Plugins/Experimental/LearningAgents/Source/Learning/Private/LearningAction.cpp
--- a/Engine/Plugins/Experimental/LearningAgents/Source/Learning/Private/LearningAction.cpp
+++ b/Engine/Plugins/Experimental/LearningAgents/Source/Learning/Private/LearningAction.cpp
@@ -291,6 +291,8 @@ namespace UE::Learning::Action
 
 		FEncodingData ElementData;
 		ElementData.EncodingSize = Parameters.EncodingSize;
+		ElementData.LayerNum = Parameters.LayerNum;
+		ElementData.ActivationFunction = Parameters.ActivationFunction;
 		ElementData.ElementIndex = SubElementObjects.Num();
 
 		SubElementNames.Add(NAME_None);
@@ -426,6 +428,8 @@ namespace UE::Learning::Action
 		FSchemaEncodingParameters Parameters;
 		Parameters.Element = SubElementObjects[ElementData.ElementIndex];
 		Parameters.EncodingSize = ElementData.EncodingSize;
+		Parameters.LayerNum = ElementData.LayerNum;
+		Parameters.ActivationFunction = ElementData.ActivationFunction;
 		return Parameters;
 	}
 
@@ -812,6 +816,17 @@ namespace UE::Learning::Action
 
 	namespace Private
 	{
+		static inline NNE::RuntimeBasic::FModelBuilder::EActivationFunction GetNNEActivationFunction(const EEncodingActivationFunction ActivationFunction)
+		{
+			switch (ActivationFunction)
+			{
+			case EEncodingActivationFunction::ReLU: return NNE::RuntimeBasic::FModelBuilder::EActivationFunction::ReLU;
+			case EEncodingActivationFunction::ELU: return NNE::RuntimeBasic::FModelBuilder::EActivationFunction::ELU;
+			case EEncodingActivationFunction::TanH: return NNE::RuntimeBasic::FModelBuilder::EActivationFunction::TanH;
+			default: UE_LEARNING_NOT_IMPLEMENTED(); return NNE::RuntimeBasic::FModelBuilder::EActivationFunction::ReLU;
+			}
+		}
+
 		NNE::RuntimeBasic::FModelBuilderElement MakeDecoderNetworkFromSchema(
 			NNE::RuntimeBasic::FModelBuilder& Builder,
 			const FSchema& Schema,
@@ -963,8 +978,14 @@ namespace UE::Learning::Action
 				const int32 SubElementEncodedSize = Schema.GetEncodedVectorSize(Parameters.Element);
 
 				ReturnElement = Builder.MakeSequence({
-					Builder.MakeELU(Parameters.EncodingSize),
-					Builder.MakeLinearWithRandomKaimingWeights(Parameters.EncodingSize, SubElementEncodedSize),
+					Builder.MakeActivation(Parameters.EncodingSize, GetNNEActivationFunction(Parameters.ActivationFunction)),
+					Builder.MakeMLPWithRandomKaimingWeights(
+						Parameters.EncodingSize, 
+						SubElementEncodedSize, 
+						Parameters.EncodingSize, 
+						Parameters.LayerNum + 1,  // Add 1 to account for input layer 
+						GetNNEActivationFunction(Parameters.ActivationFunction),
+						false),
 					MakeDecoderNetworkFromSchema(Builder, Schema, Parameters.Element),
 				});
 				break;
diff --git a/Engine/Plugins/Experimental/LearningAgents/Source/Learning/Private/LearningObservation.cpp b/Engine/Plugins/Experimental/LearningAgents/Source/Learning/Private/LearningObservation.cpp
--- a/Engine/Plugins/Experimental/LearningAgents/Source/Learning/Private/LearningObservation.cpp
+++ b/Engine/Plugins/Experimental/LearningAgents/Source/Learning/Private/LearningObservation.cpp
@@ -212,6 +212,8 @@ namespace UE::Learning::Observation
 		FEncodingData ElementData;
 		ElementData.ElementIndex = SubElementObjects.Num();
 		ElementData.EncodingSize = Parameters.EncodingSize;
+		ElementData.LayerNum = Parameters.LayerNum;
+		ElementData.ActivationFunction = Parameters.ActivationFunction;
 
 		SubElementNames.Add(NAME_None);
 		SubElementObjects.Add(Parameters.Element);
@@ -333,6 +335,8 @@ namespace UE::Learning::Observation
 		FSchemaEncodingParameters Parameters;
 		Parameters.Element = SubElementObjects[ElementData.ElementIndex];
 		Parameters.EncodingSize = ElementData.EncodingSize;
+		Parameters.LayerNum = ElementData.LayerNum;
+		Parameters.ActivationFunction = ElementData.ActivationFunction;
 		return Parameters;
 	}
 
@@ -665,6 +669,17 @@ namespace UE::Learning::Observation
 
 	namespace Private
 	{
+		static inline NNE::RuntimeBasic::FModelBuilder::EActivationFunction GetNNEActivationFunction(const EEncodingActivationFunction ActivationFunction)
+		{
+			switch (ActivationFunction)
+			{
+			case EEncodingActivationFunction::ReLU: return NNE::RuntimeBasic::FModelBuilder::EActivationFunction::ReLU;
+			case EEncodingActivationFunction::ELU: return NNE::RuntimeBasic::FModelBuilder::EActivationFunction::ELU;
+			case EEncodingActivationFunction::TanH: return NNE::RuntimeBasic::FModelBuilder::EActivationFunction::TanH;
+			default: UE_LEARNING_NOT_IMPLEMENTED(); return NNE::RuntimeBasic::FModelBuilder::EActivationFunction::ReLU;
+			}
+		}
+
 		NNE::RuntimeBasic::FModelBuilderElement MakeEncoderNetworkFromSchema(
 			NNE::RuntimeBasic::FModelBuilder& Builder,
 			const FSchema& Schema,
@@ -793,8 +808,13 @@ namespace UE::Learning::Observation
 
 				ReturnElement = Builder.MakeSequence({
 					MakeEncoderNetworkFromSchema(Builder, Schema, Parameters.Element),
-					Builder.MakeLinearWithRandomKaimingWeights(SubElementEncodedSize, Parameters.EncodingSize),
-					Builder.MakeELU(Parameters.EncodingSize),
+					Builder.MakeMLPWithRandomKaimingWeights(
+						SubElementEncodedSize,
+						Parameters.EncodingSize,
+						Parameters.EncodingSize,
+						Parameters.LayerNum + 1, // Add 1 to account for input layer
+						GetNNEActivationFunction(Parameters.ActivationFunction),
+						true)
 					});
 				break;
 			}
diff --git a/Engine/Plugins/Experimental/LearningAgents/Source/Learning/Public/LearningAction.h b/Engine/Plugins/Experimental/LearningAgents/Source/Learning/Public/LearningAction.h
--- a/Engine/Plugins/Experimental/LearningAgents/Source/Learning/Public/LearningAction.h
+++ b/Engine/Plugins/Experimental/LearningAgents/Source/Learning/Public/LearningAction.h
@@ -117,6 +117,14 @@ namespace UE::Learning::Action
 		int32 Num = 0;
 	};
 
+	/** Activation Function to use for encoding */
+	enum class EEncodingActivationFunction : uint8
+	{
+		ELU = 0,
+		ReLU = 1,
+		TanH = 2,
+	};
+
 	struct FSchemaEncodingParameters
 	{
 		// The sub-element.
@@ -124,6 +132,12 @@ namespace UE::Learning::Action
 
 		// The size at which the sub-element should be encoded.
 		int32 EncodingSize = 32;
+
+		// The number of layers in the encoding
+		int32 LayerNum = 1;
+
+		// The activation function to use for encoding
+		EEncodingActivationFunction ActivationFunction = EEncodingActivationFunction::ELU;
 	};
 
 	/**
@@ -230,6 +244,8 @@ namespace UE::Learning::Action
 		{
 			int32 EncodingSize = INDEX_NONE;
 			int32 ElementIndex = INDEX_NONE;
+			int32 LayerNum = INDEX_NONE;
+			EEncodingActivationFunction ActivationFunction = EEncodingActivationFunction::ELU;
 		};
 
 		uint32 Generation = 0;
diff --git a/Engine/Plugins/Experimental/LearningAgents/Source/Learning/Public/LearningObservation.h b/Engine/Plugins/Experimental/LearningAgents/Source/Learning/Public/LearningObservation.h
--- a/Engine/Plugins/Experimental/LearningAgents/Source/Learning/Public/LearningObservation.h
+++ b/Engine/Plugins/Experimental/LearningAgents/Source/Learning/Public/LearningObservation.h
@@ -120,6 +120,14 @@ namespace UE::Learning::Observation
 		int32 ValueEncodingSize = 32;
 	};
 
+	/** Activation Function to use for encoding */
+	enum class EEncodingActivationFunction : uint8
+	{
+		ELU = 0,
+		ReLU = 1,
+		TanH = 2,
+	};
+
 	struct FSchemaEncodingParameters
 	{
 		// The sub-element.
@@ -127,6 +135,12 @@ namespace UE::Learning::Observation
 
 		// The size at which the sub-element should be encoded.
 		int32 EncodingSize = 32;
+
+		// The number of layers in the encoding
+		int32 LayerNum = 1;
+
+		// The activation function to use for encoding
+		EEncodingActivationFunction ActivationFunction = EEncodingActivationFunction::ELU;
 	};
 
 	/**
@@ -228,6 +242,8 @@ namespace UE::Learning::Observation
 		{
 			int32 ElementIndex = INDEX_NONE;
 			int32 EncodingSize = INDEX_NONE;
+			int32 LayerNum = INDEX_NONE;
+			EEncodingActivationFunction ActivationFunction = EEncodingActivationFunction::ELU;
 		};
 
 		uint32 Generation = 0;
diff --git a/Engine/Plugins/Experimental/LearningAgents/Source/LearningAgents/Private/LearningAgentsActions.cpp b/Engine/Plugins/Experimental/LearningAgents/Source/LearningAgents/Private/LearningAgentsActions.cpp
--- a/Engine/Plugins/Experimental/LearningAgents/Source/LearningAgents/Private/LearningAgentsActions.cpp
+++ b/Engine/Plugins/Experimental/LearningAgents/Source/LearningAgents/Private/LearningAgentsActions.cpp
@@ -509,6 +509,17 @@ namespace UE::Learning::Agents::Action::Private
 			PriorProbabilities[Idx] = FMath::Clamp(PriorProbabilities[Idx], 0.0f, 1.0f);
 		}
 	}
+
+	static inline Learning::Action::EEncodingActivationFunction GetEncodingActivationFunction(const ELearningAgentsActivationFunction ActivationFunction)
+	{
+		switch (ActivationFunction)
+		{
+		case ELearningAgentsActivationFunction::ReLU: return Learning::Action::EEncodingActivationFunction::ReLU;
+		case ELearningAgentsActivationFunction::ELU: return Learning::Action::EEncodingActivationFunction::ELU;
+		case ELearningAgentsActivationFunction::TanH: return Learning::Action::EEncodingActivationFunction::TanH;
+		default: UE_LEARNING_NOT_IMPLEMENTED(); return Learning::Action::EEncodingActivationFunction::ReLU;
+		}
+	}
 }
 
 bool ULearningAgentsActionSchema::ValidateObjectMatchesSchema(
@@ -1016,11 +1027,17 @@ FLearningAgentsActionSchemaElement ULearningAgentsActionSchema::SpecifyEitherAct
 	return SpecifyExclusiveUnionActionFromArrayViews({ TEXT("A"), TEXT("B") }, { A, B }, { 1.0f - PriorProbabilityOfA, PriorProbabilityOfA }, Name);
 }
 
-FLearningAgentsActionSchemaElement ULearningAgentsActionSchema::SpecifyEncodingAction(const FLearningAgentsActionSchemaElement Element, const int32 EncodingSize, const FName Name)
+FLearningAgentsActionSchemaElement ULearningAgentsActionSchema::SpecifyEncodingAction(const FLearningAgentsActionSchemaElement Element, const int32 EncodingSize, const int32 LayerNum, const ELearningAgentsActivationFunction ActivationFunction, const FName Name)
 {
 	if (EncodingSize < 1)
 	{
-		UE_LOG(LogLearning, Error, TEXT("%s: Invalid Action EncodingSize '%i'."), *GetName(), EncodingSize);
+		UE_LOG(LogLearning, Error, TEXT("%s: Invalid Action EncodingSize '%i' - must be greater than zero."), *GetName(), EncodingSize);
+		return FLearningAgentsActionSchemaElement();
+	}
+
+	if (LayerNum < 1)
+	{
+		UE_LOG(LogLearning, Error, TEXT("%s: Invalid Action LayerNum '%i' - must be greater than zero."), *GetName(), LayerNum);
 		return FLearningAgentsActionSchemaElement();
 	}
 
@@ -1030,7 +1047,7 @@ FLearningAgentsActionSchemaElement ULearningAgentsActionSchema::SpecifyEncodingA
 		return FLearningAgentsActionSchemaElement();
 	}
 
-	return { ActionSchema.CreateEncoding({ Element.SchemaElement, EncodingSize }, Name) };
+	return { ActionSchema.CreateEncoding({ Element.SchemaElement, EncodingSize, LayerNum, UE::Learning::Agents::Action::Private::GetEncodingActivationFunction(ActivationFunction) }, Name)};
 }
 
 FLearningAgentsActionSchemaElement ULearningAgentsActionSchema::SpecifyBoolAction(const float PriorProbability, const FName Name)
diff --git a/Engine/Plugins/Experimental/LearningAgents/Source/LearningAgents/Private/LearningAgentsObservations.cpp b/Engine/Plugins/Experimental/LearningAgents/Source/LearningAgents/Private/LearningAgentsObservations.cpp
--- a/Engine/Plugins/Experimental/LearningAgents/Source/LearningAgents/Private/LearningAgentsObservations.cpp
+++ b/Engine/Plugins/Experimental/LearningAgents/Source/LearningAgents/Private/LearningAgentsObservations.cpp
@@ -503,6 +503,17 @@ namespace UE::Learning::Agents::Observation::Private
 			FMath::Exp(V.Y),
 			FMath::Exp(V.Z));
 	}
+
+	static inline Learning::Observation::EEncodingActivationFunction GetEncodingActivationFunction(const ELearningAgentsActivationFunction ActivationFunction)
+	{
+		switch (ActivationFunction)
+		{
+		case ELearningAgentsActivationFunction::ReLU: return Learning::Observation::EEncodingActivationFunction::ReLU;
+		case ELearningAgentsActivationFunction::ELU: return Learning::Observation::EEncodingActivationFunction::ELU;
+		case ELearningAgentsActivationFunction::TanH: return Learning::Observation::EEncodingActivationFunction::TanH;
+		default: UE_LEARNING_NOT_IMPLEMENTED(); return Learning::Observation::EEncodingActivationFunction::ReLU;
+		}
+	}
 }
 
 FTransform ULearningAgentsObservationFunctions::ProjectTransformOntoGroundPlane(const FTransform Transform, const FVector LocalForwardVector, const float GroundPlaneHeight)
@@ -664,7 +675,7 @@ FLearningAgentsObservationSchemaElement ULearningAgentsObservationSchema::Specif
 {
 	if (EncodingSize < 1)
 	{
-		UE_LOG(LogLearning, Error, TEXT("%s: Invalid Observation EncodingSize '%i'."), *GetName(), EncodingSize);
+		UE_LOG(LogLearning, Error, TEXT("%s: Invalid Observation EncodingSize '%i' - must be greater than zero."), *GetName(), EncodingSize);
 		return FLearningAgentsObservationSchemaElement();
 	}
 
@@ -720,7 +731,7 @@ FLearningAgentsObservationSchemaElement ULearningAgentsObservationSchema::Specif
 {
 	if (EncodingSize < 1)
 	{
-		UE_LOG(LogLearning, Error, TEXT("%s: Invalid Observation EncodingSize '%i'."), *GetName(), EncodingSize);
+		UE_LOG(LogLearning, Error, TEXT("%s: Invalid Observation EncodingSize '%i' - must be greater than zero."), *GetName(), EncodingSize);
 		return FLearningAgentsObservationSchemaElement();
 	}
 
@@ -762,7 +773,7 @@ FLearningAgentsObservationSchemaElement ULearningAgentsObservationSchema::Specif
 {
 	if (AttentionEncodingSize < 1 || AttentionHeadNum < 1 || ValueEncodingSize < 1)
 	{
-		UE_LOG(LogLearning, Error, TEXT("%s: Invalid Observation Parameters: AttentionEncodingSize: %i, AttentionHeadNum: %i, ValueEncodingSize: %i."), *GetName(), AttentionEncodingSize, AttentionHeadNum, ValueEncodingSize);
+		UE_LOG(LogLearning, Error, TEXT("%s: Invalid Observation Parameters: AttentionEncodingSize: %i, AttentionHeadNum: %i, ValueEncodingSize: %i - must be greater than zero."), *GetName(), AttentionEncodingSize, AttentionHeadNum, ValueEncodingSize);
 		return FLearningAgentsObservationSchemaElement();
 	}
 
@@ -818,7 +829,7 @@ FLearningAgentsObservationSchemaElement ULearningAgentsObservationSchema::Specif
 {
 	if (AttentionEncodingSize < 1 || AttentionHeadNum < 1 || ValueEncodingSize < 1)
 	{
-		UE_LOG(LogLearning, Error, TEXT("%s: Invalid Observation Parameters: AttentionEncodingSize: %i, AttentionHeadNum: %i, ValueEncodingSize: %i."), *GetName(), AttentionEncodingSize, AttentionHeadNum, ValueEncodingSize);
+		UE_LOG(LogLearning, Error, TEXT("%s: Invalid Observation Parameters: AttentionEncodingSize: %i, AttentionHeadNum: %i, ValueEncodingSize: %i - must be greater than zero."), *GetName(), AttentionEncodingSize, AttentionHeadNum, ValueEncodingSize);
 		return FLearningAgentsObservationSchemaElement();
 	}
 
@@ -882,7 +893,7 @@ FLearningAgentsObservationSchemaElement ULearningAgentsObservationSchema::Specif
 {
 	if (MaxNum < 0)
 	{
-		UE_LOG(LogLearning, Error, TEXT("%s: Invalid Observation Set MaxNum %i."), *GetName(), MaxNum);
+		UE_LOG(LogLearning, Error, TEXT("%s: Invalid Observation Set MaxNum %i - must be greater than or equal to zero."), *GetName(), MaxNum);
 		return FLearningAgentsObservationSchemaElement();
 	}
 
@@ -893,7 +904,7 @@ FLearningAgentsObservationSchemaElement ULearningAgentsObservationSchema::Specif
 
 	if (AttentionEncodingSize < 1 || AttentionHeadNum < 1 || ValueEncodingSize < 1)
 	{
-		UE_LOG(LogLearning, Error, TEXT("%s: Invalid Observation Parameters: AttentionEncodingSize: %i, AttentionHeadNum: %i, ValueEncodingSize: %i."), *GetName(), AttentionEncodingSize, AttentionHeadNum, ValueEncodingSize);
+		UE_LOG(LogLearning, Error, TEXT("%s: Invalid Observation Parameters: AttentionEncodingSize: %i, AttentionHeadNum: %i, ValueEncodingSize: %i - must be greater than zero."), *GetName(), AttentionEncodingSize, AttentionHeadNum, ValueEncodingSize);
 		return FLearningAgentsObservationSchemaElement();
 	}
 
@@ -959,11 +970,17 @@ FLearningAgentsObservationSchemaElement ULearningAgentsObservationSchema::Specif
 	return SpecifyExclusiveUnionObservationFromArrayViews({ TEXT("A"), TEXT("B") }, { A, B }, EncodingSize, Name);
 }
 
-FLearningAgentsObservationSchemaElement ULearningAgentsObservationSchema::SpecifyEncodingObservation(const FLearningAgentsObservationSchemaElement Element, const int32 EncodingSize, const FName Name)
+FLearningAgentsObservationSchemaElement ULearningAgentsObservationSchema::SpecifyEncodingObservation(const FLearningAgentsObservationSchemaElement Element, const int32 EncodingSize, const int32 LayerNum, const ELearningAgentsActivationFunction ActivationFunction, const FName Name)
 {
 	if (EncodingSize < 1)
 	{
-		UE_LOG(LogLearning, Error, TEXT("%s: Invalid Observation EncodingSize '%i'."), *GetName(), EncodingSize);
+		UE_LOG(LogLearning, Error, TEXT("%s: Invalid Observation EncodingSize '%i' - must be greater than zero."), *GetName(), EncodingSize);
+		return FLearningAgentsObservationSchemaElement();
+	}
+
+	if (LayerNum < 1)
+	{
+		UE_LOG(LogLearning, Error, TEXT("%s: Invalid Observation LayerNum '%i' - must be greater than zero."), *GetName(), LayerNum);
 		return FLearningAgentsObservationSchemaElement();
 	}
 
@@ -973,7 +990,7 @@ FLearningAgentsObservationSchemaElement ULearningAgentsObservationSchema::Specif
 		return FLearningAgentsObservationSchemaElement();
 	}
 
-	return { ObservationSchema.CreateEncoding({ Element.SchemaElement, EncodingSize }, Name) };
+	return { ObservationSchema.CreateEncoding({ Element.SchemaElement, EncodingSize, LayerNum, UE::Learning::Agents::Observation::Private::GetEncodingActivationFunction(ActivationFunction) }, Name)};
 }
 
 FLearningAgentsObservationSchemaElement ULearningAgentsObservationSchema::SpecifyBoolObservation(const FName Name)
diff --git a/Engine/Plugins/Experimental/LearningAgents/Source/LearningAgents/Public/LearningAgentsActions.h b/Engine/Plugins/Experimental/LearningAgents/Source/LearningAgents/Public/LearningAgentsActions.h
--- a/Engine/Plugins/Experimental/LearningAgents/Source/LearningAgents/Public/LearningAgentsActions.h
+++ b/Engine/Plugins/Experimental/LearningAgents/Source/LearningAgents/Public/LearningAgentsActions.h
@@ -5,6 +5,8 @@
 #include "LearningArray.h"
 #include "LearningAction.h"
 
+#include "LearningAgentsNeuralNetwork.h" // Included for ELearningAgentsActivationFunction
+
 #include "LearningAgentsActions.generated.h"
 
 class ULearningAgentsActionSchema;
@@ -145,7 +147,7 @@ public:
 	FLearningAgentsActionSchemaElement SpecifyEitherAction(const FLearningAgentsActionSchemaElement A, const FLearningAgentsActionSchemaElement B, const float PriorProbabilityOfA = 0.5f, const FName Name = TEXT("Either"));
 
 	UFUNCTION(BlueprintPure, Category = "LearningAgents")
-	FLearningAgentsActionSchemaElement SpecifyEncodingAction(const FLearningAgentsActionSchemaElement Element, const int32 EncodingSize = 128, const FName Name = TEXT("Encoding"));
+	FLearningAgentsActionSchemaElement SpecifyEncodingAction(const FLearningAgentsActionSchemaElement Element, const int32 EncodingSize = 128, const int32 LayerNum = 1, const ELearningAgentsActivationFunction ActivationFunction = ELearningAgentsActivationFunction::ELU, const FName Name = TEXT("Encoding"));
 
 	UFUNCTION(BlueprintPure, Category = "LearningAgents")
 	FLearningAgentsActionSchemaElement SpecifyBoolAction(const float PriorProbability = 0.5f, const FName Name = TEXT("Bool"));
diff --git a/Engine/Plugins/Experimental/LearningAgents/Source/LearningAgents/Public/LearningAgentsObservations.h b/Engine/Plugins/Experimental/LearningAgents/Source/LearningAgents/Public/LearningAgentsObservations.h
--- a/Engine/Plugins/Experimental/LearningAgents/Source/LearningAgents/Public/LearningAgentsObservations.h
+++ b/Engine/Plugins/Experimental/LearningAgents/Source/LearningAgents/Public/LearningAgentsObservations.h
@@ -5,6 +5,8 @@
 #include "LearningArray.h"
 #include "LearningObservation.h"
 
+#include "LearningAgentsNeuralNetwork.h" // Included for ELearningAgentsActivationFunction
+
 #include "Engine/EngineTypes.h"
 #include "Kismet/BlueprintFunctionLibrary.h"
 
@@ -177,7 +179,7 @@ public:
 	FLearningAgentsObservationSchemaElement SpecifyEitherObservation(const FLearningAgentsObservationSchemaElement A, const FLearningAgentsObservationSchemaElement B, const int32 EncodingSize = 128, const FName Name = TEXT("Either"));
 
 	UFUNCTION(BlueprintPure, Category = "LearningAgents")
-	FLearningAgentsObservationSchemaElement SpecifyEncodingObservation(const FLearningAgentsObservationSchemaElement Element, const int32 EncodingSize = 128, const FName Name = TEXT("Encoding"));
+	FLearningAgentsObservationSchemaElement SpecifyEncodingObservation(const FLearningAgentsObservationSchemaElement Element, const int32 EncodingSize = 128, const int32 LayerNum = 1, const ELearningAgentsActivationFunction ActivationFunction = ELearningAgentsActivationFunction::ELU, const FName Name = TEXT("Encoding"));
 
 	/**
 	 * Specifies a new bool observation. A true or false observation.
