LearningAgents: Added Visual Logger for Rewards

#virtualized

diff --git a/Engine/Plugins/Experimental/LearningAgents/Source/LearningAgents/Private/LearningAgentsManagerListener.cpp b/Engine/Plugins/Experimental/LearningAgents/Source/LearningAgents/Private/LearningAgentsManagerListener.cpp
--- a/Engine/Plugins/Experimental/LearningAgents/Source/LearningAgents/Private/LearningAgentsManagerListener.cpp
+++ b/Engine/Plugins/Experimental/LearningAgents/Source/LearningAgents/Private/LearningAgentsManagerListener.cpp
@@ -96,3 +96,20 @@ bool ULearningAgentsManagerListener::IsSetup() const
 {
 	return bIsSetup;
 }
+
+const ULearningAgentsVisualLoggerObject* ULearningAgentsManagerListener::GetOrAddVisualLoggerObject(const FName Name)
+{
+	const TObjectPtr<const ULearningAgentsVisualLoggerObject>* Value = VisualLoggerObjects.Find(Name);
+
+	if (Value)
+	{
+		return Value->Get();
+	}
+	else
+	{
+		const FName UniqueName = MakeUniqueObjectName(this, ULearningAgentsVisualLoggerObject::StaticClass(), Name, EUniqueObjectNameOptions::GloballyUnique);
+		const ULearningAgentsVisualLoggerObject* NewLoggerObject = NewObject<ULearningAgentsVisualLoggerObject>(this, UniqueName);
+		VisualLoggerObjects.Add(Name, NewLoggerObject);
+		return NewLoggerObject;
+	}
+}
diff --git a/Engine/Plugins/Experimental/LearningAgents/Source/LearningAgents/Private/LearningAgentsObservations.cpp b/Engine/Plugins/Experimental/LearningAgents/Source/LearningAgents/Private/LearningAgentsObservations.cpp
--- a/Engine/Plugins/Experimental/LearningAgents/Source/LearningAgents/Private/LearningAgentsObservations.cpp
+++ b/Engine/Plugins/Experimental/LearningAgents/Source/LearningAgents/Private/LearningAgentsObservations.cpp
@@ -2,6 +2,9 @@
 
 #include "LearningAgentsObservations.h"
 
+#include "LearningAgentsManagerListener.h"
+#include "LearningAgentsDebug.h"
+
 #include "LearningArray.h"
 #include "LearningLog.h"
 
@@ -11,63 +14,6 @@
 #include "Engine/World.h"
 #include "Components/SplineComponent.h"
 
-#include "VisualLogger/VisualLogger.h"
-
-#ifndef UE_LEARNING_AGENTS_ENABLE_VISUAL_LOG
-#define UE_LEARNING_AGENTS_ENABLE_VISUAL_LOG ENABLE_VISUAL_LOG
-#endif
-
-#if UE_LEARNING_AGENTS_ENABLE_VISUAL_LOG
-
-#define UE_LEARNING_AGENTS_VLOG_STRING(Owner, Category, Verbosity, Location, Color, Format, ...) \
-	UE_VLOG_LOCATION(Owner, Category, Verbosity, Location, 0.0f, Color, Format, ##__VA_ARGS__)
-
-#define UE_LEARNING_AGENTS_VLOG_LOCATION(Owner, Category, Verbosity, Location, Radius, Color, Format, ...) \
-	UE_VLOG_LOCATION(Owner, Category, Verbosity, Location, Radius, Color, Format, ##__VA_ARGS__)
-
-#define UE_LEARNING_AGENTS_VLOG_ARROW(Owner, Category, Verbosity, Start, End, Color, Format, ...) \
-	UE_VLOG_ARROW(Owner, Category, Verbosity, Start, End, Color, Format, ##__VA_ARGS__)
-
-#define UE_LEARNING_AGENTS_VLOG_SEGMENT(Owner, Category, Verbosity, Start, End, Color, Format, ...) \
-	UE_VLOG_SEGMENT(Owner, Category, Verbosity, Start, End, Color, Format, ##__VA_ARGS__)
-
-#define UE_LEARNING_AGENTS_VLOG_CIRCLE(Owner, Category, Verbosity, Center, UpAxis, Radius, Color, Format, ...) \
-	UE_VLOG_CIRCLE(Owner, Category, Verbosity, Center, UpAxis, Radius, Color, Format, ##__VA_ARGS__)
-
-#define UE_LEARNING_AGENTS_VLOG_OBOX(Owner, Category, Verbosity, Box, Matrix, Color, Format, ...) \
-	UE_VLOG_OBOX(Owner, Category, Verbosity, Box, Matrix, Color, Format, ##__VA_ARGS__)
-
-#define UE_LEARNING_AGENTS_VLOG_PLANE(Owner, Category, Verbosity, Location, Rotation, Axis0, Axis1, Color, Format, ...) \
-	UE_LEARNING_AGENTS_VLOG_OBOX(Owner, Category, Verbosity, FBox(FVector(-25.0f, -25.0f, -0.1f), FVector(25.0f, 25.0f, 0.1f)), UE::Learning::Agents::Debug::PlaneMatrix(Rotation, Location, Axis0, Axis1), Color, Format, ##__VA_ARGS__)
-
-#define UE_LEARNING_AGENTS_VLOG_MATRIX(Owner, Category, Verbosity, Matrix, Color, Format, ...) \
-	{ \
-		UE_LEARNING_AGENTS_VLOG_SEGMENT(Owner, Category, Verbosity, Matrix.TransformPosition(FVector::ZeroVector), Matrix.TransformPosition(15.0f * FVector::ForwardVector), FColor::Red, TEXT("")); \
-		UE_LEARNING_AGENTS_VLOG_SEGMENT(Owner, Category, Verbosity, Matrix.TransformPosition(FVector::ZeroVector), Matrix.TransformPosition(15.0f * FVector::RightVector), FColor::Green, TEXT("")); \
-		UE_LEARNING_AGENTS_VLOG_SEGMENT(Owner, Category, Verbosity, Matrix.TransformPosition(FVector::ZeroVector), Matrix.TransformPosition(15.0f * FVector::UpVector), FColor::Blue, TEXT("")); \
-		UE_LEARNING_AGENTS_VLOG_STRING(Owner, Category, Verbosity, Matrix.TransformPosition(FVector::ZeroVector) + FVector(0.0f, 0.0f, 20.0f), Color, Format, ##__VA_ARGS__); \
-	}
-
-#define UE_LEARNING_AGENTS_VLOG_TRANSFORM(Owner, Category, Verbosity, Location, Rotation, Color, Format, ...) \
-	{ \
-		UE_LEARNING_AGENTS_VLOG_SEGMENT(Owner, Category, Verbosity, Location, Location + 15.0f * Rotation.RotateVector(FVector::ForwardVector), FColor::Red, TEXT("")); \
-		UE_LEARNING_AGENTS_VLOG_SEGMENT(Owner, Category, Verbosity, Location, Location + 15.0f * Rotation.RotateVector(FVector::RightVector), FColor::Green, TEXT("")); \
-		UE_LEARNING_AGENTS_VLOG_SEGMENT(Owner, Category, Verbosity, Location, Location + 15.0f * Rotation.RotateVector(FVector::UpVector), FColor::Blue, TEXT("")); \
-		UE_LEARNING_AGENTS_VLOG_OBOX(Owner, Category, Verbosity, FBox(5.0f * FVector(-1, -1, -1), 5.0f * FVector(1, 1, 1)), FTransform(Rotation, Location, FVector::OneVector).ToMatrixNoScale(), Color, TEXT("")); \
-		UE_LEARNING_AGENTS_VLOG_STRING(Owner, Category, Verbosity, Location + FVector(0.0f, 0.0f, 20.0f), Color, Format, ##__VA_ARGS__); \
-	}
-
-#define UE_LEARNING_AGENTS_VLOG_ANGLE(Owner, Category, Verbosity, Angle, RelativeAngle, Location, Radius, Color, Format, ...) \
-	{ \
-		UE_LEARNING_AGENTS_VLOG_CIRCLE(Owner, Category, Verbosity, Location, FVector::UpVector, Radius, Color, TEXT("")); \
-		UE_LEARNING_AGENTS_VLOG_SEGMENT(Owner, Category, Verbosity, Location, Location + Radius * FVector(FMath::Sin(RelativeAngle), FMath::Cos(RelativeAngle), 0.0f), Color, TEXT("")); \
-		UE_LEARNING_AGENTS_VLOG_SEGMENT(Owner, Category, Verbosity, Location, Location + Radius * FVector(FMath::Sin(Angle), FMath::Cos(Angle), 0.0f), Color, TEXT("")); \
-		UE_LEARNING_AGENTS_VLOG_LOCATION(Owner, Category, Verbosity, Location + Radius * FVector(FMath::Sin(Angle), FMath::Cos(Angle), 0.0f), static_cast<uint16>(Radius / 20.0f), Color, Format, ##__VA_ARGS__); \
-	}
-
-#endif
-
-
 bool operator==(const FLearningAgentsObservationObjectElement& Lhs, const FLearningAgentsObservationObjectElement& Rhs)
 {
 	return Lhs.ObjectElement.Index == Rhs.ObjectElement.Index;
@@ -526,7 +472,7 @@ FTransform ULearningAgentsObservationFunctions::ProjectTransformOntoGroundPlane(
 	return FTransform(FQuat::FindBetweenNormals(FVector::ForwardVector, Direction), Position, Transform.GetScale3D());
 }
 
-UEnum* ULearningAgentsObservationFunctions::FindEnumByName(const FString Name)
+UEnum* ULearningAgentsObservationFunctions::FindEnumByName(const FString& Name)
 {
 	return FindObject<UEnum>(nullptr, *Name);
 }
@@ -1589,20 +1535,24 @@ FLearningAgentsObservationObjectElement ULearningAgentsObservationObject::MakeBo
 	const bool bValue,
 	const FName Name,
 	const bool bVisualLoggerEnabled,
+	ULearningAgentsManagerListener* VisualLoggerListener,
+	const int32 VisualLoggerAgentId,
 	const FVector VisualLoggerLocation,
 	const FLinearColor VisualLoggerColor)
 {
 	const float EncodedBool = bValue ? 1.0f : -1.0f;
 
 #if UE_LEARNING_AGENTS_ENABLE_VISUAL_LOG
-	if (bVisualLoggerEnabled)
+	if (bVisualLoggerEnabled && VisualLoggerListener)
 	{
-		const ULearningAgentsObservationVisualLoggerObject* VisualLoggerObject = GetOrAddVisualLoggerObject(Name);
+		const ULearningAgentsVisualLoggerObject* VisualLoggerObject = VisualLoggerListener->GetOrAddVisualLoggerObject(Name);
 
 		UE_LEARNING_AGENTS_VLOG_STRING(VisualLoggerObject, LogLearning, Display, VisualLoggerLocation,
 			VisualLoggerColor.ToFColor(true),
-			TEXT("Name: %s\nValue: [%s]\nEncoded: [% 6.2f]"),
+			TEXT("Listener: %s\nName: %s\nAgent Id: % 3i\nValue: [%s]\nEncoded: [% 6.2f]"),
+			*VisualLoggerListener->GetName(),
 			*Name.ToString(),
+			VisualLoggerAgentId,
 			bValue ? TEXT("true") : TEXT("false"),
 			EncodedBool);
 	}
@@ -1616,20 +1566,24 @@ FLearningAgentsObservationObjectElement ULearningAgentsObservationObject::MakeFl
 	const float FloatScale,
 	const FName Name,
 	const bool bVisualLoggerEnabled,
+	ULearningAgentsManagerListener* VisualLoggerListener,
+	const int32 VisualLoggerAgentId,
 	const FVector VisualLoggerLocation,
 	const FLinearColor VisualLoggerColor)
 {
 	const float EncodedValue = Value / FMath::Max(FloatScale, UE_SMALL_NUMBER);
 
 #if UE_LEARNING_AGENTS_ENABLE_VISUAL_LOG
-	if (bVisualLoggerEnabled)
+	if (bVisualLoggerEnabled && VisualLoggerListener)
 	{
-		const ULearningAgentsObservationVisualLoggerObject* VisualLoggerObject = GetOrAddVisualLoggerObject(Name);
+		const ULearningAgentsVisualLoggerObject* VisualLoggerObject = VisualLoggerListener->GetOrAddVisualLoggerObject(Name);
 
 		UE_LEARNING_AGENTS_VLOG_STRING(VisualLoggerObject, LogLearning, Display, VisualLoggerLocation,
 			VisualLoggerColor.ToFColor(true),
-			TEXT("Name: %s\nValue: [% 6.1f]\nScale: [% 6.2f]\nEncoded: [% 6.2f]"),
+			TEXT("Listener: %s\nName: %s\nAgent Id: % 3i\nValue: [% 6.1f]\nScale: [% 6.2f]\nEncoded: [% 6.2f]"),
+			*VisualLoggerListener->GetName(),
 			*Name.ToString(),
+			VisualLoggerAgentId,
 			Value,
 			FloatScale,
 			EncodedValue);
@@ -1645,6 +1599,8 @@ FLearningAgentsObservationObjectElement ULearningAgentsObservationObject::MakeLo
 	const float LocationScale,
 	const FName Name,
 	const bool bVisualLoggerEnabled,
+	ULearningAgentsManagerListener* VisualLoggerListener,
+	const int32 VisualLoggerAgentId,
 	const FVector VisualLoggerLocation,
 	const FLinearColor VisualLoggerColor)
 {
@@ -1655,9 +1611,9 @@ FLearningAgentsObservationObjectElement ULearningAgentsObservationObject::MakeLo
 		LocalLocation.Z / FMath::Max(LocationScale, UE_SMALL_NUMBER));
 
 #if UE_LEARNING_AGENTS_ENABLE_VISUAL_LOG
-	if (bVisualLoggerEnabled)
+	if (bVisualLoggerEnabled && VisualLoggerListener)
 	{
-		const ULearningAgentsObservationVisualLoggerObject* VisualLoggerObject = GetOrAddVisualLoggerObject(Name);
+		const ULearningAgentsVisualLoggerObject* VisualLoggerObject = VisualLoggerListener->GetOrAddVisualLoggerObject(Name);
 
 		UE_LEARNING_AGENTS_VLOG_LOCATION(VisualLoggerObject, LogLearning, Display,
 			Location,
@@ -1679,8 +1635,10 @@ FLearningAgentsObservationObjectElement ULearningAgentsObservationObject::MakeLo
 
 		UE_LEARNING_AGENTS_VLOG_STRING(VisualLoggerObject, LogLearning, Display, VisualLoggerLocation,
 			VisualLoggerColor.ToFColor(true),
-			TEXT("Name: %s\nLocation: [% 6.1f % 6.1f % 6.1f]\nLocal Location: [% 6.1f % 6.1f % 6.1f]\nScale: [% 6.2f]\nEncoded: [% 6.2f % 6.2f % 6.2f]"),
+			TEXT("Listener: %s\nName: %s\nAgent Id: % 3i\nLocation: [% 6.1f % 6.1f % 6.1f]\nLocal Location: [% 6.1f % 6.1f % 6.1f]\nScale: [% 6.2f]\nEncoded: [% 6.2f % 6.2f % 6.2f]"),
+			*VisualLoggerListener->GetName(),
 			*Name.ToString(),
+			VisualLoggerAgentId,
 			Location.X, Location.Y, Location.Z,
 			LocalLocation.X, LocalLocation.Y, LocalLocation.Z,
 			LocationScale,
@@ -1700,6 +1658,9 @@ FLearningAgentsObservationObjectElement ULearningAgentsObservationObject::MakeRo
 	const FRotator RelativeRotation,
 	const FName Name,
 	const bool bVisualLoggerEnabled,
+	ULearningAgentsManagerListener* VisualLoggerListener,
+	const int32 VisualLoggerAgentId,
+	const FVector VisualLoggerRotationLocation,
 	const FVector VisualLoggerLocation,
 	const FLinearColor VisualLoggerColor)
 {
@@ -1707,7 +1668,7 @@ FLearningAgentsObservationObjectElement ULearningAgentsObservationObject::MakeRo
 
 	return MakeRotationObservationFromQuat(
 		FQuat::MakeFromRotator(Rotation), FQuat::MakeFromRotator(RelativeRotation), Name,
-		bVisualLoggerEnabled, VisualLoggerLocation,	VisualLoggerColor);
+		bVisualLoggerEnabled, VisualLoggerListener, VisualLoggerAgentId, VisualLoggerRotationLocation, VisualLoggerLocation, VisualLoggerColor);
 }
 
 FLearningAgentsObservationObjectElement ULearningAgentsObservationObject::MakeRotationObservationFromQuat(
@@ -1715,6 +1676,9 @@ FLearningAgentsObservationObjectElement ULearningAgentsObservationObject::MakeRo
 	const FQuat RelativeRotation,
 	const FName Name,
 	const bool bVisualLoggerEnabled,
+	ULearningAgentsManagerListener* VisualLoggerListener,
+	const int32 VisualLoggerAgentId,
+	const FVector VisualLoggerRotationLocation,
 	const FVector VisualLoggerLocation,
 	const FLinearColor VisualLoggerColor)
 {
@@ -1723,20 +1687,22 @@ FLearningAgentsObservationObjectElement ULearningAgentsObservationObject::MakeRo
 	const FVector LocalAxisRight = LocalRotation.GetRightVector();
 
 #if UE_LEARNING_AGENTS_ENABLE_VISUAL_LOG
-	if (bVisualLoggerEnabled)
+	if (bVisualLoggerEnabled && VisualLoggerListener)
 	{
-		const ULearningAgentsObservationVisualLoggerObject* VisualLoggerObject = GetOrAddVisualLoggerObject(Name);
+		const ULearningAgentsVisualLoggerObject* VisualLoggerObject = VisualLoggerListener->GetOrAddVisualLoggerObject(Name);
 
 		UE_LEARNING_AGENTS_VLOG_TRANSFORM(VisualLoggerObject, LogLearning, Display,
-			VisualLoggerLocation,
+			VisualLoggerRotationLocation,
 			LocalRotation.Rotator(),
 			VisualLoggerColor.ToFColor(true),
 			TEXT(""));
 
 		UE_LEARNING_AGENTS_VLOG_STRING(VisualLoggerObject, LogLearning, Display, VisualLoggerLocation,
 			VisualLoggerColor.ToFColor(true),
-			TEXT("Name: %s\nRotation: [% 6.1f % 6.1f % 6.1f % 6.1f]\nLocal Rotation: [% 6.1f % 6.1f % 6.1f % 6.1f]\nEncoded Forward: [% 6.2f % 6.2f % 6.2f]\nEncoded Right: [% 6.2f % 6.2f % 6.2f]"),
+			TEXT("Listener: %s\nName: %s\nAgent Id: % 3i\nRotation: [% 6.1f % 6.1f % 6.1f % 6.1f]\nLocal Rotation: [% 6.1f % 6.1f % 6.1f % 6.1f]\nEncoded Forward: [% 6.2f % 6.2f % 6.2f]\nEncoded Right: [% 6.2f % 6.2f % 6.2f]"),
+			*VisualLoggerListener->GetName(),
 			*Name.ToString(),
+			VisualLoggerAgentId,
 			Rotation.W, Rotation.X, Rotation.Y, Rotation.Z,
 			LocalRotation.W, LocalRotation.X, LocalRotation.Y, LocalRotation.Z,
 			LocalAxisForward.X, LocalAxisForward.Y, LocalAxisForward.Z,
@@ -1759,6 +1725,9 @@ FLearningAgentsObservationObjectElement ULearningAgentsObservationObject::MakeSc
 	const FVector RelativeScale,
 	const FName Name,
 	const bool bVisualLoggerEnabled,
+	ULearningAgentsManagerListener* VisualLoggerListener,
+	const int32 VisualLoggerAgentId,
+	const FVector VisualLoggerScaleLocation,
 	const FVector VisualLoggerLocation,
 	const FLinearColor VisualLoggerColor)
 {
@@ -1767,14 +1736,16 @@ FLearningAgentsObservationObjectElement ULearningAgentsObservationObject::MakeSc
 		UE::Learning::Agents::Observation::Private::VectorLogSafe(RelativeScale);
 
 #if UE_LEARNING_AGENTS_ENABLE_VISUAL_LOG
-	if (bVisualLoggerEnabled)
+	if (bVisualLoggerEnabled && VisualLoggerListener)
 	{
-		const ULearningAgentsObservationVisualLoggerObject* VisualLoggerObject = GetOrAddVisualLoggerObject(Name);
+		const ULearningAgentsVisualLoggerObject* VisualLoggerObject = VisualLoggerListener->GetOrAddVisualLoggerObject(Name);
 
 		UE_LEARNING_AGENTS_VLOG_STRING(VisualLoggerObject, LogLearning, Display, VisualLoggerLocation,
 			VisualLoggerColor.ToFColor(true),
-			TEXT("Name: %s\nScale: [% 6.1f % 6.1f % 6.1f]\nEncoded: [% 6.2f % 6.2f % 6.2f]"),
+			TEXT("Listener: %s\nName: %s\nAgent Id: % 3i\nScale: [% 6.1f % 6.1f % 6.1f]\nEncoded: [% 6.2f % 6.2f % 6.2f]"),
+			*VisualLoggerListener->GetName(),
 			*Name.ToString(),
+			VisualLoggerAgentId,
 			Scale.X, Scale.Y, Scale.Z,
 			LocalLogScale.X, LocalLogScale.Y, LocalLogScale.Z);
 	}
@@ -1793,15 +1764,17 @@ FLearningAgentsObservationObjectElement ULearningAgentsObservationObject::MakeTr
 	const float LocationScale,
 	const FName Name,
 	const bool bVisualLoggerEnabled,
+	ULearningAgentsManagerListener* VisualLoggerListener,
+	const int32 VisualLoggerAgentId,
 	const FVector VisualLoggerLocation,
 	const FLinearColor VisualLoggerColor)
 {
 	const FTransform LocalTransform = Transform * RelativeTransform.Inverse();
 	
 #if UE_LEARNING_AGENTS_ENABLE_VISUAL_LOG
-	if (bVisualLoggerEnabled)
+	if (bVisualLoggerEnabled && VisualLoggerListener)
 	{
-		const ULearningAgentsObservationVisualLoggerObject* VisualLoggerObject = GetOrAddVisualLoggerObject(Name);
+		const ULearningAgentsVisualLoggerObject* VisualLoggerObject = VisualLoggerListener->GetOrAddVisualLoggerObject(Name);
 
 		const FVector LocalLocation = LocalTransform.GetLocation();
 		const FRotator LocalRotation = LocalTransform.Rotator();
@@ -1819,8 +1792,10 @@ FLearningAgentsObservationObjectElement ULearningAgentsObservationObject::MakeTr
 
 		UE_LEARNING_AGENTS_VLOG_STRING(VisualLoggerObject, LogLearning, Display, VisualLoggerLocation,
 			VisualLoggerColor.ToFColor(true),
-			TEXT("Name: %s\nLocation: [% 6.1f % 6.1f % 6.1f]\nLocal Location: [% 6.1f % 6.1f % 6.1f]\nRotation: [% 6.1f % 6.1f % 6.1f]\nLocal Rotation: [% 6.1f % 6.1f % 6.1f]\nScale: [% 6.1f % 6.1f % 6.1f]\nLocal Scale: [% 6.1f % 6.1f % 6.1f]\nLocation Scale: [% 6.1f]"),
+			TEXT("Listener: %s\nName: %s\nAgent Id: % 3i\nLocation: [% 6.1f % 6.1f % 6.1f]\nLocal Location: [% 6.1f % 6.1f % 6.1f]\nRotation: [% 6.1f % 6.1f % 6.1f]\nLocal Rotation: [% 6.1f % 6.1f % 6.1f]\nScale: [% 6.1f % 6.1f % 6.1f]\nLocal Scale: [% 6.1f % 6.1f % 6.1f]\nLocation Scale: [% 6.1f]"),
+			*VisualLoggerListener->GetName(),
 			*Name.ToString(),
+			VisualLoggerAgentId,
 			Location.X, Location.Y, Location.Z,
 			LocalLocation.X, LocalLocation.Y, LocalLocation.Z,
 			Rotation.Roll, Rotation.Pitch, Rotation.Yaw,
@@ -1876,6 +1851,9 @@ FLearningAgentsObservationObjectElement ULearningAgentsObservationObject::MakeDi
 	const FTransform RelativeTransform,
 	const FName Name,
 	const bool bVisualLoggerEnabled,
+	ULearningAgentsManagerListener* VisualLoggerListener,
+	const int32 VisualLoggerAgentId,
+	const FVector VisualLoggerDirectionLocation,
 	const FVector VisualLoggerLocation,
 	const float VisualLoggerArrowLength,
 	const FLinearColor VisualLoggerColor)
@@ -1883,13 +1861,13 @@ FLearningAgentsObservationObjectElement ULearningAgentsObservationObject::MakeDi
 	const FVector LocalDirection = RelativeTransform.InverseTransformVectorNoScale(Direction).GetSafeNormal(UE_SMALL_NUMBER, FVector::ForwardVector);
 
 #if UE_LEARNING_AGENTS_ENABLE_VISUAL_LOG
-	if (bVisualLoggerEnabled)
+	if (bVisualLoggerEnabled && VisualLoggerListener)
 	{
-		const ULearningAgentsObservationVisualLoggerObject* VisualLoggerObject = GetOrAddVisualLoggerObject(Name);
+		const ULearningAgentsVisualLoggerObject* VisualLoggerObject = VisualLoggerListener->GetOrAddVisualLoggerObject(Name);
 
 		UE_LEARNING_AGENTS_VLOG_ARROW(VisualLoggerObject, LogLearning, Display,
-			RelativeTransform.GetTranslation(),
-			RelativeTransform.GetTranslation() + VisualLoggerArrowLength * Direction,
+			VisualLoggerDirectionLocation,
+			VisualLoggerDirectionLocation + VisualLoggerArrowLength * Direction,
 			VisualLoggerColor.ToFColor(true),
 			TEXT(""));
 
@@ -1901,8 +1879,10 @@ FLearningAgentsObservationObjectElement ULearningAgentsObservationObject::MakeDi
 
 		UE_LEARNING_AGENTS_VLOG_STRING(VisualLoggerObject, LogLearning, Display, VisualLoggerLocation,
 			VisualLoggerColor.ToFColor(true),
-			TEXT("Name: %s\nDirection: [% 6.1f % 6.1f % 6.1f]\nLocal Direction: [% 6.1f % 6.1f % 6.1f]\nEncoded: [% 6.2f % 6.2f % 6.2f]"),
+			TEXT("Listener: %s\nName: %s\nAgent Id: % 3i\nDirection: [% 6.1f % 6.1f % 6.1f]\nLocal Direction: [% 6.1f % 6.1f % 6.1f]\nEncoded: [% 6.2f % 6.2f % 6.2f]"),
+			*VisualLoggerListener->GetName(),
 			*Name.ToString(),
+			VisualLoggerAgentId,
 			Direction.X, Direction.Y, Direction.Z,
 			LocalDirection.X, LocalDirection.Y, LocalDirection.Z,
 			LocalDirection.X, LocalDirection.Y, LocalDirection.Z); // We don't have a scaling so the encoded value is the same
@@ -1923,6 +1903,8 @@ FLearningAgentsObservationObjectElement ULearningAgentsObservationObject::MakeLo
 	const float LocationScale,
 	const FName Name,
 	const bool bVisualLoggerEnabled,
+	ULearningAgentsManagerListener* VisualLoggerListener,
+	const int32 VisualLoggerAgentId,
 	const FVector VisualLoggerLocation,
 	const FLinearColor VisualLoggerColor)
 {
@@ -1932,7 +1914,7 @@ FLearningAgentsObservationObjectElement ULearningAgentsObservationObject::MakeLo
 		return FLearningAgentsObservationObjectElement();
 	}
 
-	return MakeLocationObservation(SplineComponent->GetLocationAtDistanceAlongSpline(DistanceAlongSpline, ESplineCoordinateSpace::World), RelativeTransform, LocationScale, Name, bVisualLoggerEnabled, VisualLoggerLocation, VisualLoggerColor);
+	return MakeLocationObservation(SplineComponent->GetLocationAtDistanceAlongSpline(DistanceAlongSpline, ESplineCoordinateSpace::World), RelativeTransform, LocationScale, Name, bVisualLoggerEnabled, VisualLoggerListener, VisualLoggerAgentId, VisualLoggerLocation, VisualLoggerColor);
 }
 
 FLearningAgentsObservationObjectElement ULearningAgentsObservationObject::MakeProportionAlongSplineObservation(const USplineComponent* SplineComponent, const float DistanceAlongSpline, const FName Name)
@@ -1964,6 +1946,8 @@ FLearningAgentsObservationObjectElement ULearningAgentsObservationObject::MakeDi
 	const FTransform RelativeTransform,
 	const FName Name,
 	const bool bVisualLoggerEnabled,
+	ULearningAgentsManagerListener* VisualLoggerListener,
+	const int32 VisualLoggerAgentId,
 	const FVector VisualLoggerLocation,
 	const float VisualLoggerArrowLength,
 	const FLinearColor VisualLoggerColor)
@@ -1974,7 +1958,16 @@ FLearningAgentsObservationObjectElement ULearningAgentsObservationObject::MakeDi
 		return FLearningAgentsObservationObjectElement();
 	}
 
-	return MakeDirectionObservation(SplineComponent->GetDirectionAtDistanceAlongSpline(DistanceAlongSpline, ESplineCoordinateSpace::World), RelativeTransform, Name, bVisualLoggerEnabled, VisualLoggerLocation, VisualLoggerArrowLength, VisualLoggerColor);
+#if UE_LEARNING_AGENTS_ENABLE_VISUAL_LOG
+	const FVector VisualLoggerDirectionLocation = 
+		(bVisualLoggerEnabled && VisualLoggerListener) ?
+		SplineComponent->GetLocationAtDistanceAlongSpline(DistanceAlongSpline, ESplineCoordinateSpace::World) : 
+		FVector::ZeroVector;
+#else
+	const FVector VisualLoggerDirectionLocation = FVector::ZeroVector;
+#endif
+
+	return MakeDirectionObservation(SplineComponent->GetDirectionAtDistanceAlongSpline(DistanceAlongSpline, ESplineCoordinateSpace::World), RelativeTransform, Name, bVisualLoggerEnabled, VisualLoggerListener, VisualLoggerAgentId, VisualLoggerDirectionLocation, VisualLoggerLocation, VisualLoggerArrowLength, VisualLoggerColor);
 }
 
 FLearningAgentsObservationObjectElement ULearningAgentsObservationObject::MakePropertiesAlongSplineObservation(const USplineComponent* SplineComponent, const float DistanceAlongSpline, const FTransform RelativeTransform, const float LocationScale, const FName Name)
@@ -2013,13 +2006,10 @@ FLearningAgentsObservationObjectElement ULearningAgentsObservationObject::MakePr
 
 	if (bHit)
 	{
-		//DrawDebugLine(GetWorld(), RayStartWorld, RayStartWorld + TraceHit.Time * (RayEndWorld - RayStartWorld), FColor::Red, false, -1.0f, 0, 5.0f);
-		//DrawDebugLine(GetWorld(), RayStartWorld + TraceHit.Time * (RayEndWorld - RayStartWorld), RayEndWorld, FColor::Blue, false, -1.0f, 0, 5.0f);
 		return MakeFloatObservation(1.0f - TraceHit.Time, 1.0f, Name);
 	}
 	else
 	{
-		//DrawDebugLine(GetWorld(), RayStartWorld, RayEndWorld, FColor::Red, false, -1.0f, 0, 5.0f);
 		return MakeFloatObservation(0.0f, 1.0f, Name);
 	}
 }
@@ -3555,18 +3545,3 @@ bool ULearningAgentsObservationObject::GetProportionAlongRaysObservationToArrayV
 	return true;
 }
 
-const ULearningAgentsObservationVisualLoggerObject* ULearningAgentsObservationObject::GetOrAddVisualLoggerObject(const FName Name)
-{
-	const TObjectPtr<const ULearningAgentsObservationVisualLoggerObject>* Value = VisualLoggerObjects.Find(Name);
-	if (Value)
-	{
-		return Value->Get();
-	}
-	else
-	{
-		const FName UniqueName = MakeUniqueObjectName(this, ULearningAgentsObservationVisualLoggerObject::StaticClass(), Name, EUniqueObjectNameOptions::GloballyUnique);
-		const ULearningAgentsObservationVisualLoggerObject* NewLoggerObject = NewObject<ULearningAgentsObservationVisualLoggerObject>(this, UniqueName);
-		VisualLoggerObjects.Add(Name, NewLoggerObject);
-		return NewLoggerObject;
-	}
-}
diff --git a/Engine/Plugins/Experimental/LearningAgents/Source/LearningAgents/Public/LearningAgentsActions.h b/Engine/Plugins/Experimental/LearningAgents/Source/LearningAgents/Public/LearningAgentsActions.h
--- a/Engine/Plugins/Experimental/LearningAgents/Source/LearningAgents/Public/LearningAgentsActions.h
+++ b/Engine/Plugins/Experimental/LearningAgents/Source/LearningAgents/Public/LearningAgentsActions.h
@@ -310,107 +310,107 @@ public:
 
 public:
 
-	UFUNCTION(BlueprintPure = false, Category = "LearningAgents")
-	UPARAM(DisplayName = "Success") bool GetNullAction(const FLearningAgentsActionObjectElement Element, const FName Name = TEXT("Null")) const;
+	UFUNCTION(BlueprintPure = false, Category = "LearningAgents", meta = (ReturnDisplayName = "Success"))
+	bool GetNullAction(const FLearningAgentsActionObjectElement Element, const FName Name = TEXT("Null")) const;
 
-	UFUNCTION(BlueprintPure = false, Category = "LearningAgents")
-	UPARAM(DisplayName = "Success") bool GetContinuousActionNum(int32& OutNum, const FLearningAgentsActionObjectElement Element, const FName Name = TEXT("Continuous")) const;
+	UFUNCTION(BlueprintPure = false, Category = "LearningAgents", meta = (ReturnDisplayName = "Success"))
+	bool GetContinuousActionNum(int32& OutNum, const FLearningAgentsActionObjectElement Element, const FName Name = TEXT("Continuous")) const;
 
-	UFUNCTION(BlueprintPure = false, Category = "LearningAgents")
-	UPARAM(DisplayName = "Success") bool GetContinuousAction(TArray<float>& OutValues, const FLearningAgentsActionObjectElement Element, const FName Name = TEXT("Continuous")) const;
-	UPARAM(DisplayName = "Success") bool GetContinuousActionToArrayView(TArrayView<float> OutValues, const FLearningAgentsActionObjectElement Element, const FName Name = TEXT("Continuous")) const;
+	UFUNCTION(BlueprintPure = false, Category = "LearningAgents", meta = (ReturnDisplayName = "Success"))
+	bool GetContinuousAction(TArray<float>& OutValues, const FLearningAgentsActionObjectElement Element, const FName Name = TEXT("Continuous")) const;
+	bool GetContinuousActionToArrayView(TArrayView<float> OutValues, const FLearningAgentsActionObjectElement Element, const FName Name = TEXT("Continuous")) const;
 	
-	UFUNCTION(BlueprintPure = false, Category = "LearningAgents")
-	UPARAM(DisplayName = "Success") bool GetExclusiveDiscreteAction(int32& OutIndex, const FLearningAgentsActionObjectElement Element, const FName Name = TEXT("DiscreteExclusive")) const;
+	UFUNCTION(BlueprintPure = false, Category = "LearningAgents", meta = (ReturnDisplayName = "Success"))
+	bool GetExclusiveDiscreteAction(int32& OutIndex, const FLearningAgentsActionObjectElement Element, const FName Name = TEXT("DiscreteExclusive")) const;
 
-	UFUNCTION(BlueprintPure = false, Category = "LearningAgents")
-	UPARAM(DisplayName = "Success") bool GetInclusiveDiscreteActionNum(int32& OutNum, const FLearningAgentsActionObjectElement Element, const FName Name = TEXT("DiscreteInclusive")) const;
+	UFUNCTION(BlueprintPure = false, Category = "LearningAgents", meta = (ReturnDisplayName = "Success"))
+	bool GetInclusiveDiscreteActionNum(int32& OutNum, const FLearningAgentsActionObjectElement Element, const FName Name = TEXT("DiscreteInclusive")) const;
 
-	UFUNCTION(BlueprintPure = false, Category = "LearningAgents")
-	UPARAM(DisplayName = "Success") bool GetInclusiveDiscreteAction(TArray<int32>& OutIndices, const FLearningAgentsActionObjectElement Element, const FName Name = TEXT("DiscreteInclusive")) const;
-	UPARAM(DisplayName = "Success") bool GetInclusiveDiscreteActionToArrayView(TArrayView<int32> OutIndices, const FLearningAgentsActionObjectElement Element, const FName Name = TEXT("DiscreteInclusive")) const;
+	UFUNCTION(BlueprintPure = false, Category = "LearningAgents", meta = (ReturnDisplayName = "Success"))
+	bool GetInclusiveDiscreteAction(TArray<int32>& OutIndices, const FLearningAgentsActionObjectElement Element, const FName Name = TEXT("DiscreteInclusive")) const;
+	bool GetInclusiveDiscreteActionToArrayView(TArrayView<int32> OutIndices, const FLearningAgentsActionObjectElement Element, const FName Name = TEXT("DiscreteInclusive")) const;
 
-	UFUNCTION(BlueprintPure = false, Category = "LearningAgents")
-	UPARAM(DisplayName = "Success") bool GetStructActionNum(int32& OutNum, const FLearningAgentsActionObjectElement Element, const FName Name = TEXT("Struct")) const;
+	UFUNCTION(BlueprintPure = false, Category = "LearningAgents", meta = (ReturnDisplayName = "Success"))
+	bool GetStructActionNum(int32& OutNum, const FLearningAgentsActionObjectElement Element, const FName Name = TEXT("Struct")) const;
 
-	UFUNCTION(BlueprintPure = false, Category = "LearningAgents")
-	UPARAM(DisplayName = "Success") bool GetStructAction(TMap<FName, FLearningAgentsActionObjectElement>& OutElements, const FLearningAgentsActionObjectElement Element, const FName Name = TEXT("Struct")) const;
+	UFUNCTION(BlueprintPure = false, Category = "LearningAgents", meta = (ReturnDisplayName = "Success"))
+	bool GetStructAction(TMap<FName, FLearningAgentsActionObjectElement>& OutElements, const FLearningAgentsActionObjectElement Element, const FName Name = TEXT("Struct")) const;
 	
-	UFUNCTION(BlueprintPure = false, Category = "LearningAgents")
-	UPARAM(DisplayName = "Success") bool GetStructActionToArrays(TArray<FName>& OutElementNames, TArray<FLearningAgentsActionObjectElement>& OutElements, const FLearningAgentsActionObjectElement Element, const FName Name = TEXT("Struct")) const;
-	UPARAM(DisplayName = "Success") bool GetStructActionToArrayViews(TArrayView<FName> OutElementNames, TArrayView<FLearningAgentsActionObjectElement> OutElements, const FLearningAgentsActionObjectElement Element, const FName Name = TEXT("Struct")) const;
+	UFUNCTION(BlueprintPure = false, Category = "LearningAgents", meta = (ReturnDisplayName = "Success"))
+	bool GetStructActionToArrays(TArray<FName>& OutElementNames, TArray<FLearningAgentsActionObjectElement>& OutElements, const FLearningAgentsActionObjectElement Element, const FName Name = TEXT("Struct")) const;
+	bool GetStructActionToArrayViews(TArrayView<FName> OutElementNames, TArrayView<FLearningAgentsActionObjectElement> OutElements, const FLearningAgentsActionObjectElement Element, const FName Name = TEXT("Struct")) const;
 
-	UFUNCTION(BlueprintPure = false, Category = "LearningAgents")
-	UPARAM(DisplayName = "Success") bool GetExclusiveUnionAction(FName& OutElementName, FLearningAgentsActionObjectElement& OutElement, const FLearningAgentsActionObjectElement Element, const FName Name = TEXT("ExclusiveUnion")) const;
+	UFUNCTION(BlueprintPure = false, Category = "LearningAgents", meta = (ReturnDisplayName = "Success"))
+	bool GetExclusiveUnionAction(FName& OutElementName, FLearningAgentsActionObjectElement& OutElement, const FLearningAgentsActionObjectElement Element, const FName Name = TEXT("ExclusiveUnion")) const;
 
-	UFUNCTION(BlueprintPure = false, Category = "LearningAgents")
-	UPARAM(DisplayName = "Success") bool GetInclusiveUnionActionNum(int32& OutNum, const FLearningAgentsActionObjectElement Element, const FName Name = TEXT("InclusiveUnion")) const;
+	UFUNCTION(BlueprintPure = false, Category = "LearningAgents", meta = (ReturnDisplayName = "Success"))
+	bool GetInclusiveUnionActionNum(int32& OutNum, const FLearningAgentsActionObjectElement Element, const FName Name = TEXT("InclusiveUnion")) const;
 
-	UFUNCTION(BlueprintPure = false, Category = "LearningAgents")
-	UPARAM(DisplayName = "Success") bool GetInclusiveUnionAction(TMap<FName, FLearningAgentsActionObjectElement>& OutElements, const FLearningAgentsActionObjectElement Element, const FName Name = TEXT("InclusiveUnion")) const;
+	UFUNCTION(BlueprintPure = false, Category = "LearningAgents", meta = (ReturnDisplayName = "Success"))
+	bool GetInclusiveUnionAction(TMap<FName, FLearningAgentsActionObjectElement>& OutElements, const FLearningAgentsActionObjectElement Element, const FName Name = TEXT("InclusiveUnion")) const;
 	
-	UFUNCTION(BlueprintPure = false, Category = "LearningAgents")
-	UPARAM(DisplayName = "Success") bool GetInclusiveUnionActionToArrays(TArray<FName>& OutElementNames, TArray<FLearningAgentsActionObjectElement>& OutElements, const FLearningAgentsActionObjectElement Element, const FName Name = TEXT("InclusiveUnion")) const;
-	UPARAM(DisplayName = "Success") bool GetInclusiveUnionActionToArrayViews(TArrayView<FName> OutElementNames, TArrayView<FLearningAgentsActionObjectElement> OutElements, const FLearningAgentsActionObjectElement Element, const FName Name = TEXT("InclusiveUnion")) const;
+	UFUNCTION(BlueprintPure = false, Category = "LearningAgents", meta = (ReturnDisplayName = "Success"))
+	bool GetInclusiveUnionActionToArrays(TArray<FName>& OutElementNames, TArray<FLearningAgentsActionObjectElement>& OutElements, const FLearningAgentsActionObjectElement Element, const FName Name = TEXT("InclusiveUnion")) const;
+	bool GetInclusiveUnionActionToArrayViews(TArrayView<FName> OutElementNames, TArrayView<FLearningAgentsActionObjectElement> OutElements, const FLearningAgentsActionObjectElement Element, const FName Name = TEXT("InclusiveUnion")) const;
 
-	UFUNCTION(BlueprintPure = false, Category = "LearningAgents")
-	UPARAM(DisplayName = "Success") bool GetStaticArrayActionNum(int32& OutNum, const FLearningAgentsActionObjectElement Element, const FName Name = TEXT("StaticArray")) const;
+	UFUNCTION(BlueprintPure = false, Category = "LearningAgents", meta = (ReturnDisplayName = "Success"))
+	bool GetStaticArrayActionNum(int32& OutNum, const FLearningAgentsActionObjectElement Element, const FName Name = TEXT("StaticArray")) const;
 
-	UFUNCTION(BlueprintPure = false, Category = "LearningAgents")
-	UPARAM(DisplayName = "Success") bool GetStaticArrayAction(TArray<FLearningAgentsActionObjectElement>& OutElements, const FLearningAgentsActionObjectElement Element, const FName Name = TEXT("StaticArray")) const;
-	UPARAM(DisplayName = "Success") bool GetStaticArrayActionToArrayView(TArrayView<FLearningAgentsActionObjectElement> OutElements, const FLearningAgentsActionObjectElement Element, const FName Name = TEXT("StaticArray")) const;
+	UFUNCTION(BlueprintPure = false, Category = "LearningAgents", meta = (ReturnDisplayName = "Success"))
+	bool GetStaticArrayAction(TArray<FLearningAgentsActionObjectElement>& OutElements, const FLearningAgentsActionObjectElement Element, const FName Name = TEXT("StaticArray")) const;
+	bool GetStaticArrayActionToArrayView(TArrayView<FLearningAgentsActionObjectElement> OutElements, const FLearningAgentsActionObjectElement Element, const FName Name = TEXT("StaticArray")) const;
 
-	UFUNCTION(BlueprintPure = false, Category = "LearningAgents")
-	UPARAM(DisplayName = "Success") bool GetPairAction(FLearningAgentsActionObjectElement& OutKey, FLearningAgentsActionObjectElement& OutValue, const FLearningAgentsActionObjectElement Element, const FName Name = TEXT("Pair")) const;
+	UFUNCTION(BlueprintPure = false, Category = "LearningAgents", meta = (ReturnDisplayName = "Success"))
+	bool GetPairAction(FLearningAgentsActionObjectElement& OutKey, FLearningAgentsActionObjectElement& OutValue, const FLearningAgentsActionObjectElement Element, const FName Name = TEXT("Pair")) const;
 
-	UFUNCTION(BlueprintPure = false, Category = "LearningAgents")
-	UPARAM(DisplayName = "Success") bool GetEnumAction(uint8& OutEnumValue, const UEnum* Enum, const FLearningAgentsActionObjectElement Element, const FName Name = TEXT("Enum")) const;
+	UFUNCTION(BlueprintPure = false, Category = "LearningAgents", meta = (ReturnDisplayName = "Success"))
+	bool GetEnumAction(uint8& OutEnumValue, const UEnum* Enum, const FLearningAgentsActionObjectElement Element, const FName Name = TEXT("Enum")) const;
 
-	UFUNCTION(BlueprintPure = false, Category = "LearningAgents")
-	UPARAM(DisplayName = "Success") bool GetBitmaskAction(int32& OutBitmaskValue, const UEnum* Enum, const FLearningAgentsActionObjectElement Element, const FName Name = TEXT("Bitmask")) const;
+	UFUNCTION(BlueprintPure = false, Category = "LearningAgents", meta = (ReturnDisplayName = "Success"))
+	bool GetBitmaskAction(int32& OutBitmaskValue, const UEnum* Enum, const FLearningAgentsActionObjectElement Element, const FName Name = TEXT("Bitmask")) const;
 
-	UFUNCTION(BlueprintPure = false, Category = "LearningAgents", meta = (ExpandEnumAsExecs = "OutOption"))
-	UPARAM(DisplayName = "Success") bool GetOptionalAction(ELearningAgentsOptionalAction& OutOption, FLearningAgentsActionObjectElement& OutElement, const FLearningAgentsActionObjectElement Element, const FName Name = TEXT("Optional")) const;
+	UFUNCTION(BlueprintPure = false, Category = "LearningAgents", meta = (ExpandEnumAsExecs = "OutOption", ReturnDisplayName = "Success"))
+	bool GetOptionalAction(ELearningAgentsOptionalAction& OutOption, FLearningAgentsActionObjectElement& OutElement, const FLearningAgentsActionObjectElement Element, const FName Name = TEXT("Optional")) const;
 	
-	UFUNCTION(BlueprintPure = false, Category = "LearningAgents", meta = (ExpandEnumAsExecs = "OutEither"))
-	UPARAM(DisplayName = "Success") bool GetEitherAction(ELearningAgentsEitherAction& OutEither, FLearningAgentsActionObjectElement& OutElement, const FLearningAgentsActionObjectElement Element, const FName Name = TEXT("Either")) const;
+	UFUNCTION(BlueprintPure = false, Category = "LearningAgents", meta = (ExpandEnumAsExecs = "OutEither", ReturnDisplayName = "Success"))
+	bool GetEitherAction(ELearningAgentsEitherAction& OutEither, FLearningAgentsActionObjectElement& OutElement, const FLearningAgentsActionObjectElement Element, const FName Name = TEXT("Either")) const;
 
-	UFUNCTION(BlueprintPure = false, Category = "LearningAgents")
-	UPARAM(DisplayName = "Success") bool GetEncodingAction(FLearningAgentsActionObjectElement& OutElement, const FLearningAgentsActionObjectElement Element, const FName Name = TEXT("Encoding")) const;
+	UFUNCTION(BlueprintPure = false, Category = "LearningAgents", meta = (ReturnDisplayName = "Success"))
+	bool GetEncodingAction(FLearningAgentsActionObjectElement& OutElement, const FLearningAgentsActionObjectElement Element, const FName Name = TEXT("Encoding")) const;
 
-	UFUNCTION(BlueprintPure = false, Category = "LearningAgents")
-	UPARAM(DisplayName = "Success") bool GetBoolAction(bool& bOutValue, const FLearningAgentsActionObjectElement Element, const FName Name = TEXT("Bool")) const;
+	UFUNCTION(BlueprintPure = false, Category = "LearningAgents", meta = (ReturnDisplayName = "Success"))
+	bool GetBoolAction(bool& bOutValue, const FLearningAgentsActionObjectElement Element, const FName Name = TEXT("Bool")) const;
 
-	UFUNCTION(BlueprintPure = false, Category = "LearningAgents")
-	UPARAM(DisplayName = "Success") bool GetFloatAction(float& OutValue, const FLearningAgentsActionObjectElement Element, const float FloatScale = 1.0f, const FName Name = TEXT("Float")) const;
+	UFUNCTION(BlueprintPure = false, Category = "LearningAgents", meta = (ReturnDisplayName = "Success"))
+	bool GetFloatAction(float& OutValue, const FLearningAgentsActionObjectElement Element, const float FloatScale = 1.0f, const FName Name = TEXT("Float")) const;
 
-	UFUNCTION(BlueprintPure = false, Category = "LearningAgents")
-	UPARAM(DisplayName = "Success") bool GetLocationAction(FVector& OutLocation, const FLearningAgentsActionObjectElement Element, const FTransform RelativeTransform = FTransform(), const float LocationScale = 100.0f, const FName Name = TEXT("Location")) const;
+	UFUNCTION(BlueprintPure = false, Category = "LearningAgents", meta = (ReturnDisplayName = "Success"))
+	bool GetLocationAction(FVector& OutLocation, const FLearningAgentsActionObjectElement Element, const FTransform RelativeTransform = FTransform(), const float LocationScale = 100.0f, const FName Name = TEXT("Location")) const;
 
-	UFUNCTION(BlueprintPure = false, Category = "LearningAgents")
-	UPARAM(DisplayName = "Success") bool GetRotationAction(FRotator& OutRotation, const FLearningAgentsActionObjectElement Element, const FRotator RelativeRotation = FRotator::ZeroRotator, const float RotationScale = 90.0f, const FName Name = TEXT("Rotation")) const;
-	UPARAM(DisplayName = "Success") bool GetRotationActionAsQuat(FQuat& OutRotation, const FLearningAgentsActionObjectElement Element, const FQuat RelativeRotation = FQuat::Identity, const float RotationScale = 90.0f, const FName Name = TEXT("Rotation")) const;
+	UFUNCTION(BlueprintPure = false, Category = "LearningAgents", meta = (ReturnDisplayName = "Success"))
+	bool GetRotationAction(FRotator& OutRotation, const FLearningAgentsActionObjectElement Element, const FRotator RelativeRotation = FRotator::ZeroRotator, const float RotationScale = 90.0f, const FName Name = TEXT("Rotation")) const;
+	bool GetRotationActionAsQuat(FQuat& OutRotation, const FLearningAgentsActionObjectElement Element, const FQuat RelativeRotation = FQuat::Identity, const float RotationScale = 90.0f, const FName Name = TEXT("Rotation")) const;
 
-	UFUNCTION(BlueprintPure = false, Category = "LearningAgents")
-	UPARAM(DisplayName = "Success") bool GetScaleAction(FVector& OutScale, const FLearningAgentsActionObjectElement Element, const FVector RelativeScale = FVector(1,1,1), const float Scale = 1.0f, const FName Name = TEXT("Scale")) const;
+	UFUNCTION(BlueprintPure = false, Category = "LearningAgents", meta = (ReturnDisplayName = "Success"))
+	bool GetScaleAction(FVector& OutScale, const FLearningAgentsActionObjectElement Element, const FVector RelativeScale = FVector(1,1,1), const float Scale = 1.0f, const FName Name = TEXT("Scale")) const;
 
-	UFUNCTION(BlueprintPure = false, Category = "LearningAgents")
-	UPARAM(DisplayName = "Success") bool GetTransformAction(FTransform& OutTransform, const FLearningAgentsActionObjectElement Element, const FTransform RelativeTransform = FTransform(), const float LocationScale = 100.0f, const float RotationScale = 1.0f, const float ScaleScale = 1.0f, const FName Name = TEXT("Transform")) const;
+	UFUNCTION(BlueprintPure = false, Category = "LearningAgents", meta = (ReturnDisplayName = "Success"))
+	bool GetTransformAction(FTransform& OutTransform, const FLearningAgentsActionObjectElement Element, const FTransform RelativeTransform = FTransform(), const float LocationScale = 100.0f, const float RotationScale = 1.0f, const float ScaleScale = 1.0f, const FName Name = TEXT("Transform")) const;
 
-	UFUNCTION(BlueprintPure = false, Category = "LearningAgents")
-	UPARAM(DisplayName = "Success") bool GetAngleAction(float& OutAngle, const FLearningAgentsActionObjectElement Element, const float RelativeAngle = 0.0f, const float AngleScale = 90.0f, const FName Name = TEXT("Angle")) const;
+	UFUNCTION(BlueprintPure = false, Category = "LearningAgents", meta = (ReturnDisplayName = "Success"))
+	bool GetAngleAction(float& OutAngle, const FLearningAgentsActionObjectElement Element, const float RelativeAngle = 0.0f, const float AngleScale = 90.0f, const FName Name = TEXT("Angle")) const;
 
-	UFUNCTION(BlueprintPure = false, Category = "LearningAgents")
-	UPARAM(DisplayName = "Success") bool GetAngleActionRadians(float& OutAngle, const FLearningAgentsActionObjectElement Element, const float RelativeAngle = 0.0f, const float AngleScale = 1.57079632679f, const FName Name = TEXT("Angle")) const;
+	UFUNCTION(BlueprintPure = false, Category = "LearningAgents", meta = (ReturnDisplayName = "Success"))
+	bool GetAngleActionRadians(float& OutAngle, const FLearningAgentsActionObjectElement Element, const float RelativeAngle = 0.0f, const float AngleScale = 1.57079632679f, const FName Name = TEXT("Angle")) const;
 
-	UFUNCTION(BlueprintPure = false, Category = "LearningAgents")
-	UPARAM(DisplayName = "Success") bool GetVelocityAction(FVector& OutVelocity, const FLearningAgentsActionObjectElement Element, const FTransform RelativeTransform = FTransform(), const float VelocityScale = 200.0f, const FName Name = TEXT("Velocity")) const;
+	UFUNCTION(BlueprintPure = false, Category = "LearningAgents", meta = (ReturnDisplayName = "Success"))
+	bool GetVelocityAction(FVector& OutVelocity, const FLearningAgentsActionObjectElement Element, const FTransform RelativeTransform = FTransform(), const float VelocityScale = 200.0f, const FName Name = TEXT("Velocity")) const;
 
-	UFUNCTION(BlueprintPure = false, Category = "LearningAgents")
-	UPARAM(DisplayName = "Success") bool GetDirectionAction(FVector& OutDirection, const FLearningAgentsActionObjectElement Element, const FTransform RelativeTransform = FTransform(), const FName Name = TEXT("Direction")) const;
+	UFUNCTION(BlueprintPure = false, Category = "LearningAgents", meta = (ReturnDisplayName = "Success"))
+	bool GetDirectionAction(FVector& OutDirection, const FLearningAgentsActionObjectElement Element, const FTransform RelativeTransform = FTransform(), const FName Name = TEXT("Direction")) const;
 
-	UFUNCTION(BlueprintPure = false, Category = "LearningAgents")
-	UPARAM(DisplayName = "Success") bool GetSpeedAction(float& OutSpeed, const FLearningAgentsActionObjectElement Element, const float SpeedScale = 200.0f, const FName Name = TEXT("Speed")) const;
+	UFUNCTION(BlueprintPure = false, Category = "LearningAgents", meta = (ReturnDisplayName = "Success"))
+	bool GetSpeedAction(float& OutSpeed, const FLearningAgentsActionObjectElement Element, const float SpeedScale = 200.0f, const FName Name = TEXT("Speed")) const;
 
 private:
 
diff --git a/Engine/Plugins/Experimental/LearningAgents/Source/LearningAgents/Public/LearningAgentsDebug.h b/Engine/Plugins/Experimental/LearningAgents/Source/LearningAgents/Public/LearningAgentsDebug.h
new file mode 100644
--- /dev/null
+++ b/Engine/Plugins/Experimental/LearningAgents/Source/LearningAgents/Public/LearningAgentsDebug.h
@@ -0,0 +1,56 @@
+// Copyright Epic Games, Inc. All Rights Reserved.
+
+#pragma once
+
+#include "VisualLogger/VisualLogger.h"
+
+#ifndef UE_LEARNING_AGENTS_ENABLE_VISUAL_LOG
+#define UE_LEARNING_AGENTS_ENABLE_VISUAL_LOG ENABLE_VISUAL_LOG
+#endif
+
+#if UE_LEARNING_AGENTS_ENABLE_VISUAL_LOG
+
+#define UE_LEARNING_AGENTS_VLOG_STRING(Owner, Category, Verbosity, Location, Color, Format, ...) \
+	UE_VLOG_LOCATION(Owner, Category, Verbosity, Location, 0.0f, Color, Format, ##__VA_ARGS__)
+
+#define UE_LEARNING_AGENTS_VLOG_LOCATION(Owner, Category, Verbosity, Location, Radius, Color, Format, ...) \
+	UE_VLOG_LOCATION(Owner, Category, Verbosity, Location, Radius, Color, Format, ##__VA_ARGS__)
+
+#define UE_LEARNING_AGENTS_VLOG_ARROW(Owner, Category, Verbosity, Start, End, Color, Format, ...) \
+	UE_VLOG_ARROW(Owner, Category, Verbosity, Start, End, Color, Format, ##__VA_ARGS__)
+
+#define UE_LEARNING_AGENTS_VLOG_SEGMENT(Owner, Category, Verbosity, Start, End, Color, Format, ...) \
+	UE_VLOG_SEGMENT(Owner, Category, Verbosity, Start, End, Color, Format, ##__VA_ARGS__)
+
+#define UE_LEARNING_AGENTS_VLOG_CIRCLE(Owner, Category, Verbosity, Center, UpAxis, Radius, Color, Format, ...) \
+	UE_VLOG_CIRCLE(Owner, Category, Verbosity, Center, UpAxis, Radius, Color, Format, ##__VA_ARGS__)
+
+#define UE_LEARNING_AGENTS_VLOG_OBOX(Owner, Category, Verbosity, Box, Matrix, Color, Format, ...) \
+	UE_VLOG_OBOX(Owner, Category, Verbosity, Box, Matrix, Color, Format, ##__VA_ARGS__)
+
+#define UE_LEARNING_AGENTS_VLOG_MATRIX(Owner, Category, Verbosity, Matrix, Color, Format, ...) \
+	{ \
+		UE_LEARNING_AGENTS_VLOG_SEGMENT(Owner, Category, Verbosity, Matrix.TransformPosition(FVector::ZeroVector), Matrix.TransformPosition(15.0f * FVector::ForwardVector), FColor::Red, TEXT("")); \
+		UE_LEARNING_AGENTS_VLOG_SEGMENT(Owner, Category, Verbosity, Matrix.TransformPosition(FVector::ZeroVector), Matrix.TransformPosition(15.0f * FVector::RightVector), FColor::Green, TEXT("")); \
+		UE_LEARNING_AGENTS_VLOG_SEGMENT(Owner, Category, Verbosity, Matrix.TransformPosition(FVector::ZeroVector), Matrix.TransformPosition(15.0f * FVector::UpVector), FColor::Blue, TEXT("")); \
+		UE_LEARNING_AGENTS_VLOG_STRING(Owner, Category, Verbosity, Matrix.TransformPosition(FVector::ZeroVector) + FVector(0.0f, 0.0f, 20.0f), Color, Format, ##__VA_ARGS__); \
+	}
+
+#define UE_LEARNING_AGENTS_VLOG_TRANSFORM(Owner, Category, Verbosity, Location, Rotation, Color, Format, ...) \
+	{ \
+		UE_LEARNING_AGENTS_VLOG_SEGMENT(Owner, Category, Verbosity, Location, Location + 15.0f * Rotation.RotateVector(FVector::ForwardVector), FColor::Red, TEXT("")); \
+		UE_LEARNING_AGENTS_VLOG_SEGMENT(Owner, Category, Verbosity, Location, Location + 15.0f * Rotation.RotateVector(FVector::RightVector), FColor::Green, TEXT("")); \
+		UE_LEARNING_AGENTS_VLOG_SEGMENT(Owner, Category, Verbosity, Location, Location + 15.0f * Rotation.RotateVector(FVector::UpVector), FColor::Blue, TEXT("")); \
+		UE_LEARNING_AGENTS_VLOG_OBOX(Owner, Category, Verbosity, FBox(5.0f * FVector(-1, -1, -1), 5.0f * FVector(1, 1, 1)), FTransform(Rotation, Location, FVector::OneVector).ToMatrixNoScale(), Color, TEXT("")); \
+		UE_LEARNING_AGENTS_VLOG_STRING(Owner, Category, Verbosity, Location + FVector(0.0f, 0.0f, 20.0f), Color, Format, ##__VA_ARGS__); \
+	}
+
+#define UE_LEARNING_AGENTS_VLOG_ANGLE(Owner, Category, Verbosity, Angle, RelativeAngle, Location, Radius, Color, Format, ...) \
+	{ \
+		UE_LEARNING_AGENTS_VLOG_CIRCLE(Owner, Category, Verbosity, Location, FVector::UpVector, Radius, Color, TEXT("")); \
+		UE_LEARNING_AGENTS_VLOG_SEGMENT(Owner, Category, Verbosity, Location, Location + Radius * FVector(FMath::Sin(RelativeAngle), FMath::Cos(RelativeAngle), 0.0f), Color, TEXT("")); \
+		UE_LEARNING_AGENTS_VLOG_SEGMENT(Owner, Category, Verbosity, Location, Location + Radius * FVector(FMath::Sin(Angle), FMath::Cos(Angle), 0.0f), Color, TEXT("")); \
+		UE_LEARNING_AGENTS_VLOG_LOCATION(Owner, Category, Verbosity, Location + Radius * FVector(FMath::Sin(Angle), FMath::Cos(Angle), 0.0f), static_cast<uint16>(Radius / 20.0f), Color, Format, ##__VA_ARGS__); \
+	}
+
+#endif
diff --git a/Engine/Plugins/Experimental/LearningAgents/Source/LearningAgents/Public/LearningAgentsManagerListener.h b/Engine/Plugins/Experimental/LearningAgents/Source/LearningAgents/Public/LearningAgentsManagerListener.h
--- a/Engine/Plugins/Experimental/LearningAgents/Source/LearningAgents/Public/LearningAgentsManagerListener.h
+++ b/Engine/Plugins/Experimental/LearningAgents/Source/LearningAgents/Public/LearningAgentsManagerListener.h
@@ -13,6 +13,13 @@
 
 class ULearningAgentsManager;
 
+/** Dummy class used for visual logging */
+UCLASS(BlueprintType)
+class LEARNINGAGENTS_API ULearningAgentsVisualLoggerObject : public UObject
+{
+	GENERATED_BODY()
+};
+
 /**
  * Base class for objects which can be added to a ULearningAgentsManager to receive callbacks whenever agents are added, remove or reset.
  */
@@ -122,6 +129,9 @@ public:
 	/** Checks if the object has an associated agent manager. */
 	bool HasAgentManager() const;
 
+	/** Either gets or adds a visual logger object for the given name */
+	const ULearningAgentsVisualLoggerObject* GetOrAddVisualLoggerObject(const FName Name);
+
 protected:
 
 	/** True if this object has been setup. Otherwise, false. */
@@ -131,4 +141,8 @@ protected:
 	/** The manager this object is associated with. */
 	UPROPERTY(VisibleAnywhere, Transient, Category = "LearningAgents")
 	TObjectPtr<ULearningAgentsManager> Manager;
+
+	/** The visual logger objects associated with this listener. */
+	UPROPERTY(VisibleAnywhere, Transient, Category = "LearningAgents")
+	TMap<FName, TObjectPtr<const ULearningAgentsVisualLoggerObject>> VisualLoggerObjects;
 };
diff --git a/Engine/Plugins/Experimental/LearningAgents/Source/LearningAgents/Public/LearningAgentsObservations.h b/Engine/Plugins/Experimental/LearningAgents/Source/LearningAgents/Public/LearningAgentsObservations.h
--- a/Engine/Plugins/Experimental/LearningAgents/Source/LearningAgents/Public/LearningAgentsObservations.h
+++ b/Engine/Plugins/Experimental/LearningAgents/Source/LearningAgents/Public/LearningAgentsObservations.h
@@ -12,6 +12,7 @@
 
 #include "LearningAgentsObservations.generated.h"
 
+class ULearningAgentsManagerListener;
 class ULearningAgentsObservationSchema;
 class ULearningAgentsObservationObject;
 struct FLearningAgentsObservationSchemaElement;
@@ -76,7 +77,7 @@ class ULearningAgentsObservationFunctions : public UBlueprintFunctionLibrary
 
 	/** Find an Enum type by Name. This can be used to find Enum types defined in C++. This call can be expensive so the result should be cached. */
 	UFUNCTION(BlueprintCallable, Category = "LearningAgents")
-	static UEnum* FindEnumByName(const FString Name);
+	static UEnum* FindEnumByName(const FString& Name);
 };
 
 /**
@@ -254,12 +255,6 @@ private:
 	UE::Learning::Observation::FSchema ObservationSchema;
 };
 
-UCLASS(BlueprintType)
-class LEARNINGAGENTS_API ULearningAgentsObservationVisualLoggerObject : public UObject
-{
-	GENERATED_BODY()
-};
-
 /**
  * Observation Object
  *
@@ -382,15 +377,19 @@ public:
 	 * @param Value The new value of this observation.
 	 * @param Name The name of the corresponding observation. Must match the name given during Specify.
 	 * @param bVisualLoggerEnabled When true, debug data will be sent to the visual logger.
+	 * @param VisualLoggerListener The listener object which is making this observation. This must be set to use logging.
+	 * @param VisualLoggerAgentId The agent id associated with this observation.
 	 * @param VisualLoggerLocation A location for the visual logger information in the world.
 	 * @param VisualLoggerColor The color for the visual logger display.
 	 * @return The newly created observation object element.
 	 */
-	UFUNCTION(BlueprintPure, Category = "LearningAgents", meta = (AdvancedDisplay = 2))
+	UFUNCTION(BlueprintPure, Category = "LearningAgents", meta = (AdvancedDisplay = 2, DefaultToSelf = "VisualLoggerListener"))
 	FLearningAgentsObservationObjectElement MakeBoolObservation(
 		const bool bValue,
 		const FName Name = TEXT("Bool"),
 		const bool bVisualLoggerEnabled = false,
+		ULearningAgentsManagerListener* VisualLoggerListener = nullptr,
+		const int32 VisualLoggerAgentId = -1,
 		const FVector VisualLoggerLocation = FVector::ZeroVector,
 		const FLinearColor VisualLoggerColor = FLinearColor::Red);
 
@@ -401,16 +400,20 @@ public:
 	 * @param FloatScale Used to normalize the data for this observation.
 	 * @param Name The name of the corresponding observation. Must match the name given during Specify.
 	 * @param bVisualLoggerEnabled When true, debug data will be sent to the visual logger.
+	 * @param VisualLoggerListener The listener object which is making this observation. This must be set to use logging.
+	 * @param VisualLoggerAgentId The agent id associated with this observation.
 	 * @param VisualLoggerLocation A location for the visual logger information in the world.
 	 * @param VisualLoggerColor The color for the visual logger display.
 	 * @return The newly created observation object element.
 	 */
-	UFUNCTION(BlueprintPure, Category = "LearningAgents", meta = (AdvancedDisplay = 3))
+	UFUNCTION(BlueprintPure, Category = "LearningAgents", meta = (AdvancedDisplay = 3, DefaultToSelf = "VisualLoggerListener"))
 	FLearningAgentsObservationObjectElement MakeFloatObservation(
 		const float Value,
 		const float FloatScale = 1.0f,
 		const FName Name = TEXT("Float"),
 		const bool bVisualLoggerEnabled = false,
+		ULearningAgentsManagerListener* VisualLoggerListener = nullptr,
+		const int32 VisualLoggerAgentId = -1,
 		const FVector VisualLoggerLocation = FVector::ZeroVector,
 		const FLinearColor VisualLoggerColor = FLinearColor::Red);
 
@@ -422,17 +425,21 @@ public:
 	 * @param LocationScale Used to normalize the data for this observation.
 	 * @param Name The name of the corresponding observation. Must match the name given during Specify.
 	 * @param bVisualLoggerEnabled When true, debug data will be sent to the visual logger.
+	 * @param VisualLoggerListener The listener object which is making this observation. This must be set to use logging.
+	 * @param VisualLoggerAgentId The agent id associated with this observation.
 	 * @param VisualLoggerLocation A location for the visual logger information in the world.
 	 * @param VisualLoggerColor The color for the visual logger display.
 	 * @return The newly created observation object element.
 	 */
-	UFUNCTION(BlueprintPure, Category = "LearningAgents", meta = (AdvancedDisplay = 4))
+	UFUNCTION(BlueprintPure, Category = "LearningAgents", meta = (AdvancedDisplay = 4, DefaultToSelf = "VisualLoggerListener"))
 	FLearningAgentsObservationObjectElement MakeLocationObservation(
 		const FVector Location,
 		const FTransform RelativeTransform = FTransform(),
 		const float LocationScale = 100.0f,
 		const FName Name = TEXT("Location"),
 		const bool bVisualLoggerEnabled = false,
+		ULearningAgentsManagerListener* VisualLoggerListener = nullptr,
+		const int32 VisualLoggerAgentId = -1,
 		const FVector VisualLoggerLocation = FVector::ZeroVector,
 		const FLinearColor VisualLoggerColor = FLinearColor::Red);
 
@@ -443,16 +450,22 @@ public:
 	 * @param RelativeRotation The rotation the provided rotation should be encoded relative to.
 	 * @param Name The name of the corresponding observation. Must match the name given during Specify.
 	 * @param bVisualLoggerEnabled When true, debug data will be sent to the visual logger.
+	 * @param VisualLoggerListener The listener object which is making this observation. This must be set to use logging.
+	 * @param VisualLoggerAgentId The agent id associated with this observation.
+	 * @param VisualLoggerRotationLocation A location for the visual logger to display the rotation in the world.
 	 * @param VisualLoggerLocation A location for the visual logger information in the world.
 	 * @param VisualLoggerColor The color for the visual logger display.
 	 * @return The newly created observation object element.
 	 */
-	UFUNCTION(BlueprintPure, Category = "LearningAgents", meta = (AdvancedDisplay = 3))
+	UFUNCTION(BlueprintPure, Category = "LearningAgents", meta = (AdvancedDisplay = 3, DefaultToSelf = "VisualLoggerListener"))
 	FLearningAgentsObservationObjectElement MakeRotationObservation(
 		const FRotator Rotation,
 		const FRotator RelativeRotation = FRotator::ZeroRotator,
 		const FName Name = TEXT("Rotation"),
 		const bool bVisualLoggerEnabled = false,
+		ULearningAgentsManagerListener* VisualLoggerListener = nullptr,
+		const int32 VisualLoggerAgentId = -1,
+		const FVector VisualLoggerRotationLocation = FVector::ZeroVector,
 		const FVector VisualLoggerLocation = FVector::ZeroVector,
 		const FLinearColor VisualLoggerColor = FLinearColor::Red);
 
@@ -463,16 +476,22 @@ public:
 	 * @param RelativeRotation The rotation the provided rotation should be encoded relative to.
 	 * @param Name The name of the corresponding observation. Must match the name given during Specify.
 	 * @param bVisualLoggerEnabled When true, debug data will be sent to the visual logger.
+	 * @param VisualLoggerListener The listener object which is making this observation. This must be set to use logging.
+	 * @param VisualLoggerAgentId The agent id associated with this observation.
+	 * @param VisualLoggerRotationLocation A location for the visual logger to display the rotation in the world.
 	 * @param VisualLoggerLocation A location for the visual logger information in the world.
 	 * @param VisualLoggerColor The color for the visual logger display.
 	 * @return The newly created observation object element.
 	 */
-	UFUNCTION(BlueprintPure, Category = "LearningAgents", meta = (AdvancedDisplay = 3))
+	UFUNCTION(BlueprintPure, Category = "LearningAgents", meta = (AdvancedDisplay = 3, DefaultToSelf = "VisualLoggerListener"))
 	FLearningAgentsObservationObjectElement MakeRotationObservationFromQuat(
 		const FQuat Rotation,
 		const FQuat RelativeRotation,
 		const FName Name = TEXT("Rotation"),
 		const bool bVisualLoggerEnabled = false,
+		ULearningAgentsManagerListener* VisualLoggerListener = nullptr,
+		const int32 VisualLoggerAgentId = -1,
+		const FVector VisualLoggerRotationLocation = FVector::ZeroVector,
 		const FVector VisualLoggerLocation = FVector::ZeroVector,
 		const FLinearColor VisualLoggerColor = FLinearColor::Red);
 
@@ -483,16 +502,22 @@ public:
 	 * @param RelativeScale The scale the provided scale should be encoded relative to.
 	 * @param Name The name of the corresponding observation. Must match the name given during Specify.
 	 * @param bVisualLoggerEnabled When true, debug data will be sent to the visual logger.
+	 * @param VisualLoggerListener The listener object which is making this observation. This must be set to use logging.
+	 * @param VisualLoggerAgentId The agent id associated with this observation.
+	 * @param VisualLoggerScaleLocation A location for the visual logger to display the scale in the world.
 	 * @param VisualLoggerLocation A location for the visual logger information in the world.
 	 * @param VisualLoggerColor The color for the visual logger display.
 	 * @return The newly created observation object element.
 	 */
-	UFUNCTION(BlueprintPure, Category = "LearningAgents", meta = (AdvancedDisplay = 3))
+	UFUNCTION(BlueprintPure, Category = "LearningAgents", meta = (AdvancedDisplay = 3, DefaultToSelf = "VisualLoggerListener"))
 	FLearningAgentsObservationObjectElement MakeScaleObservation(
 		const FVector Scale,
 		const FVector RelativeScale = FVector(1,1,1),
 		const FName Name = TEXT("Scale"),
 		const bool bVisualLoggerEnabled = false,
+		ULearningAgentsManagerListener* VisualLoggerListener = nullptr,
+		const int32 VisualLoggerAgentId = -1,
+		const FVector VisualLoggerScaleLocation = FVector::ZeroVector,
 		const FVector VisualLoggerLocation = FVector::ZeroVector,
 		const FLinearColor VisualLoggerColor = FLinearColor::Red);
 
@@ -504,17 +529,21 @@ public:
 	 * @param LocationScale Used to normalize the transform's location for this observation.
 	 * @param Name The name of the corresponding observation. Must match the name given during Specify.
 	 * @param bVisualLoggerEnabled When true, debug data will be sent to the visual logger.
+	 * @param VisualLoggerListener The listener object which is making this observation. This must be set to use logging.
+	 * @param VisualLoggerAgentId The agent id associated with this observation.
 	 * @param VisualLoggerLocation A location for the visual logger information in the world.
 	 * @param VisualLoggerColor The color for the visual logger display.
 	 * @return The newly created observation object element.
 	 */
-	UFUNCTION(BlueprintPure, Category = "LearningAgents", meta = (AdvancedDisplay = 4))
+	UFUNCTION(BlueprintPure, Category = "LearningAgents", meta = (AdvancedDisplay = 4, DefaultToSelf = "VisualLoggerListener"))
 	FLearningAgentsObservationObjectElement MakeTransformObservation(
 		const FTransform Transform,
 		const FTransform RelativeTransform = FTransform(),
 		const float LocationScale = 100.0f,
 		const FName Name = TEXT("Transform"),
 		const bool bVisualLoggerEnabled = false,
+		ULearningAgentsManagerListener* VisualLoggerListener = nullptr,
+		const int32 VisualLoggerAgentId = -1,
 		const FVector VisualLoggerLocation = FVector::ZeroVector,
 		const FLinearColor VisualLoggerColor = FLinearColor::Red);
 
@@ -527,19 +556,22 @@ public:
 	UFUNCTION(BlueprintPure, Category = "LearningAgents")
 	FLearningAgentsObservationObjectElement MakeVelocityObservation(const FVector Velocity, const FTransform RelativeTransform = FTransform(), const float VelocityScale = 200.0f, const FName Name = TEXT("Velocity"));
 
-	UFUNCTION(BlueprintPure, Category = "LearningAgents", meta = (AdvancedDisplay = 3))
+	UFUNCTION(BlueprintPure, Category = "LearningAgents", meta = (AdvancedDisplay = 3, DefaultToSelf = "VisualLoggerListener"))
 	FLearningAgentsObservationObjectElement MakeDirectionObservation(
 		const FVector Direction,
 		const FTransform RelativeTransform = FTransform(),
 		const FName Name = TEXT("Direction"),
 		const bool bVisualLoggerEnabled = false,
+		ULearningAgentsManagerListener* VisualLoggerListener = nullptr,
+		const int32 VisualLoggerAgentId = -1,
+		const FVector VisualLoggerDirectionLocation = FVector::ZeroVector,
 		const FVector VisualLoggerLocation = FVector::ZeroVector,
 		const float VisualLoggerArrowLength = 100.0f,
 		const FLinearColor VisualLoggerColor = FLinearColor::Red);
 
 	// Spline Observations
 
-	UFUNCTION(BlueprintPure, Category = "LearningAgents", meta = (AdvancedDisplay = 5))
+	UFUNCTION(BlueprintPure, Category = "LearningAgents", meta = (AdvancedDisplay = 5, DefaultToSelf = "VisualLoggerListener"))
 	FLearningAgentsObservationObjectElement MakeLocationAlongSplineObservation(
 		const USplineComponent* SplineComponent,
 		const float DistanceAlongSpline,
@@ -547,19 +579,23 @@ public:
 		const float LocationScale = 100.0f,
 		const FName Name = TEXT("LocationAlongSpline"),
 		const bool bVisualLoggerEnabled = false,
+		ULearningAgentsManagerListener* VisualLoggerListener = nullptr,
+		const int32 VisualLoggerAgentId = -1,
 		const FVector VisualLoggerLocation = FVector::ZeroVector,
 		const FLinearColor VisualLoggerColor = FLinearColor::Red);
 
 	UFUNCTION(BlueprintPure, Category = "LearningAgents")
 	FLearningAgentsObservationObjectElement MakeProportionAlongSplineObservation(const USplineComponent* SplineComponent, const float DistanceAlongSpline, const FName Name = TEXT("ProportionAlongSpline"));
 
-	UFUNCTION(BlueprintPure, Category = "LearningAgents", meta = (AdvancedDisplay = 4))
+	UFUNCTION(BlueprintPure, Category = "LearningAgents", meta = (AdvancedDisplay = 4, DefaultToSelf = "VisualLoggerListener"))
 	FLearningAgentsObservationObjectElement MakeDirectionAlongSplineObservation(
 		const USplineComponent* SplineComponent,
 		const float DistanceAlongSpline,
 		const FTransform RelativeTransform = FTransform(),
 		const FName Name = TEXT("DirectionAlongSpline"),
 		const bool bVisualLoggerEnabled = false,
+		ULearningAgentsManagerListener* VisualLoggerListener = nullptr,
+		const int32 VisualLoggerAgentId = -1,
 		const FVector VisualLoggerLocation = FVector::ZeroVector,
 		const float VisualLoggerArrowLength = 100.0f,
 		const FLinearColor VisualLoggerColor = FLinearColor::Red);
@@ -579,171 +615,166 @@ public:
 
 	// Basic Observations
 
-	UFUNCTION(BlueprintPure = false, Category = "LearningAgents")
-	UPARAM(DisplayName = "Success") bool GetNullObservation(const FLearningAgentsObservationObjectElement Element, const FName Name = TEXT("Null")) const;
+	UFUNCTION(BlueprintPure = false, Category = "LearningAgents", meta=(ReturnDisplayName = "Success"))
+	bool GetNullObservation(const FLearningAgentsObservationObjectElement Element, const FName Name = TEXT("Null")) const;
 
-	UFUNCTION(BlueprintPure = false, Category = "LearningAgents")
-	UPARAM(DisplayName = "Success") bool GetContinuousObservationNum(int32& OutNum, const FLearningAgentsObservationObjectElement Element, const FName Name = TEXT("Continuous")) const;
+	UFUNCTION(BlueprintPure = false, Category = "LearningAgents", meta = (ReturnDisplayName = "Success"))
+	bool GetContinuousObservationNum(int32& OutNum, const FLearningAgentsObservationObjectElement Element, const FName Name = TEXT("Continuous")) const;
 
-	UFUNCTION(BlueprintPure = false, Category = "LearningAgents")
-	UPARAM(DisplayName = "Success") bool GetContinuousObservation(TArray<float>& OutValues, const FLearningAgentsObservationObjectElement Element, const FName Name = TEXT("Continuous")) const;
-	UPARAM(DisplayName = "Success") bool GetContinuousObservationToArrayView(TArrayView<float> OutValues, const FLearningAgentsObservationObjectElement Element, const FName Name = TEXT("Continuous")) const;
+	UFUNCTION(BlueprintPure = false, Category = "LearningAgents", meta = (ReturnDisplayName = "Success"))
+	bool GetContinuousObservation(TArray<float>& OutValues, const FLearningAgentsObservationObjectElement Element, const FName Name = TEXT("Continuous")) const;
+	bool GetContinuousObservationToArrayView(TArrayView<float> OutValues, const FLearningAgentsObservationObjectElement Element, const FName Name = TEXT("Continuous")) const;
 
-	UFUNCTION(BlueprintPure = false, Category = "LearningAgents")
-	UPARAM(DisplayName = "Success") bool GetExclusiveDiscreteObservation(int32& OutIndex, const FLearningAgentsObservationObjectElement Element, const FName Name = TEXT("ExclusiveDiscrete")) const;
+	UFUNCTION(BlueprintPure = false, Category = "LearningAgents", meta = (ReturnDisplayName = "Success"))
+	bool GetExclusiveDiscreteObservation(int32& OutIndex, const FLearningAgentsObservationObjectElement Element, const FName Name = TEXT("ExclusiveDiscrete")) const;
 
-	UFUNCTION(BlueprintPure = false, Category = "LearningAgents")
-	UPARAM(DisplayName = "Success") bool GetInclusiveDiscreteObservationNum(int32& OutNum, const FLearningAgentsObservationObjectElement Element, const FName Name = TEXT("InclusiveDiscrete")) const;
+	UFUNCTION(BlueprintPure = false, Category = "LearningAgents", meta = (ReturnDisplayName = "Success"))
+	bool GetInclusiveDiscreteObservationNum(int32& OutNum, const FLearningAgentsObservationObjectElement Element, const FName Name = TEXT("InclusiveDiscrete")) const;
 
-	UFUNCTION(BlueprintPure = false, Category = "LearningAgents")
-	UPARAM(DisplayName = "Success") bool GetInclusiveDiscreteObservation(TArray<int32>& OutIndices, const FLearningAgentsObservationObjectElement Element, const FName Name = TEXT("InclusiveDiscrete")) const;
-	UPARAM(DisplayName = "Success") bool GetInclusiveDiscreteObservationToArrayView(TArrayView<int32> OutIndices, const FLearningAgentsObservationObjectElement Element, const FName Name = TEXT("InclusiveDiscrete")) const;
+	UFUNCTION(BlueprintPure = false, Category = "LearningAgents", meta = (ReturnDisplayName = "Success"))
+	bool GetInclusiveDiscreteObservation(TArray<int32>& OutIndices, const FLearningAgentsObservationObjectElement Element, const FName Name = TEXT("InclusiveDiscrete")) const;
+	bool GetInclusiveDiscreteObservationToArrayView(TArrayView<int32> OutIndices, const FLearningAgentsObservationObjectElement Element, const FName Name = TEXT("InclusiveDiscrete")) const;
 
-	UFUNCTION(BlueprintPure = false, Category = "LearningAgents")
-	UPARAM(DisplayName = "Success") bool GetIndexObservation(int32& OutIndex, const FLearningAgentsObservationObjectElement Element, const FName Name = TEXT("Index")) const;
+	UFUNCTION(BlueprintPure = false, Category = "LearningAgents", meta = (ReturnDisplayName = "Success"))
+	bool GetIndexObservation(int32& OutIndex, const FLearningAgentsObservationObjectElement Element, const FName Name = TEXT("Index")) const;
 
-	UFUNCTION(BlueprintPure = false, Category = "LearningAgents")
-	UPARAM(DisplayName = "Success") bool GetCountObservation(int32& OutNum, const FLearningAgentsObservationObjectElement Element, const int32 MaxNum, const FName Name = TEXT("Count")) const;
+	UFUNCTION(BlueprintPure = false, Category = "LearningAgents", meta = (ReturnDisplayName = "Success"))
+	bool GetCountObservation(int32& OutNum, const FLearningAgentsObservationObjectElement Element, const int32 MaxNum, const FName Name = TEXT("Count")) const;
 
-	UFUNCTION(BlueprintPure = false, Category = "LearningAgents")
-	UPARAM(DisplayName = "Success") bool GetStructObservationNum(int32& OutNum, const FLearningAgentsObservationObjectElement Element, const FName Name = TEXT("Struct")) const;
+	UFUNCTION(BlueprintPure = false, Category = "LearningAgents", meta = (ReturnDisplayName = "Success"))
+	bool GetStructObservationNum(int32& OutNum, const FLearningAgentsObservationObjectElement Element, const FName Name = TEXT("Struct")) const;
 
-	UFUNCTION(BlueprintPure = false, Category = "LearningAgents")
-	UPARAM(DisplayName = "Success") bool GetStructObservation(TMap<FName, FLearningAgentsObservationObjectElement>& OutElements, const FLearningAgentsObservationObjectElement Element, const FName Name = TEXT("Struct")) const;
+	UFUNCTION(BlueprintPure = false, Category = "LearningAgents", meta = (ReturnDisplayName = "Success"))
+	bool GetStructObservation(TMap<FName, FLearningAgentsObservationObjectElement>& OutElements, const FLearningAgentsObservationObjectElement Element, const FName Name = TEXT("Struct")) const;
 	
-	UFUNCTION(BlueprintPure = false, Category = "LearningAgents")
-	UPARAM(DisplayName = "Success") bool GetStructObservationToArrays(TArray<FName>& OutElementNames, TArray<FLearningAgentsObservationObjectElement>& OutElements, const FLearningAgentsObservationObjectElement Element, const FName Name = TEXT("Struct")) const;
-	UPARAM(DisplayName = "Success") bool GetStructObservationToArrayViews(TArrayView<FName> OutElementNames, TArrayView<FLearningAgentsObservationObjectElement> OutElements, const FLearningAgentsObservationObjectElement Element, const FName Name = TEXT("Struct")) const;
+	UFUNCTION(BlueprintPure = false, Category = "LearningAgents", meta = (ReturnDisplayName = "Success"))
+	bool GetStructObservationToArrays(TArray<FName>& OutElementNames, TArray<FLearningAgentsObservationObjectElement>& OutElements, const FLearningAgentsObservationObjectElement Element, const FName Name = TEXT("Struct")) const;
+	bool GetStructObservationToArrayViews(TArrayView<FName> OutElementNames, TArrayView<FLearningAgentsObservationObjectElement> OutElements, const FLearningAgentsObservationObjectElement Element, const FName Name = TEXT("Struct")) const;
 
-	UFUNCTION(BlueprintPure = false, Category = "LearningAgents")
-	UPARAM(DisplayName = "Success") bool GetExclusiveUnionObservation(FName& OutElementName, FLearningAgentsObservationObjectElement& OutElement, const FLearningAgentsObservationObjectElement Element, const FName Name = TEXT("ExclusiveUnion")) const;
+	UFUNCTION(BlueprintPure = false, Category = "LearningAgents", meta = (ReturnDisplayName = "Success"))
+	bool GetExclusiveUnionObservation(FName& OutElementName, FLearningAgentsObservationObjectElement& OutElement, const FLearningAgentsObservationObjectElement Element, const FName Name = TEXT("ExclusiveUnion")) const;
 
-	UFUNCTION(BlueprintPure = false, Category = "LearningAgents")
-	UPARAM(DisplayName = "Success") bool GetInclusiveUnionObservationNum(int32& OutNum, const FLearningAgentsObservationObjectElement Element, const FName Name = TEXT("InclusiveUnion")) const;
+	UFUNCTION(BlueprintPure = false, Category = "LearningAgents", meta = (ReturnDisplayName = "Success"))
+	bool GetInclusiveUnionObservationNum(int32& OutNum, const FLearningAgentsObservationObjectElement Element, const FName Name = TEXT("InclusiveUnion")) const;
 
-	UFUNCTION(BlueprintPure = false, Category = "LearningAgents")
-	UPARAM(DisplayName = "Success") bool GetInclusiveUnionObservation(TMap<FName, FLearningAgentsObservationObjectElement>& OutElements, const FLearningAgentsObservationObjectElement Element, const FName Name = TEXT("InclusiveUnion")) const;
+	UFUNCTION(BlueprintPure = false, Category = "LearningAgents", meta = (ReturnDisplayName = "Success"))
+	bool GetInclusiveUnionObservation(TMap<FName, FLearningAgentsObservationObjectElement>& OutElements, const FLearningAgentsObservationObjectElement Element, const FName Name = TEXT("InclusiveUnion")) const;
 	
-	UFUNCTION(BlueprintPure = false, Category = "LearningAgents")
-	UPARAM(DisplayName = "Success") bool GetInclusiveUnionObservationToArrays(TArray<FName>& OutElementNames, TArray<FLearningAgentsObservationObjectElement>& OutElements, const FLearningAgentsObservationObjectElement Element, const FName Name = TEXT("InclusiveUnion")) const;
-	UPARAM(DisplayName = "Success") bool GetInclusiveUnionObservationToArrayViews(TArrayView<FName> OutElementNames, TArrayView<FLearningAgentsObservationObjectElement> OutElements, const FLearningAgentsObservationObjectElement Element, const FName Name = TEXT("InclusiveUnion")) const;
+	UFUNCTION(BlueprintPure = false, Category = "LearningAgents", meta = (ReturnDisplayName = "Success"))
+	bool GetInclusiveUnionObservationToArrays(TArray<FName>& OutElementNames, TArray<FLearningAgentsObservationObjectElement>& OutElements, const FLearningAgentsObservationObjectElement Element, const FName Name = TEXT("InclusiveUnion")) const;
+	bool GetInclusiveUnionObservationToArrayViews(TArrayView<FName> OutElementNames, TArrayView<FLearningAgentsObservationObjectElement> OutElements, const FLearningAgentsObservationObjectElement Element, const FName Name = TEXT("InclusiveUnion")) const;
 
-	UFUNCTION(BlueprintPure = false, Category = "LearningAgents")
-	UPARAM(DisplayName = "Success") bool GetStaticArrayObservationNum(int32& OutNum, const FLearningAgentsObservationObjectElement Element, const FName Name = TEXT("StaticArray")) const;
+	UFUNCTION(BlueprintPure = false, Category = "LearningAgents", meta = (ReturnDisplayName = "Success"))
+	bool GetStaticArrayObservationNum(int32& OutNum, const FLearningAgentsObservationObjectElement Element, const FName Name = TEXT("StaticArray")) const;
 
-	UFUNCTION(BlueprintPure = false, Category = "LearningAgents")
-	UPARAM(DisplayName = "Success") bool GetStaticArrayObservation(TArray<FLearningAgentsObservationObjectElement>& OutElements, const FLearningAgentsObservationObjectElement Element, const FName Name = TEXT("StaticArray")) const;
-	UPARAM(DisplayName = "Success") bool GetStaticArrayObservationToArrayView(TArrayView<FLearningAgentsObservationObjectElement> OutElements, const FLearningAgentsObservationObjectElement Element, const FName Name = TEXT("StaticArray")) const;
+	UFUNCTION(BlueprintPure = false, Category = "LearningAgents", meta = (ReturnDisplayName = "Success"))
+	bool GetStaticArrayObservation(TArray<FLearningAgentsObservationObjectElement>& OutElements, const FLearningAgentsObservationObjectElement Element, const FName Name = TEXT("StaticArray")) const;
+	bool GetStaticArrayObservationToArrayView(TArrayView<FLearningAgentsObservationObjectElement> OutElements, const FLearningAgentsObservationObjectElement Element, const FName Name = TEXT("StaticArray")) const;
 
-	UFUNCTION(BlueprintPure = false, Category = "LearningAgents")
-	UPARAM(DisplayName = "Success") bool GetSetObservationNum(int32& OutNum, const FLearningAgentsObservationObjectElement Element, const FName Name = TEXT("Set")) const;
+	UFUNCTION(BlueprintPure = false, Category = "LearningAgents", meta = (ReturnDisplayName = "Success"))
+	bool GetSetObservationNum(int32& OutNum, const FLearningAgentsObservationObjectElement Element, const FName Name = TEXT("Set")) const;
 
-	UFUNCTION(BlueprintPure = false, Category = "LearningAgents")
-	UPARAM(DisplayName = "Success") bool GetSetObservation(TSet<FLearningAgentsObservationObjectElement>& OutElements, const FLearningAgentsObservationObjectElement Element, const FName Name = TEXT("Set")) const;
+	UFUNCTION(BlueprintPure = false, Category = "LearningAgents", meta = (ReturnDisplayName = "Success"))
+	bool GetSetObservation(TSet<FLearningAgentsObservationObjectElement>& OutElements, const FLearningAgentsObservationObjectElement Element, const FName Name = TEXT("Set")) const;
 	
-	UFUNCTION(BlueprintPure = false, Category = "LearningAgents")
-	UPARAM(DisplayName = "Success") bool GetSetObservationToArray(TArray<FLearningAgentsObservationObjectElement>& OutElements, const FLearningAgentsObservationObjectElement Element, const FName Name = TEXT("Set")) const;
-	UPARAM(DisplayName = "Success") bool GetSetObservationToArrayView(TArrayView<FLearningAgentsObservationObjectElement> OutElements, const FLearningAgentsObservationObjectElement Element, const FName Name = TEXT("Set")) const;
+	UFUNCTION(BlueprintPure = false, Category = "LearningAgents", meta = (ReturnDisplayName = "Success"))
+	bool GetSetObservationToArray(TArray<FLearningAgentsObservationObjectElement>& OutElements, const FLearningAgentsObservationObjectElement Element, const FName Name = TEXT("Set")) const;
+	bool GetSetObservationToArrayView(TArrayView<FLearningAgentsObservationObjectElement> OutElements, const FLearningAgentsObservationObjectElement Element, const FName Name = TEXT("Set")) const;
 
-	UFUNCTION(BlueprintPure = false, Category = "LearningAgents")
-	UPARAM(DisplayName = "Success") bool GetPairObservation(FLearningAgentsObservationObjectElement& OutKey, FLearningAgentsObservationObjectElement& OutValue, const FLearningAgentsObservationObjectElement Element, const FName Name = TEXT("Pair")) const;
+	UFUNCTION(BlueprintPure = false, Category = "LearningAgents", meta = (ReturnDisplayName = "Success"))
+	bool GetPairObservation(FLearningAgentsObservationObjectElement& OutKey, FLearningAgentsObservationObjectElement& OutValue, const FLearningAgentsObservationObjectElement Element, const FName Name = TEXT("Pair")) const;
 
-	UFUNCTION(BlueprintPure = false, Category = "LearningAgents")
-	UPARAM(DisplayName = "Success") bool GetArrayObservationNum(int32& OutNum, const FLearningAgentsObservationObjectElement Element, const FName Name = TEXT("Array")) const;
+	UFUNCTION(BlueprintPure = false, Category = "LearningAgents", meta = (ReturnDisplayName = "Success"))
+	bool GetArrayObservationNum(int32& OutNum, const FLearningAgentsObservationObjectElement Element, const FName Name = TEXT("Array")) const;
 
-	UFUNCTION(BlueprintPure = false, Category = "LearningAgents")
-	UPARAM(DisplayName = "Success") bool GetArrayObservation(TArray<FLearningAgentsObservationObjectElement>& OutElements, const FLearningAgentsObservationObjectElement Element, const FName Name = TEXT("Array")) const;
-	UPARAM(DisplayName = "Success") bool GetArrayObservationToArrayView(TArrayView<FLearningAgentsObservationObjectElement> OutElements, const FLearningAgentsObservationObjectElement Element, const FName Name = TEXT("Array")) const;
+	UFUNCTION(BlueprintPure = false, Category = "LearningAgents", meta = (ReturnDisplayName = "Success"))
+	bool GetArrayObservation(TArray<FLearningAgentsObservationObjectElement>& OutElements, const FLearningAgentsObservationObjectElement Element, const FName Name = TEXT("Array")) const;
+	bool GetArrayObservationToArrayView(TArrayView<FLearningAgentsObservationObjectElement> OutElements, const FLearningAgentsObservationObjectElement Element, const FName Name = TEXT("Array")) const;
 
-	UFUNCTION(BlueprintPure = false, Category = "LearningAgents")
-	UPARAM(DisplayName = "Success") bool GetMapObservationNum(int32& OutNum, const FLearningAgentsObservationObjectElement Element, const FName Name = TEXT("Map")) const;
+	UFUNCTION(BlueprintPure = false, Category = "LearningAgents", meta = (ReturnDisplayName = "Success"))
+	bool GetMapObservationNum(int32& OutNum, const FLearningAgentsObservationObjectElement Element, const FName Name = TEXT("Map")) const;
 
-	UFUNCTION(BlueprintPure = false, Category = "LearningAgents")
-	UPARAM(DisplayName = "Success") bool GetMapObservation(TMap<FLearningAgentsObservationObjectElement,FLearningAgentsObservationObjectElement>& OutElements, const FLearningAgentsObservationObjectElement Element, const FName Name = TEXT("Map")) const;
+	UFUNCTION(BlueprintPure = false, Category = "LearningAgents", meta = (ReturnDisplayName = "Success"))
+	bool GetMapObservation(TMap<FLearningAgentsObservationObjectElement,FLearningAgentsObservationObjectElement>& OutElements, const FLearningAgentsObservationObjectElement Element, const FName Name = TEXT("Map")) const;
 	
-	UFUNCTION(BlueprintPure = false, Category = "LearningAgents")
-	UPARAM(DisplayName = "Success") bool GetMapObservationToArrays(TArray<FLearningAgentsObservationObjectElement>& OutKeys, TArray<FLearningAgentsObservationObjectElement>& OutValues, const FLearningAgentsObservationObjectElement Element, const FName Name = TEXT("Map")) const;
-	UPARAM(DisplayName = "Success") bool GetMapObservationToArrayViews(TArrayView<FLearningAgentsObservationObjectElement> OutKeys, TArrayView<FLearningAgentsObservationObjectElement> OutValues, const FLearningAgentsObservationObjectElement Element, const FName Name = TEXT("Map")) const;
+	UFUNCTION(BlueprintPure = false, Category = "LearningAgents", meta = (ReturnDisplayName = "Success"))
+	bool GetMapObservationToArrays(TArray<FLearningAgentsObservationObjectElement>& OutKeys, TArray<FLearningAgentsObservationObjectElement>& OutValues, const FLearningAgentsObservationObjectElement Element, const FName Name = TEXT("Map")) const;
+	bool GetMapObservationToArrayViews(TArrayView<FLearningAgentsObservationObjectElement> OutKeys, TArrayView<FLearningAgentsObservationObjectElement> OutValues, const FLearningAgentsObservationObjectElement Element, const FName Name = TEXT("Map")) const;
 
-	UFUNCTION(BlueprintPure = false, Category = "LearningAgents")
-	UPARAM(DisplayName = "Success") bool GetEnumObservation(uint8& OutEnumValue, const UEnum* Enum, const FLearningAgentsObservationObjectElement Element, const FName Name = TEXT("Enum")) const;
+	UFUNCTION(BlueprintPure = false, Category = "LearningAgents", meta = (ReturnDisplayName = "Success"))
+	bool GetEnumObservation(uint8& OutEnumValue, const UEnum* Enum, const FLearningAgentsObservationObjectElement Element, const FName Name = TEXT("Enum")) const;
 
-	UFUNCTION(BlueprintPure = false, Category = "LearningAgents")
-	UPARAM(DisplayName = "Success") bool GetBitmaskObservation(int32& OutBitmaskValue, const UEnum* Enum, const FLearningAgentsObservationObjectElement Element, const FName Name = TEXT("Bitmask")) const;
+	UFUNCTION(BlueprintPure = false, Category = "LearningAgents", meta = (ReturnDisplayName = "Success"))
+	bool GetBitmaskObservation(int32& OutBitmaskValue, const UEnum* Enum, const FLearningAgentsObservationObjectElement Element, const FName Name = TEXT("Bitmask")) const;
 
-	UFUNCTION(BlueprintPure = false, Category = "LearningAgents", meta = (ExpandEnumAsExecs = "OutOption"))
-	UPARAM(DisplayName = "Success") bool GetOptionalObservation(ELearningAgentsOptionalObservation& OutOption, FLearningAgentsObservationObjectElement& OutElement, const FLearningAgentsObservationObjectElement Element, const FName Name = TEXT("Optional")) const;
+	UFUNCTION(BlueprintPure = false, Category = "LearningAgents", meta = (ExpandEnumAsExecs = "OutOption", ReturnDisplayName = "Success"))
+	bool GetOptionalObservation(ELearningAgentsOptionalObservation& OutOption, FLearningAgentsObservationObjectElement& OutElement, const FLearningAgentsObservationObjectElement Element, const FName Name = TEXT("Optional")) const;
 
-	UFUNCTION(BlueprintPure = false, Category = "LearningAgents", meta = (ExpandEnumAsExecs = "OutEither"))
-	UPARAM(DisplayName = "Success") bool GetEitherObservation(ELearningAgentsEitherObservation& OutEither, FLearningAgentsObservationObjectElement& OutElement, const FLearningAgentsObservationObjectElement Element, const FName Name = TEXT("Either")) const;
+	UFUNCTION(BlueprintPure = false, Category = "LearningAgents", meta = (ExpandEnumAsExecs = "OutEither", ReturnDisplayName = "Success"))
+	bool GetEitherObservation(ELearningAgentsEitherObservation& OutEither, FLearningAgentsObservationObjectElement& OutElement, const FLearningAgentsObservationObjectElement Element, const FName Name = TEXT("Either")) const;
 
-	UFUNCTION(BlueprintPure = false, Category = "LearningAgents")
-	UPARAM(DisplayName = "Success") bool GetEncodingObservation(FLearningAgentsObservationObjectElement& OutElement, const FLearningAgentsObservationObjectElement Element, const FName Name = TEXT("Encoding")) const;
+	UFUNCTION(BlueprintPure = false, Category = "LearningAgents", meta = (ReturnDisplayName = "Success"))
+	bool GetEncodingObservation(FLearningAgentsObservationObjectElement& OutElement, const FLearningAgentsObservationObjectElement Element, const FName Name = TEXT("Encoding")) const;
 
-	UFUNCTION(BlueprintPure = false, Category = "LearningAgents")
-	UPARAM(DisplayName = "Success") bool GetBoolObservation(bool& bOutValue, const FLearningAgentsObservationObjectElement Element, const FName Name = TEXT("Bool")) const;
+	UFUNCTION(BlueprintPure = false, Category = "LearningAgents", meta = (ReturnDisplayName = "Success"))
+	bool GetBoolObservation(bool& bOutValue, const FLearningAgentsObservationObjectElement Element, const FName Name = TEXT("Bool")) const;
 
-	UFUNCTION(BlueprintPure = false, Category = "LearningAgents")
-	UPARAM(DisplayName = "Success") bool GetFloatObservation(float& OutValue, const FLearningAgentsObservationObjectElement Element, const float FloatScale = 1.0f, const FName Name = TEXT("Float")) const;
+	UFUNCTION(BlueprintPure = false, Category = "LearningAgents", meta = (ReturnDisplayName = "Success"))
+	bool GetFloatObservation(float& OutValue, const FLearningAgentsObservationObjectElement Element, const float FloatScale = 1.0f, const FName Name = TEXT("Float")) const;
 
-	UFUNCTION(BlueprintPure = false, Category = "LearningAgents")
-	UPARAM(DisplayName = "Success") bool GetLocationObservation(FVector& OutLocation, const FLearningAgentsObservationObjectElement Element, const FTransform RelativeTransform = FTransform(), const float LocationScale = 100.0f, const FName Name = TEXT("Location")) const;
+	UFUNCTION(BlueprintPure = false, Category = "LearningAgents", meta = (ReturnDisplayName = "Success"))
+	bool GetLocationObservation(FVector& OutLocation, const FLearningAgentsObservationObjectElement Element, const FTransform RelativeTransform = FTransform(), const float LocationScale = 100.0f, const FName Name = TEXT("Location")) const;
 
-	UFUNCTION(BlueprintPure = false, Category = "LearningAgents")
-	UPARAM(DisplayName = "Success") bool GetRotationObservation(FRotator& OutRotation, const FLearningAgentsObservationObjectElement Element, const FRotator RelativeRotation = FRotator::ZeroRotator, const FName Name = TEXT("Rotation")) const;
-	UPARAM(DisplayName = "Success") bool GetRotationObservationAsQuat(FQuat& OutRotation, const FLearningAgentsObservationObjectElement Element, const FQuat RelativeRotation = FQuat::Identity, const FName Name = TEXT("Rotation")) const;
+	UFUNCTION(BlueprintPure = false, Category = "LearningAgents", meta = (ReturnDisplayName = "Success"))
+	bool GetRotationObservation(FRotator& OutRotation, const FLearningAgentsObservationObjectElement Element, const FRotator RelativeRotation = FRotator::ZeroRotator, const FName Name = TEXT("Rotation")) const;
+	bool GetRotationObservationAsQuat(FQuat& OutRotation, const FLearningAgentsObservationObjectElement Element, const FQuat RelativeRotation = FQuat::Identity, const FName Name = TEXT("Rotation")) const;
 
-	UFUNCTION(BlueprintPure = false, Category = "LearningAgents")
-	UPARAM(DisplayName = "Success") bool GetScaleObservation(FVector& OutScale, const FLearningAgentsObservationObjectElement Element, const FVector RelativeScale = FVector(1, 1, 1), const FName Name = TEXT("Scale")) const;
+	UFUNCTION(BlueprintPure = false, Category = "LearningAgents", meta = (ReturnDisplayName = "Success"))
+	bool GetScaleObservation(FVector& OutScale, const FLearningAgentsObservationObjectElement Element, const FVector RelativeScale = FVector(1, 1, 1), const FName Name = TEXT("Scale")) const;
 
-	UFUNCTION(BlueprintPure = false, Category = "LearningAgents")
-	UPARAM(DisplayName = "Success") bool GetTransformObservation(FTransform& OutTransform, const FLearningAgentsObservationObjectElement Element, const FTransform RelativeTransform = FTransform(), const float LocationScale = 100.0f, const FName Name = TEXT("Transform")) const;
+	UFUNCTION(BlueprintPure = false, Category = "LearningAgents", meta = (ReturnDisplayName = "Success"))
+	bool GetTransformObservation(FTransform& OutTransform, const FLearningAgentsObservationObjectElement Element, const FTransform RelativeTransform = FTransform(), const float LocationScale = 100.0f, const FName Name = TEXT("Transform")) const;
 
-	UFUNCTION(BlueprintPure = false, Category = "LearningAgents")
-	UPARAM(DisplayName = "Success") bool GetAngleObservation(float& OutAngle, const FLearningAgentsObservationObjectElement Element, const float RelativeAngle = 0.0f, const FName Name = TEXT("Angle")) const;
+	UFUNCTION(BlueprintPure = false, Category = "LearningAgents", meta = (ReturnDisplayName = "Success"))
+	bool GetAngleObservation(float& OutAngle, const FLearningAgentsObservationObjectElement Element, const float RelativeAngle = 0.0f, const FName Name = TEXT("Angle")) const;
 
-	UFUNCTION(BlueprintPure = false, Category = "LearningAgents")
-	UPARAM(DisplayName = "Success") bool GetAngleObservationRadians(float& OutAngle, const FLearningAgentsObservationObjectElement Element, const float RelativeAngle = 0.0f, const FName Name = TEXT("Angle")) const;
+	UFUNCTION(BlueprintPure = false, Category = "LearningAgents", meta = (ReturnDisplayName = "Success"))
+	bool GetAngleObservationRadians(float& OutAngle, const FLearningAgentsObservationObjectElement Element, const float RelativeAngle = 0.0f, const FName Name = TEXT("Angle")) const;
 
-	UFUNCTION(BlueprintPure = false, Category = "LearningAgents")
-	UPARAM(DisplayName = "Success") bool GetVelocityObservation(FVector& OutVelocity, const FLearningAgentsObservationObjectElement Element, const FTransform RelativeTransform = FTransform(), const float VelocityScale = 200.0f, const FName Name = TEXT("Velocity")) const;
+	UFUNCTION(BlueprintPure = false, Category = "LearningAgents", meta = (ReturnDisplayName = "Success"))
+	bool GetVelocityObservation(FVector& OutVelocity, const FLearningAgentsObservationObjectElement Element, const FTransform RelativeTransform = FTransform(), const float VelocityScale = 200.0f, const FName Name = TEXT("Velocity")) const;
 
-	UFUNCTION(BlueprintPure = false, Category = "LearningAgents")
-	UPARAM(DisplayName = "Success") bool GetDirectionObservation(FVector& OutDirection, const FLearningAgentsObservationObjectElement Element, const FTransform RelativeTransform = FTransform(), const FName Name = TEXT("Direction")) const;
+	UFUNCTION(BlueprintPure = false, Category = "LearningAgents", meta = (ReturnDisplayName = "Success"))
+	bool GetDirectionObservation(FVector& OutDirection, const FLearningAgentsObservationObjectElement Element, const FTransform RelativeTransform = FTransform(), const FName Name = TEXT("Direction")) const;
 
 	// Spline Observations
 
-	UFUNCTION(BlueprintPure = false, Category = "LearningAgents")
-	UPARAM(DisplayName = "Success") bool GetLocationAlongSplineObservation(FVector& OutLocation, const FLearningAgentsObservationObjectElement Element, const FTransform RelativeTransform = FTransform(), const float LocationScale = 100.0f, const FName Name = TEXT("LocationAlongSpline"));
+	UFUNCTION(BlueprintPure = false, Category = "LearningAgents", meta = (ReturnDisplayName = "Success"))
+	bool GetLocationAlongSplineObservation(FVector& OutLocation, const FLearningAgentsObservationObjectElement Element, const FTransform RelativeTransform = FTransform(), const float LocationScale = 100.0f, const FName Name = TEXT("LocationAlongSpline"));
 
-	UFUNCTION(BlueprintPure = false, Category = "LearningAgents")
-	UPARAM(DisplayName = "Success") bool GetProportionAlongSplineObservation(bool& bOutIsClosedLoop, float& OutAngle, float& OutPropotion, const FLearningAgentsObservationObjectElement Element, const FName Name = TEXT("ProportionAlongSpline"));
+	UFUNCTION(BlueprintPure = false, Category = "LearningAgents", meta = (ReturnDisplayName = "Success"))
+	bool GetProportionAlongSplineObservation(bool& bOutIsClosedLoop, float& OutAngle, float& OutPropotion, const FLearningAgentsObservationObjectElement Element, const FName Name = TEXT("ProportionAlongSpline"));
 
-	UFUNCTION(BlueprintPure = false, Category = "LearningAgents")
-	UPARAM(DisplayName = "Success") bool GetDirectionAlongSplineObservation(FVector& OutDirection, const FLearningAgentsObservationObjectElement Element, const FTransform RelativeTransform = FTransform(), const FName Name = TEXT("DirectionAlongSpline"));
+	UFUNCTION(BlueprintPure = false, Category = "LearningAgents", meta = (ReturnDisplayName = "Success"))
+	bool GetDirectionAlongSplineObservation(FVector& OutDirection, const FLearningAgentsObservationObjectElement Element, const FTransform RelativeTransform = FTransform(), const FName Name = TEXT("DirectionAlongSpline"));
 
-	UFUNCTION(BlueprintPure = false, Category = "LearningAgents")
-	UPARAM(DisplayName = "Success") bool GetPropertiesAlongSplineObservation(FVector& OutLocation, bool& bOutIsClosedLoop, float& OutAngle, float& OutPropotion, FVector& OutDirection, const FLearningAgentsObservationObjectElement Element, const FTransform RelativeTransform = FTransform(), const float LocationScale = 100.0f, const FName Name = TEXT("PropertiesAlongSpline"));
+	UFUNCTION(BlueprintPure = false, Category = "LearningAgents", meta = (ReturnDisplayName = "Success"))
+	bool GetPropertiesAlongSplineObservation(FVector& OutLocation, bool& bOutIsClosedLoop, float& OutAngle, float& OutPropotion, FVector& OutDirection, const FLearningAgentsObservationObjectElement Element, const FTransform RelativeTransform = FTransform(), const float LocationScale = 100.0f, const FName Name = TEXT("PropertiesAlongSpline"));
 
 	// Ray Cast Observations
 
-	UFUNCTION(BlueprintPure, Category = "LearningAgents")
-	UPARAM(DisplayName = "Success") bool GetProportionAlongRayObservation(float& OutProportion, const FLearningAgentsObservationObjectElement Element, const FName Name = TEXT("ProportionAlongRay"));
+	UFUNCTION(BlueprintPure, Category = "LearningAgents", meta = (ReturnDisplayName = "Success"))
+	bool GetProportionAlongRayObservation(float& OutProportion, const FLearningAgentsObservationObjectElement Element, const FName Name = TEXT("ProportionAlongRay"));
 
-	UFUNCTION(BlueprintPure, Category = "LearningAgents")
-	UPARAM(DisplayName = "Success") bool GetProportionAlongRaysObservationNum(int32& OutProportionNum, const FLearningAgentsObservationObjectElement Element, const FName Name = TEXT("ProportionAlongRays"));
+	UFUNCTION(BlueprintPure, Category = "LearningAgents", meta = (ReturnDisplayName = "Success"))
+	bool GetProportionAlongRaysObservationNum(int32& OutProportionNum, const FLearningAgentsObservationObjectElement Element, const FName Name = TEXT("ProportionAlongRays"));
 
-	UFUNCTION(BlueprintPure, Category = "LearningAgents")
-	UPARAM(DisplayName = "Success") bool GetProportionAlongRaysObservation(TArray<float>& OutProportions, const FLearningAgentsObservationObjectElement Element, const FName Name = TEXT("ProportionAlongRays"));
-	UPARAM(DisplayName = "Success") bool GetProportionAlongRaysObservationToArrayView(TArrayView<float> OutProportions, const FLearningAgentsObservationObjectElement Element, const FName Name = TEXT("ProportionAlongRays"));
+	UFUNCTION(BlueprintPure, Category = "LearningAgents", meta = (ReturnDisplayName = "Success"))
+	bool GetProportionAlongRaysObservation(TArray<float>& OutProportions, const FLearningAgentsObservationObjectElement Element, const FName Name = TEXT("ProportionAlongRays"));
+	bool GetProportionAlongRaysObservationToArrayView(TArrayView<float> OutProportions, const FLearningAgentsObservationObjectElement Element, const FName Name = TEXT("ProportionAlongRays"));
 
 
 private:
 
 	UE::Learning::Observation::FObject ObservationObject;
-
-	const ULearningAgentsObservationVisualLoggerObject* GetOrAddVisualLoggerObject(const FName Name);
-
-	UPROPERTY()
-	TMap<FName, TObjectPtr<const ULearningAgentsObservationVisualLoggerObject>> VisualLoggerObjects;
 };
diff --git a/Engine/Plugins/Experimental/LearningAgents/Source/LearningAgentsTraining/Private/LearningAgentsRewards.cpp b/Engine/Plugins/Experimental/LearningAgents/Source/LearningAgentsTraining/Private/LearningAgentsRewards.cpp
--- a/Engine/Plugins/Experimental/LearningAgents/Source/LearningAgentsTraining/Private/LearningAgentsRewards.cpp
+++ b/Engine/Plugins/Experimental/LearningAgents/Source/LearningAgentsTraining/Private/LearningAgentsRewards.cpp
@@ -2,57 +2,486 @@
 
 #include "LearningAgentsRewards.h"
 
+#include "LearningAgentsManagerListener.h"
+#include "LearningAgentsDebug.h"
+
 #include "LearningLog.h"
 
 #include "Components/SplineComponent.h"
 
-float ULearningAgentsRewards::RewardOnCondition(const bool bCondition, const float RewardScale)
+float ULearningAgentsRewards::MakeReward(
+	const float RewardValue,
+	const float RewardScale,
+	const FName Name,
+	const bool bVisualLoggerEnabled,
+	ULearningAgentsManagerListener* VisualLoggerListener,
+	const int32 VisualLoggerAgentId, 
+	const FVector VisualLoggerLocation,
+	const FLinearColor VisualLoggerColor)
 {
-	return bCondition ? RewardScale : 0.0f;
+	const float Reward = RewardValue * RewardScale;
+
+#if UE_LEARNING_AGENTS_ENABLE_VISUAL_LOG
+	if (bVisualLoggerEnabled && VisualLoggerListener)
+	{
+		const ULearningAgentsVisualLoggerObject* VisualLoggerObject = VisualLoggerListener->GetOrAddVisualLoggerObject(Name);
+
+		UE_LEARNING_AGENTS_VLOG_STRING(VisualLoggerObject, LogLearning, Display, VisualLoggerLocation,
+			VisualLoggerColor.ToFColor(true),
+			TEXT("Listener: %s\nName: %s\nAgent Id: % 3i\nValue: [% 6.1f]\nScale: [% 6.2f]\nReward: [% 6.2f]"),
+			*VisualLoggerListener->GetName(),
+			*Name.ToString(),
+			VisualLoggerAgentId,
+			RewardValue,
+			RewardScale,
+			Reward);
+	}
+#endif
+
+	return Reward;
+}
+
+float ULearningAgentsRewards::MakeRewardOnCondition(
+	const bool bCondition, 
+	const float RewardScale,
+	const FName Name,
+	const bool bVisualLoggerEnabled,
+	ULearningAgentsManagerListener* VisualLoggerListener,
+	const int32 VisualLoggerAgentId,
+	const FVector VisualLoggerLocation,
+	const FLinearColor VisualLoggerColor)
+{
+	const float Reward = bCondition ? RewardScale : 0.0f;
+
+#if UE_LEARNING_AGENTS_ENABLE_VISUAL_LOG
+	if (bVisualLoggerEnabled && VisualLoggerListener)
+	{
+		const ULearningAgentsVisualLoggerObject* VisualLoggerObject = VisualLoggerListener->GetOrAddVisualLoggerObject(Name);
+
+		UE_LEARNING_AGENTS_VLOG_STRING(VisualLoggerObject, LogLearning, Display, VisualLoggerLocation,
+			VisualLoggerColor.ToFColor(true),
+			TEXT("Listener: %s\nName: %s\nAgent Id: % 3i\nCondition: [%s]\nScale: [% 6.2f]\nReward: [% 6.2f]"),
+			*VisualLoggerListener->GetName(),
+			*Name.ToString(),
+			VisualLoggerAgentId,
+			bCondition ? TEXT("true") : TEXT("false"),
+			RewardScale,
+			Reward);
+	}
+#endif
+
+	return Reward;
 }
 
-float ULearningAgentsRewards::RewardOnLocationDifferenceBelowThreshold(const FVector LocationA, const FVector LocationB, const float DistanceThreshold, const float RewardScale)
+float ULearningAgentsRewards::MakeRewardOnLocationDifferenceBelowThreshold(
+	const FVector LocationA, 
+	const FVector LocationB, 
+	const float DistanceThreshold, 
+	const float RewardScale,
+	const FName Name,
+	const bool bVisualLoggerEnabled,
+	ULearningAgentsManagerListener* VisualLoggerListener,
+	const int32 VisualLoggerAgentId,
+	const FVector VisualLoggerLocation,
+	const FLinearColor VisualLoggerColor)
 {
-	return RewardOnCondition(FVector::Distance(LocationA, LocationB) < DistanceThreshold, RewardScale);
+	const float Distance = FVector::Distance(LocationA, LocationB);
+	const bool bCondition = Distance < DistanceThreshold;
+	const float Reward = MakeRewardOnCondition(bCondition, RewardScale);
+
+#if UE_LEARNING_AGENTS_ENABLE_VISUAL_LOG
+	if (bVisualLoggerEnabled && VisualLoggerListener)
+	{
+		const ULearningAgentsVisualLoggerObject* VisualLoggerObject = VisualLoggerListener->GetOrAddVisualLoggerObject(Name);
+
+		UE_LEARNING_AGENTS_VLOG_LOCATION(VisualLoggerObject, LogLearning, Display,
+			LocationA,
+			10,
+			VisualLoggerColor.ToFColor(true),
+			TEXT(""));
+
+		UE_LEARNING_AGENTS_VLOG_SEGMENT(VisualLoggerObject, LogLearning, Display,
+			LocationA,
+			LocationB,
+			VisualLoggerColor.ToFColor(true),
+			TEXT(""));
+
+		UE_LEARNING_AGENTS_VLOG_LOCATION(VisualLoggerObject, LogLearning, Display,
+			LocationB,
+			10,
+			VisualLoggerColor.ToFColor(true),
+			TEXT(""));
+
+		UE_LEARNING_AGENTS_VLOG_STRING(VisualLoggerObject, LogLearning, Display, VisualLoggerLocation,
+			VisualLoggerColor.ToFColor(true),
+			TEXT("Listener: %s\nName: %s\nAgent Id: % 3i\nLocationA: [% 6.1f % 6.1f % 6.1f]\nLocationB: [% 6.1f % 6.1f % 6.1f]\nDistance: [% 6.2f]\nThreshold: [% 6.2f]\nCondition: [%s]\nScale: [% 6.2f]\nReward: [% 6.2f]"),
+			*VisualLoggerListener->GetName(),
+			*Name.ToString(),
+			VisualLoggerAgentId,
+			LocationA.X, LocationA.Y, LocationA.Z,
+			LocationB.X, LocationB.Y, LocationB.Z,
+			Distance,
+			DistanceThreshold,
+			bCondition ? TEXT("true") : TEXT("false"),
+			RewardScale,
+			Reward);
+	}
+#endif
+
+	return Reward;
 }
 
-float ULearningAgentsRewards::RewardFromLocationSimilarity(const FVector LocationA, const FVector LocationB, const float LocationScale, const float RewardScale)
+float ULearningAgentsRewards::MakeRewardOnLocationDifferenceAboveThreshold(
+	const FVector LocationA,
+	const FVector LocationB,
+	const float DistanceThreshold,
+	const float RewardScale,
+	const FName Name,
+	const bool bVisualLoggerEnabled,
+	ULearningAgentsManagerListener* VisualLoggerListener,
+	const int32 VisualLoggerAgentId,
+	const FVector VisualLoggerLocation,
+	const FLinearColor VisualLoggerColor)
+{
+	const float Distance = FVector::Distance(LocationA, LocationB);
+	const bool bCondition = Distance > DistanceThreshold;
+	const float Reward = MakeRewardOnCondition(bCondition, RewardScale);
+
+#if UE_LEARNING_AGENTS_ENABLE_VISUAL_LOG
+	if (bVisualLoggerEnabled && VisualLoggerListener)
+	{
+		const ULearningAgentsVisualLoggerObject* VisualLoggerObject = VisualLoggerListener->GetOrAddVisualLoggerObject(Name);
+
+		UE_LEARNING_AGENTS_VLOG_LOCATION(VisualLoggerObject, LogLearning, Display,
+			LocationA,
+			10,
+			VisualLoggerColor.ToFColor(true),
+			TEXT(""));
+
+		UE_LEARNING_AGENTS_VLOG_SEGMENT(VisualLoggerObject, LogLearning, Display,
+			LocationA,
+			LocationB,
+			VisualLoggerColor.ToFColor(true),
+			TEXT(""));
+
+		UE_LEARNING_AGENTS_VLOG_LOCATION(VisualLoggerObject, LogLearning, Display,
+			LocationB,
+			10,
+			VisualLoggerColor.ToFColor(true),
+			TEXT(""));
+
+		UE_LEARNING_AGENTS_VLOG_STRING(VisualLoggerObject, LogLearning, Display, VisualLoggerLocation,
+			VisualLoggerColor.ToFColor(true),
+			TEXT("Listener: %s\nName: %s\nAgent Id: % 3i\nLocationA: [% 6.1f % 6.1f % 6.1f]\nLocationB: [% 6.1f % 6.1f % 6.1f]\nDistance: [% 6.2f]\nThreshold: [% 6.2f]\nCondition: [%s]\nScale: [% 6.2f]\nReward: [% 6.2f]"),
+			*VisualLoggerListener->GetName(),
+			*Name.ToString(),
+			VisualLoggerAgentId,
+			LocationA.X, LocationA.Y, LocationA.Z,
+			LocationB.X, LocationB.Y, LocationB.Z,
+			Distance,
+			DistanceThreshold,
+			bCondition ? TEXT("true") : TEXT("false"),
+			RewardScale,
+			Reward);
+	}
+#endif
+
+	return Reward;
+}
+
+float ULearningAgentsRewards::MakeRewardFromLocationSimilarity(
+	const FVector LocationA, 
+	const FVector LocationB, 
+	const float LocationScale, 
+	const float RewardScale,
+	const FName Name,
+	const bool bVisualLoggerEnabled,
+	ULearningAgentsManagerListener* VisualLoggerListener,
+	const int32 VisualLoggerAgentId,
+	const FVector VisualLoggerLocation,
+	const FLinearColor VisualLoggerColor)
 {
 	const float LocationDifference = FVector::Dist(LocationA, LocationB);
-	return RewardScale * FMath::InvExpApprox(FMath::Square(LocationDifference / FMath::Max(LocationScale, UE_SMALL_NUMBER)));
+	const float Similarity = FMath::InvExpApprox(FMath::Square(LocationDifference / FMath::Max(LocationScale, UE_SMALL_NUMBER)));
+	const float Reward = Similarity * RewardScale;
+
+#if UE_LEARNING_AGENTS_ENABLE_VISUAL_LOG
+	if (bVisualLoggerEnabled && VisualLoggerListener)
+	{
+		const ULearningAgentsVisualLoggerObject* VisualLoggerObject = VisualLoggerListener->GetOrAddVisualLoggerObject(Name);
+
+		UE_LEARNING_AGENTS_VLOG_LOCATION(VisualLoggerObject, LogLearning, Display,
+			LocationA,
+			10,
+			VisualLoggerColor.ToFColor(true),
+			TEXT(""));
+
+		UE_LEARNING_AGENTS_VLOG_SEGMENT(VisualLoggerObject, LogLearning, Display,
+			LocationA,
+			LocationB,
+			VisualLoggerColor.ToFColor(true),
+			TEXT(""));
+
+		UE_LEARNING_AGENTS_VLOG_LOCATION(VisualLoggerObject, LogLearning, Display,
+			LocationB,
+			10,
+			VisualLoggerColor.ToFColor(true),
+			TEXT(""));
+
+		UE_LEARNING_AGENTS_VLOG_STRING(VisualLoggerObject, LogLearning, Display, VisualLoggerLocation,
+			VisualLoggerColor.ToFColor(true),
+			TEXT("Listener: %s\nName: %s\nAgent Id: % 3i\nLocationA: [% 6.1f % 6.1f % 6.1f]\nLocationB: [% 6.1f % 6.1f % 6.1f]\nDifference: [% 6.2f]\nLocationScale: [% 6.2f]\nSimilarity: [% 6.2f]\nScale: [% 6.2f]\nReward: [% 6.2f]"),
+			*VisualLoggerListener->GetName(),
+			*Name.ToString(),
+			VisualLoggerAgentId,
+			LocationA.X, LocationA.Y, LocationA.Z,
+			LocationB.X, LocationB.Y, LocationB.Z,
+			LocationDifference,
+			LocationScale,
+			Similarity,
+			RewardScale,
+			Reward);
+	}
+#endif
+
+	return Reward;
 }
 
-float ULearningAgentsRewards::RewardFromAngleSimilarity(const float AngleA, const float AngleB, const float AngleScale, const float RewardScale)
+float ULearningAgentsRewards::MakeRewardFromAngleSimilarity(
+	const float AngleA, 
+	const float AngleB, 
+	const float AngleScale, 
+	const float RewardScale,
+	const FName Name,
+	const bool bVisualLoggerEnabled,
+	ULearningAgentsManagerListener* VisualLoggerListener,
+	const int32 VisualLoggerAgentId,
+	const FVector VisualLoggerAngleLocationA,
+	const FVector VisualLoggerAngleLocationB,
+	const FVector VisualLoggerLocation,
+	const FLinearColor VisualLoggerColor)
 {
-	const float AngleDifference = FMath::FindDeltaAngleRadians(FMath::DegreesToRadians(AngleA), FMath::DegreesToRadians(AngleB));
-	return RewardScale * FMath::InvExpApprox(FMath::Square(AngleDifference / FMath::Max(FMath::DegreesToRadians(AngleScale), UE_SMALL_NUMBER)));
+	const float AngleDifference = FMath::FindDeltaAngleDegrees(AngleA, AngleB);
+	const float Similarity = FMath::InvExpApprox(FMath::Square(AngleDifference / FMath::Max(AngleScale, UE_SMALL_NUMBER)));
+	const float Reward = Similarity * RewardScale;
+
+#if UE_LEARNING_AGENTS_ENABLE_VISUAL_LOG
+	if (bVisualLoggerEnabled && VisualLoggerListener)
+	{
+		const ULearningAgentsVisualLoggerObject* VisualLoggerObject = VisualLoggerListener->GetOrAddVisualLoggerObject(Name);
+
+		UE_LEARNING_AGENTS_VLOG_ANGLE(
+			VisualLoggerObject,
+			LogLearning,
+			Display,
+			AngleA,
+			0.0f,
+			VisualLoggerAngleLocationA,
+			5.0f,
+			VisualLoggerColor.ToFColor(true),
+			TEXT(""));
+
+		UE_LEARNING_AGENTS_VLOG_ANGLE(
+			VisualLoggerObject,
+			LogLearning,
+			Display,
+			AngleB,
+			0.0f,
+			VisualLoggerAngleLocationB,
+			5.0f,
+			VisualLoggerColor.ToFColor(true),
+			TEXT(""));
+
+		UE_LEARNING_AGENTS_VLOG_STRING(VisualLoggerObject, LogLearning, Display, VisualLoggerLocation,
+			VisualLoggerColor.ToFColor(true),
+			TEXT("Listener: %s\nName: %s\nAgent Id: % 3i\nAngleA [% 6.2f]\nAngleB [% 6.2f]\nDifference: [% 6.2f]\nAngleScale: [% 6.2f]\nSimilarity: [% 6.2f]\nScale: [% 6.2f]\nReward: [% 6.2f]"),
+			*VisualLoggerListener->GetName(),
+			*Name.ToString(),
+			VisualLoggerAgentId,
+			AngleA,
+			AngleB,
+			AngleDifference,
+			AngleScale,
+			Similarity,
+			RewardScale,
+			Reward);
+	}
+#endif
+
+	return Reward;
 }
 
-float ULearningAgentsRewards::RewardFromRotationSimilarityAsQuats(const FQuat RotationA, const FQuat RotationB, const float AngleScale, const float RewardScale)
+float ULearningAgentsRewards::MakeRewardFromRotationSimilarityAsQuats(
+	const FQuat RotationA, 
+	const FQuat RotationB, 
+	const float AngleScale, 
+	const float RewardScale,
+	const FName Name,
+	const bool bVisualLoggerEnabled,
+	ULearningAgentsManagerListener* VisualLoggerListener,
+	const int32 VisualLoggerAgentId,
+	const FVector VisualLoggerRotationLocationA,
+	const FVector VisualLoggerRotationLocationB,
+	const FVector VisualLoggerLocation,
+	const FLinearColor VisualLoggerColor)
 {
 	FQuat Difference = RotationA.Inverse() * RotationB;
 	Difference.EnforceShortestArcWith(FQuat::Identity);
-	return RewardScale * FMath::InvExpApprox(FMath::Square(Difference.GetAngle() / (FMath::Max(FMath::DegreesToRadians(AngleScale), UE_SMALL_NUMBER))));
+	const float AngleDifference = FMath::RadiansToDegrees(Difference.GetAngle());
+	const float Similarity = FMath::InvExpApprox(FMath::Square(AngleDifference / (FMath::Max(AngleScale, UE_SMALL_NUMBER))));
+	const float Reward = Similarity * RewardScale;
+
+#if UE_LEARNING_AGENTS_ENABLE_VISUAL_LOG
+	if (bVisualLoggerEnabled && VisualLoggerListener)
+	{
+		const ULearningAgentsVisualLoggerObject* VisualLoggerObject = VisualLoggerListener->GetOrAddVisualLoggerObject(Name);
+
+		UE_LEARNING_AGENTS_VLOG_TRANSFORM(
+			VisualLoggerObject,
+			LogLearning,
+			Display,
+			VisualLoggerRotationLocationA,
+			RotationA,
+			VisualLoggerColor.ToFColor(true),
+			TEXT(""));
+
+		UE_LEARNING_AGENTS_VLOG_TRANSFORM(
+			VisualLoggerObject,
+			LogLearning,
+			Display,
+			VisualLoggerRotationLocationB,
+			RotationB,
+			VisualLoggerColor.ToFColor(true),
+			TEXT(""));
+
+		UE_LEARNING_AGENTS_VLOG_STRING(VisualLoggerObject, LogLearning, Display, VisualLoggerLocation,
+			VisualLoggerColor.ToFColor(true),
+			TEXT("Listener: %s\nName: %s\nAgent Id: % 3i\nRotationA [% 6.2f % 6.2f % 6.2f % 6.2f]\nRotationB [% 6.2f % 6.2f % 6.2f % 6.2f]\nDifference: [% 6.2f]\nAngleScale: [% 6.2f]\nSimilarity: [% 6.2f]\nScale: [% 6.2f]\nReward: [% 6.2f]"),
+			*VisualLoggerListener->GetName(),
+			*Name.ToString(),
+			VisualLoggerAgentId,
+			RotationA.X, RotationA.Y, RotationA.Z, RotationA.W,
+			RotationB.X, RotationB.Y, RotationB.Z, RotationB.W,
+			AngleDifference,
+			AngleScale,
+			Similarity,
+			RewardScale,
+			Reward);
+	}
+#endif
+
+	return Reward;
 }
 
-float ULearningAgentsRewards::RewardFromRotationSimilarity(const FRotator RotationA, const FRotator RotationB, const float AngleScale, const float RewardScale)
+float ULearningAgentsRewards::MakeRewardFromRotationSimilarity(
+	const FRotator RotationA, 
+	const FRotator RotationB, 
+	const float AngleScale, 
+	const float RewardScale,
+	const FName Name,
+	const bool bVisualLoggerEnabled,
+	ULearningAgentsManagerListener* VisualLoggerListener,
+	const int32 VisualLoggerAgentId,
+	const FVector VisualLoggerRotationLocationA,
+	const FVector VisualLoggerRotationLocationB,
+	const FVector VisualLoggerLocation,
+	const FLinearColor VisualLoggerColor)
 {
-	return RewardScale * RewardFromRotationSimilarityAsQuats(RotationA.Quaternion(), RotationB.Quaternion(), AngleScale);
+	return MakeRewardFromRotationSimilarityAsQuats(
+		RotationA.Quaternion(), 
+		RotationB.Quaternion(), 
+		AngleScale, 
+		RewardScale, 
+		Name, 
+		bVisualLoggerEnabled, 
+		VisualLoggerListener, 
+		VisualLoggerAgentId,
+		VisualLoggerRotationLocationA,
+		VisualLoggerRotationLocationB,
+		VisualLoggerLocation, 
+		VisualLoggerColor);
 }
 
-float ULearningAgentsRewards::RewardFromDirectionSimilarity(const FVector DirectionA, const FVector DirectionB, const float AngleScale, const float RewardScale)
+float ULearningAgentsRewards::MakeRewardFromDirectionSimilarity(
+	const FVector DirectionA, 
+	const FVector DirectionB, 
+	const float AngleScale, 
+	const float RewardScale,
+	const FName Name,
+	const bool bVisualLoggerEnabled,
+	ULearningAgentsManagerListener* VisualLoggerListener,
+	const int32 VisualLoggerAgentId,
+	const FVector VisualLoggerDirectionLocationA,
+	const FVector VisualLoggerDirectionLocationB,
+	const FVector VisualLoggerLocation,
+	const float VisualLoggerArrowLength,
+	const FLinearColor VisualLoggerColor)
 {
-	const float AngleDifference = FMath::Acos(DirectionA.Dot(DirectionB));
-	return RewardScale * FMath::InvExpApprox(FMath::Square(AngleDifference / FMath::Max(FMath::DegreesToRadians(AngleScale), UE_SMALL_NUMBER)));
+	const float AngleDifference = FMath::RadiansToDegrees(FMath::Acos(DirectionA.GetSafeNormal(UE_SMALL_NUMBER, FVector::ForwardVector).Dot(DirectionB.GetSafeNormal(UE_SMALL_NUMBER, FVector::ForwardVector))));
+	const float Similarity = FMath::InvExpApprox(FMath::Square(AngleDifference / FMath::Max(AngleScale, UE_SMALL_NUMBER)));
+	const float Reward = Similarity * RewardScale;
+
+#if UE_LEARNING_AGENTS_ENABLE_VISUAL_LOG
+	if (bVisualLoggerEnabled && VisualLoggerListener)
+	{
+		const ULearningAgentsVisualLoggerObject* VisualLoggerObject = VisualLoggerListener->GetOrAddVisualLoggerObject(Name);
+
+		UE_LEARNING_AGENTS_VLOG_ARROW(
+			VisualLoggerObject,
+			LogLearning,
+			Display,
+			VisualLoggerDirectionLocationA,
+			VisualLoggerDirectionLocationA + VisualLoggerArrowLength * DirectionA,
+			VisualLoggerColor.ToFColor(true),
+			TEXT(""));
+
+		UE_LEARNING_AGENTS_VLOG_ARROW(
+			VisualLoggerObject,
+			LogLearning,
+			Display,
+			VisualLoggerDirectionLocationB,
+			VisualLoggerDirectionLocationB + VisualLoggerArrowLength * DirectionB,
+			VisualLoggerColor.ToFColor(true),
+			TEXT(""));
+
+		UE_LEARNING_AGENTS_VLOG_STRING(VisualLoggerObject, LogLearning, Display, VisualLoggerLocation,
+			VisualLoggerColor.ToFColor(true),
+			TEXT("Listener: %s\nName: %s\nAgent Id: % 3i\nDirectionA [% 6.2f % 6.2f % 6.2f]\nDirectionB [% 6.2f % 6.2f % 6.2f]\nDifference: [% 6.2f]\nAngleScale: [% 6.2f]\nSimilarity: [% 6.2f]\nScale: [% 6.2f]\nReward: [% 6.2f]"),
+			*VisualLoggerListener->GetName(),
+			*Name.ToString(),
+			VisualLoggerAgentId,
+			DirectionA.X, DirectionA.Y, DirectionA.Z,
+			DirectionB.X, DirectionB.Y, DirectionB.Z,
+			AngleDifference,
+			AngleScale,
+			Similarity,
+			RewardScale,
+			Reward);
+	}
+#endif
+
+	return Reward;
 }
 
 // Spline Rewards
 
-float ULearningAgentsRewards::RewardFromVelocityAlongSpline(const USplineComponent* SplineComponent, const FVector Position, const FVector Velocity, const float VelocityScale, const float RewardScale, const float FiniteDifferenceDelta)
+float ULearningAgentsRewards::MakeRewardFromVelocityAlongSpline(
+	const USplineComponent* SplineComponent, 
+	const FVector Location, 
+	const FVector Velocity, 
+	const float VelocityScale, 
+	const float RewardScale, 
+	const float FiniteDifferenceDelta,
+	const FName Name,
+	const bool bVisualLoggerEnabled,
+	ULearningAgentsManagerListener* VisualLoggerListener,
+	const int32 VisualLoggerAgentId,
+	const FVector VisualLoggerLocation,
+	const FLinearColor VisualLoggerColor)
 {
 	if (!SplineComponent)
 	{
-		UE_LOG(LogLearning, Error, TEXT("RewardFromVelocityAlongSpline: SplineComponent is nullptr."));
+		UE_LOG(LogLearning, Error, TEXT("MakeRewardFromVelocityAlongSpline: SplineComponent is nullptr."));
 		return 0.0f;
 	}
 
@@ -60,12 +489,12 @@ float ULearningAgentsRewards::RewardFromVelocityAlongSpline(const USplineCompone
 
 	if (FiniteDiff < UE_KINDA_SMALL_NUMBER)
 	{
-		UE_LOG(LogLearning, Warning, TEXT("RewardFromVelocityAlongSpline: FiniteDifferenceDelta is too small (%6.4f). Clamping to %6.4f."), FiniteDifferenceDelta, UE_KINDA_SMALL_NUMBER);
+		UE_LOG(LogLearning, Warning, TEXT("MakeRewardFromVelocityAlongSpline: FiniteDifferenceDelta is too small (%6.4f). Clamping to %6.4f."), FiniteDifferenceDelta, UE_KINDA_SMALL_NUMBER);
 		FiniteDiff = UE_KINDA_SMALL_NUMBER;
 	}
 
-	const float RawDistance0 = SplineComponent->GetDistanceAlongSplineAtLocation(Position, ESplineCoordinateSpace::World);
-	const float RawDistance1 = SplineComponent->GetDistanceAlongSplineAtLocation(Position + FiniteDiff * Velocity.GetSafeNormal(), ESplineCoordinateSpace::World);
+	const float RawDistance0 = SplineComponent->GetDistanceAlongSplineAtLocation(Location, ESplineCoordinateSpace::World);
+	const float RawDistance1 = SplineComponent->GetDistanceAlongSplineAtLocation(Location + FiniteDiff * Velocity.GetSafeNormal(), ESplineCoordinateSpace::World);
 
 	float Distance0 = RawDistance0, Distance1 = RawDistance1;
 
@@ -83,49 +512,40 @@ float ULearningAgentsRewards::RewardFromVelocityAlongSpline(const USplineCompone
 		}
 	}
 
-	return RewardScale * ((Distance1 - Distance0) / FiniteDiff) * (Velocity.Length() / FMath::Max(VelocityScale, UE_SMALL_NUMBER));
-}
-
-
-// Penalties
+	const float DistanceVelocity = ((Distance1 - Distance0) / FiniteDiff) * (Velocity.Length() / FMath::Max(VelocityScale, UE_SMALL_NUMBER));
+	const float Reward = RewardScale * DistanceVelocity;
 
-float ULearningAgentsRewards::PenaltyOnCondition(const bool bCondition, const float PenaltyScale)
-{
-	return bCondition ? -PenaltyScale : 0.0f;
-}
+#if UE_LEARNING_AGENTS_ENABLE_VISUAL_LOG
+	if (bVisualLoggerEnabled && VisualLoggerListener)
+	{
+		const FVector SplineLocation0 = SplineComponent->GetLocationAtDistanceAlongSpline(RawDistance0, ESplineCoordinateSpace::World);
+		const FVector SplineLocation1 = SplineComponent->GetLocationAtDistanceAlongSpline(RawDistance1, ESplineCoordinateSpace::World);
+		const FVector SplineVelocityDirection = (SplineLocation1 - SplineLocation0).GetSafeNormal(UE_SMALL_NUMBER, FVector::ForwardVector);
 
-float ULearningAgentsRewards::PenaltyOnLocationDifferenceAboveThreshold(const FVector LocationA, const FVector LocationB, const float DistanceThreshold, const float PenaltyScale)
-{
-	return PenaltyOnCondition(FVector::Distance(LocationA, LocationB) > DistanceThreshold, PenaltyScale);
-}
+		const ULearningAgentsVisualLoggerObject* VisualLoggerObject = VisualLoggerListener->GetOrAddVisualLoggerObject(Name);
 
-float ULearningAgentsRewards::PenaltyFromLocationDifference(const FVector LocationA, const FVector LocationB, const float LocationScale, const float PenaltyScale)
-{
-	const float LocationDifference = FVector::Dist(LocationA, LocationB);
-	return PenaltyScale  * (-LocationDifference / FMath::Max(LocationScale, UE_SMALL_NUMBER));
-}
-
-float ULearningAgentsRewards::PenaltyFromAngleDifference(const float AngleA, const float AngleB, const float AngleScale, const float PenaltyScale)
-{
-	const float AngleDifference = FMath::FindDeltaAngleRadians(FMath::DegreesToRadians(AngleA), FMath::DegreesToRadians(AngleB));
-	return PenaltyScale  * (-AngleDifference / FMath::Max(FMath::DegreesToRadians(AngleScale), UE_SMALL_NUMBER));
-}
+		UE_LEARNING_AGENTS_VLOG_ARROW(
+			VisualLoggerObject,
+			LogLearning,
+			Display,
+			SplineLocation0,
+			SplineLocation0 + DistanceVelocity * SplineVelocityDirection,
+			VisualLoggerColor.ToFColor(true),
+			TEXT(""));
 
-float ULearningAgentsRewards::PenaltyFromRotationDifferenceAsQuats(const FQuat RotationA, const FQuat RotationB, const float AngleScale, const float PenaltyScale)
-{
-	FQuat Difference = RotationA.Inverse() * RotationB;
-	Difference.EnforceShortestArcWith(FQuat::Identity);
-	return PenaltyScale  * (-Difference.GetAngle() / (FMath::Max(FMath::DegreesToRadians(AngleScale), UE_SMALL_NUMBER)));
-}
-
-float ULearningAgentsRewards::PenaltyFromRotationDifference(const FRotator RotationA, const FRotator RotationB, const float AngleScale, const float PenaltyScale)
-{
-	return PenaltyFromRotationDifferenceAsQuats(RotationA.Quaternion(), RotationB.Quaternion(), AngleScale, PenaltyScale);
-}
+		UE_LEARNING_AGENTS_VLOG_STRING(VisualLoggerObject, LogLearning, Display, VisualLoggerLocation,
+			VisualLoggerColor.ToFColor(true),
+			TEXT("Listener: %s\nName: %s\nAgent Id: % 3i\nLocation [% 6.2f % 6.2f % 6.2f]\nVelocity [% 6.2f % 6.2f % 6.2f]\nDistanceVelocity: [% 6.2f]\nScale: [% 6.2f]\nReward: [% 6.2f]"),
+			*VisualLoggerListener->GetName(),
+			*Name.ToString(),
+			VisualLoggerAgentId,
+			Location.X, Location.Y, Location.Z,
+			Velocity.X, Velocity.Y, Velocity.Z,
+			DistanceVelocity,
+			RewardScale,
+			Reward);
+	}
+#endif
 
-float ULearningAgentsRewards::PenaltyFromDirectionDifference(const FVector DirectionA, const FVector DirectionB, const float AngleScale, const float PenaltyScale)
-{
-	const float AngleDifference = FMath::Acos(DirectionA.Dot(DirectionB));
-	return PenaltyScale  * (-AngleDifference / FMath::Max(FMath::DegreesToRadians(AngleScale), UE_SMALL_NUMBER));
+	return Reward;
 }
-
diff --git a/Engine/Plugins/Experimental/LearningAgents/Source/LearningAgentsTraining/Public/LearningAgentsRewards.h b/Engine/Plugins/Experimental/LearningAgents/Source/LearningAgentsTraining/Public/LearningAgentsRewards.h
--- a/Engine/Plugins/Experimental/LearningAgents/Source/LearningAgentsTraining/Public/LearningAgentsRewards.h
+++ b/Engine/Plugins/Experimental/LearningAgents/Source/LearningAgentsTraining/Public/LearningAgentsRewards.h
@@ -6,6 +6,7 @@
 
 #include "LearningAgentsRewards.generated.h"
 
+class ULearningAgentsManagerListener;
 class USplineComponent;
 
 UCLASS(BlueprintType)
@@ -15,52 +16,297 @@ class LEARNINGAGENTSTRAINING_API ULearningAgentsRewards : public UBlueprintFunct
 
 public:
 
-	// Rewards
+	/**
+	 * Make a reward from a float value.
+	 *
+	 * @param RewardValue The float value used to create the reward.
+	 * @param RewardScale The scale of the reward. Use a negative scale to create a penalty.
+	 * @param Name The name of the reward. Used for debugging.
+	 * @param bVisualLoggerEnabled When true, debug data will be sent to the visual logger.
+	 * @param VisualLoggerListener The listener object which is making this reward. This must be set to use logging.
+	 * @param VisualLoggerAgentId The agent id associated with this reward.
+	 * @param VisualLoggerLocation A location for the visual logger information in the world.
+	 * @param VisualLoggerColor The color for the visual logger display.
+	 * @return The resulting reward value.
+	 */
+	UFUNCTION(BlueprintPure, Category = "LearningAgents", meta = (AdvancedDisplay = 3, DefaultToSelf = "VisualLoggerListener"))
+	static float MakeReward(
+		const float RewardValue, 
+		const float RewardScale = 1.0f, 
+		const FName Name = TEXT("Reward"),
+		const bool bVisualLoggerEnabled = false,
+		ULearningAgentsManagerListener* VisualLoggerListener = nullptr,
+		const int32 VisualLoggerAgentId = -1,
+		const FVector VisualLoggerLocation = FVector::ZeroVector,
+		const FLinearColor VisualLoggerColor = FLinearColor::Green);
 
-	UFUNCTION(BlueprintPure, Category = "LearningAgents")
-	static float RewardOnCondition(const bool bCondition, const float RewardScale = 1.0f);
+	/**
+	 * Make a reward which is equal to RewardScale when bCondition is true, otherwise returns zero.
+	 *
+	 * @param bCondition The condition under which to create a reward.
+	 * @param RewardScale The scale of the reward. Use a negative scale to create a penalty.
+	 * @param Name The name of the reward. Used for debugging.
+	 * @param bVisualLoggerEnabled When true, debug data will be sent to the visual logger.
+	 * @param VisualLoggerListener The listener object which is making this reward. This must be set to use logging.
+	 * @param VisualLoggerAgentId The agent id associated with this reward.
+	 * @param VisualLoggerLocation A location for the visual logger information in the world.
+	 * @param VisualLoggerColor The color for the visual logger display.
+	 * @return The resulting reward value.
+	 */
+	UFUNCTION(BlueprintPure, Category = "LearningAgents", meta = (AdvancedDisplay = 3, DefaultToSelf = "VisualLoggerListener"))
+	static float MakeRewardOnCondition(
+		const bool bCondition, 
+		const float RewardScale = 1.0f,
+		const FName Name = TEXT("RewardOnCondition"),
+		const bool bVisualLoggerEnabled = false,
+		ULearningAgentsManagerListener* VisualLoggerListener = nullptr,
+		const int32 VisualLoggerAgentId = -1,
+		const FVector VisualLoggerLocation = FVector::ZeroVector,
+		const FLinearColor VisualLoggerColor = FLinearColor::Green);
 
-	UFUNCTION(BlueprintPure, Category = "LearningAgents")
-	static float RewardOnLocationDifferenceBelowThreshold(const FVector LocationA, const FVector LocationB, const float DistanceThreshold = 100.0f, const float RewardScale = 1.0f);
+	/**
+	 * Make a reward when the distance between two locations is below a threshold, otherwise returns zero.
+	 *
+	 * @param LocationA The first location.
+	 * @param LocationB The second location.
+	 * @param DistanceThreshold The distance threshold.
+	 * @param RewardScale The scale of the reward. Use a negative scale to create a penalty.
+	 * @param Name The name of the reward. Used for debugging.
+	 * @param bVisualLoggerEnabled When true, debug data will be sent to the visual logger.
+	 * @param VisualLoggerListener The listener object which is making this reward. This must be set to use logging.
+	 * @param VisualLoggerAgentId The agent id associated with this reward.
+	 * @param VisualLoggerLocation A location for the visual logger information in the world.
+	 * @param VisualLoggerColor The color for the visual logger display.
+	 * @return The resulting reward value.
+	 */
+	UFUNCTION(BlueprintPure, Category = "LearningAgents", meta = (AdvancedDisplay = 5, DefaultToSelf = "VisualLoggerListener"))
+	static float MakeRewardOnLocationDifferenceBelowThreshold(
+		const FVector LocationA, 
+		const FVector LocationB, 
+		const float DistanceThreshold = 100.0f, 
+		const float RewardScale = 1.0f,
+		const FName Name = TEXT("RewardOnLocationDifferenceBelowThreshold"),
+		const bool bVisualLoggerEnabled = false,
+		ULearningAgentsManagerListener* VisualLoggerListener = nullptr,
+		const int32 VisualLoggerAgentId = -1,
+		const FVector VisualLoggerLocation = FVector::ZeroVector,
+		const FLinearColor VisualLoggerColor = FLinearColor::Green);
 
-	UFUNCTION(BlueprintPure, Category = "LearningAgents")
-	static float RewardFromLocationSimilarity(const FVector LocationA, const FVector LocationB, const float LocationScale = 100.0f, const float RewardScale = 1.0f);
+	/**
+	 * Make a reward when the distance between two locations is above a threshold, otherwise returns zero.
+	 *
+	 * @param LocationA The first location.
+	 * @param LocationB The second location.
+	 * @param DistanceThreshold The distance threshold.
+	 * @param RewardScale The scale of the reward. Use a negative scale to create a penalty.
+	 * @param Name The name of the reward. Used for debugging.
+	 * @param bVisualLoggerEnabled When true, debug data will be sent to the visual logger.
+	 * @param VisualLoggerListener The listener object which is making this reward. This must be set to use logging.
+	 * @param VisualLoggerAgentId The agent id associated with this reward.
+	 * @param VisualLoggerLocation A location for the visual logger information in the world.
+	 * @param VisualLoggerColor The color for the visual logger display.
+	 * @return The resulting reward value.
+	 */
+	UFUNCTION(BlueprintPure, Category = "LearningAgents", meta = (AdvancedDisplay = 5, DefaultToSelf = "VisualLoggerListener"))
+	static float MakeRewardOnLocationDifferenceAboveThreshold(
+		const FVector LocationA,
+		const FVector LocationB,
+		const float DistanceThreshold = 100.0f,
+		const float RewardScale = 1.0f,
+		const FName Name = TEXT("RewardOnLocationDifferenceAboveThreshold"),
+		const bool bVisualLoggerEnabled = false,
+		ULearningAgentsManagerListener* VisualLoggerListener = nullptr,
+		const int32 VisualLoggerAgentId = -1,
+		const FVector VisualLoggerLocation = FVector::ZeroVector,
+		const FLinearColor VisualLoggerColor = FLinearColor::Green);
 
-	UFUNCTION(BlueprintPure, Category = "LearningAgents")
-	static float RewardFromRotationSimilarity(const FRotator RotationA, const FRotator RotationB, const float AngleScale = 90.0f, const float RewardScale = 1.0f);
-	static float RewardFromRotationSimilarityAsQuats(const FQuat RotationA, const FQuat RotationB, const float AngleScale = 90.0f, const float RewardScale = 1.0f);
+	/**
+	 * Make a reward based on how similar two locations are.
+	 *
+	 * @param LocationA The first location.
+	 * @param LocationB The second location.
+	 * @param LocationScale The expected scale for the distance between locations.
+	 * @param RewardScale The scale of the reward. Use a negative scale to create a penalty.
+	 * @param Name The name of the reward. Used for debugging.
+	 * @param bVisualLoggerEnabled When true, debug data will be sent to the visual logger.
+	 * @param VisualLoggerListener The listener object which is making this reward. This must be set to use logging.
+	 * @param VisualLoggerAgentId The agent id associated with this reward.
+	 * @param VisualLoggerLocation A location for the visual logger information in the world.
+	 * @param VisualLoggerColor The color for the visual logger display.
+	 * @return The resulting reward value.
+	 */
+	UFUNCTION(BlueprintPure, Category = "LearningAgents", meta = (AdvancedDisplay = 5, DefaultToSelf = "VisualLoggerListener"))
+	static float MakeRewardFromLocationSimilarity(
+		const FVector LocationA, 
+		const FVector LocationB, 
+		const float LocationScale = 100.0f, 
+		const float RewardScale = 1.0f,
+		const FName Name = TEXT("RewardFromLocationSimilarity"),
+		const bool bVisualLoggerEnabled = false,
+		ULearningAgentsManagerListener* VisualLoggerListener = nullptr,
+		const int32 VisualLoggerAgentId = -1,
+		const FVector VisualLoggerLocation = FVector::ZeroVector,
+		const FLinearColor VisualLoggerColor = FLinearColor::Green);
 
-	UFUNCTION(BlueprintPure, Category = "LearningAgents")
-	static float RewardFromAngleSimilarity(const float AngleA, const float AngleB, const float AngleScale = 90.0f, const float RewardScale = 1.0f);
+	/**
+	 * Make a reward based on how similar two rotations are.
+	 *
+	 * @param RotationA The first rotation.
+	 * @param RotationB The second rotation.
+	 * @param AngleScale The expected scale for the angle between rotations.
+	 * @param RewardScale The scale of the reward. Use a negative scale to create a penalty.
+	 * @param Name The name of the reward. Used for debugging.
+	 * @param bVisualLoggerEnabled When true, debug data will be sent to the visual logger.
+	 * @param VisualLoggerListener The listener object which is making this reward. This must be set to use logging.
+	 * @param VisualLoggerAgentId The agent id associated with this reward.
+	 * @param VisualLoggerRotationLocationA A location for the visual logger to display the first rotation in the world.
+	 * @param VisualLoggerRotationLocationB A location for the visual logger to display the second rotation in the world.
+	 * @param VisualLoggerLocation A location for the visual logger information in the world.
+	 * @param VisualLoggerColor The color for the visual logger display.
+	 * @return The resulting reward value.
+	 */
+	UFUNCTION(BlueprintPure, Category = "LearningAgents", meta = (AdvancedDisplay = 5, DefaultToSelf = "VisualLoggerListener"))
+	static float MakeRewardFromRotationSimilarity(
+		const FRotator RotationA, 
+		const FRotator RotationB, 
+		const float AngleScale = 90.0f, 
+		const float RewardScale = 1.0f,
+		const FName Name = TEXT("RewardFromRotationSimilarity"),
+		const bool bVisualLoggerEnabled = false,
+		ULearningAgentsManagerListener* VisualLoggerListener = nullptr,
+		const int32 VisualLoggerAgentId = -1,
+		const FVector VisualLoggerRotationLocationA = FVector::ZeroVector,
+		const FVector VisualLoggerRotationLocationB = FVector::ZeroVector,
+		const FVector VisualLoggerLocation = FVector::ZeroVector,
+		const FLinearColor VisualLoggerColor = FLinearColor::Green);
 
-	UFUNCTION(BlueprintPure, Category = "LearningAgents")
-	static float RewardFromDirectionSimilarity(const FVector DirectionA, const FVector DirectionB, const float AngleScale = 90.0f, const float RewardScale = 1.0f);
+	/**
+	 * Make a reward based on how similar two rotations are represented as quaternions.
+	 *
+	 * @param RotationA The first rotation.
+	 * @param RotationB The second rotation.
+	 * @param AngleScale The expected scale for the angle between rotations.
+	 * @param RewardScale The scale of the reward. Use a negative scale to create a penalty.
+	 * @param Name The name of the reward. Used for debugging.
+	 * @param bVisualLoggerEnabled When true, debug data will be sent to the visual logger.
+	 * @param VisualLoggerListener The listener object which is making this reward. This must be set to use logging.
+	 * @param VisualLoggerAgentId The agent id associated with this reward.
+	 * @param VisualLoggerRotationLocationA A location for the visual logger to display the first rotation in the world.
+	 * @param VisualLoggerRotationLocationB A location for the visual logger to display the second rotation in the world.
+	 * @param VisualLoggerLocation A location for the visual logger information in the world.
+	 * @param VisualLoggerColor The color for the visual logger display.
+	 * @return The resulting reward value.
+	 */
+	UFUNCTION(BlueprintPure, Category = "LearningAgents", meta = (AdvancedDisplay = 5, DefaultToSelf = "VisualLoggerListener"))
+	static float MakeRewardFromRotationSimilarityAsQuats(
+		const FQuat RotationA, 
+		const FQuat RotationB, 
+		const float AngleScale = 90.0f,
+		const float RewardScale = 1.0f,
+		const FName Name = TEXT("RewardFromRotationSimilarity"),
+		const bool bVisualLoggerEnabled = false,
+		ULearningAgentsManagerListener* VisualLoggerListener = nullptr,
+		const int32 VisualLoggerAgentId = -1,
+		const FVector VisualLoggerRotationLocationA = FVector::ZeroVector,
+		const FVector VisualLoggerRotationLocationB = FVector::ZeroVector,
+		const FVector VisualLoggerLocation = FVector::ZeroVector,
+		const FLinearColor VisualLoggerColor = FLinearColor::Green);
 
-	// Spline Rewards
-
-	UFUNCTION(BlueprintPure, Category = "LearningAgents")
-	static float RewardFromVelocityAlongSpline(const USplineComponent* SplineComponent, const FVector Position, const FVector Velocity, const float VelocityScale = 200.0f, const float RewardScale = 1.0f, const float FiniteDifferenceDelta = 10.0f);
-
-	// Penalties
-
-	UFUNCTION(BlueprintPure, Category = "LearningAgents")
-	static float PenaltyOnCondition(const bool bCondition, const float PenaltyScale = 1.0f);
-
-	UFUNCTION(BlueprintPure, Category = "LearningAgents")
-	static float PenaltyOnLocationDifferenceAboveThreshold(const FVector LocationA, const FVector LocationB, const float DistanceThreshold = 100.0f, const float PenaltyScale = 1.0f);
-
-	UFUNCTION(BlueprintPure, Category = "LearningAgents")
-	static float PenaltyFromLocationDifference(const FVector LocationA, const FVector LocationB, const float LocationScale = 100.0f, const float PenaltyScale = 1.0f);
-
-	UFUNCTION(BlueprintPure, Category = "LearningAgents")
-	static float PenaltyFromAngleDifference(const float AngleA, const float AngleB, const float AngleScale = 90.0f, const float PenaltyScale = 1.0f);
-
-	UFUNCTION(BlueprintPure, Category = "LearningAgents")
-	static float PenaltyFromRotationDifference(const FRotator RotationA, const FRotator RotationB, const float AngleScale = 90.0f, const float PenaltyScale = 1.0f);
-	static float PenaltyFromRotationDifferenceAsQuats(const FQuat RotationA, const FQuat RotationB, const float AngleScale = 90.0f, const float PenaltyScale = 1.0f);
-
-	UFUNCTION(BlueprintPure, Category = "LearningAgents")
-	static float PenaltyFromDirectionDifference(const FVector DirectionA, const FVector DirectionB, const float AngleScale = 90.0f, const float PenaltyScale = 1.0f);
+	/**
+	 * Make a reward based on how similar two angles are.
+	 *
+	 * @param AngleA The first angle.
+	 * @param AngleB The second angle.
+	 * @param AngleScale The expected scale for the angle between angles.
+	 * @param RewardScale The scale of the reward. Use a negative scale to create a penalty.
+	 * @param Name The name of the reward. Used for debugging.
+	 * @param bVisualLoggerEnabled When true, debug data will be sent to the visual logger.
+	 * @param VisualLoggerListener The listener object which is making this reward. This must be set to use logging.
+	 * @param VisualLoggerAgentId The agent id associated with this reward.
+	 * @param VisualLoggerAngleLocationA A location for the visual logger to display the first angle in the world.
+	 * @param VisualLoggerAngleLocationB A location for the visual logger to display the second angle in the world.
+	 * @param VisualLoggerLocation A location for the visual logger information in the world.
+	 * @param VisualLoggerColor The color for the visual logger display.
+	 * @return The resulting reward value.
+	 */
+	UFUNCTION(BlueprintPure, Category = "LearningAgents", meta = (AdvancedDisplay = 5, DefaultToSelf = "VisualLoggerListener"))
+	static float MakeRewardFromAngleSimilarity(
+		const float AngleA, 
+		const float AngleB, 
+		const float AngleScale = 90.0f, 
+		const float RewardScale = 1.0f,
+		const FName Name = TEXT("RewardFromAngleSimilarity"),
+		const bool bVisualLoggerEnabled = false,
+		ULearningAgentsManagerListener* VisualLoggerListener = nullptr,
+		const int32 VisualLoggerAgentId = -1,
+		const FVector VisualLoggerAngleLocationA = FVector::ZeroVector,
+		const FVector VisualLoggerAngleLocationB = FVector::ZeroVector,
+		const FVector VisualLoggerLocation = FVector::ZeroVector,
+		const FLinearColor VisualLoggerColor = FLinearColor::Green);
 
+	/**
+	 * Make a reward based on how similar two directions are.
+	 *
+	 * @param DirectionA The first angle.
+	 * @param DirectionB The second angle.
+	 * @param AngleScale The expected scale for the angle between the directions.
+	 * @param RewardScale The scale of the reward. Use a negative scale to create a penalty.
+	 * @param Name The name of the reward. Used for debugging.
+	 * @param bVisualLoggerEnabled When true, debug data will be sent to the visual logger.
+	 * @param VisualLoggerListener The listener object which is making this reward. This must be set to use logging.
+	 * @param VisualLoggerAgentId The agent id associated with this reward.
+	 * @param VisualLoggerDirectionLocationA A location for the visual logger to display the first direction in the world.
+	 * @param VisualLoggerDirectionLocationB A location for the visual logger to display the second direction in the world.
+	 * @param VisualLoggerLocation A location for the visual logger information in the world.
+	 * @param VisualLoggerArrowLength The length of the arrow to display for the directions.
+	 * @param VisualLoggerColor The color for the visual logger display.
+	 * @return The resulting reward value.
+	 */
+	UFUNCTION(BlueprintPure, Category = "LearningAgents", meta = (AdvancedDisplay = 5, DefaultToSelf = "VisualLoggerListener", DirectionA="1.0,0.0,0.0", DirectionB = "1.0,0.0,0.0"))
+	static float MakeRewardFromDirectionSimilarity(
+		const FVector DirectionA, 
+		const FVector DirectionB, 
+		const float AngleScale = 90.0f, 
+		const float RewardScale = 1.0f,
+		const FName Name = TEXT("RewardFromDirectionSimilarity"),
+		const bool bVisualLoggerEnabled = false,
+		ULearningAgentsManagerListener* VisualLoggerListener = nullptr,
+		const int32 VisualLoggerAgentId = -1,
+		const FVector VisualLoggerDirectionLocationA = FVector::ZeroVector,
+		const FVector VisualLoggerDirectionLocationB = FVector::ZeroVector,
+		const FVector VisualLoggerLocation = FVector::ZeroVector,
+		const float VisualLoggerArrowLength = 100.0f,
+		const FLinearColor VisualLoggerColor = FLinearColor::Green);
 
+	/**
+	 * Make a reward based on the velocity an object is moving along a spline.
+	 *
+	 * @param SplineComponent The spline.
+	 * @param Location The location of the object.
+	 * @param Velocity The velocity of the object.
+	 * @param VelocityScale The expected scale for the velocity.
+	 * @param RewardScale The scale of the reward. Use a negative scale to create a penalty.
+	 * @param FiniteDifferenceDelta The finite difference to use when computing the velocity along the spline.
+	 * @param Name The name of the reward. Used for debugging.
+	 * @param bVisualLoggerEnabled When true, debug data will be sent to the visual logger.
+	 * @param VisualLoggerListener The listener object which is making this reward. This must be set to use logging.
+	 * @param VisualLoggerAgentId The agent id associated with this reward.
+	 * @param VisualLoggerLocation A location for the visual logger information in the world.
+	 * @param VisualLoggerColor The color for the visual logger display.
+	 * @return The resulting reward value.
+	 */
+	UFUNCTION(BlueprintPure, Category = "LearningAgents", meta = (AdvancedDisplay = 7, DefaultToSelf = "VisualLoggerListener"))
+	static float MakeRewardFromVelocityAlongSpline(
+		const USplineComponent* SplineComponent, 
+		const FVector Location, 
+		const FVector Velocity, 
+		const float VelocityScale = 200.0f, 
+		const float RewardScale = 1.0f, 
+		const float FiniteDifferenceDelta = 10.0f,
+		const FName Name = TEXT("RewardFromVelocityAlongSpline"),
+		const bool bVisualLoggerEnabled = false,
+		ULearningAgentsManagerListener* VisualLoggerListener = nullptr,
+		const int32 VisualLoggerAgentId = -1,
+		const FVector VisualLoggerLocation = FVector::ZeroVector,
+		const FLinearColor VisualLoggerColor = FLinearColor::Green);
 };
