LearningAgents: Refactored to make Observation and Action functions static and use "AdvancedDisplay" markup. Added Visual Logging and comments for Completions.

#virtualized

diff --git a/Engine/Plugins/Experimental/LearningAgents/Source/Learning/Private/LearningAction.cpp b/Engine/Plugins/Experimental/LearningAgents/Source/Learning/Private/LearningAction.cpp
--- a/Engine/Plugins/Experimental/LearningAgents/Source/Learning/Private/LearningAction.cpp
+++ b/Engine/Plugins/Experimental/LearningAgents/Source/Learning/Private/LearningAction.cpp
@@ -117,10 +117,10 @@ namespace UE::Learning::Action
 		}
 	}
 
-	FSchemaElement FSchema::CreateNull(const FName Name)
+	FSchemaElement FSchema::CreateNull(const FName Tag)
 	{
 		const int32 Index = Types.Add(EType::Null);
-		Names.Add(Name);
+		Tags.Add(Tag);
 		EncodedVectorSizes.Add(0);
 		ActionVectorSizes.Add(0);
 		ActionDistributionVectorSizes.Add(0);
@@ -129,13 +129,13 @@ namespace UE::Learning::Action
 		return { Index, Generation };
 	}
 
-	FSchemaElement FSchema::CreateContinuous(const FSchemaContinuousParameters Parameters, const FName Name)
+	FSchemaElement FSchema::CreateContinuous(const FSchemaContinuousParameters Parameters, const FName Tag)
 	{
 		FContinuousData ElementData;
 		ElementData.Num = Parameters.Num;
 
 		const int32 Index = Types.Add(EType::Continuous);
-		Names.Add(Name);
+		Tags.Add(Tag);
 		EncodedVectorSizes.Add(2 * Parameters.Num);
 		ActionVectorSizes.Add(Parameters.Num);
 		ActionDistributionVectorSizes.Add(2 * Parameters.Num);
@@ -144,7 +144,7 @@ namespace UE::Learning::Action
 		return { Index, Generation };
 	}
 
-	FSchemaElement FSchema::CreateDiscreteExclusive(const FSchemaDiscreteExclusiveParameters Parameters, const FName Name)
+	FSchemaElement FSchema::CreateDiscreteExclusive(const FSchemaDiscreteExclusiveParameters Parameters, const FName Tag)
 	{
 		UE_LEARNING_CHECK(Parameters.PriorProbabilities.Num() == Parameters.Num);
 		UE_LEARNING_CHECK(Private::CheckPriorProbabilitiesExclusive(Parameters.PriorProbabilities));
@@ -156,7 +156,7 @@ namespace UE::Learning::Action
 		PriorProbabilities.Append(Parameters.PriorProbabilities);
 
 		const int32 Index = Types.Add(EType::DiscreteExclusive);
-		Names.Add(Name);
+		Tags.Add(Tag);
 		EncodedVectorSizes.Add(Parameters.Num);
 		ActionVectorSizes.Add(Parameters.Num);
 		ActionDistributionVectorSizes.Add(Parameters.Num);
@@ -165,7 +165,7 @@ namespace UE::Learning::Action
 		return { Index, Generation };
 	}
 
-	FSchemaElement FSchema::CreateDiscreteInclusive(const FSchemaDiscreteInclusiveParameters Parameters, const FName Name)
+	FSchemaElement FSchema::CreateDiscreteInclusive(const FSchemaDiscreteInclusiveParameters Parameters, const FName Tag)
 	{
 		UE_LEARNING_CHECK(Parameters.PriorProbabilities.Num() == Parameters.Num);
 		UE_LEARNING_CHECK(Private::CheckPriorProbabilitiesInclusive(Parameters.PriorProbabilities));
@@ -177,7 +177,7 @@ namespace UE::Learning::Action
 		PriorProbabilities.Append(Parameters.PriorProbabilities);
 
 		const int32 Index = Types.Add(EType::DiscreteInclusive);
-		Names.Add(Name);
+		Tags.Add(Tag);
 		EncodedVectorSizes.Add(Parameters.Num);
 		ActionVectorSizes.Add(Parameters.Num);
 		ActionDistributionVectorSizes.Add(Parameters.Num);
@@ -186,7 +186,7 @@ namespace UE::Learning::Action
 		return { Index, Generation };
 	}
 
-	FSchemaElement FSchema::CreateAnd(const FSchemaAndParameters Parameters, const FName Name)
+	FSchemaElement FSchema::CreateAnd(const FSchemaAndParameters Parameters, const FName Tag)
 	{
 		UE_LEARNING_CHECK(Parameters.Elements.Num() == Parameters.ElementNames.Num());
 		UE_LEARNING_CHECK(!Private::ContainsDuplicates(Parameters.ElementNames));
@@ -200,7 +200,7 @@ namespace UE::Learning::Action
 		SubElementObjects.Append(Parameters.Elements);
 
 		const int32 Index = Types.Add(EType::And);
-		Names.Add(Name);
+		Tags.Add(Tag);
 		EncodedVectorSizes.Add(Private::GetTotalEncodedActionVectorSize(*this, Parameters.Elements));
 		ActionVectorSizes.Add(Private::GetTotalActionVectorSize(*this, Parameters.Elements));
 		ActionDistributionVectorSizes.Add(Private::GetTotalActionDistributionVectorSize(*this, Parameters.Elements));
@@ -209,7 +209,7 @@ namespace UE::Learning::Action
 		return { Index, Generation };
 	}
 
-	FSchemaElement FSchema::CreateOrExclusive(const FSchemaOrExclusiveParameters Parameters, const FName Name)
+	FSchemaElement FSchema::CreateOrExclusive(const FSchemaOrExclusiveParameters Parameters, const FName Tag)
 	{
 		UE_LEARNING_CHECK(Parameters.Elements.Num() == Parameters.ElementNames.Num());
 		UE_LEARNING_CHECK(!Private::ContainsDuplicates(Parameters.ElementNames));
@@ -227,7 +227,7 @@ namespace UE::Learning::Action
 		PriorProbabilities.Append(Parameters.PriorProbabilities);
 
 		const int32 Index = Types.Add(EType::OrExclusive);
-		Names.Add(Name);
+		Tags.Add(Tag);
 		EncodedVectorSizes.Add(Private::GetTotalEncodedActionVectorSize(*this, Parameters.Elements) + Parameters.Elements.Num());
 		ActionVectorSizes.Add(Private::GetMaxActionVectorSize(*this, Parameters.Elements) + Parameters.Elements.Num());
 		ActionDistributionVectorSizes.Add(Private::GetTotalActionDistributionVectorSize(*this, Parameters.Elements) + Parameters.Elements.Num());
@@ -236,7 +236,7 @@ namespace UE::Learning::Action
 		return { Index, Generation };
 	}
 
-	FSchemaElement FSchema::CreateOrInclusive(const FSchemaOrInclusiveParameters Parameters, const FName Name)
+	FSchemaElement FSchema::CreateOrInclusive(const FSchemaOrInclusiveParameters Parameters, const FName Tag)
 	{
 		UE_LEARNING_CHECK(Parameters.Elements.Num() == Parameters.ElementNames.Num());
 		UE_LEARNING_CHECK(!Private::ContainsDuplicates(Parameters.ElementNames));
@@ -254,7 +254,7 @@ namespace UE::Learning::Action
 		PriorProbabilities.Append(Parameters.PriorProbabilities);
 
 		const int32 Index = Types.Add(EType::OrInclusive);
-		Names.Add(Name);
+		Tags.Add(Tag);
 		EncodedVectorSizes.Add(Private::GetTotalEncodedActionVectorSize(*this, Parameters.Elements) + Parameters.Elements.Num());
 		ActionVectorSizes.Add(Private::GetTotalActionVectorSize(*this, Parameters.Elements) + Parameters.Elements.Num());
 		ActionDistributionVectorSizes.Add(Private::GetTotalActionDistributionVectorSize(*this, Parameters.Elements) + Parameters.Elements.Num());
@@ -263,7 +263,7 @@ namespace UE::Learning::Action
 		return { Index, Generation };
 	}
 
-	FSchemaElement FSchema::CreateArray(const FSchemaArrayParameters Parameters, const FName Name)
+	FSchemaElement FSchema::CreateArray(const FSchemaArrayParameters Parameters, const FName Tag)
 	{
 		UE_LEARNING_CHECK(IsValid(Parameters.Element));
 		UE_LEARNING_CHECK(Parameters.Num >= 0);
@@ -276,7 +276,7 @@ namespace UE::Learning::Action
 		SubElementObjects.Add(Parameters.Element);
 
 		const int32 Index = Types.Add(EType::Array);
-		Names.Add(Name);
+		Tags.Add(Tag);
 		EncodedVectorSizes.Add(GetEncodedVectorSize(Parameters.Element) * Parameters.Num);
 		ActionVectorSizes.Add(GetActionVectorSize(Parameters.Element) * Parameters.Num);
 		ActionDistributionVectorSizes.Add(GetActionDistributionVectorSize(Parameters.Element) * Parameters.Num);
@@ -285,7 +285,7 @@ namespace UE::Learning::Action
 		return { Index, Generation };
 	}
 
-	FSchemaElement FSchema::CreateEncoding(const FSchemaEncodingParameters Parameters, const FName Name)
+	FSchemaElement FSchema::CreateEncoding(const FSchemaEncodingParameters Parameters, const FName Tag)
 	{
 		UE_LEARNING_CHECK(IsValid(Parameters.Element));
 
@@ -299,7 +299,7 @@ namespace UE::Learning::Action
 		SubElementObjects.Add(Parameters.Element);
 
 		const int32 Index = Types.Add(EType::Encoding);
-		Names.Add(Name);
+		Tags.Add(Tag);
 		EncodedVectorSizes.Add(ElementData.EncodingSize);
 		ActionVectorSizes.Add(GetActionVectorSize(Parameters.Element));
 		ActionDistributionVectorSizes.Add(GetActionDistributionVectorSize(Parameters.Element));
@@ -319,10 +319,10 @@ namespace UE::Learning::Action
 		return Types[Element.Index];
 	}
 
-	FName FSchema::GetName(const FSchemaElement Element) const
+	FName FSchema::GetTag(const FSchemaElement Element) const
 	{
 		UE_LEARNING_CHECK(IsValid(Element));
-		return Names[Element.Index];
+		return Tags[Element.Index];
 	}
 
 	int32 FSchema::GetEncodedVectorSize(const FSchemaElement Element) const
@@ -441,7 +441,7 @@ namespace UE::Learning::Action
 	void FSchema::Empty()
 	{
 		Types.Empty();
-		Names.Empty();
+		Tags.Empty();
 		EncodedVectorSizes.Empty();
 		ActionVectorSizes.Empty();
 		ActionDistributionVectorSizes.Empty();
@@ -466,7 +466,7 @@ namespace UE::Learning::Action
 	void FSchema::Reset()
 	{
 		Types.Reset();
-		Names.Reset();
+		Tags.Reset();
 		EncodedVectorSizes.Reset();
 		ActionVectorSizes.Reset();
 		ActionDistributionVectorSizes.Reset();
@@ -488,10 +488,10 @@ namespace UE::Learning::Action
 		Generation++;
 	}
 
-	FObjectElement FObject::CreateNull(const FName Name)
+	FObjectElement FObject::CreateNull(const FName Tag)
 	{
 		const int32 Index = Types.Add(EType::Null);
-		Names.Add(Name);
+		Tags.Add(Tag);
 
 		ContinuousDataOffsets.Add(ContinuousValues.Num());
 		ContinuousDataNums.Add(0);
@@ -505,10 +505,10 @@ namespace UE::Learning::Action
 		return { Index, Generation };
 	}
 
-	FObjectElement FObject::CreateContinuous(const FObjectContinuousParameters Parameters, const FName Name)
+	FObjectElement FObject::CreateContinuous(const FObjectContinuousParameters Parameters, const FName Tag)
 	{
 		const int32 Index = Types.Add(EType::Continuous);
-		Names.Add(Name);
+		Tags.Add(Tag);
 
 		ContinuousDataOffsets.Add(ContinuousValues.Num());
 		ContinuousDataNums.Add(Parameters.Values.Num());
@@ -524,10 +524,10 @@ namespace UE::Learning::Action
 		return { Index, Generation };
 	}
 
-	FObjectElement FObject::CreateDiscreteExclusive(const FObjectDiscreteExclusiveParameters Parameters, const FName Name)
+	FObjectElement FObject::CreateDiscreteExclusive(const FObjectDiscreteExclusiveParameters Parameters, const FName Tag)
 	{
 		const int32 Index = Types.Add(EType::DiscreteExclusive);
-		Names.Add(Name);
+		Tags.Add(Tag);
 
 		ContinuousDataOffsets.Add(ContinuousValues.Num());
 		ContinuousDataNums.Add(0);
@@ -543,10 +543,10 @@ namespace UE::Learning::Action
 		return { Index, Generation };
 	}
 
-	FObjectElement FObject::CreateDiscreteInclusive(const FObjectDiscreteInclusiveParameters Parameters, const FName Name)
+	FObjectElement FObject::CreateDiscreteInclusive(const FObjectDiscreteInclusiveParameters Parameters, const FName Tag)
 	{
 		const int32 Index = Types.Add(EType::DiscreteInclusive);
-		Names.Add(Name);
+		Tags.Add(Tag);
 
 		ContinuousDataOffsets.Add(ContinuousValues.Num());
 		ContinuousDataNums.Add(0);
@@ -562,14 +562,14 @@ namespace UE::Learning::Action
 		return { Index, Generation };
 	}
 
-	FObjectElement FObject::CreateAnd(const FObjectAndParameters Parameters, const FName Name)
+	FObjectElement FObject::CreateAnd(const FObjectAndParameters Parameters, const FName Tag)
 	{
 		UE_LEARNING_CHECK(Parameters.Elements.Num() == Parameters.ElementNames.Num());
 		UE_LEARNING_CHECK(!Private::ContainsDuplicates(Parameters.ElementNames));
 		UE_LEARNING_CHECK(Private::CheckAllValid(*this, Parameters.Elements));
 
 		const int32 Index = Types.Add(EType::And);
-		Names.Add(Name);
+		Tags.Add(Tag);
 
 		ContinuousDataOffsets.Add(ContinuousValues.Num());
 		ContinuousDataNums.Add(0);
@@ -586,12 +586,12 @@ namespace UE::Learning::Action
 		return { Index, Generation };
 	}
 
-	FObjectElement FObject::CreateOrExclusive(const FObjectOrExclusiveParameters Parameters, const FName Name)
+	FObjectElement FObject::CreateOrExclusive(const FObjectOrExclusiveParameters Parameters, const FName Tag)
 	{
 		UE_LEARNING_CHECK(IsValid(Parameters.Element));
 
 		const int32 Index = Types.Add(EType::OrExclusive);
-		Names.Add(Name);
+		Tags.Add(Tag);
 
 		ContinuousDataOffsets.Add(ContinuousValues.Num());
 		ContinuousDataNums.Add(0);
@@ -608,14 +608,14 @@ namespace UE::Learning::Action
 		return { Index, Generation };
 	}
 
-	FObjectElement FObject::CreateOrInclusive(const FObjectOrInclusiveParameters Parameters, const FName Name)
+	FObjectElement FObject::CreateOrInclusive(const FObjectOrInclusiveParameters Parameters, const FName Tag)
 	{
 		UE_LEARNING_CHECK(Parameters.Elements.Num() == Parameters.ElementNames.Num());
 		UE_LEARNING_CHECK(!Private::ContainsDuplicates(Parameters.ElementNames));
 		UE_LEARNING_CHECK(Private::CheckAllValid(*this, Parameters.Elements));
 
 		const int32 Index = Types.Add(EType::OrInclusive);
-		Names.Add(Name);
+		Tags.Add(Tag);
 
 		ContinuousDataOffsets.Add(ContinuousValues.Num());
 		ContinuousDataNums.Add(0);
@@ -632,12 +632,12 @@ namespace UE::Learning::Action
 		return { Index, Generation };
 	}
 
-	FObjectElement FObject::CreateArray(const FObjectArrayParameters Parameters, const FName Name)
+	FObjectElement FObject::CreateArray(const FObjectArrayParameters Parameters, const FName Tag)
 	{
 		UE_LEARNING_CHECK(Private::CheckAllValid(*this, Parameters.Elements));
 
 		const int32 Index = Types.Add(EType::Array);
-		Names.Add(Name);
+		Tags.Add(Tag);
 
 		ContinuousDataOffsets.Add(ContinuousValues.Num());
 		ContinuousDataNums.Add(0);
@@ -657,12 +657,12 @@ namespace UE::Learning::Action
 		return { Index, Generation };
 	}
 
-	FObjectElement FObject::CreateEncoding(const FObjectEncodingParameters Parameters, const FName Name)
+	FObjectElement FObject::CreateEncoding(const FObjectEncodingParameters Parameters, const FName Tag)
 	{
 		UE_LEARNING_CHECK(IsValid(Parameters.Element));
 
 		const int32 Index = Types.Add(EType::Encoding);
-		Names.Add(Name);
+		Tags.Add(Tag);
 
 		ContinuousDataOffsets.Add(ContinuousValues.Num());
 		ContinuousDataNums.Add(0);
@@ -690,10 +690,10 @@ namespace UE::Learning::Action
 		return Types[Element.Index];
 	}
 	
-	FName FObject::GetName(const FObjectElement Element) const
+	FName FObject::GetTag(const FObjectElement Element) const
 	{
 		UE_LEARNING_CHECK(IsValid(Element));
-		return Names[Element.Index];
+		return Tags[Element.Index];
 	}
 
 	FObjectContinuousParameters FObject::GetContinuous(const FObjectElement Element) const
@@ -779,7 +779,7 @@ namespace UE::Learning::Action
 	void FObject::Empty()
 	{
 		Types.Empty();
-		Names.Empty();
+		Tags.Empty();
 		ContinuousDataOffsets.Empty();
 		ContinuousDataNums.Empty();
 		DiscreteDataOffsets.Empty();
@@ -798,7 +798,7 @@ namespace UE::Learning::Action
 	void FObject::Reset()
 	{
 		Types.Reset();
-		Names.Reset();
+		Tags.Reset();
 		ContinuousDataOffsets.Reset();
 		ContinuousDataNums.Reset();
 		DiscreteDataOffsets.Reset();
@@ -1689,7 +1689,7 @@ namespace UE::Learning::Action
 		// Check that the types match
 
 		const EType SchemaElementType = Schema.GetType(SchemaElement);
-		const FName SchemaElementName = Schema.GetName(SchemaElement);
+		const FName SchemaElementTag = Schema.GetTag(SchemaElement);
 
 		// Get Action Vector Size
 
@@ -1703,7 +1703,7 @@ namespace UE::Learning::Action
 
 		case EType::Null:
 		{
-			OutObjectElement = OutObject.CreateNull(SchemaElementName);
+			OutObjectElement = OutObject.CreateNull(SchemaElementTag);
 			return;
 		}
 
@@ -1711,7 +1711,7 @@ namespace UE::Learning::Action
 		{
 			UE_LEARNING_CHECK(ActionVectorSize == Schema.GetContinuous(SchemaElement).Num);
 
-			OutObjectElement = OutObject.CreateContinuous({ MakeArrayView(ActionVector.GetData(), ActionVector.Num()) }, SchemaElementName);
+			OutObjectElement = OutObject.CreateContinuous({ MakeArrayView(ActionVector.GetData(), ActionVector.Num()) }, SchemaElementTag);
 			return;
 		}
 
@@ -1732,7 +1732,7 @@ namespace UE::Learning::Action
 			}
 			UE_LEARNING_CHECK(ExclusiveIndex != INDEX_NONE);
 
-			OutObjectElement = OutObject.CreateDiscreteExclusive({ ExclusiveIndex }, SchemaElementName);
+			OutObjectElement = OutObject.CreateDiscreteExclusive({ ExclusiveIndex }, SchemaElementTag);
 			return;
 		}
 
@@ -1752,7 +1752,7 @@ namespace UE::Learning::Action
 				}
 			}
 
-			OutObjectElement = OutObject.CreateDiscreteInclusive({ InclusiveIndices }, SchemaElementName);
+			OutObjectElement = OutObject.CreateDiscreteInclusive({ InclusiveIndices }, SchemaElementTag);
 			return;
 		}
 
@@ -1781,7 +1781,7 @@ namespace UE::Learning::Action
 			}
 			UE_LEARNING_CHECK(SubElementOffset == ActionVectorSize);
 
-			OutObjectElement = OutObject.CreateAnd({ Parameters.ElementNames, SubElements }, SchemaElementName);
+			OutObjectElement = OutObject.CreateAnd({ Parameters.ElementNames, SubElements }, SchemaElementTag);
 			return;
 		}
 
@@ -1817,7 +1817,7 @@ namespace UE::Learning::Action
 				Parameters.Elements[SchemaElementIndex],
 				ActionVector.Slice(0, SubElementSize));
 
-			OutObjectElement = OutObject.CreateOrExclusive({ Parameters.ElementNames[SchemaElementIndex], SubElement }, SchemaElementName);
+			OutObjectElement = OutObject.CreateOrExclusive({ Parameters.ElementNames[SchemaElementIndex], SubElement }, SchemaElementTag);
 			return;
 		}
 
@@ -1860,7 +1860,7 @@ namespace UE::Learning::Action
 			}
 			UE_LEARNING_CHECK(SubElementOffset + Parameters.Elements.Num() == ActionVectorSize);
 
-			OutObjectElement = OutObject.CreateOrInclusive({ SubElementNames, SubElements }, SchemaElementName);
+			OutObjectElement = OutObject.CreateOrInclusive({ SubElementNames, SubElements }, SchemaElementTag);
 			return;
 		}
 
@@ -1885,7 +1885,7 @@ namespace UE::Learning::Action
 					ActionVector.Slice(ElementIdx * SubElementSize, SubElementSize));
 			}
 
-			OutObjectElement = OutObject.CreateArray({ SubElements }, SchemaElementName);
+			OutObjectElement = OutObject.CreateArray({ SubElements }, SchemaElementTag);
 			return;
 		}
 
@@ -1901,7 +1901,7 @@ namespace UE::Learning::Action
 				Parameters.Element,
 				ActionVector);
 
-			OutObjectElement = OutObject.CreateEncoding({ SubElement }, SchemaElementName);
+			OutObjectElement = OutObject.CreateEncoding({ SubElement }, SchemaElementTag);
 			return;
 		}
 
diff --git a/Engine/Plugins/Experimental/LearningAgents/Source/Learning/Private/LearningObservation.cpp b/Engine/Plugins/Experimental/LearningAgents/Source/Learning/Private/LearningObservation.cpp
--- a/Engine/Plugins/Experimental/LearningAgents/Source/Learning/Private/LearningObservation.cpp
+++ b/Engine/Plugins/Experimental/LearningAgents/Source/Learning/Private/LearningObservation.cpp
@@ -64,10 +64,10 @@ namespace UE::Learning::Observation
 		}
 	}
 
-	FSchemaElement FSchema::CreateNull(const FName Name)
+	FSchemaElement FSchema::CreateNull(const FName Tag)
 	{
 		const int32 Index = Types.Add(EType::Null);
-		Names.Add(Name);
+		Tags.Add(Tag);
 		ObservationVectorSizes.Add(0);
 		EncodedVectorSizes.Add(0);
 		TypeDataIndices.Add(INDEX_NONE);
@@ -75,7 +75,7 @@ namespace UE::Learning::Observation
 		return { Index, Generation };
 	}
 
-	FSchemaElement FSchema::CreateContinuous(const FSchemaContinuousParameters Parameters, const FName Name)
+	FSchemaElement FSchema::CreateContinuous(const FSchemaContinuousParameters Parameters, const FName Tag)
 	{
 		UE_LEARNING_CHECK(Parameters.Num >= 0);
 
@@ -83,7 +83,7 @@ namespace UE::Learning::Observation
 		ElementData.Num = Parameters.Num;
 
 		const int32 Index = Types.Add(EType::Continuous);
-		Names.Add(Name);
+		Tags.Add(Tag);
 		ObservationVectorSizes.Add(Parameters.Num);
 		EncodedVectorSizes.Add(Parameters.Num);
 		TypeDataIndices.Add(ContinuousData.Add(ElementData));
@@ -91,7 +91,7 @@ namespace UE::Learning::Observation
 		return { Index, Generation };
 	}
 
-	FSchemaElement FSchema::CreateAnd(const FSchemaAndParameters Parameters, const FName Name)
+	FSchemaElement FSchema::CreateAnd(const FSchemaAndParameters Parameters, const FName Tag)
 	{
 		UE_LEARNING_CHECK(Parameters.Elements.Num() == Parameters.ElementNames.Num());
 		UE_LEARNING_CHECK(!Private::ContainsDuplicates(Parameters.ElementNames));
@@ -105,7 +105,7 @@ namespace UE::Learning::Observation
 		SubElementObjects.Append(Parameters.Elements);
 
 		const int32 Index = Types.Add(EType::And);
-		Names.Add(Name);
+		Tags.Add(Tag);
 		ObservationVectorSizes.Add(Private::GetTotalObservationVectorSize(*this, Parameters.Elements));
 		EncodedVectorSizes.Add(Private::GetTotalEncodedObservationVectorSize(*this, Parameters.Elements));
 		TypeDataIndices.Add(AndData.Add(ElementData));
@@ -113,7 +113,7 @@ namespace UE::Learning::Observation
 		return { Index, Generation };
 	}
 
-	FSchemaElement FSchema::CreateOrExclusive(const FSchemaOrExclusiveParameters Parameters, const FName Name)
+	FSchemaElement FSchema::CreateOrExclusive(const FSchemaOrExclusiveParameters Parameters, const FName Tag)
 	{
 		UE_LEARNING_CHECK(Parameters.Elements.Num() == Parameters.ElementNames.Num());
 		UE_LEARNING_CHECK(!Private::ContainsDuplicates(Parameters.ElementNames));
@@ -128,7 +128,7 @@ namespace UE::Learning::Observation
 		SubElementObjects.Append(Parameters.Elements);
 
 		const int32 Index = Types.Add(EType::OrExclusive);
-		Names.Add(Name);
+		Tags.Add(Tag);
 		ObservationVectorSizes.Add(Private::GetMaxObservationVectorSize(*this, Parameters.Elements) + Parameters.Elements.Num());
 		EncodedVectorSizes.Add(Parameters.EncodingSize + Parameters.Elements.Num());
 		TypeDataIndices.Add(OrExclusiveData.Add(ElementData));
@@ -136,7 +136,7 @@ namespace UE::Learning::Observation
 		return { Index, Generation };
 	}
 
-	FSchemaElement FSchema::CreateOrInclusive(const FSchemaOrInclusiveParameters Parameters, const FName Name)
+	FSchemaElement FSchema::CreateOrInclusive(const FSchemaOrInclusiveParameters Parameters, const FName Tag)
 	{
 		UE_LEARNING_CHECK(Parameters.Elements.Num() == Parameters.ElementNames.Num());
 		UE_LEARNING_CHECK(!Private::ContainsDuplicates(Parameters.ElementNames));
@@ -153,7 +153,7 @@ namespace UE::Learning::Observation
 		SubElementObjects.Append(Parameters.Elements);
 
 		const int32 Index = Types.Add(EType::OrInclusive);
-		Names.Add(Name);
+		Tags.Add(Tag);
 		ObservationVectorSizes.Add(Private::GetTotalObservationVectorSize(*this, Parameters.Elements) + Parameters.Elements.Num());
 		EncodedVectorSizes.Add(Parameters.AttentionHeadNum * Parameters.ValueEncodingSize + Parameters.Elements.Num());
 		TypeDataIndices.Add(OrInclusiveData.Add(ElementData));
@@ -161,7 +161,7 @@ namespace UE::Learning::Observation
 		return { Index, Generation };
 	}
 
-	FSchemaElement FSchema::CreateArray(const FSchemaArrayParameters Parameters, const FName Name)
+	FSchemaElement FSchema::CreateArray(const FSchemaArrayParameters Parameters, const FName Tag)
 	{
 		UE_LEARNING_CHECK(Parameters.Num >= 0);
 		UE_LEARNING_CHECK(IsValid(Parameters.Element));
@@ -174,7 +174,7 @@ namespace UE::Learning::Observation
 		SubElementObjects.Add(Parameters.Element);
 
 		const int32 Index = Types.Add(EType::Array);
-		Names.Add(Name);
+		Tags.Add(Tag);
 		ObservationVectorSizes.Add(GetObservationVectorSize(Parameters.Element) * Parameters.Num);
 		EncodedVectorSizes.Add(GetEncodedVectorSize(Parameters.Element) * Parameters.Num);
 		TypeDataIndices.Add(ArrayData.Add(ElementData));
@@ -182,7 +182,7 @@ namespace UE::Learning::Observation
 		return { Index, Generation };
 	}
 
-	FSchemaElement FSchema::CreateSet(const FSchemaSetParameters Parameters, const FName Name)
+	FSchemaElement FSchema::CreateSet(const FSchemaSetParameters Parameters, const FName Tag)
 	{
 		UE_LEARNING_CHECK(IsValid(Parameters.Element));
 
@@ -197,7 +197,7 @@ namespace UE::Learning::Observation
 		SubElementObjects.Add(Parameters.Element);
 
 		const int32 Index = Types.Add(EType::Set);
-		Names.Add(Name);
+		Tags.Add(Tag);
 		ObservationVectorSizes.Add(GetObservationVectorSize(Parameters.Element) * Parameters.MaxNum + Parameters.MaxNum);
 		EncodedVectorSizes.Add(Parameters.ValueEncodingSize * Parameters.AttentionHeadNum + 1);
 		TypeDataIndices.Add(SetData.Add(ElementData));
@@ -205,7 +205,7 @@ namespace UE::Learning::Observation
 		return { Index, Generation };
 	}
 
-	FSchemaElement FSchema::CreateEncoding(const FSchemaEncodingParameters Parameters, const FName Name)
+	FSchemaElement FSchema::CreateEncoding(const FSchemaEncodingParameters Parameters, const FName Tag)
 	{
 		UE_LEARNING_CHECK(IsValid(Parameters.Element));
 
@@ -219,7 +219,7 @@ namespace UE::Learning::Observation
 		SubElementObjects.Add(Parameters.Element);
 
 		const int32 Index = Types.Add(EType::Encoding);
-		Names.Add(Name);
+		Tags.Add(Tag);
 		ObservationVectorSizes.Add(GetObservationVectorSize(Parameters.Element));
 		EncodedVectorSizes.Add(Parameters.EncodingSize);
 		TypeDataIndices.Add(EncodingData.Add(ElementData));
@@ -238,10 +238,10 @@ namespace UE::Learning::Observation
 		return Types[Element.Index];
 	}
 
-	FName FSchema::GetName(const FSchemaElement Element) const
+	FName FSchema::GetTag(const FSchemaElement Element) const
 	{
 		UE_LEARNING_CHECK(IsValid(Element));
-		return Names[Element.Index];
+		return Tags[Element.Index];
 	}
 
 	int32 FSchema::GetObservationVectorSize(const FSchemaElement Element) const
@@ -348,7 +348,7 @@ namespace UE::Learning::Observation
 	void FSchema::Empty()
 	{
 		Types.Empty();
-		Names.Empty();
+		Tags.Empty();
 		ObservationVectorSizes.Empty();
 		EncodedVectorSizes.Empty();
 		TypeDataIndices.Empty();
@@ -370,7 +370,7 @@ namespace UE::Learning::Observation
 	void FSchema::Reset()
 	{
 		Types.Reset();
-		Names.Reset();
+		Tags.Reset();
 		ObservationVectorSizes.Reset();
 		EncodedVectorSizes.Reset();
 		TypeDataIndices.Reset();
@@ -389,10 +389,10 @@ namespace UE::Learning::Observation
 		Generation++;
 	}
 
-	FObjectElement FObject::CreateNull(const FName Name)
+	FObjectElement FObject::CreateNull(const FName Tag)
 	{
 		const int32 Index = Types.Add(EType::Null);
-		Names.Add(Name);
+		Tags.Add(Tag);
 
 		ContinuousDataOffsets.Add(ContinuousValues.Num());
 		ContinuousDataNums.Add(0);
@@ -403,10 +403,10 @@ namespace UE::Learning::Observation
 		return { Index, Generation };
 	}
 
-	FObjectElement FObject::CreateContinuous(const FObjectContinuousParameters Parameters, const FName Name)
+	FObjectElement FObject::CreateContinuous(const FObjectContinuousParameters Parameters, const FName Tag)
 	{
 		const int32 Index = Types.Add(EType::Continuous);
-		Names.Add(Name);
+		Tags.Add(Tag);
 
 		ContinuousDataOffsets.Add(ContinuousValues.Num());
 		ContinuousDataNums.Add(Parameters.Values.Num());
@@ -419,14 +419,14 @@ namespace UE::Learning::Observation
 		return { Index, Generation };
 	}
 
-	FObjectElement FObject::CreateAnd(const FObjectAndParameters Parameters, const FName Name)
+	FObjectElement FObject::CreateAnd(const FObjectAndParameters Parameters, const FName Tag)
 	{
 		UE_LEARNING_CHECK(Parameters.Elements.Num() == Parameters.ElementNames.Num());
 		UE_LEARNING_CHECK(!Private::ContainsDuplicates(Parameters.ElementNames));
 		UE_LEARNING_CHECK(Private::CheckAllValid(*this, Parameters.Elements));
 
 		const int32 Index = Types.Add(EType::And);
-		Names.Add(Name);
+		Tags.Add(Tag);
 
 		ContinuousDataOffsets.Add(ContinuousValues.Num());
 		ContinuousDataNums.Add(0);
@@ -440,12 +440,12 @@ namespace UE::Learning::Observation
 		return { Index, Generation };
 	}
 
-	FObjectElement FObject::CreateOrExclusive(const FObjectOrExclusiveParameters Parameters, const FName Name)
+	FObjectElement FObject::CreateOrExclusive(const FObjectOrExclusiveParameters Parameters, const FName Tag)
 	{
 		UE_LEARNING_CHECK(IsValid(Parameters.Element));
 
 		const int32 Index = Types.Add(EType::OrExclusive);
-		Names.Add(Name);
+		Tags.Add(Tag);
 
 		ContinuousDataOffsets.Add(ContinuousValues.Num());
 		ContinuousDataNums.Add(0);
@@ -459,14 +459,14 @@ namespace UE::Learning::Observation
 		return { Index, Generation };
 	}
 
-	FObjectElement FObject::CreateOrInclusive(const FObjectOrInclusiveParameters Parameters, const FName Name)
+	FObjectElement FObject::CreateOrInclusive(const FObjectOrInclusiveParameters Parameters, const FName Tag)
 	{
 		UE_LEARNING_CHECK(Parameters.Elements.Num() == Parameters.ElementNames.Num());
 		UE_LEARNING_CHECK(!Private::ContainsDuplicates(Parameters.ElementNames));
 		UE_LEARNING_CHECK(Private::CheckAllValid(*this, Parameters.Elements));
 
 		const int32 Index = Types.Add(EType::OrInclusive);
-		Names.Add(Name);
+		Tags.Add(Tag);
 
 		ContinuousDataOffsets.Add(ContinuousValues.Num());
 		ContinuousDataNums.Add(0);
@@ -480,12 +480,12 @@ namespace UE::Learning::Observation
 		return { Index, Generation };
 	}
 
-	FObjectElement FObject::CreateArray(const FObjectArrayParameters Parameters, const FName Name)
+	FObjectElement FObject::CreateArray(const FObjectArrayParameters Parameters, const FName Tag)
 	{
 		UE_LEARNING_CHECK(Private::CheckAllValid(*this, Parameters.Elements));
 
 		const int32 Index = Types.Add(EType::Array);
-		Names.Add(Name);
+		Tags.Add(Tag);
 
 		ContinuousDataOffsets.Add(ContinuousValues.Num());
 		ContinuousDataNums.Add(0);
@@ -502,12 +502,12 @@ namespace UE::Learning::Observation
 		return { Index, Generation };
 	}
 
-	FObjectElement FObject::CreateSet(const FObjectSetParameters Parameters, const FName Name)
+	FObjectElement FObject::CreateSet(const FObjectSetParameters Parameters, const FName Tag)
 	{
 		UE_LEARNING_CHECK(Private::CheckAllValid(*this, Parameters.Elements));
 
 		const int32 Index = Types.Add(EType::Set);
-		Names.Add(Name);
+		Tags.Add(Tag);
 
 		ContinuousDataOffsets.Add(ContinuousValues.Num());
 		ContinuousDataNums.Add(0);
@@ -524,12 +524,12 @@ namespace UE::Learning::Observation
 		return { Index, Generation };
 	}
 
-	FObjectElement FObject::CreateEncoding(const FObjectEncodingParameters Parameters, const FName Name)
+	FObjectElement FObject::CreateEncoding(const FObjectEncodingParameters Parameters, const FName Tag)
 	{
 		UE_LEARNING_CHECK(IsValid(Parameters.Element));
 
 		const int32 Index = Types.Add(EType::Encoding);
-		Names.Add(Name);
+		Tags.Add(Tag);
 
 		ContinuousDataOffsets.Add(ContinuousValues.Num());
 		ContinuousDataNums.Add(0);
@@ -554,10 +554,10 @@ namespace UE::Learning::Observation
 		return Types[Element.Index];
 	}
 
-	FName FObject::GetName(const FObjectElement Element) const
+	FName FObject::GetTag(const FObjectElement Element) const
 	{
 		UE_LEARNING_CHECK(IsValid(Element));
-		return Names[Element.Index];
+		return Tags[Element.Index];
 	}
 
 	FObjectContinuousParameters FObject::GetContinuous(const FObjectElement Element) const
@@ -634,7 +634,7 @@ namespace UE::Learning::Observation
 	void FObject::Empty()
 	{
 		Types.Empty();
-		Names.Empty();
+		Tags.Empty();
 
 		ContinuousDataOffsets.Empty();
 		ContinuousDataNums.Empty();
@@ -652,7 +652,7 @@ namespace UE::Learning::Observation
 	void FObject::Reset()
 	{
 		Types.Reset();
-		Names.Reset();
+		Tags.Reset();
 
 		ContinuousDataOffsets.Reset();
 		ContinuousDataNums.Reset();
@@ -827,11 +827,11 @@ namespace UE::Learning::Observation
 
 			UE_LEARNING_CHECKF(ReturnElement.GetInputSize() == Schema.GetObservationVectorSize(SchemaElement),
 				TEXT("Encoder Network Input unexpected size for %s. Got %i, expected %i according to Schema."),
-				*Schema.GetName(SchemaElement).ToString(), ReturnElement.GetInputSize(), Schema.GetObservationVectorSize(SchemaElement));
+				*Schema.GetTag(SchemaElement).ToString(), ReturnElement.GetInputSize(), Schema.GetObservationVectorSize(SchemaElement));
 
 			UE_LEARNING_CHECKF(ReturnElement.GetOutputSize() == Schema.GetEncodedVectorSize(SchemaElement),
 				TEXT("Encoder Network Output unexpected size for %s. Got %i, expected %i according to Schema."),
-				*Schema.GetName(SchemaElement).ToString(), ReturnElement.GetOutputSize(), Schema.GetEncodedVectorSize(SchemaElement));
+				*Schema.GetTag(SchemaElement).ToString(), ReturnElement.GetOutputSize(), Schema.GetEncodedVectorSize(SchemaElement));
 
 			return ReturnElement;
 		}
@@ -1309,7 +1309,7 @@ namespace UE::Learning::Observation
 		// Check that the types match
 
 		const EType SchemaElementType = Schema.GetType(SchemaElement);
-		const FName SchemaElementName = Schema.GetName(SchemaElement);
+		const FName SchemaElementTag = Schema.GetTag(SchemaElement);
 
 		// Get Observation Vector Size
 
@@ -1323,7 +1323,7 @@ namespace UE::Learning::Observation
 
 		case EType::Null:
 		{
-			OutObjectElement = OutObject.CreateNull(SchemaElementName);
+			OutObjectElement = OutObject.CreateNull(SchemaElementTag);
 			return;
 		}
 
@@ -1331,7 +1331,7 @@ namespace UE::Learning::Observation
 		{
 			UE_LEARNING_CHECK(ObservationVectorSize == Schema.GetContinuous(SchemaElement).Num);
 
-			OutObjectElement = OutObject.CreateContinuous({ MakeArrayView(ObservationVector.GetData(), ObservationVector.Num()) }, SchemaElementName);
+			OutObjectElement = OutObject.CreateContinuous({ MakeArrayView(ObservationVector.GetData(), ObservationVector.Num()) }, SchemaElementTag);
 			return;
 		}
 
@@ -1360,7 +1360,7 @@ namespace UE::Learning::Observation
 			}
 			UE_LEARNING_CHECK(SubElementOffset == ObservationVectorSize);
 
-			OutObjectElement = OutObject.CreateAnd({ Parameters.ElementNames, SubElements }, SchemaElementName);
+			OutObjectElement = OutObject.CreateAnd({ Parameters.ElementNames, SubElements }, SchemaElementTag);
 			return;
 		}
 
@@ -1396,7 +1396,7 @@ namespace UE::Learning::Observation
 				Parameters.Elements[SchemaElementIndex],
 				ObservationVector.Slice(0, SubElementSize));
 
-			OutObjectElement = OutObject.CreateOrExclusive({ Parameters.ElementNames[SchemaElementIndex], SubElement }, SchemaElementName);
+			OutObjectElement = OutObject.CreateOrExclusive({ Parameters.ElementNames[SchemaElementIndex], SubElement }, SchemaElementTag);
 			return;
 		}
 
@@ -1442,7 +1442,7 @@ namespace UE::Learning::Observation
 			}
 			UE_LEARNING_CHECK(SubElementOffset + Parameters.Elements.Num() == ObservationVectorSize);
 
-			OutObjectElement = OutObject.CreateOrInclusive({ SubElementNames, SubElements }, SchemaElementName);
+			OutObjectElement = OutObject.CreateOrInclusive({ SubElementNames, SubElements }, SchemaElementTag);
 			return;
 		}
 
@@ -1467,7 +1467,7 @@ namespace UE::Learning::Observation
 					ObservationVector.Slice(ElementIdx * SubElementSize, SubElementSize));
 			}
 
-			OutObjectElement = OutObject.CreateArray({ SubElements }, SchemaElementName);
+			OutObjectElement = OutObject.CreateArray({ SubElements }, SchemaElementTag);
 			return;
 		}
 
@@ -1505,7 +1505,7 @@ namespace UE::Learning::Observation
 				SubElements.Add(SubElement);
 			}
 
-			OutObjectElement = OutObject.CreateSet({ SubElements }, SchemaElementName);
+			OutObjectElement = OutObject.CreateSet({ SubElements }, SchemaElementTag);
 			return;
 		}
 
@@ -1521,7 +1521,7 @@ namespace UE::Learning::Observation
 				Parameters.Element,
 				ObservationVector);
 
-			OutObjectElement = OutObject.CreateEncoding({ SubElement }, SchemaElementName);
+			OutObjectElement = OutObject.CreateEncoding({ SubElement }, SchemaElementTag);
 			return;
 		}
 
diff --git a/Engine/Plugins/Experimental/LearningAgents/Source/Learning/Public/LearningAction.h b/Engine/Plugins/Experimental/LearningAgents/Source/Learning/Public/LearningAction.h
--- a/Engine/Plugins/Experimental/LearningAgents/Source/Learning/Public/LearningAction.h
+++ b/Engine/Plugins/Experimental/LearningAgents/Source/Learning/Public/LearningAction.h
@@ -149,15 +149,15 @@ namespace UE::Learning::Action
 	 */
 	struct LEARNING_API FSchema
 	{
-		FSchemaElement CreateNull(const FName Name = NAME_None);
-		FSchemaElement CreateContinuous(const FSchemaContinuousParameters Parameters, const FName Name = NAME_None);
-		FSchemaElement CreateDiscreteExclusive(const FSchemaDiscreteExclusiveParameters Parameters, const FName Name = NAME_None);
-		FSchemaElement CreateDiscreteInclusive(const FSchemaDiscreteInclusiveParameters Parameters, const FName Name = NAME_None);
-		FSchemaElement CreateAnd(const FSchemaAndParameters Parameters, const FName Name = NAME_None);
-		FSchemaElement CreateOrExclusive(const FSchemaOrExclusiveParameters Parameters, const FName Name = NAME_None);
-		FSchemaElement CreateOrInclusive(const FSchemaOrInclusiveParameters Parameters, const FName Name = NAME_None);
-		FSchemaElement CreateArray(const FSchemaArrayParameters Parameters, const FName Name = NAME_None);
-		FSchemaElement CreateEncoding(const FSchemaEncodingParameters Parameters, const FName Name = NAME_None);
+		FSchemaElement CreateNull(const FName Tag = NAME_None);
+		FSchemaElement CreateContinuous(const FSchemaContinuousParameters Parameters, const FName Tag = NAME_None);
+		FSchemaElement CreateDiscreteExclusive(const FSchemaDiscreteExclusiveParameters Parameters, const FName Tag = NAME_None);
+		FSchemaElement CreateDiscreteInclusive(const FSchemaDiscreteInclusiveParameters Parameters, const FName Tag = NAME_None);
+		FSchemaElement CreateAnd(const FSchemaAndParameters Parameters, const FName Tag = NAME_None);
+		FSchemaElement CreateOrExclusive(const FSchemaOrExclusiveParameters Parameters, const FName Tag = NAME_None);
+		FSchemaElement CreateOrInclusive(const FSchemaOrInclusiveParameters Parameters, const FName Tag = NAME_None);
+		FSchemaElement CreateArray(const FSchemaArrayParameters Parameters, const FName Tag = NAME_None);
+		FSchemaElement CreateEncoding(const FSchemaEncodingParameters Parameters, const FName Tag = NAME_None);
 
 		FSchemaContinuousParameters GetContinuous(const FSchemaElement Element) const;
 		FSchemaDiscreteExclusiveParameters GetDiscreteExclusive(const FSchemaElement Element) const;
@@ -174,8 +174,8 @@ namespace UE::Learning::Action
 		// Gets the type of the given element
 		EType GetType(const FSchemaElement Element) const;
 
-		// Gets the name of the given element
-		FName GetName(const FSchemaElement Element) const;
+		// Gets the tag of the given element
+		FName GetTag(const FSchemaElement Element) const;
 
 		// Get the encoded vector size of the given element
 		int32 GetEncodedVectorSize(const FSchemaElement Element) const;
@@ -252,7 +252,7 @@ namespace UE::Learning::Action
 
 		/** These have entries for each Schema Element */
 		TArray<EType> Types;
-		TArray<FName> Names;
+		TArray<FName> Tags;
 		TArray<int32> EncodedVectorSizes;
 		TArray<int32> ActionVectorSizes;
 		TArray<int32> ActionDistributionVectorSizes;
@@ -350,15 +350,15 @@ namespace UE::Learning::Action
 	 */
 	struct LEARNING_API FObject
 	{
-		FObjectElement CreateNull(const FName Name = NAME_None);
-		FObjectElement CreateContinuous(const FObjectContinuousParameters Parameters, const FName Name = NAME_None);
-		FObjectElement CreateDiscreteExclusive(const FObjectDiscreteExclusiveParameters Parameters, const FName Name = NAME_None);
-		FObjectElement CreateDiscreteInclusive(const FObjectDiscreteInclusiveParameters Parameters, const FName Name = NAME_None);
-		FObjectElement CreateAnd(const FObjectAndParameters Parameters, const FName Name = NAME_None);
-		FObjectElement CreateOrExclusive(const FObjectOrExclusiveParameters Parameters, const FName Name = NAME_None);
-		FObjectElement CreateOrInclusive(const FObjectOrInclusiveParameters Parameters, const FName Name = NAME_None);
-		FObjectElement CreateArray(const FObjectArrayParameters Parameters, const FName Name = NAME_None);
-		FObjectElement CreateEncoding(const FObjectEncodingParameters Parameters, const FName Name = NAME_None);
+		FObjectElement CreateNull(const FName Tag = NAME_None);
+		FObjectElement CreateContinuous(const FObjectContinuousParameters Parameters, const FName Tag = NAME_None);
+		FObjectElement CreateDiscreteExclusive(const FObjectDiscreteExclusiveParameters Parameters, const FName Tag = NAME_None);
+		FObjectElement CreateDiscreteInclusive(const FObjectDiscreteInclusiveParameters Parameters, const FName Tag = NAME_None);
+		FObjectElement CreateAnd(const FObjectAndParameters Parameters, const FName Tag = NAME_None);
+		FObjectElement CreateOrExclusive(const FObjectOrExclusiveParameters Parameters, const FName Tag = NAME_None);
+		FObjectElement CreateOrInclusive(const FObjectOrInclusiveParameters Parameters, const FName Tag = NAME_None);
+		FObjectElement CreateArray(const FObjectArrayParameters Parameters, const FName Tag = NAME_None);
+		FObjectElement CreateEncoding(const FObjectEncodingParameters Parameters, const FName Tag = NAME_None);
 
 		FObjectContinuousParameters GetContinuous(const FObjectElement Element) const;
 		FObjectDiscreteExclusiveParameters GetDiscreteExclusive(const FObjectElement Element) const;
@@ -375,8 +375,8 @@ namespace UE::Learning::Action
 		// Gets the type of the given element
 		EType GetType(const FObjectElement Element) const;
 
-		// Gets the name of the given element
-		FName GetName(const FObjectElement Element) const;
+		// Gets the tag of the given element
+		FName GetTag(const FObjectElement Element) const;
 
 		// Get the current generation
 		uint32 GetGeneration() const;
@@ -392,7 +392,7 @@ namespace UE::Learning::Action
 		uint32 Generation = 0;
 
 		TArray<EType> Types;
-		TArray<FName> Names;
+		TArray<FName> Tags;
 		TArray<int32> ContinuousDataOffsets;
 		TArray<int32> ContinuousDataNums;
 		TArray<int32> DiscreteDataOffsets;
diff --git a/Engine/Plugins/Experimental/LearningAgents/Source/Learning/Public/LearningObservation.h b/Engine/Plugins/Experimental/LearningAgents/Source/Learning/Public/LearningObservation.h
--- a/Engine/Plugins/Experimental/LearningAgents/Source/Learning/Public/LearningObservation.h
+++ b/Engine/Plugins/Experimental/LearningAgents/Source/Learning/Public/LearningObservation.h
@@ -152,14 +152,14 @@ namespace UE::Learning::Observation
 	 */
 	struct LEARNING_API FSchema
 	{
-		FSchemaElement CreateNull(const FName Name = NAME_None);
-		FSchemaElement CreateContinuous(const FSchemaContinuousParameters Parameters, const FName Name = NAME_None);
-		FSchemaElement CreateAnd(const FSchemaAndParameters Parameters, const FName Name = NAME_None);
-		FSchemaElement CreateOrExclusive(const FSchemaOrExclusiveParameters Parameters, const FName Name = NAME_None);
-		FSchemaElement CreateOrInclusive(const FSchemaOrInclusiveParameters Parameters, const FName Name = NAME_None);
-		FSchemaElement CreateArray(const FSchemaArrayParameters Parameters, const FName Name = NAME_None);
-		FSchemaElement CreateSet(const FSchemaSetParameters Parameters, const FName Name = NAME_None);
-		FSchemaElement CreateEncoding(const FSchemaEncodingParameters Parameters, const FName Name = NAME_None);
+		FSchemaElement CreateNull(const FName Tag = NAME_None);
+		FSchemaElement CreateContinuous(const FSchemaContinuousParameters Parameters, const FName Tag = NAME_None);
+		FSchemaElement CreateAnd(const FSchemaAndParameters Parameters, const FName Tag = NAME_None);
+		FSchemaElement CreateOrExclusive(const FSchemaOrExclusiveParameters Parameters, const FName Tag = NAME_None);
+		FSchemaElement CreateOrInclusive(const FSchemaOrInclusiveParameters Parameters, const FName Tag = NAME_None);
+		FSchemaElement CreateArray(const FSchemaArrayParameters Parameters, const FName Tag = NAME_None);
+		FSchemaElement CreateSet(const FSchemaSetParameters Parameters, const FName Tag = NAME_None);
+		FSchemaElement CreateEncoding(const FSchemaEncodingParameters Parameters, const FName Tag = NAME_None);
 
 		FSchemaContinuousParameters GetContinuous(const FSchemaElement Element) const;
 		FSchemaAndParameters GetAnd(const FSchemaElement Element) const;
@@ -175,8 +175,8 @@ namespace UE::Learning::Observation
 		// Gets the type of the given element
 		EType GetType(const FSchemaElement Element) const;
 
-		// Gets the name of the given element
-		FName GetName(const FSchemaElement Element) const;
+		// Gets the tag of the given element
+		FName GetTag(const FSchemaElement Element) const;
 
 		// Get the observation vector size of the given element
 		int32 GetObservationVectorSize(const FSchemaElement Element) const;
@@ -250,7 +250,7 @@ namespace UE::Learning::Observation
 
 		/** These have entries for each Schema Element */
 		TArray<EType> Types;
-		TArray<FName> Names;
+		TArray<FName> Tags;
 		TArray<int32> ObservationVectorSizes;
 		TArray<int32> EncodedVectorSizes;
 		TArray<int32> TypeDataIndices;
@@ -337,14 +337,14 @@ namespace UE::Learning::Observation
 	 */
 	struct LEARNING_API FObject
 	{
-		FObjectElement CreateNull(const FName Name = NAME_None);
-		FObjectElement CreateContinuous(const FObjectContinuousParameters Parameters, const FName Name = NAME_None);
-		FObjectElement CreateAnd(const FObjectAndParameters Parameters, const FName Name = NAME_None);
-		FObjectElement CreateOrExclusive(const FObjectOrExclusiveParameters Parameters, const FName Name = NAME_None);
-		FObjectElement CreateOrInclusive(const FObjectOrInclusiveParameters Parameters, const FName Name = NAME_None);
-		FObjectElement CreateArray(const FObjectArrayParameters Parametes, const FName Name = NAME_None);
-		FObjectElement CreateSet(const FObjectSetParameters Parameters, const FName Name = NAME_None);
-		FObjectElement CreateEncoding(const FObjectEncodingParameters Parameters, const FName Name = NAME_None);
+		FObjectElement CreateNull(const FName Tag = NAME_None);
+		FObjectElement CreateContinuous(const FObjectContinuousParameters Parameters, const FName Tag = NAME_None);
+		FObjectElement CreateAnd(const FObjectAndParameters Parameters, const FName Tag = NAME_None);
+		FObjectElement CreateOrExclusive(const FObjectOrExclusiveParameters Parameters, const FName Tag = NAME_None);
+		FObjectElement CreateOrInclusive(const FObjectOrInclusiveParameters Parameters, const FName Tag = NAME_None);
+		FObjectElement CreateArray(const FObjectArrayParameters Parametes, const FName Tag = NAME_None);
+		FObjectElement CreateSet(const FObjectSetParameters Parameters, const FName Tag = NAME_None);
+		FObjectElement CreateEncoding(const FObjectEncodingParameters Parameters, const FName Tag = NAME_None);
 
 		FObjectContinuousParameters GetContinuous(const FObjectElement Element) const;
 		FObjectAndParameters GetAnd(const FObjectElement Element) const;
@@ -360,8 +360,8 @@ namespace UE::Learning::Observation
 		// Gets the type of the given element
 		EType GetType(const FObjectElement Element) const;
 
-		// Gets the name of the given element
-		FName GetName(const FObjectElement Element) const;
+		// Gets the tag of the given element
+		FName GetTag(const FObjectElement Element) const;
 
 		// Get the current generation
 		uint32 GetGeneration() const;
@@ -377,7 +377,7 @@ namespace UE::Learning::Observation
 		uint32 Generation = 0;
 
 		TArray<EType> Types;
-		TArray<FName> Names;
+		TArray<FName> Tags;
 		TArray<int32> ContinuousDataOffsets;
 		TArray<int32> ContinuousDataNums;
 		TArray<int32> SubElementDataOffsets;
