LearningAgents: Refactored to make Observation and Action functions static and use "AdvancedDisplay" markup. Added Visual Logging and comments for Completions.

diff --git a/Engine/Plugins/Experimental/LearningAgents/Source/LearningAgents/Private/LearningAgentsActions.cpp b/Engine/Plugins/Experimental/LearningAgents/Source/LearningAgents/Private/LearningAgentsActions.cpp
--- a/Engine/Plugins/Experimental/LearningAgents/Source/LearningAgents/Private/LearningAgentsActions.cpp
+++ b/Engine/Plugins/Experimental/LearningAgents/Source/LearningAgents/Private/LearningAgentsActions.cpp
@@ -17,11 +17,6 @@ uint32 GetTypeHash(const FLearningAgentsActionObjectElement& Element)
 	return (uint32)Element.ObjectElement.Index;
 }
 
-const UE::Learning::Action::FSchema& ULearningAgentsActionSchema::GetActionSchema() const
-{
-	return ActionSchema;
-}
-
 namespace UE::Learning::Agents::Action::Private
 {
 	static inline bool ContainsDuplicates(const TArrayView<const int32> Indices)
@@ -61,32 +56,31 @@ namespace UE::Learning::Agents::Action::Private
 		const Learning::Action::FSchema& Schema,
 		const Learning::Action::FSchemaElement SchemaElement,
 		const Learning::Action::FObject& Object,
-		const Learning::Action::FObjectElement ObjectElement,
-		const FString& ObjectName)
+		const Learning::Action::FObjectElement ObjectElement)
 	{
 		// Check Elements are Valid
 
 		if (!Schema.IsValid(SchemaElement))
 		{
-			UE_LOG(LogLearning, Error, TEXT("%s: Invalid Action Schema Element."), *ObjectName);
+			UE_LOG(LogLearning, Error, TEXT("ValidateObjectMatchesSchema: Invalid Action Schema Element."));
 			return false;
 		}
 
 		if (!Object.IsValid(ObjectElement))
 		{
-			UE_LOG(LogLearning, Error, TEXT("%s: Invalid Action Object Element."), *ObjectName);
+			UE_LOG(LogLearning, Error, TEXT("ValidateObjectMatchesSchema: Invalid Action Object Element."));
 			return false;
 		}
 
 		// Check Names Match
 
-		const FName ActionSchemaElementName = Schema.GetName(SchemaElement);
-		const FName ActionObjectElementName = Object.GetName(ObjectElement);
+		const FName ActionSchemaElementTag = Schema.GetTag(SchemaElement);
+		const FName ActionObjectElementTag = Object.GetTag(ObjectElement);
 
-		if (ActionSchemaElementName != ActionObjectElementName)
+		if (ActionSchemaElementTag != ActionObjectElementTag)
 		{
-			UE_LOG(LogLearning, Warning, TEXT("%s: Action name does not match Schema. Expected '%s', got '%s'."),
-				*ObjectName, *ActionSchemaElementName.ToString(), *ActionObjectElementName.ToString());
+			UE_LOG(LogLearning, Warning, TEXT("ValidateObjectMatchesSchema: Action tag does not match Schema. Expected '%s', got '%s'."),
+				*ActionSchemaElementTag.ToString(), *ActionObjectElementTag.ToString());
 		}
 
 		// Check Types Match
@@ -96,9 +90,8 @@ namespace UE::Learning::Agents::Action::Private
 
 		if (ActionSchemaElementType != ActionObjectElementType)
 		{
-			UE_LOG(LogLearning, Error, TEXT("%s: Action '%s' type does not match Schema. Expected type '%s', got type '%s'."),
-				*ObjectName,
-				*ActionSchemaElementName.ToString(),
+			UE_LOG(LogLearning, Error, TEXT("ValidateObjectMatchesSchema: Action '%s' type does not match Schema. Expected type '%s', got type '%s'."),
+				*ActionSchemaElementTag.ToString(),
 				GetActionTypeString(ActionSchemaElementType),
 				GetActionTypeString(ActionObjectElementType));
 			return false;
@@ -117,9 +110,8 @@ namespace UE::Learning::Agents::Action::Private
 
 			if (SchemaElementSize != ObjectElementSize)
 			{
-				UE_LOG(LogLearning, Error, TEXT("%s: Action '%s' size does not match Schema. Expected '%i', got '%i'."),
-					*ObjectName,
-					*ActionSchemaElementName.ToString(),
+				UE_LOG(LogLearning, Error, TEXT("ValidateObjectMatchesSchema: Action '%s' size does not match Schema. Expected '%i', got '%i'."),
+					*ActionSchemaElementTag.ToString(),
 					SchemaElementSize,
 					ObjectElementSize);
 				return false;
@@ -135,9 +127,8 @@ namespace UE::Learning::Agents::Action::Private
 
 			if (ObjectElementIndex < 0 || ObjectElementIndex >= SchemaElementSize)
 			{
-				UE_LOG(LogLearning, Error, TEXT("%s: Action '%s' index out of range for Schema. Expected '<%i', got '%i'."),
-					*ObjectName,
-					*ActionSchemaElementName.ToString(),
+				UE_LOG(LogLearning, Error, TEXT("ValidateObjectMatchesSchema: Action '%s' index out of range for Schema. Expected '<%i', got '%i'."),
+					*ActionSchemaElementTag.ToString(),
 					SchemaElementSize,
 					ObjectElementIndex);
 				return false;
@@ -153,9 +144,8 @@ namespace UE::Learning::Agents::Action::Private
 
 			if (ObjectElementIndices.Num() > SchemaElementSize)
 			{
-				UE_LOG(LogLearning, Error, TEXT("%s: Action '%s' too many indices provided. Expected at most '%i', got '%i'."),
-					*ObjectName,
-					*ActionSchemaElementName.ToString(),
+				UE_LOG(LogLearning, Error, TEXT("ValidateObjectMatchesSchema: Action '%s' too many indices provided. Expected at most '%i', got '%i'."),
+					*ActionSchemaElementTag.ToString(),
 					SchemaElementSize,
 					ObjectElementIndices.Num());
 				return false;
@@ -165,9 +155,8 @@ namespace UE::Learning::Agents::Action::Private
 			{
 				if (ObjectElementIndices[SubElementIdx] < 0 || ObjectElementIndices[SubElementIdx] >= SchemaElementSize)
 				{
-					UE_LOG(LogLearning, Error, TEXT("%s: Action '%s' index out of range for Schema. Expected '<%i', got '%i'."),
-						*ObjectName,
-						*ActionSchemaElementName.ToString(),
+					UE_LOG(LogLearning, Error, TEXT("ValidateObjectMatchesSchema: Action '%s' index out of range for Schema. Expected '<%i', got '%i'."),
+						*ActionSchemaElementTag.ToString(),
 						SchemaElementSize,
 						ObjectElementIndices[SubElementIdx]);
 					return false;
@@ -186,9 +175,8 @@ namespace UE::Learning::Agents::Action::Private
 
 			if (SchemaParameters.Elements.Num() != ObjectParameters.Elements.Num())
 			{
-				UE_LOG(LogLearning, Error, TEXT("%s: Action '%s' number of sub-elements does not match Schema. Expected '%i', got '%i'."),
-					*ObjectName,
-					*ActionSchemaElementName.ToString(),
+				UE_LOG(LogLearning, Error, TEXT("ValidateObjectMatchesSchema: Action '%s' number of sub-elements does not match Schema. Expected '%i', got '%i'."),
+					*ActionSchemaElementTag.ToString(),
 					SchemaParameters.Elements.Num(),
 					ObjectParameters.Elements.Num());
 				return false;
@@ -200,9 +188,8 @@ namespace UE::Learning::Agents::Action::Private
 
 				if (ObjectElementIdx == INDEX_NONE)
 				{
-					UE_LOG(LogLearning, Error, TEXT("%s: Action '%s' does not include '%s' action required by Schema."),
-						*ObjectName,
-						*ActionSchemaElementName.ToString(),
+					UE_LOG(LogLearning, Error, TEXT("ValidateObjectMatchesSchema: Action '%s' does not include '%s' action required by Schema."),
+						*ActionSchemaElementTag.ToString(),
 						*SchemaParameters.ElementNames[SchemaElementIdx].ToString());
 					return false;
 				}
@@ -211,8 +198,7 @@ namespace UE::Learning::Agents::Action::Private
 					Schema,
 					SchemaParameters.Elements[SchemaElementIdx],
 					Object,
-					ObjectParameters.Elements[ObjectElementIdx],
-					ObjectName))
+					ObjectParameters.Elements[ObjectElementIdx]))
 				{
 					return false;
 				}
@@ -231,9 +217,8 @@ namespace UE::Learning::Agents::Action::Private
 
 			if (SchemaSubElementIdx == INDEX_NONE)
 			{
-				UE_LOG(LogLearning, Error, TEXT("%s: Action '%s' Schema does not include '%s' action."),
-					*ObjectName,
-					*ActionSchemaElementName.ToString(),
+				UE_LOG(LogLearning, Error, TEXT("ValidateObjectMatchesSchema: Action '%s' Schema does not include '%s' action."),
+					*ActionSchemaElementTag.ToString(),
 					*ObjectParameters.ElementName.ToString());
 				return false;
 			}
@@ -242,8 +227,7 @@ namespace UE::Learning::Agents::Action::Private
 				Schema,
 				SchemaParameters.Elements[SchemaSubElementIdx],
 				Object,
-				ObjectParameters.Element,
-				ObjectName);
+				ObjectParameters.Element);
 		}
 
 		case Learning::Action::EType::OrInclusive:
@@ -253,9 +237,8 @@ namespace UE::Learning::Agents::Action::Private
 
 			if (ObjectParameters.Elements.Num() > SchemaParameters.Elements.Num())
 			{
-				UE_LOG(LogLearning, Error, TEXT("%s: Action '%s' too many sub-actions provided. Expected at most '%i', got '%i'."),
-					*ObjectName,
-					*ActionSchemaElementName.ToString(),
+				UE_LOG(LogLearning, Error, TEXT("ValidateObjectMatchesSchema: Action '%s' too many sub-actions provided. Expected at most '%i', got '%i'."),
+					*ActionSchemaElementTag.ToString(),
 					SchemaParameters.Elements.Num(),
 					ObjectParameters.Elements.Num());
 				return false;
@@ -267,9 +250,8 @@ namespace UE::Learning::Agents::Action::Private
 
 				if (SchemaSubElementIdx == INDEX_NONE)
 				{
-					UE_LOG(LogLearning, Error, TEXT("%s: Action '%s' Schema does not include '%s' action."),
-						*ObjectName,
-						*ActionSchemaElementName.ToString(),
+					UE_LOG(LogLearning, Error, TEXT("ValidateObjectMatchesSchema: Action '%s' Schema does not include '%s' action."),
+						*ActionSchemaElementTag.ToString(),
 						*ObjectParameters.ElementNames[ObjectSubElementIdx].ToString());
 					return false;
 				}
@@ -278,8 +260,7 @@ namespace UE::Learning::Agents::Action::Private
 					Schema,
 					SchemaParameters.Elements[SchemaSubElementIdx],
 					Object,
-					ObjectParameters.Elements[ObjectSubElementIdx],
-					ObjectName))
+					ObjectParameters.Elements[ObjectSubElementIdx]))
 				{
 					return false;
 				}
@@ -295,9 +276,8 @@ namespace UE::Learning::Agents::Action::Private
 
 			if (ObjectParameters.Elements.Num() != SchemaParameters.Num)
 			{
-				UE_LOG(LogLearning, Error, TEXT("%s: Action '%s' array incorrect size. Expected '%i' elements, got '%i'."),
-					*ObjectName,
-					*ActionSchemaElementName.ToString(),
+				UE_LOG(LogLearning, Error, TEXT("ValidateObjectMatchesSchema: Action '%s' array incorrect size. Expected '%i' elements, got '%i'."),
+					*ActionSchemaElementTag.ToString(),
 					SchemaParameters.Num,
 					ObjectParameters.Elements.Num());
 				return false;
@@ -309,8 +289,7 @@ namespace UE::Learning::Agents::Action::Private
 					Schema,
 					SchemaParameters.Element,
 					Object,
-					ObjectParameters.Elements[ElementIdx],
-					ObjectName))
+					ObjectParameters.Elements[ElementIdx]))
 				{
 					return false;
 				}
@@ -328,8 +307,7 @@ namespace UE::Learning::Agents::Action::Private
 					Schema,
 					SchemaParameters.Element,
 					Object,
-					ObjectParameters.Element,
-					ObjectName);
+					ObjectParameters.Element);
 		}
 
 		default:
@@ -344,23 +322,22 @@ namespace UE::Learning::Agents::Action::Private
 		const UE::Learning::Action::FObject& Object,
 		const UE::Learning::Action::FObjectElement ObjectElement,
 		const FString& Indentation,
-		const FString& Prefix,
-		const FString& ObjectName)
+		const FString& Prefix)
 	{
 		if (!Object.IsValid(ObjectElement))
 		{
-			UE_LOG(LogLearning, Error, TEXT("%s: Invalid Action Object Element."), *ObjectName);
+			UE_LOG(LogLearning, Error, TEXT("LogAction: Invalid Action Object Element."));
 			return;
 		}
 
 		const UE::Learning::Action::EType Type = Object.GetType(ObjectElement);
-		const FName Name = Object.GetName(ObjectElement);
+		const FName Tag = Object.GetTag(ObjectElement);
 
 		switch (Type)
 		{
 		case UE::Learning::Action::EType::Null:
 		{
-			UE_LOG(LogLearning, Display, TEXT("%s%s \"%s\" (%s)"), *Indentation, *Prefix, *Name.ToString(), GetActionTypeString(Type));
+			UE_LOG(LogLearning, Display, TEXT("%s%s \"%s\" (%s)"), *Indentation, *Prefix, *Tag.ToString(), GetActionTypeString(Type));
 			return;
 		}
 
@@ -368,7 +345,7 @@ namespace UE::Learning::Agents::Action::Private
 		{
 			const UE::Learning::Action::FObjectContinuousParameters Parameters = Object.GetContinuous(ObjectElement);
 
-			UE_LOG(LogLearning, Display, TEXT("%s%s \"%s\" (%s) %s"), *Indentation, *Prefix, *Name.ToString(), GetActionTypeString(Type), *UE::Learning::Array::FormatFloat(Parameters.Values));
+			UE_LOG(LogLearning, Display, TEXT("%s%s \"%s\" (%s) %s"), *Indentation, *Prefix, *Tag.ToString(), GetActionTypeString(Type), *UE::Learning::Array::FormatFloat(Parameters.Values));
 			return;
 		}
 
@@ -376,7 +353,7 @@ namespace UE::Learning::Agents::Action::Private
 		{
 			const UE::Learning::Action::FObjectDiscreteExclusiveParameters Parameters = Object.GetDiscreteExclusive(ObjectElement);
 
-			UE_LOG(LogLearning, Display, TEXT("%s%s \"%s\" (%s) %i"), *Indentation, *Prefix, *Name.ToString(), GetActionTypeString(Type), Parameters.DiscreteIndex);
+			UE_LOG(LogLearning, Display, TEXT("%s%s \"%s\" (%s) %i"), *Indentation, *Prefix, *Tag.ToString(), GetActionTypeString(Type), Parameters.DiscreteIndex);
 			return;
 		}
 
@@ -384,7 +361,7 @@ namespace UE::Learning::Agents::Action::Private
 		{
 			const UE::Learning::Action::FObjectDiscreteInclusiveParameters Parameters = Object.GetDiscreteInclusive(ObjectElement);
 
-			UE_LOG(LogLearning, Display, TEXT("%s%s \"%s\" (%s) %s"), *Indentation, *Prefix, *Name.ToString(), GetActionTypeString(Type), *UE::Learning::Array::FormatInt32(Parameters.DiscreteIndices));
+			UE_LOG(LogLearning, Display, TEXT("%s%s \"%s\" (%s) %s"), *Indentation, *Prefix, *Tag.ToString(), GetActionTypeString(Type), *UE::Learning::Array::FormatInt32(Parameters.DiscreteIndices));
 			return;
 		}
 
@@ -392,10 +369,10 @@ namespace UE::Learning::Agents::Action::Private
 		{
 			const UE::Learning::Action::FObjectAndParameters Parameters = Object.GetAnd(ObjectElement);
 
-			UE_LOG(LogLearning, Display, TEXT("%s%s \"%s\" (%s)"), *Indentation, *Prefix, *Name.ToString(), GetActionTypeString(Type));
+			UE_LOG(LogLearning, Display, TEXT("%s%s \"%s\" (%s)"), *Indentation, *Prefix, *Tag.ToString(), GetActionTypeString(Type));
 			for (int32 SubElementIdx = 0; SubElementIdx < Parameters.Elements.Num(); SubElementIdx++)
 			{
-				LogAction(Object, Parameters.Elements[SubElementIdx], *(Indentation + TEXT("    ")), FString::Printf(TEXT("| \"%s\": "), *Parameters.ElementNames[SubElementIdx].ToString()), ObjectName);
+				LogAction(Object, Parameters.Elements[SubElementIdx], *(Indentation + TEXT("    ")), FString::Printf(TEXT("| \"%s\": "), *Parameters.ElementNames[SubElementIdx].ToString()));
 			}
 
 			return;
@@ -405,8 +382,8 @@ namespace UE::Learning::Agents::Action::Private
 		{
 			const UE::Learning::Action::FObjectOrExclusiveParameters Parameters = Object.GetOrExclusive(ObjectElement);
 
-			UE_LOG(LogLearning, Display, TEXT("%s%s \"%s\" (%s)"), *Indentation, *Prefix, *Name.ToString(), GetActionTypeString(Type));
-			LogAction(Object, Parameters.Element, *(Indentation + TEXT("    ")), FString::Printf(TEXT("| \"%s\": "), *Parameters.ElementName.ToString()), ObjectName);
+			UE_LOG(LogLearning, Display, TEXT("%s%s \"%s\" (%s)"), *Indentation, *Prefix, *Tag.ToString(), GetActionTypeString(Type));
+			LogAction(Object, Parameters.Element, *(Indentation + TEXT("    ")), FString::Printf(TEXT("| \"%s\": "), *Parameters.ElementName.ToString()));
 
 			return;
 		}
@@ -415,10 +392,10 @@ namespace UE::Learning::Agents::Action::Private
 		{
 			const UE::Learning::Action::FObjectOrInclusiveParameters Parameters = Object.GetOrInclusive(ObjectElement);
 
-			UE_LOG(LogLearning, Display, TEXT("%s%s \"%s\" (%s)"), *Indentation, *Prefix, *Name.ToString(), GetActionTypeString(Type));
+			UE_LOG(LogLearning, Display, TEXT("%s%s \"%s\" (%s)"), *Indentation, *Prefix, *Tag.ToString(), GetActionTypeString(Type));
 			for (int32 SubElementIdx = 0; SubElementIdx < Parameters.Elements.Num(); SubElementIdx++)
 			{
-				LogAction(Object, Parameters.Elements[SubElementIdx], *(Indentation + TEXT("    ")), FString::Printf(TEXT("| \"%s\": "), *Parameters.ElementNames[SubElementIdx].ToString()), ObjectName);
+				LogAction(Object, Parameters.Elements[SubElementIdx], *(Indentation + TEXT("    ")), FString::Printf(TEXT("| \"%s\": "), *Parameters.ElementNames[SubElementIdx].ToString()));
 			}
 
 			return;
@@ -428,10 +405,10 @@ namespace UE::Learning::Agents::Action::Private
 		{
 			const UE::Learning::Action::FObjectArrayParameters Parameters = Object.GetArray(ObjectElement);
 
-			UE_LOG(LogLearning, Display, TEXT("%s%s \"%s\" (%s)"), *Indentation, *Prefix, *Name.ToString(), GetActionTypeString(Type));
+			UE_LOG(LogLearning, Display, TEXT("%s%s \"%s\" (%s)"), *Indentation, *Prefix, *Tag.ToString(), GetActionTypeString(Type));
 			for (int32 SubElementIdx = 0; SubElementIdx < Parameters.Elements.Num(); SubElementIdx++)
 			{
-				LogAction(Object, Parameters.Elements[SubElementIdx], *(Indentation + TEXT("    ")), FString::Printf(TEXT("| %3i:"), SubElementIdx), ObjectName);
+				LogAction(Object, Parameters.Elements[SubElementIdx], *(Indentation + TEXT("    ")), FString::Printf(TEXT("| %3i:"), SubElementIdx));
 			}
 
 			return;
@@ -441,8 +418,8 @@ namespace UE::Learning::Agents::Action::Private
 		{
 			const UE::Learning::Action::FObjectEncodingParameters Parameters = Object.GetEncoding(ObjectElement);
 
-			UE_LOG(LogLearning, Display, TEXT("%s%s \"%s\" (%s)"), *Indentation, *Prefix, *Name.ToString(), GetActionTypeString(Type));
-			LogAction(Object, Parameters.Element, *(Indentation + TEXT("    ")), TEXT("|"), ObjectName);
+			UE_LOG(LogLearning, Display, TEXT("%s%s \"%s\" (%s)"), *Indentation, *Prefix, *Tag.ToString(), GetActionTypeString(Type));
+			LogAction(Object, Parameters.Element, *(Indentation + TEXT("    ")), TEXT("|"));
 
 			return;
 		}
@@ -465,14 +442,14 @@ namespace UE::Learning::Agents::Action::Private
 			FMath::Exp(V.Z));
 	}
 
-	static inline void NormalizeProbabilitiesExclusive(TArrayView<float> PriorProbabilities, const FString& ObjectName)
+	static inline void NormalizeProbabilitiesExclusive(TArrayView<float> PriorProbabilities)
 	{
 		float Total = 0.0f;
 		for (int32 Idx = 0; Idx < PriorProbabilities.Num(); Idx++)
 		{
 			if (PriorProbabilities[Idx] < 0.0f || PriorProbabilities[Idx] > 1.0f)
 			{
-				UE_LOG(LogLearning, Warning, TEXT("%s: Invalid Prior Probability Given (%f), must be in range 0 to 1."), *ObjectName, PriorProbabilities[Idx]);
+				UE_LOG(LogLearning, Warning, TEXT("NormalizeProbabilitiesExclusive: Invalid Prior Probability Given (%f), must be in range 0 to 1."), PriorProbabilities[Idx]);
 			}
 
 			PriorProbabilities[Idx] = FMath::Clamp(PriorProbabilities[Idx], 0.0f, 1.0f);
@@ -481,7 +458,7 @@ namespace UE::Learning::Agents::Action::Private
 
 		if (PriorProbabilities.Num() > 0 && FMath::Abs(Total) < UE_SMALL_NUMBER)
 		{
-			UE_LOG(LogLearning, Warning, TEXT("%s: Prior Probabilities are too small. Should sum to 1."), *ObjectName);
+			UE_LOG(LogLearning, Warning, TEXT("NormalizeProbabilitiesExclusive: Prior Probabilities are too small. Should sum to 1."));
 
 			for (int32 Idx = 0; Idx < PriorProbabilities.Num(); Idx++)
 			{
@@ -497,13 +474,13 @@ namespace UE::Learning::Agents::Action::Private
 		}
 	}
 
-	static inline void NormalizeProbabilitiesInclusive(TArrayView<float> PriorProbabilities, const FString& ObjectName)
+	static inline void NormalizeProbabilitiesInclusive(TArrayView<float> PriorProbabilities)
 	{
 		for (int32 Idx = 0; Idx < PriorProbabilities.Num(); Idx++)
 		{
 			if (PriorProbabilities[Idx] < 0.0f || PriorProbabilities[Idx] > 1.0f)
 			{
-				UE_LOG(LogLearning, Warning, TEXT("%s: Invalid Prior Probability Given (%f), must be in range 0 to 1."), *ObjectName, PriorProbabilities[Idx]);
+				UE_LOG(LogLearning, Warning, TEXT("NormalizeProbabilitiesInclusive: Invalid Prior Probability Given (%f), must be in range 0 to 1."), PriorProbabilities[Idx]);
 			}
 
 			PriorProbabilities[Idx] = FMath::Clamp(PriorProbabilities[Idx], 0.0f, 1.0f);
@@ -522,56 +499,86 @@ namespace UE::Learning::Agents::Action::Private
 	}
 }
 
-bool ULearningAgentsActionSchema::ValidateObjectMatchesSchema(
+bool ULearningAgentsActions::ValidateObjectMatchesSchema(
+	const ULearningAgentsActionSchema* Schema,
 	const FLearningAgentsActionSchemaElement SchemaElement,
 	const ULearningAgentsActionObject* Object,
-	const FLearningAgentsActionObjectElement ObjectElement) const
+	const FLearningAgentsActionObjectElement ObjectElement)
 {
+	if (!Schema)
+	{
+		UE_LOG(LogLearning, Error, TEXT("ValidateObjectMatchesSchema: Schema is nullptr."));
+		return false;
+	}
+
+	if (!Object)
+	{
+		UE_LOG(LogLearning, Error, TEXT("ValidateObjectMatchesSchema: Object is nullptr."));
+		return false;
+	}
+
 	return UE::Learning::Agents::Action::Private::ValidateObjectMatchesSchema(
-		ActionSchema,
+		Schema->ActionSchema,
 		SchemaElement.SchemaElement,
-		Object->GetActionObject(),
-		ObjectElement.ObjectElement,
-		GetName());
+		Object->ActionObject,
+		ObjectElement.ObjectElement);
 }
 
-FLearningAgentsActionSchemaElement ULearningAgentsActionSchema::SpecifyNullAction(const FName Name)
+FLearningAgentsActionSchemaElement ULearningAgentsActions::SpecifyNullAction(ULearningAgentsActionSchema* Schema, const FName Tag)
 {
-	return { ActionSchema.CreateNull(Name) };
+	if (!Schema)
+	{
+		UE_LOG(LogLearning, Error, TEXT("SpecifyNullAction: Schema is nullptr."));
+		return FLearningAgentsActionSchemaElement();
+	}
+
+	return { Schema->ActionSchema.CreateNull(Tag) };
 }
 
-FLearningAgentsActionSchemaElement ULearningAgentsActionSchema::SpecifyContinuousAction(const int32 Size, const FName Name)
+FLearningAgentsActionSchemaElement ULearningAgentsActions::SpecifyContinuousAction(ULearningAgentsActionSchema* Schema, const int32 Size, const FName Tag)
 {
+	if (!Schema)
+	{
+		UE_LOG(LogLearning, Error, TEXT("SpecifyContinuousAction: Schema is nullptr."));
+		return FLearningAgentsActionSchemaElement();
+	}
+
 	if (Size < 0)
 	{
-		UE_LOG(LogLearning, Error, TEXT("%s: Invalid Continuous Action Size '%i'."), *GetName(), Size);
+		UE_LOG(LogLearning, Error, TEXT("SpecifyContinuousAction: Invalid Continuous Action Size '%i'."), Size);
 		return FLearningAgentsActionSchemaElement();
 	}
 
 	if (Size == 0)
 	{
-		UE_LOG(LogLearning, Warning, TEXT("%s: Specifying zero-sized Continuous Action."), *GetName());
+		UE_LOG(LogLearning, Warning, TEXT("SpecifyContinuousAction: Specifying zero-sized Continuous Action."));
 	}
 
-	return { ActionSchema.CreateContinuous({ Size }, Name) };
+	return { Schema->ActionSchema.CreateContinuous({ Size }, Tag) };
 }
 
-FLearningAgentsActionSchemaElement ULearningAgentsActionSchema::SpecifyExclusiveDiscreteAction(const int32 Size, const TArray<float>& PriorProbabilities, const FName Name)
+FLearningAgentsActionSchemaElement ULearningAgentsActions::SpecifyExclusiveDiscreteAction(ULearningAgentsActionSchema* Schema, const int32 Size, const TArray<float>& PriorProbabilities, const FName Tag)
 {
-	return SpecifyExclusiveDiscreteActionFromArrayView(Size, PriorProbabilities, Name);
+	return SpecifyExclusiveDiscreteActionFromArrayView(Schema, Size, PriorProbabilities, Tag);
 }
 
-FLearningAgentsActionSchemaElement ULearningAgentsActionSchema::SpecifyExclusiveDiscreteActionFromArrayView(const int32 Size, const TArrayView<const float> PriorProbabilities, const FName Name)
+FLearningAgentsActionSchemaElement ULearningAgentsActions::SpecifyExclusiveDiscreteActionFromArrayView(ULearningAgentsActionSchema* Schema, const int32 Size, const TArrayView<const float> PriorProbabilities, const FName Tag)
 {
+	if (!Schema)
+	{
+		UE_LOG(LogLearning, Error, TEXT("SpecifyExclusiveDiscreteActionFromArrayView: Schema is nullptr."));
+		return FLearningAgentsActionSchemaElement();
+	}
+
 	if (Size < 0)
 	{
-		UE_LOG(LogLearning, Error, TEXT("%s: Invalid DiscreteExclusive Action Size '%i'."), *GetName(), Size);
+		UE_LOG(LogLearning, Error, TEXT("SpecifyExclusiveDiscreteActionFromArrayView: Invalid DiscreteExclusive Action Size '%i'."), Size);
 		return FLearningAgentsActionSchemaElement();
 	}
 
 	if (Size == 0)
 	{
-		UE_LOG(LogLearning, Warning, TEXT("%s: Specifying zero-sized Exclusive Discrete Action."), *GetName());
+		UE_LOG(LogLearning, Warning, TEXT("SpecifyExclusiveDiscreteActionFromArrayView: Specifying zero-sized Exclusive Discrete Action."));
 	}
 
 	TArray<float, TInlineAllocator<16>> NormalizedPriorProbabilities;
@@ -580,27 +587,33 @@ FLearningAgentsActionSchemaElement ULearningAgentsActionSchema::SpecifyExclusive
 	{
 		NormalizedPriorProbabilities[Idx] = PriorProbabilities[Idx];
 	}
-	UE::Learning::Agents::Action::Private::NormalizeProbabilitiesExclusive(NormalizedPriorProbabilities, GetName());
+	UE::Learning::Agents::Action::Private::NormalizeProbabilitiesExclusive(NormalizedPriorProbabilities);
 
-	return { ActionSchema.CreateDiscreteExclusive({ Size, MakeArrayView(NormalizedPriorProbabilities) }, Name)};
+	return { Schema->ActionSchema.CreateDiscreteExclusive({ Size, MakeArrayView(NormalizedPriorProbabilities) }, Tag)};
 }
 
-FLearningAgentsActionSchemaElement ULearningAgentsActionSchema::SpecifyInclusiveDiscreteAction(const int32 Size, const TArray<float>& PriorProbabilities, const FName Name)
+FLearningAgentsActionSchemaElement ULearningAgentsActions::SpecifyInclusiveDiscreteAction(ULearningAgentsActionSchema* Schema, const int32 Size, const TArray<float>& PriorProbabilities, const FName Tag)
 {
-	return SpecifyInclusiveDiscreteActionFromArrayView(Size, PriorProbabilities, Name);
+	return SpecifyInclusiveDiscreteActionFromArrayView(Schema, Size, PriorProbabilities, Tag);
 }
 
-FLearningAgentsActionSchemaElement ULearningAgentsActionSchema::SpecifyInclusiveDiscreteActionFromArrayView(const int32 Size, const TArrayView<const float> PriorProbabilities, const FName Name)
+FLearningAgentsActionSchemaElement ULearningAgentsActions::SpecifyInclusiveDiscreteActionFromArrayView(ULearningAgentsActionSchema* Schema, const int32 Size, const TArrayView<const float> PriorProbabilities, const FName Tag)
 {
+	if (!Schema)
+	{
+		UE_LOG(LogLearning, Error, TEXT("SpecifyInclusiveDiscreteActionFromArrayView: Schema is nullptr."));
+		return FLearningAgentsActionSchemaElement();
+	}
+
 	if (Size < 0)
 	{
-		UE_LOG(LogLearning, Error, TEXT("%s: Invalid DiscreteInclusive Action Size '%i'."), *GetName(), Size);
+		UE_LOG(LogLearning, Error, TEXT("SpecifyInclusiveDiscreteActionFromArrayView: Invalid DiscreteInclusive Action Size '%i'."), Size);
 		return FLearningAgentsActionSchemaElement();
 	}
 
 	if (Size == 0)
 	{
-		UE_LOG(LogLearning, Warning, TEXT("%s: Specifying zero-sized Inclusive Discrete Action."), *GetName());
+		UE_LOG(LogLearning, Warning, TEXT("SpecifyInclusiveDiscreteActionFromArrayView: Specifying zero-sized Inclusive Discrete Action."));
 	}
 
 	TArray<float, TInlineAllocator<16>> NormalizedPriorProbabilities;
@@ -609,16 +622,16 @@ FLearningAgentsActionSchemaElement ULearningAgentsActionSchema::SpecifyInclusive
 	{
 		NormalizedPriorProbabilities[Idx] = PriorProbabilities[Idx];
 	}
-	UE::Learning::Agents::Action::Private::NormalizeProbabilitiesInclusive(NormalizedPriorProbabilities, GetName());
+	UE::Learning::Agents::Action::Private::NormalizeProbabilitiesInclusive(NormalizedPriorProbabilities);
 
-	return { ActionSchema.CreateDiscreteInclusive({ Size, MakeArrayView(NormalizedPriorProbabilities) }, Name) };
+	return { Schema->ActionSchema.CreateDiscreteInclusive({ Size, MakeArrayView(NormalizedPriorProbabilities) }, Tag) };
 }
 
-FLearningAgentsActionSchemaElement ULearningAgentsActionSchema::SpecifyStructAction(const TMap<FName, FLearningAgentsActionSchemaElement>& Elements, const FName Name)
+FLearningAgentsActionSchemaElement ULearningAgentsActions::SpecifyStructAction(ULearningAgentsActionSchema* Schema, const TMap<FName, FLearningAgentsActionSchemaElement>& Elements, const FName Tag)
 {
 	if (Elements.Num() == 0)
 	{
-		UE_LOG(LogLearning, Warning, TEXT("%s: Specifying zero-sized Struct Action."), *GetName());
+		UE_LOG(LogLearning, Warning, TEXT("SpecifyStructAction: Specifying zero-sized Struct Action."));
 	}
 
 	const int32 SubElementNum = Elements.Num();
@@ -656,30 +669,36 @@ FLearningAgentsActionSchemaElement ULearningAgentsActionSchema::SpecifyStructAct
 		SortedSubElements[Idx] = SubElements[SubElementIndices[Idx]];
 	}
 
-	return SpecifyStructActionFromArrayViews(SortedSubElementNames, SortedSubElements, Name);
+	return SpecifyStructActionFromArrayViews(Schema, SortedSubElementNames, SortedSubElements, Tag);
 }
 
-FLearningAgentsActionSchemaElement ULearningAgentsActionSchema::SpecifyStructActionFromArrays(const TArray<FName>& ElementNames, const TArray<FLearningAgentsActionSchemaElement>& Elements, const FName Name)
+FLearningAgentsActionSchemaElement ULearningAgentsActions::SpecifyStructActionFromArrays(ULearningAgentsActionSchema* Schema, const TArray<FName>& ElementNames, const TArray<FLearningAgentsActionSchemaElement>& Elements, const FName Tag)
 {
-	return SpecifyStructActionFromArrayViews(ElementNames, Elements, Name);
+	return SpecifyStructActionFromArrayViews(Schema, ElementNames, Elements, Tag);
 }
 
-FLearningAgentsActionSchemaElement ULearningAgentsActionSchema::SpecifyStructActionFromArrayViews(const TArrayView<const FName> ElementNames, const TArrayView<const FLearningAgentsActionSchemaElement> Elements, const FName Name)
+FLearningAgentsActionSchemaElement ULearningAgentsActions::SpecifyStructActionFromArrayViews(ULearningAgentsActionSchema* Schema, const TArrayView<const FName> ElementNames, const TArrayView<const FLearningAgentsActionSchemaElement> Elements, const FName Tag)
 {
+	if (!Schema)
+	{
+		UE_LOG(LogLearning, Error, TEXT("SpecifyStructActionFromArrayViews: Schema is nullptr."));
+		return FLearningAgentsActionSchemaElement();
+	}
+
 	if (Elements.Num() == 0)
 	{
-		UE_LOG(LogLearning, Warning, TEXT("%s: Specifying zero-sized Struct Action."), *GetName());
+		UE_LOG(LogLearning, Warning, TEXT("SpecifyStructActionFromArrayViews: Specifying zero-sized Struct Action."));
 	}
 
 	if (Elements.Num() != ElementNames.Num())
 	{
-		UE_LOG(LogLearning, Error, TEXT("%s: Number of elements (%i) must match number of names (%i)."), *GetName(), Elements.Num(), ElementNames.Num());
+		UE_LOG(LogLearning, Error, TEXT("SpecifyStructActionFromArrayViews: Number of elements (%i) must match number of names (%i)."), Elements.Num(), ElementNames.Num());
 		return FLearningAgentsActionSchemaElement();
 	}
 
 	if (UE::Learning::Agents::Action::Private::ContainsDuplicates(ElementNames))
 	{
-		UE_LOG(LogLearning, Error, TEXT("%s: Element Names contain duplicates."), *GetName());
+		UE_LOG(LogLearning, Error, TEXT("SpecifyStructActionFromArrayViews: Element Names contain duplicates."));
 		return FLearningAgentsActionSchemaElement();
 	}
 
@@ -688,23 +707,23 @@ FLearningAgentsActionSchemaElement ULearningAgentsActionSchema::SpecifyStructAct
 
 	for (const FLearningAgentsActionSchemaElement& Element : Elements)
 	{
-		if (!ActionSchema.IsValid(Element.SchemaElement))
+		if (!Schema->ActionSchema.IsValid(Element.SchemaElement))
 		{
-			UE_LOG(LogLearning, Error, TEXT("%s: Invalid Action Object."), *GetName());
+			UE_LOG(LogLearning, Error, TEXT("SpecifyStructActionFromArrayViews: Invalid Action Object."));
 			return FLearningAgentsActionSchemaElement();
 		}
 
 		SubElements.Add(Element.SchemaElement);
 	}
 
-	return { ActionSchema.CreateAnd({ ElementNames, SubElements }, Name) };
+	return { Schema->ActionSchema.CreateAnd({ ElementNames, SubElements }, Tag) };
 }
 
-FLearningAgentsActionSchemaElement ULearningAgentsActionSchema::SpecifyExclusiveUnionAction(const TMap<FName, FLearningAgentsActionSchemaElement>& Elements, const TMap<FName, float>& PriorProbabilities, const FName Name)
+FLearningAgentsActionSchemaElement ULearningAgentsActions::SpecifyExclusiveUnionAction(ULearningAgentsActionSchema* Schema, const TMap<FName, FLearningAgentsActionSchemaElement>& Elements, const TMap<FName, float>& PriorProbabilities, const FName Tag)
 {
 	if (Elements.Num() == 0)
 	{
-		UE_LOG(LogLearning, Warning, TEXT("%s: Specifying zero-sized Exclusive Union Action."), *GetName());
+		UE_LOG(LogLearning, Warning, TEXT("SpecifyExclusiveUnionAction: Specifying zero-sized Exclusive Union Action."));
 	}
 
 	const int32 SubElementNum = Elements.Num();
@@ -749,30 +768,36 @@ FLearningAgentsActionSchemaElement ULearningAgentsActionSchema::SpecifyExclusive
 		SortedSubElementPriorProbabilities[Idx] = SubElementPriorProbabilities[SubElementIndices[Idx]];
 	}
 
-	return SpecifyExclusiveUnionActionFromArrayViews(SortedSubElementNames, SortedSubElements, SortedSubElementPriorProbabilities, Name);
+	return SpecifyExclusiveUnionActionFromArrayViews(Schema, SortedSubElementNames, SortedSubElements, SortedSubElementPriorProbabilities, Tag);
 }
 
-FLearningAgentsActionSchemaElement ULearningAgentsActionSchema::SpecifyExclusiveUnionActionFromArrays(const TArray<FName>& ElementNames, const TArray<FLearningAgentsActionSchemaElement>& Elements, const TArray<float>& PriorProbabilities, const FName Name)
+FLearningAgentsActionSchemaElement ULearningAgentsActions::SpecifyExclusiveUnionActionFromArrays(ULearningAgentsActionSchema* Schema, const TArray<FName>& ElementNames, const TArray<FLearningAgentsActionSchemaElement>& Elements, const TArray<float>& PriorProbabilities, const FName Tag)
 {
-	return SpecifyExclusiveUnionActionFromArrayViews(ElementNames, Elements, PriorProbabilities, Name);
+	return SpecifyExclusiveUnionActionFromArrayViews(Schema, ElementNames, Elements, PriorProbabilities, Tag);
 }
 
-FLearningAgentsActionSchemaElement ULearningAgentsActionSchema::SpecifyExclusiveUnionActionFromArrayViews(const TArrayView<const FName> ElementNames, const TArrayView<const FLearningAgentsActionSchemaElement> Elements, const TArrayView<const float> PriorProbabilities, const FName Name)
+FLearningAgentsActionSchemaElement ULearningAgentsActions::SpecifyExclusiveUnionActionFromArrayViews(ULearningAgentsActionSchema* Schema, const TArrayView<const FName> ElementNames, const TArrayView<const FLearningAgentsActionSchemaElement> Elements, const TArrayView<const float> PriorProbabilities, const FName Tag)
 {
+	if (!Schema)
+	{
+		UE_LOG(LogLearning, Error, TEXT("SpecifyExclusiveUnionActionFromArrayViews: Schema is nullptr."));
+		return FLearningAgentsActionSchemaElement();
+	}
+
 	if (Elements.Num() == 0)
 	{
-		UE_LOG(LogLearning, Warning, TEXT("%s: Specifying zero-sized Exclusive Union Action."), *GetName());
+		UE_LOG(LogLearning, Warning, TEXT("SpecifyExclusiveUnionActionFromArrayViews: Specifying zero-sized Exclusive Union Action."));
 	}
 
 	if (Elements.Num() != ElementNames.Num())
 	{
-		UE_LOG(LogLearning, Error, TEXT("%s: Number of elements (%i) must match number of names (%i)."), *GetName(), Elements.Num(), ElementNames.Num());
+		UE_LOG(LogLearning, Error, TEXT("SpecifyExclusiveUnionActionFromArrayViews: Number of elements (%i) must match number of names (%i)."), Elements.Num(), ElementNames.Num());
 		return FLearningAgentsActionSchemaElement();
 	}
 
 	if (UE::Learning::Agents::Action::Private::ContainsDuplicates(ElementNames))
 	{
-		UE_LOG(LogLearning, Error, TEXT("%s: Element Names contain duplicates."), *GetName());
+		UE_LOG(LogLearning, Error, TEXT("SpecifyExclusiveUnionActionFromArrayViews: Element Names contain duplicates."));
 		return FLearningAgentsActionSchemaElement();
 	}
 
@@ -781,9 +806,9 @@ FLearningAgentsActionSchemaElement ULearningAgentsActionSchema::SpecifyExclusive
 
 	for (const FLearningAgentsActionSchemaElement& Element : Elements)
 	{
-		if (!ActionSchema.IsValid(Element.SchemaElement))
+		if (!Schema->ActionSchema.IsValid(Element.SchemaElement))
 		{
-			UE_LOG(LogLearning, Error, TEXT("%s: Invalid Action Object."), *GetName());
+			UE_LOG(LogLearning, Error, TEXT("SpecifyExclusiveUnionActionFromArrayViews: Invalid Action Object."));
 			return FLearningAgentsActionSchemaElement();
 		}
 
@@ -796,16 +821,16 @@ FLearningAgentsActionSchemaElement ULearningAgentsActionSchema::SpecifyExclusive
 	{
 		NormalizedPriorProbabilities[Idx] = PriorProbabilities[Idx];
 	}
-	UE::Learning::Agents::Action::Private::NormalizeProbabilitiesExclusive(NormalizedPriorProbabilities, GetName());
+	UE::Learning::Agents::Action::Private::NormalizeProbabilitiesExclusive(NormalizedPriorProbabilities);
 
-	return { ActionSchema.CreateOrExclusive({ ElementNames, SubElements, MakeArrayView(NormalizedPriorProbabilities) }, Name)};
+	return { Schema->ActionSchema.CreateOrExclusive({ ElementNames, SubElements, MakeArrayView(NormalizedPriorProbabilities) }, Tag)};
 }
 
-FLearningAgentsActionSchemaElement ULearningAgentsActionSchema::SpecifyInclusiveUnionAction(const TMap<FName, FLearningAgentsActionSchemaElement>& Elements, const TMap<FName, float>& PriorProbabilities, const FName Name)
+FLearningAgentsActionSchemaElement ULearningAgentsActions::SpecifyInclusiveUnionAction(ULearningAgentsActionSchema* Schema, const TMap<FName, FLearningAgentsActionSchemaElement>& Elements, const TMap<FName, float>& PriorProbabilities, const FName Tag)
 {
 	if (Elements.Num() == 0)
 	{
-		UE_LOG(LogLearning, Warning, TEXT("%s: Specifying zero-sized Inclusive Union Action."), *GetName());
+		UE_LOG(LogLearning, Warning, TEXT("SpecifyInclusiveUnionAction: Specifying zero-sized Inclusive Union Action."));
 	}
 
 	const int32 SubElementNum = Elements.Num();
@@ -850,30 +875,36 @@ FLearningAgentsActionSchemaElement ULearningAgentsActionSchema::SpecifyInclusive
 		SortedSubElementPriorProbabilities[Idx] = SubElementPriorProbabilities[SubElementIndices[Idx]];
 	}
 
-	return SpecifyInclusiveUnionActionFromArrayViews(SortedSubElementNames, SortedSubElements, SortedSubElementPriorProbabilities, Name);
+	return SpecifyInclusiveUnionActionFromArrayViews(Schema, SortedSubElementNames, SortedSubElements, SortedSubElementPriorProbabilities, Tag);
 }
 
-FLearningAgentsActionSchemaElement ULearningAgentsActionSchema::SpecifyInclusiveUnionActionFromArrays(const TArray<FName> ElementNames, const TArray<FLearningAgentsActionSchemaElement>& Elements, const TArray<float>& PriorProbabilities, const FName Name)
+FLearningAgentsActionSchemaElement ULearningAgentsActions::SpecifyInclusiveUnionActionFromArrays(ULearningAgentsActionSchema* Schema, const TArray<FName> ElementNames, const TArray<FLearningAgentsActionSchemaElement>& Elements, const TArray<float>& PriorProbabilities, const FName Tag)
 {
-	return SpecifyInclusiveUnionActionFromArrayViews(ElementNames, Elements, PriorProbabilities, Name);
+	return SpecifyInclusiveUnionActionFromArrayViews(Schema, ElementNames, Elements, PriorProbabilities, Tag);
 }
 
-FLearningAgentsActionSchemaElement ULearningAgentsActionSchema::SpecifyInclusiveUnionActionFromArrayViews(const TArrayView<const FName> ElementNames, const TArrayView<const FLearningAgentsActionSchemaElement> Elements, const TArrayView<const float> PriorProbabilities, const FName Name)
+FLearningAgentsActionSchemaElement ULearningAgentsActions::SpecifyInclusiveUnionActionFromArrayViews(ULearningAgentsActionSchema* Schema, const TArrayView<const FName> ElementNames, const TArrayView<const FLearningAgentsActionSchemaElement> Elements, const TArrayView<const float> PriorProbabilities, const FName Tag)
 {
+	if (!Schema)
+	{
+		UE_LOG(LogLearning, Error, TEXT("SpecifyInclusiveUnionActionFromArrayViews: Schema is nullptr."));
+		return FLearningAgentsActionSchemaElement();
+	}
+
 	if (Elements.Num() == 0)
 	{
-		UE_LOG(LogLearning, Warning, TEXT("%s: Specifying zero-sized Inclusive Union Action."), *GetName());
+		UE_LOG(LogLearning, Warning, TEXT("SpecifyInclusiveUnionActionFromArrayViews: Specifying zero-sized Inclusive Union Action."));
 	}
 
 	if (Elements.Num() != ElementNames.Num())
 	{
-		UE_LOG(LogLearning, Error, TEXT("%s: Number of elements (%i) must match number of names (%i)."), *GetName(), Elements.Num(), ElementNames.Num());
+		UE_LOG(LogLearning, Error, TEXT("SpecifyInclusiveUnionActionFromArrayViews: Number of elements (%i) must match number of names (%i)."), Elements.Num(), ElementNames.Num());
 		return FLearningAgentsActionSchemaElement();
 	}
 
 	if (UE::Learning::Agents::Action::Private::ContainsDuplicates(ElementNames))
 	{
-		UE_LOG(LogLearning, Error, TEXT("%s: Element Names contain duplicates."), *GetName());
+		UE_LOG(LogLearning, Error, TEXT("SpecifyInclusiveUnionActionFromArrayViews: Element Names contain duplicates."));
 		return FLearningAgentsActionSchemaElement();
 	}
 
@@ -882,9 +913,9 @@ FLearningAgentsActionSchemaElement ULearningAgentsActionSchema::SpecifyInclusive
 
 	for (const FLearningAgentsActionSchemaElement& Element : Elements)
 	{
-		if (!ActionSchema.IsValid(Element.SchemaElement))
+		if (!Schema->ActionSchema.IsValid(Element.SchemaElement))
 		{
-			UE_LOG(LogLearning, Error, TEXT("%s: Invalid Action Object."), *GetName());
+			UE_LOG(LogLearning, Error, TEXT("SpecifyInclusiveUnionActionFromArrayViews: Invalid Action Object."));
 			return FLearningAgentsActionSchemaElement();
 		}
 
@@ -897,43 +928,49 @@ FLearningAgentsActionSchemaElement ULearningAgentsActionSchema::SpecifyInclusive
 	{
 		NormalizedPriorProbabilities[Idx] = PriorProbabilities[Idx];
 	}
-	UE::Learning::Agents::Action::Private::NormalizeProbabilitiesInclusive(NormalizedPriorProbabilities, GetName());
+	UE::Learning::Agents::Action::Private::NormalizeProbabilitiesInclusive(NormalizedPriorProbabilities);
 
-	return { ActionSchema.CreateOrInclusive({ ElementNames, SubElements, MakeArrayView(NormalizedPriorProbabilities) }, Name)};
+	return { Schema->ActionSchema.CreateOrInclusive({ ElementNames, SubElements, MakeArrayView(NormalizedPriorProbabilities) }, Tag)};
 }
 
-FLearningAgentsActionSchemaElement ULearningAgentsActionSchema::SpecifyStaticArrayAction(const FLearningAgentsActionSchemaElement Element, const int32 Num, const FName Name)
+FLearningAgentsActionSchemaElement ULearningAgentsActions::SpecifyStaticArrayAction(ULearningAgentsActionSchema* Schema, const FLearningAgentsActionSchemaElement Element, const int32 Num, const FName Tag)
 {
+	if (!Schema)
+	{
+		UE_LOG(LogLearning, Error, TEXT("SpecifyStaticArrayAction: Schema is nullptr."));
+		return FLearningAgentsActionSchemaElement();
+	}
+
 	if (Num < 0)
 	{
-		UE_LOG(LogLearning, Error, TEXT("%s: Invalid Action Static Array Num %i."), *GetName(), Num);
+		UE_LOG(LogLearning, Error, TEXT("SpecifyStaticArrayAction: Invalid Action Static Array Num %i."), Num);
 		return FLearningAgentsActionSchemaElement();
 	}
 
 	if (Num == 0)
 	{
-		UE_LOG(LogLearning, Warning, TEXT("%s: Specifying zero-sized Static Array Action."), *GetName());
+		UE_LOG(LogLearning, Warning, TEXT("SpecifyStaticArrayAction: Specifying zero-sized Static Array Action."));
 	}
 
-	if (!ActionSchema.IsValid(Element.SchemaElement))
+	if (!Schema->ActionSchema.IsValid(Element.SchemaElement))
 	{
-		UE_LOG(LogLearning, Error, TEXT("%s: Invalid Action Object."), *GetName());
+		UE_LOG(LogLearning, Error, TEXT("SpecifyStaticArrayAction: Invalid Action Object."));
 		return FLearningAgentsActionSchemaElement();
 	}
 
-	return { ActionSchema.CreateArray({ Element.SchemaElement, Num }, Name) };
+	return { Schema->ActionSchema.CreateArray({ Element.SchemaElement, Num }, Tag) };
 }
 
-FLearningAgentsActionSchemaElement ULearningAgentsActionSchema::SpecifyPairAction(const FLearningAgentsActionSchemaElement Key, const FLearningAgentsActionSchemaElement Value, const FName Name)
+FLearningAgentsActionSchemaElement ULearningAgentsActions::SpecifyPairAction(ULearningAgentsActionSchema* Schema, const FLearningAgentsActionSchemaElement Key, const FLearningAgentsActionSchemaElement Value, const FName Tag)
 {
-	return SpecifyStructActionFromArrayViews({ TEXT("Key"), TEXT("Value") }, { Key, Value }, Name);
+	return SpecifyStructActionFromArrayViews(Schema, { TEXT("Key"), TEXT("Value") }, { Key, Value }, Tag);
 }
 
-FLearningAgentsActionSchemaElement ULearningAgentsActionSchema::SpecifyEnumAction(const UEnum* Enum, const TMap<uint8, float>& PriorProbabilities, const FName Name)
+FLearningAgentsActionSchemaElement ULearningAgentsActions::SpecifyEnumAction(ULearningAgentsActionSchema* Schema, const UEnum* Enum, const TMap<uint8, float>& PriorProbabilities, const FName Tag)
 {
 	if (!Enum)
 	{
-		UE_LOG(LogLearning, Error, TEXT("%s: Enum is nullptr."), *GetName());
+		UE_LOG(LogLearning, Error, TEXT("SpecifyEnumAction: Enum is nullptr."));
 		return FLearningAgentsActionSchemaElement();
 	}
 
@@ -948,36 +985,36 @@ FLearningAgentsActionSchemaElement ULearningAgentsActionSchema::SpecifyEnumActio
 		}
 	}
 
-	return SpecifyEnumActionFromArrayView(Enum, PriorProbabilitiesArray, Name);
+	return SpecifyEnumActionFromArrayView(Schema, Enum, PriorProbabilitiesArray, Tag);
 }
 
-FLearningAgentsActionSchemaElement ULearningAgentsActionSchema::SpecifyEnumActionFromArray(const UEnum* Enum, const TArray<float>& PriorProbabilities, const FName Name)
+FLearningAgentsActionSchemaElement ULearningAgentsActions::SpecifyEnumActionFromArray(ULearningAgentsActionSchema* Schema, const UEnum* Enum, const TArray<float>& PriorProbabilities, const FName Tag)
 {
-	return SpecifyEnumActionFromArrayView(Enum, PriorProbabilities, Name);
+	return SpecifyEnumActionFromArrayView(Schema, Enum, PriorProbabilities, Tag);
 }
 
-FLearningAgentsActionSchemaElement ULearningAgentsActionSchema::SpecifyEnumActionFromArrayView(const UEnum* Enum, const TArrayView<const float> PriorProbabilities, const FName Name)
+FLearningAgentsActionSchemaElement ULearningAgentsActions::SpecifyEnumActionFromArrayView(ULearningAgentsActionSchema* Schema, const UEnum* Enum, const TArrayView<const float> PriorProbabilities, const FName Tag)
 {
 	if (!Enum)
 	{
-		UE_LOG(LogLearning, Error, TEXT("%s: Enum is nullptr."), *GetName());
+		UE_LOG(LogLearning, Error, TEXT("SpecifyEnumActionFromArrayView: Enum is nullptr."));
 		return FLearningAgentsActionSchemaElement();
 	}
 
-	return SpecifyExclusiveDiscreteActionFromArrayView(Enum->NumEnums() - 1, PriorProbabilities, Name);
+	return SpecifyExclusiveDiscreteActionFromArrayView(Schema, Enum->NumEnums() - 1, PriorProbabilities, Tag);
 }
 
-FLearningAgentsActionSchemaElement ULearningAgentsActionSchema::SpecifyBitmaskAction(const UEnum* Enum, const TMap<uint8, float>& PriorProbabilities, const FName Name)
+FLearningAgentsActionSchemaElement ULearningAgentsActions::SpecifyBitmaskAction(ULearningAgentsActionSchema* Schema, const UEnum* Enum, const TMap<uint8, float>& PriorProbabilities, const FName Tag)
 {
 	if (!Enum)
 	{
-		UE_LOG(LogLearning, Error, TEXT("%s: Enum is nullptr."), *GetName());
+		UE_LOG(LogLearning, Error, TEXT("SpecifyBitmaskAction: Enum is nullptr."));
 		return FLearningAgentsActionSchemaElement();
 	}
 
 	if (Enum->NumEnums() - 1 > 32)
 	{
-		UE_LOG(LogLearning, Error, TEXT("%s: Too many values in Enum to use as Bitmask (%i)."), *GetName(), Enum->NumEnums() - 1);
+		UE_LOG(LogLearning, Error, TEXT("SpecifyBitmaskAction: Too many values in Enum to use as Bitmask (%i)."), Enum->NumEnums() - 1);
 		return FLearningAgentsActionSchemaElement();
 	}
 
@@ -992,181 +1029,208 @@ FLearningAgentsActionSchemaElement ULearningAgentsActionSchema::SpecifyBitmaskAc
 		}
 	}
 
-	return SpecifyBitmaskActionFromArrayView(Enum, PriorProbabilitiesArray, Name);
+	return SpecifyBitmaskActionFromArrayView(Schema, Enum, PriorProbabilitiesArray, Tag);
 }
 
-FLearningAgentsActionSchemaElement ULearningAgentsActionSchema::SpecifyBitmaskActionFromArray(const UEnum* Enum, const TArray<float>& PriorProbabilities, const FName Name)
+FLearningAgentsActionSchemaElement ULearningAgentsActions::SpecifyBitmaskActionFromArray(ULearningAgentsActionSchema* Schema, const UEnum* Enum, const TArray<float>& PriorProbabilities, const FName Tag)
 {
-	return SpecifyBitmaskActionFromArrayView(Enum, PriorProbabilities, Name);
+	return SpecifyBitmaskActionFromArrayView(Schema, Enum, PriorProbabilities, Tag);
 }
 
-FLearningAgentsActionSchemaElement ULearningAgentsActionSchema::SpecifyBitmaskActionFromArrayView(const UEnum* Enum, const TArrayView<const float> PriorProbabilities, const FName Name)
+FLearningAgentsActionSchemaElement ULearningAgentsActions::SpecifyBitmaskActionFromArrayView(ULearningAgentsActionSchema* Schema, const UEnum* Enum, const TArrayView<const float> PriorProbabilities, const FName Tag)
 {
 	if (!Enum)
 	{
-		UE_LOG(LogLearning, Error, TEXT("%s: Enum is nullptr."), *GetName());
+		UE_LOG(LogLearning, Error, TEXT("SpecifyBitmaskActionFromArrayView: Enum is nullptr."));
 		return FLearningAgentsActionSchemaElement();
 	}
 
 	if (Enum->NumEnums() - 1 > 32)
 	{
-		UE_LOG(LogLearning, Error, TEXT("%s: Too many values in Enum to use as Bitmask (%i)."), *GetName(), Enum->NumEnums() - 1);
+		UE_LOG(LogLearning, Error, TEXT("SpecifyBitmaskActionFromArrayView: Too many values in Enum to use as Bitmask (%i)."), Enum->NumEnums() - 1);
 		return FLearningAgentsActionSchemaElement();
 	}
 
-	return SpecifyInclusiveDiscreteActionFromArrayView(Enum->NumEnums() - 1, PriorProbabilities, Name);
+	return SpecifyInclusiveDiscreteActionFromArrayView(Schema, Enum->NumEnums() - 1, PriorProbabilities, Tag);
 }
 
-FLearningAgentsActionSchemaElement ULearningAgentsActionSchema::SpecifyOptionalAction(const FLearningAgentsActionSchemaElement Element, const float PriorProbability, const FName Name)
+FLearningAgentsActionSchemaElement ULearningAgentsActions::SpecifyOptionalAction(ULearningAgentsActionSchema* Schema, const FLearningAgentsActionSchemaElement Element, const float PriorProbability, const FName Tag)
 {
-	return SpecifyExclusiveUnionActionFromArrayViews({ TEXT("Null"), TEXT("Valid") }, { SpecifyNullAction(), Element }, { 1.0f - PriorProbability, PriorProbability }, Name);
+	return SpecifyExclusiveUnionActionFromArrayViews(Schema, { TEXT("Null"), TEXT("Valid") }, { SpecifyNullAction(Schema), Element }, { 1.0f - PriorProbability, PriorProbability }, Tag);
 }
 
-FLearningAgentsActionSchemaElement ULearningAgentsActionSchema::SpecifyEitherAction(const FLearningAgentsActionSchemaElement A, const FLearningAgentsActionSchemaElement B, const float PriorProbabilityOfA, const FName Name)
+FLearningAgentsActionSchemaElement ULearningAgentsActions::SpecifyEitherAction(ULearningAgentsActionSchema* Schema, const FLearningAgentsActionSchemaElement A, const FLearningAgentsActionSchemaElement B, const float PriorProbabilityOfA, const FName Tag)
 {
-	return SpecifyExclusiveUnionActionFromArrayViews({ TEXT("A"), TEXT("B") }, { A, B }, { 1.0f - PriorProbabilityOfA, PriorProbabilityOfA }, Name);
+	return SpecifyExclusiveUnionActionFromArrayViews(Schema, { TEXT("A"), TEXT("B") }, { A, B }, { 1.0f - PriorProbabilityOfA, PriorProbabilityOfA }, Tag);
 }
 
-FLearningAgentsActionSchemaElement ULearningAgentsActionSchema::SpecifyEncodingAction(const FLearningAgentsActionSchemaElement Element, const int32 EncodingSize, const int32 LayerNum, const ELearningAgentsActivationFunction ActivationFunction, const FName Name)
+FLearningAgentsActionSchemaElement ULearningAgentsActions::SpecifyEncodingAction(ULearningAgentsActionSchema* Schema, const FLearningAgentsActionSchemaElement Element, const int32 EncodingSize, const int32 LayerNum, const ELearningAgentsActivationFunction ActivationFunction, const FName Tag)
 {
+	if (!Schema)
+	{
+		UE_LOG(LogLearning, Error, TEXT("SpecifyEncodingAction: Schema is nullptr."));
+		return FLearningAgentsActionSchemaElement();
+	}
+
 	if (EncodingSize < 1)
 	{
-		UE_LOG(LogLearning, Error, TEXT("%s: Invalid Action EncodingSize '%i' - must be greater than zero."), *GetName(), EncodingSize);
+		UE_LOG(LogLearning, Error, TEXT("SpecifyEncodingAction: Invalid Action EncodingSize '%i' - must be greater than zero."), EncodingSize);
 		return FLearningAgentsActionSchemaElement();
 	}
 
 	if (LayerNum < 1)
 	{
-		UE_LOG(LogLearning, Error, TEXT("%s: Invalid Action LayerNum '%i' - must be greater than zero."), *GetName(), LayerNum);
+		UE_LOG(LogLearning, Error, TEXT("SpecifyEncodingAction: Invalid Action LayerNum '%i' - must be greater than zero."), LayerNum);
 		return FLearningAgentsActionSchemaElement();
 	}
 
-	if (!ActionSchema.IsValid(Element.SchemaElement))
+	if (!Schema->ActionSchema.IsValid(Element.SchemaElement))
 	{
-		UE_LOG(LogLearning, Error, TEXT("%s: Invalid Action Object."), *GetName());
+		UE_LOG(LogLearning, Error, TEXT("SpecifyEncodingAction: Invalid Action Object."));
 		return FLearningAgentsActionSchemaElement();
 	}
 
-	return { ActionSchema.CreateEncoding({ Element.SchemaElement, EncodingSize, LayerNum, UE::Learning::Agents::Action::Private::GetEncodingActivationFunction(ActivationFunction) }, Name)};
+	return { Schema->ActionSchema.CreateEncoding({ Element.SchemaElement, EncodingSize, LayerNum, UE::Learning::Agents::Action::Private::GetEncodingActivationFunction(ActivationFunction) }, Tag)};
 }
 
-FLearningAgentsActionSchemaElement ULearningAgentsActionSchema::SpecifyBoolAction(const float PriorProbability, const FName Name)
+FLearningAgentsActionSchemaElement ULearningAgentsActions::SpecifyBoolAction(ULearningAgentsActionSchema* Schema, const float PriorProbability, const FName Tag)
 {
-	return SpecifyExclusiveDiscreteActionFromArrayView(2, { 1.0f - PriorProbability, PriorProbability }, Name);
+	return SpecifyExclusiveDiscreteActionFromArrayView(Schema, 2, { 1.0f - PriorProbability, PriorProbability }, Tag);
 }
 
-FLearningAgentsActionSchemaElement ULearningAgentsActionSchema::SpecifyFloatAction(const FName Name)
+FLearningAgentsActionSchemaElement ULearningAgentsActions::SpecifyFloatAction(ULearningAgentsActionSchema* Schema, const FName Tag)
 {
-	return SpecifyContinuousAction(1, Name);
+	return SpecifyContinuousAction(Schema, 1, Tag);
 }
 
-FLearningAgentsActionSchemaElement ULearningAgentsActionSchema::SpecifyLocationAction(const FName Name)
+FLearningAgentsActionSchemaElement ULearningAgentsActions::SpecifyLocationAction(ULearningAgentsActionSchema* Schema, const FName Tag)
 {
-	return SpecifyContinuousAction(3, Name);
+	return SpecifyContinuousAction(Schema, 3, Tag);
 }
 
-FLearningAgentsActionSchemaElement ULearningAgentsActionSchema::SpecifyRotationAction(const FName Name)
+FLearningAgentsActionSchemaElement ULearningAgentsActions::SpecifyRotationAction(ULearningAgentsActionSchema* Schema, const FName Tag)
 {
-	return SpecifyContinuousAction(3, Name);
+	return SpecifyContinuousAction(Schema, 3, Tag);
 }
 
-FLearningAgentsActionSchemaElement ULearningAgentsActionSchema::SpecifyScaleAction(const FName Name)
+FLearningAgentsActionSchemaElement ULearningAgentsActions::SpecifyScaleAction(ULearningAgentsActionSchema* Schema, const FName Tag)
 {
-	return SpecifyContinuousAction(3, Name);
+	return SpecifyContinuousAction(Schema, 3, Tag);
 }
 
-FLearningAgentsActionSchemaElement ULearningAgentsActionSchema::SpecifyTransformAction(const FName Name)
+FLearningAgentsActionSchemaElement ULearningAgentsActions::SpecifyTransformAction(ULearningAgentsActionSchema* Schema, const FName Tag)
 {
-	return SpecifyStructActionFromArrayViews(
+	return SpecifyStructActionFromArrayViews(Schema,
 		{
 			TEXT("Location"),
 			TEXT("Rotation"),
 			TEXT("Scale")
 		},
 		{
-			SpecifyLocationAction(),
-			SpecifyRotationAction(),
-			SpecifyScaleAction()
+			SpecifyLocationAction(Schema),
+			SpecifyRotationAction(Schema),
+			SpecifyScaleAction(Schema)
 		}, 
-		Name);
+		Tag);
 }
 
-FLearningAgentsActionSchemaElement ULearningAgentsActionSchema::SpecifyAngleAction(const FName Name)
+FLearningAgentsActionSchemaElement ULearningAgentsActions::SpecifyAngleAction(ULearningAgentsActionSchema* Schema, const FName Tag)
 {
-	return SpecifyContinuousAction(1, Name);
+	return SpecifyContinuousAction(Schema, 1, Tag);
 }
 
-FLearningAgentsActionSchemaElement ULearningAgentsActionSchema::SpecifyVelocityAction(const FName Name)
+FLearningAgentsActionSchemaElement ULearningAgentsActions::SpecifyVelocityAction(ULearningAgentsActionSchema* Schema, const FName Tag)
 {
-	return SpecifyContinuousAction(3, Name);
+	return SpecifyContinuousAction(Schema, 3, Tag);
 }
 
-FLearningAgentsActionSchemaElement ULearningAgentsActionSchema::SpecifyDirectionAction(const FName Name)
+FLearningAgentsActionSchemaElement ULearningAgentsActions::SpecifyDirectionAction(ULearningAgentsActionSchema* Schema, const FName Tag)
 {
-	return SpecifyContinuousAction(3, Name);
+	return SpecifyContinuousAction(Schema, 3, Tag);
 }
 
-FLearningAgentsActionSchemaElement ULearningAgentsActionSchema::SpecifySpeedAction(const FName Name)
+FLearningAgentsActionSchemaElement ULearningAgentsActions::SpecifySpeedAction(ULearningAgentsActionSchema* Schema, const FName Tag)
 {
-	return SpecifyContinuousAction(1, Name);
+	return SpecifyContinuousAction(Schema, 1, Tag);
 }
 
-const UE::Learning::Action::FObject& ULearningAgentsActionObject::GetActionObject() const
-{
-	return ActionObject;
-}
 
-UE::Learning::Action::FObject& ULearningAgentsActionObject::GetActionObject()
+void ULearningAgentsActions::LogAction(const ULearningAgentsActionObject* Object, const FLearningAgentsActionObjectElement Element)
 {
-	return ActionObject;
-}
+	if (!Object)
+	{
+		UE_LOG(LogLearning, Error, TEXT("LogAction: Object is nullptr."));
+		return;
+	}
 
-void ULearningAgentsActionObject::LogAction(const FLearningAgentsActionObjectElement Element)
-{
-	UE::Learning::Agents::Action::Private::LogAction(GetActionObject(), Element.ObjectElement, TEXT(""), TEXT(""), GetName());
+	UE::Learning::Agents::Action::Private::LogAction(Object->ActionObject, Element.ObjectElement, TEXT(""), TEXT(""));
 }
 
-FLearningAgentsActionObjectElement ULearningAgentsActionObject::MakeNullAction(const FName Name)
+FLearningAgentsActionObjectElement ULearningAgentsActions::MakeNullAction(ULearningAgentsActionObject* Object, const FName Tag)
 {
-	return { ActionObject.CreateNull(Name) };
+	if (!Object)
+	{
+		UE_LOG(LogLearning, Error, TEXT("MakeNullAction: Object is nullptr."));
+		return FLearningAgentsActionObjectElement();
+	}
+
+	return { Object->ActionObject.CreateNull(Tag) };
 }
 
-FLearningAgentsActionObjectElement ULearningAgentsActionObject::MakeContinuousAction(const TArray<float>& Values, const FName Name)
+FLearningAgentsActionObjectElement ULearningAgentsActions::MakeContinuousAction(ULearningAgentsActionObject* Object, const TArray<float>& Values, const FName Tag)
 {
-	return MakeContinuousActionFromArrayView(Values, Name);
+	return MakeContinuousActionFromArrayView(Object, Values, Tag);
 }
 
-FLearningAgentsActionObjectElement ULearningAgentsActionObject::MakeContinuousActionFromArrayView(const TArrayView<const float> Values, const FName Name)
+FLearningAgentsActionObjectElement ULearningAgentsActions::MakeContinuousActionFromArrayView(ULearningAgentsActionObject* Object, const TArrayView<const float> Values, const FName Tag)
 {
+	if (!Object)
+	{
+		UE_LOG(LogLearning, Error, TEXT("MakeContinuousActionFromArrayView: Object is nullptr."));
+		return FLearningAgentsActionObjectElement();
+	}
+
 	if (Values.Num() == 0)
 	{
-		UE_LOG(LogLearning, Warning, TEXT("%s: Creating zero-sized Continuous Action."), *GetName());
+		UE_LOG(LogLearning, Warning, TEXT("MakeContinuousActionFromArrayView: Creating zero-sized Continuous Action."));
 	}
 
-	return { ActionObject.CreateContinuous({ Values }, Name) };
+	return { Object->ActionObject.CreateContinuous({ Values }, Tag) };
 }
 
-FLearningAgentsActionObjectElement ULearningAgentsActionObject::MakeExclusiveDiscreteAction(const int32 Index, const FName Name)
+FLearningAgentsActionObjectElement ULearningAgentsActions::MakeExclusiveDiscreteAction(ULearningAgentsActionObject* Object, const int32 Index, const FName Tag)
 {
+	if (!Object)
+	{
+		UE_LOG(LogLearning, Error, TEXT("MakeExclusiveDiscreteAction: Object is nullptr."));
+		return FLearningAgentsActionObjectElement();
+	}
+
 	if (Index < 0)
 	{
-		UE_LOG(LogLearning, Error, TEXT("%s: Invalid Action Index %i."), *GetName(), Index);
+		UE_LOG(LogLearning, Error, TEXT("MakeExclusiveDiscreteAction: Invalid Action Index %i."), Index);
 		return FLearningAgentsActionObjectElement();
 	}
 
-	return { ActionObject.CreateDiscreteExclusive({ Index }, Name) };
+	return { Object->ActionObject.CreateDiscreteExclusive({ Index }, Tag) };
 }
 
-FLearningAgentsActionObjectElement ULearningAgentsActionObject::MakeInclusiveDiscreteAction(const TArray<int32>& Indices, const FName Name)
+FLearningAgentsActionObjectElement ULearningAgentsActions::MakeInclusiveDiscreteAction(ULearningAgentsActionObject* Object, const TArray<int32>& Indices, const FName Tag)
 {
-	return MakeInclusiveDiscreteActionFromArrayView(Indices, Name);
+	return MakeInclusiveDiscreteActionFromArrayView(Object, Indices, Tag);
 }
 
-FLearningAgentsActionObjectElement ULearningAgentsActionObject::MakeInclusiveDiscreteActionFromArrayView(const TArrayView<const int32> Indices, const FName Name)
+FLearningAgentsActionObjectElement ULearningAgentsActions::MakeInclusiveDiscreteActionFromArrayView(ULearningAgentsActionObject* Object, const TArrayView<const int32> Indices, const FName Tag)
 {
+	if (!Object)
+	{
+		UE_LOG(LogLearning, Error, TEXT("MakeInclusiveDiscreteActionFromArrayView: Object is nullptr."));
+		return FLearningAgentsActionObjectElement();
+	}
+
 	if (UE::Learning::Agents::Action::Private::ContainsDuplicates(Indices))
 	{
-		UE_LOG(LogLearning, Error, TEXT("%s: Indices contain duplicates."), *GetName());
+		UE_LOG(LogLearning, Error, TEXT("MakeInclusiveDiscreteActionFromArrayView: Indices contain duplicates."));
 		return FLearningAgentsActionObjectElement();
 	}
 
@@ -1176,19 +1240,19 @@ FLearningAgentsActionObjectElement ULearningAgentsActionObject::MakeInclusiveDis
 	{
 		if (Indices[IndexIdx] < 0)
 		{
-			UE_LOG(LogLearning, Error, TEXT("%s: Invalid Action Index %i."), *GetName(), Indices[IndexIdx]);
+			UE_LOG(LogLearning, Error, TEXT("MakeInclusiveDiscreteActionFromArrayView: Invalid Action Index %i."), Indices[IndexIdx]);
 			return FLearningAgentsActionObjectElement();
 		}
 	}
 
-	return { ActionObject.CreateDiscreteInclusive({ Indices }, Name) };
+	return { Object->ActionObject.CreateDiscreteInclusive({ Indices }, Tag) };
 }
 
-FLearningAgentsActionObjectElement ULearningAgentsActionObject::MakeStructAction(const TMap<FName, FLearningAgentsActionObjectElement>& Elements, const FName Name)
+FLearningAgentsActionObjectElement ULearningAgentsActions::MakeStructAction(ULearningAgentsActionObject* Object, const TMap<FName, FLearningAgentsActionObjectElement>& Elements, const FName Tag)
 {
 	if (Elements.Num() == 0)
 	{
-		UE_LOG(LogLearning, Warning, TEXT("%s: Creating zero-sized Struct Action."), *GetName());
+		UE_LOG(LogLearning, Warning, TEXT("MakeStructAction: Creating zero-sized Struct Action."));
 	}
 
 	const int32 SubElementNum = Elements.Num();
@@ -1204,30 +1268,36 @@ FLearningAgentsActionObjectElement ULearningAgentsActionObject::MakeStructAction
 		SubElementNames.Add(Element.Key);
 	}
 
-	return MakeStructActionFromArrayViews(SubElementNames, SubElements);
+	return MakeStructActionFromArrayViews(Object, SubElementNames, SubElements);
 }
 
-FLearningAgentsActionObjectElement ULearningAgentsActionObject::MakeStructActionFromArrays(const TArray<FName>& ElementNames, const TArray<FLearningAgentsActionObjectElement>& Elements, const FName Name)
+FLearningAgentsActionObjectElement ULearningAgentsActions::MakeStructActionFromArrays(ULearningAgentsActionObject* Object, const TArray<FName>& ElementNames, const TArray<FLearningAgentsActionObjectElement>& Elements, const FName Tag)
 {
-	return MakeStructActionFromArrayViews(ElementNames, Elements, Name);
+	return MakeStructActionFromArrayViews(Object, ElementNames, Elements, Tag);
 }
 
-FLearningAgentsActionObjectElement ULearningAgentsActionObject::MakeStructActionFromArrayViews(const TArrayView<const FName> ElementNames, const TArrayView<const FLearningAgentsActionObjectElement> Elements, const FName Name)
+FLearningAgentsActionObjectElement ULearningAgentsActions::MakeStructActionFromArrayViews(ULearningAgentsActionObject* Object, const TArrayView<const FName> ElementNames, const TArrayView<const FLearningAgentsActionObjectElement> Elements, const FName Tag)
 {
+	if (!Object)
+	{
+		UE_LOG(LogLearning, Error, TEXT("MakeStructActionFromArrayViews: Object is nullptr."));
+		return FLearningAgentsActionObjectElement();
+	}
+
 	if (Elements.Num() == 0)
 	{
-		UE_LOG(LogLearning, Warning, TEXT("%s: Creating zero-sized Struct Action."), *GetName());
+		UE_LOG(LogLearning, Warning, TEXT("MakeStructActionFromArrayViews: Creating zero-sized Struct Action."));
 	}
 
 	if (Elements.Num() != ElementNames.Num())
 	{
-		UE_LOG(LogLearning, Error, TEXT("%s: Number of elements (%i) must match number of names (%i)."), *GetName(), Elements.Num(), ElementNames.Num());
+		UE_LOG(LogLearning, Error, TEXT("MakeStructActionFromArrayViews: Number of elements (%i) must match number of names (%i)."), Elements.Num(), ElementNames.Num());
 		return FLearningAgentsActionObjectElement();
 	}
 
 	if (UE::Learning::Agents::Action::Private::ContainsDuplicates(ElementNames))
 	{
-		UE_LOG(LogLearning, Error, TEXT("%s: Element Names contain duplicates."), *GetName());
+		UE_LOG(LogLearning, Error, TEXT("MakeStructActionFromArrayViews: Element Names contain duplicates."));
 		return FLearningAgentsActionObjectElement();
 	}
 
@@ -1236,30 +1306,36 @@ FLearningAgentsActionObjectElement ULearningAgentsActionObject::MakeStructAction
 
 	for (const FLearningAgentsActionObjectElement& Element : Elements)
 	{
-		if (!ActionObject.IsValid(Element.ObjectElement))
+		if (!Object->ActionObject.IsValid(Element.ObjectElement))
 		{
-			UE_LOG(LogLearning, Error, TEXT("%s: Invalid Action Object."), *GetName());
+			UE_LOG(LogLearning, Error, TEXT("MakeStructActionFromArrayViews: Invalid Action Object."));
 			return FLearningAgentsActionObjectElement();
 		}
 
 		SubElements.Add(Element.ObjectElement);
 	}
 
-	return { ActionObject.CreateAnd({ ElementNames, SubElements }, Name) };
+	return { Object->ActionObject.CreateAnd({ ElementNames, SubElements }, Tag) };
 }
 
-FLearningAgentsActionObjectElement ULearningAgentsActionObject::MakeExclusiveUnionAction(const FName ElementName, const FLearningAgentsActionObjectElement Element, const FName Name)
+FLearningAgentsActionObjectElement ULearningAgentsActions::MakeExclusiveUnionAction(ULearningAgentsActionObject* Object, const FName ElementName, const FLearningAgentsActionObjectElement Element, const FName Tag)
 {
-	if (!ActionObject.IsValid(Element.ObjectElement))
+	if (!Object)
+	{
+		UE_LOG(LogLearning, Error, TEXT("MakeExclusiveUnionAction: Object is nullptr."));
+		return FLearningAgentsActionObjectElement();
+	}
+
+	if (!Object->ActionObject.IsValid(Element.ObjectElement))
 	{
-		UE_LOG(LogLearning, Error, TEXT("%s: Invalid Action Object."), *GetName());
+		UE_LOG(LogLearning, Error, TEXT("MakeExclusiveUnionAction: Invalid Action Object."));
 		return FLearningAgentsActionObjectElement();
 	}
 
-	return { ActionObject.CreateOrExclusive({ ElementName, Element.ObjectElement }, Name) };
+	return { Object->ActionObject.CreateOrExclusive({ ElementName, Element.ObjectElement }, Tag) };
 }
 
-FLearningAgentsActionObjectElement ULearningAgentsActionObject::MakeInclusiveUnionAction(const TMap<FName, FLearningAgentsActionObjectElement>& Elements, const FName Name)
+FLearningAgentsActionObjectElement ULearningAgentsActions::MakeInclusiveUnionAction(ULearningAgentsActionObject* Object, const TMap<FName, FLearningAgentsActionObjectElement>& Elements, const FName Tag)
 {
 	const int32 SubElementNum = Elements.Num();
 
@@ -1274,25 +1350,31 @@ FLearningAgentsActionObjectElement ULearningAgentsActionObject::MakeInclusiveUni
 		SubElementNames.Add(Element.Key);
 	}
 
-	return MakeInclusiveUnionActionFromArrayViews(SubElementNames, SubElements, Name);
+	return MakeInclusiveUnionActionFromArrayViews(Object, SubElementNames, SubElements, Tag);
 }
 
-FLearningAgentsActionObjectElement ULearningAgentsActionObject::MakeInclusiveUnionActionFromArrays(const TArray<FName>& ElementNames, const TArray<FLearningAgentsActionObjectElement>& Elements, const FName Name)
+FLearningAgentsActionObjectElement ULearningAgentsActions::MakeInclusiveUnionActionFromArrays(ULearningAgentsActionObject* Object, const TArray<FName>& ElementNames, const TArray<FLearningAgentsActionObjectElement>& Elements, const FName Tag)
 {
-	return MakeInclusiveUnionActionFromArrayViews(ElementNames, Elements, Name);
+	return MakeInclusiveUnionActionFromArrayViews(Object, ElementNames, Elements, Tag);
 }
 
-FLearningAgentsActionObjectElement ULearningAgentsActionObject::MakeInclusiveUnionActionFromArrayViews(const TArrayView<const FName> ElementNames, const TArrayView<const FLearningAgentsActionObjectElement> Elements, const FName Name)
+FLearningAgentsActionObjectElement ULearningAgentsActions::MakeInclusiveUnionActionFromArrayViews(ULearningAgentsActionObject* Object, const TArrayView<const FName> ElementNames, const TArrayView<const FLearningAgentsActionObjectElement> Elements, const FName Tag)
 {
+	if (!Object)
+	{
+		UE_LOG(LogLearning, Error, TEXT("MakeInclusiveUnionActionFromArrayViews: Object is nullptr."));
+		return FLearningAgentsActionObjectElement();
+	}
+
 	if (Elements.Num() != ElementNames.Num())
 	{
-		UE_LOG(LogLearning, Error, TEXT("%s: Number of elements (%i) must match number of names (%i)."), *GetName(), Elements.Num(), ElementNames.Num());
+		UE_LOG(LogLearning, Error, TEXT("MakeInclusiveUnionActionFromArrayViews: Number of elements (%i) must match number of names (%i)."), Elements.Num(), ElementNames.Num());
 		return FLearningAgentsActionObjectElement();
 	}
 
 	if (UE::Learning::Agents::Action::Private::ContainsDuplicates(ElementNames))
 	{
-		UE_LOG(LogLearning, Error, TEXT("%s: Element Names contain duplicates."), *GetName());
+		UE_LOG(LogLearning, Error, TEXT("MakeInclusiveUnionActionFromArrayViews: Element Names contain duplicates."));
 		return FLearningAgentsActionObjectElement();
 	}
 
@@ -1301,28 +1383,34 @@ FLearningAgentsActionObjectElement ULearningAgentsActionObject::MakeInclusiveUni
 
 	for (const FLearningAgentsActionObjectElement& Element : Elements)
 	{
-		if (!ActionObject.IsValid(Element.ObjectElement))
+		if (!Object->ActionObject.IsValid(Element.ObjectElement))
 		{
-			UE_LOG(LogLearning, Error, TEXT("%s: Invalid Action Object."), *GetName());
+			UE_LOG(LogLearning, Error, TEXT("MakeInclusiveUnionActionFromArrayViews: Invalid Action Object."));
 			return FLearningAgentsActionObjectElement();
 		}
 
 		SubElements.Add(Element.ObjectElement);
 	}
 
-	return { ActionObject.CreateOrInclusive({ ElementNames, SubElements }, Name) };
+	return { Object->ActionObject.CreateOrInclusive({ ElementNames, SubElements }, Tag) };
 }
 
-FLearningAgentsActionObjectElement ULearningAgentsActionObject::MakeStaticArrayAction(const TArray<FLearningAgentsActionObjectElement>& Elements, const FName Name)
+FLearningAgentsActionObjectElement ULearningAgentsActions::MakeStaticArrayAction(ULearningAgentsActionObject* Object, const TArray<FLearningAgentsActionObjectElement>& Elements, const FName Tag)
 {
-	return MakeStaticArrayActionFromArrayView(Elements, Name);
+	return MakeStaticArrayActionFromArrayView(Object, Elements, Tag);
 }
 
-FLearningAgentsActionObjectElement ULearningAgentsActionObject::MakeStaticArrayActionFromArrayView(const TArrayView<const FLearningAgentsActionObjectElement> Elements, const FName Name)
+FLearningAgentsActionObjectElement ULearningAgentsActions::MakeStaticArrayActionFromArrayView(ULearningAgentsActionObject* Object, const TArrayView<const FLearningAgentsActionObjectElement> Elements, const FName Tag)
 {
+	if (!Object)
+	{
+		UE_LOG(LogLearning, Error, TEXT("MakeStaticArrayActionFromArrayView: Object is nullptr."));
+		return FLearningAgentsActionObjectElement();
+	}
+
 	if (Elements.Num() == 0)
 	{
-		UE_LOG(LogLearning, Warning, TEXT("%s: Creating zero-sized Static Array Action."), *GetName());
+		UE_LOG(LogLearning, Warning, TEXT("MakeStaticArrayActionFromArrayView: Creating zero-sized Static Array Action."));
 	}
 
 	TArray<UE::Learning::Action::FObjectElement, TInlineAllocator<16>> SubElements;
@@ -1330,28 +1418,28 @@ FLearningAgentsActionObjectElement ULearningAgentsActionObject::MakeStaticArrayA
 
 	for (const FLearningAgentsActionObjectElement& Element : Elements)
 	{
-		if (!ActionObject.IsValid(Element.ObjectElement))
+		if (!Object->ActionObject.IsValid(Element.ObjectElement))
 		{
-			UE_LOG(LogLearning, Error, TEXT("%s: Invalid Action Object."), *GetName());
+			UE_LOG(LogLearning, Error, TEXT("MakeStaticArrayActionFromArrayView: Invalid Action Object."));
 			return FLearningAgentsActionObjectElement();
 		}
 
 		SubElements.Add(Element.ObjectElement);
 	}
 
-	return { ActionObject.CreateArray({ SubElements }, Name) };
+	return { Object->ActionObject.CreateArray({ SubElements }, Tag) };
 }
 
-FLearningAgentsActionObjectElement ULearningAgentsActionObject::MakePairAction(const FLearningAgentsActionObjectElement Key, const FLearningAgentsActionObjectElement Value, const FName Name)
+FLearningAgentsActionObjectElement ULearningAgentsActions::MakePairAction(ULearningAgentsActionObject* Object, const FLearningAgentsActionObjectElement Key, const FLearningAgentsActionObjectElement Value, const FName Tag)
 {
-	return MakeStructActionFromArrayViews({ TEXT("Key"), TEXT("Value") }, { Key, Value }, Name);
+	return MakeStructActionFromArrayViews(Object, { TEXT("Key"), TEXT("Value") }, { Key, Value }, Tag);
 }
 
-FLearningAgentsActionObjectElement ULearningAgentsActionObject::MakeEnumAction(const UEnum* Enum, const uint8 EnumValue, const FName Name)
+FLearningAgentsActionObjectElement ULearningAgentsActions::MakeEnumAction(ULearningAgentsActionObject* Object, const UEnum* Enum, const uint8 EnumValue, const FName Tag)
 {
 	if (!Enum)
 	{
-		UE_LOG(LogLearning, Error, TEXT("%s: Enum is nullptr."), *GetName());
+		UE_LOG(LogLearning, Error, TEXT("MakeEnumAction: Enum is nullptr."));
 		return FLearningAgentsActionObjectElement();
 	}
 
@@ -1359,24 +1447,24 @@ FLearningAgentsActionObjectElement ULearningAgentsActionObject::MakeEnumAction(c
 
 	if (EnumValueIndex == INDEX_NONE || EnumValueIndex < 0 || EnumValueIndex >= Enum->NumEnums() - 1)
 	{
-		UE_LOG(LogLearning, Error, TEXT("%s: EnumValue %i not valid for Enum '%s'."), *GetName(), EnumValue , *Enum->GetName());
+		UE_LOG(LogLearning, Error, TEXT("MakeEnumAction: EnumValue %i not valid for Enum '%s'."), EnumValue , *Enum->GetName());
 		return FLearningAgentsActionObjectElement();
 	}
 
-	return MakeExclusiveDiscreteAction(EnumValueIndex, Name);
+	return MakeExclusiveDiscreteAction(Object, EnumValueIndex, Tag);
 }
 
-FLearningAgentsActionObjectElement ULearningAgentsActionObject::MakeBitmaskAction(const UEnum* Enum, const int32 BitmaskValue, const FName Name)
+FLearningAgentsActionObjectElement ULearningAgentsActions::MakeBitmaskAction(ULearningAgentsActionObject* Object, const UEnum* Enum, const int32 BitmaskValue, const FName Tag)
 {
 	if (!Enum)
 	{
-		UE_LOG(LogLearning, Error, TEXT("%s: Enum is nullptr."), *GetName());
+		UE_LOG(LogLearning, Error, TEXT("MakeBitmaskAction: Enum is nullptr."));
 		return FLearningAgentsActionObjectElement();
 	}
 
 	if (Enum->NumEnums() - 1 > 32)
 	{
-		UE_LOG(LogLearning, Error, TEXT("%s: Too many values in Enum to use as Bitmask (%i)."), *GetName(), Enum->NumEnums() - 1);
+		UE_LOG(LogLearning, Error, TEXT("MakeBitmaskAction: Too many values in Enum to use as Bitmask (%i)."), Enum->NumEnums() - 1);
 		return FLearningAgentsActionObjectElement();
 	}
 
@@ -1391,179 +1479,190 @@ FLearningAgentsActionObjectElement ULearningAgentsActionObject::MakeBitmaskActio
 		}
 	}
 
-	return MakeInclusiveDiscreteActionFromArrayView(BitmaskIndices, Name);
+	return MakeInclusiveDiscreteActionFromArrayView(Object, BitmaskIndices, Tag);
 }
 
-FLearningAgentsActionObjectElement ULearningAgentsActionObject::MakeOptionalAction(const FLearningAgentsActionObjectElement Element, const ELearningAgentsOptionalAction Option, const FName Name)
+FLearningAgentsActionObjectElement ULearningAgentsActions::MakeOptionalAction(ULearningAgentsActionObject* Object, const FLearningAgentsActionObjectElement Element, const ELearningAgentsOptionalAction Option, const FName Tag)
 {
-	return MakeExclusiveUnionAction(
+	return MakeExclusiveUnionAction(Object,
 		Option == ELearningAgentsOptionalAction::Null ? TEXT("Null") : TEXT("Valid"),
-		Option == ELearningAgentsOptionalAction::Null ? MakeNullAction() : Element,
-		Name);
+		Option == ELearningAgentsOptionalAction::Null ? MakeNullAction(Object) : Element,
+		Tag);
 }
 
-FLearningAgentsActionObjectElement ULearningAgentsActionObject::MakeOptionalNullAction(const FName Name)
+FLearningAgentsActionObjectElement ULearningAgentsActions::MakeOptionalNullAction(ULearningAgentsActionObject* Object, const FName Tag)
 {
-	return MakeExclusiveUnionAction(TEXT("Null"), MakeNullAction(), Name);
+	return MakeExclusiveUnionAction(Object, TEXT("Null"), MakeNullAction(Object), Tag);
 }
 
-FLearningAgentsActionObjectElement ULearningAgentsActionObject::MakeOptionalValidAction(const FLearningAgentsActionObjectElement Element, const FName Name)
+FLearningAgentsActionObjectElement ULearningAgentsActions::MakeOptionalValidAction(ULearningAgentsActionObject* Object, const FLearningAgentsActionObjectElement Element, const FName Tag)
 {
-	return MakeExclusiveUnionAction(TEXT("Valid"), Element, Name);
+	return MakeExclusiveUnionAction(Object, TEXT("Valid"), Element, Tag);
 }
 
-FLearningAgentsActionObjectElement ULearningAgentsActionObject::MakeEitherAction(const FLearningAgentsActionObjectElement Element, const ELearningAgentsEitherAction Either, const FName Name)
+FLearningAgentsActionObjectElement ULearningAgentsActions::MakeEitherAction(ULearningAgentsActionObject* Object, const FLearningAgentsActionObjectElement Element, const ELearningAgentsEitherAction Either, const FName Tag)
 {
-	return MakeExclusiveUnionAction(Either == ELearningAgentsEitherAction::A ? TEXT("A") : TEXT("B"), Element, Name);
+	return MakeExclusiveUnionAction(Object, Either == ELearningAgentsEitherAction::A ? TEXT("A") : TEXT("B"), Element, Tag);
 }
 
-FLearningAgentsActionObjectElement ULearningAgentsActionObject::MakeEitherAAction(const FLearningAgentsActionObjectElement A, const FName Name)
+FLearningAgentsActionObjectElement ULearningAgentsActions::MakeEitherAAction(ULearningAgentsActionObject* Object, const FLearningAgentsActionObjectElement A, const FName Tag)
 {
-	return MakeExclusiveUnionAction(TEXT("A"), A, Name);
+	return MakeExclusiveUnionAction(Object, TEXT("A"), A, Tag);
 }
 
-FLearningAgentsActionObjectElement ULearningAgentsActionObject::MakeEitherBAction(const FLearningAgentsActionObjectElement B, const FName Name)
+FLearningAgentsActionObjectElement ULearningAgentsActions::MakeEitherBAction(ULearningAgentsActionObject* Object, const FLearningAgentsActionObjectElement B, const FName Tag)
 {
-	return MakeExclusiveUnionAction(TEXT("B"), B, Name);
+	return MakeExclusiveUnionAction(Object, TEXT("B"), B, Tag);
 }
 
-FLearningAgentsActionObjectElement ULearningAgentsActionObject::MakeEncodingAction(const FLearningAgentsActionObjectElement Element, const FName Name)
+FLearningAgentsActionObjectElement ULearningAgentsActions::MakeEncodingAction(ULearningAgentsActionObject* Object, const FLearningAgentsActionObjectElement Element, const FName Tag)
 {
-	if (!ActionObject.IsValid(Element.ObjectElement))
+	if (!Object)
+	{
+		UE_LOG(LogLearning, Error, TEXT("MakeEncodingAction: Object is nullptr."));
+		return FLearningAgentsActionObjectElement();
+	}
+
+	if (!Object->ActionObject.IsValid(Element.ObjectElement))
 	{
-		UE_LOG(LogLearning, Error, TEXT("%s: Invalid Action Object."), *GetName());
+		UE_LOG(LogLearning, Error, TEXT("MakeEncodingAction: Invalid Action Object."));
 		return FLearningAgentsActionObjectElement();
 	}
 
-	return { ActionObject.CreateEncoding({ Element.ObjectElement }, Name) };
+	return { Object->ActionObject.CreateEncoding({ Element.ObjectElement }, Tag) };
 }
 
-FLearningAgentsActionObjectElement ULearningAgentsActionObject::MakeBoolAction(const bool bValue, const FName Name)
+FLearningAgentsActionObjectElement ULearningAgentsActions::MakeBoolAction(ULearningAgentsActionObject* Object, const bool bValue, const FName Tag)
 {
-	return MakeExclusiveDiscreteAction(bValue ? 1 : 0, Name);
+	return MakeExclusiveDiscreteAction(Object, bValue ? 1 : 0, Tag);
 }
 
-FLearningAgentsActionObjectElement ULearningAgentsActionObject::MakeFloatAction(const float Value, const float FloatScale, const FName Name)
+FLearningAgentsActionObjectElement ULearningAgentsActions::MakeFloatAction(ULearningAgentsActionObject* Object, const float Value, const float FloatScale, const FName Tag)
 {
-	return MakeContinuousActionFromArrayView({ Value / FMath::Max(FloatScale, UE_SMALL_NUMBER) }, Name);
+	return MakeContinuousActionFromArrayView(Object, { Value / FMath::Max(FloatScale, UE_SMALL_NUMBER) }, Tag);
 }
 
-FLearningAgentsActionObjectElement ULearningAgentsActionObject::MakeLocationAction(const FVector Location, const FTransform RelativeTransform, const float LocationScale, const FName Name)
+FLearningAgentsActionObjectElement ULearningAgentsActions::MakeLocationAction(ULearningAgentsActionObject* Object, const FVector Location, const FTransform RelativeTransform, const float LocationScale, const FName Tag)
 {
 	const FVector LocalLocation = RelativeTransform.InverseTransformPosition(Location);
 
-	return MakeContinuousActionFromArrayView({
+	return MakeContinuousActionFromArrayView(Object, {
 		(float)LocalLocation.X / FMath::Max(LocationScale, UE_SMALL_NUMBER),
 		(float)LocalLocation.Y / FMath::Max(LocationScale, UE_SMALL_NUMBER),
-		(float)LocalLocation.Z / FMath::Max(LocationScale, UE_SMALL_NUMBER) }, Name);
+		(float)LocalLocation.Z / FMath::Max(LocationScale, UE_SMALL_NUMBER) }, Tag);
 }
 
-FLearningAgentsActionObjectElement ULearningAgentsActionObject::MakeRotationAction(const FRotator Rotation, const FRotator RelativeRotation, const float RotationScale, const FName Name)
+FLearningAgentsActionObjectElement ULearningAgentsActions::MakeRotationAction(ULearningAgentsActionObject* Object, const FRotator Rotation, const FRotator RelativeRotation, const float RotationScale, const FName Tag)
 {
-	return MakeRotationActionFromQuat(FQuat::MakeFromRotator(Rotation), FQuat::MakeFromRotator(RelativeRotation), RotationScale, Name);
+	return MakeRotationActionFromQuat(Object, FQuat::MakeFromRotator(Rotation), FQuat::MakeFromRotator(RelativeRotation), RotationScale, Tag);
 }
 
-FLearningAgentsActionObjectElement ULearningAgentsActionObject::MakeRotationActionFromQuat(const FQuat Rotation, const FQuat RelativeRotation, const float RotationScale, const FName Name)
+FLearningAgentsActionObjectElement ULearningAgentsActions::MakeRotationActionFromQuat(ULearningAgentsActionObject* Object, const FQuat Rotation, const FQuat RelativeRotation, const float RotationScale, const FName Tag)
 {
 	FQuat LocalRotation = RelativeRotation.Inverse() * Rotation;
 	LocalRotation.EnforceShortestArcWith(FQuat::Identity);
 	const FVector RotationVector = LocalRotation.ToRotationVector();
 
-	return MakeContinuousActionFromArrayView({
+	return MakeContinuousActionFromArrayView(Object, {
 		(float)RotationVector.X / FMath::Max(FMath::DegreesToRadians(RotationScale), UE_SMALL_NUMBER),
 		(float)RotationVector.Y / FMath::Max(FMath::DegreesToRadians(RotationScale), UE_SMALL_NUMBER),
 		(float)RotationVector.Z / FMath::Max(FMath::DegreesToRadians(RotationScale), UE_SMALL_NUMBER),
-		}, Name);
+		}, Tag);
 }
 
-FLearningAgentsActionObjectElement ULearningAgentsActionObject::MakeScaleAction(const FVector Scale, const FVector RelativeScale, const FName Name)
+FLearningAgentsActionObjectElement ULearningAgentsActions::MakeScaleAction(ULearningAgentsActionObject* Object, const FVector Scale, const FVector RelativeScale, const FName Tag)
 {
 	const FVector LocalLogScale =
 		UE::Learning::Agents::Action::Private::VectorLogSafe(Scale) -
 		UE::Learning::Agents::Action::Private::VectorLogSafe(RelativeScale);
 
-	return MakeContinuousActionFromArrayView({
+	return MakeContinuousActionFromArrayView(Object, {
 		(float)LocalLogScale.X,
 		(float)LocalLogScale.Y,
 		(float)LocalLogScale.Z,
-		}, Name);
+		}, Tag);
 }
 
-FLearningAgentsActionObjectElement ULearningAgentsActionObject::MakeTransformAction(const FTransform Transform, const FTransform RelativeTransform, const float LocationScale, const FName Name)
+FLearningAgentsActionObjectElement ULearningAgentsActions::MakeTransformAction(ULearningAgentsActionObject* Object, const FTransform Transform, const FTransform RelativeTransform, const float LocationScale, const FName Tag)
 {
 	const FTransform LocalTransform = Transform * RelativeTransform.Inverse();
 
-	return MakeStructActionFromArrayViews(
+	return MakeStructActionFromArrayViews(Object,
 		{
 			TEXT("Location"),
 			TEXT("Rotation"),
 			TEXT("Scale")
 		},
 		{
-			MakeLocationAction(LocalTransform.GetLocation(), FTransform::Identity, LocationScale),
-			MakeRotationActionFromQuat(LocalTransform.GetRotation(), FQuat::Identity),
-			MakeScaleAction(LocalTransform.GetScale3D(), FVector::OneVector)
+			MakeLocationAction(Object, LocalTransform.GetLocation(), FTransform::Identity, LocationScale),
+			MakeRotationActionFromQuat(Object, LocalTransform.GetRotation(), FQuat::Identity),
+			MakeScaleAction(Object, LocalTransform.GetScale3D(), FVector::OneVector)
 		},
-		Name);
+		Tag);
 }
 
-FLearningAgentsActionObjectElement ULearningAgentsActionObject::MakeAngleAction(const float Angle, const float RelativeAngle, const float AngleScale, const FName Name)
+FLearningAgentsActionObjectElement ULearningAgentsActions::MakeAngleAction(ULearningAgentsActionObject* Object, const float Angle, const float RelativeAngle, const float AngleScale, const FName Tag)
 {
-	return MakeAngleActionRadians(FMath::DegreesToRadians(Angle), FMath::DegreesToRadians(RelativeAngle), FMath::DegreesToRadians(AngleScale), Name);
+	return MakeAngleActionRadians(Object, FMath::DegreesToRadians(Angle), FMath::DegreesToRadians(RelativeAngle), FMath::DegreesToRadians(AngleScale), Tag);
 }
 
-FLearningAgentsActionObjectElement ULearningAgentsActionObject::MakeAngleActionRadians(const float Angle, const float RelativeAngle, const float AngleScale, const FName Name)
+FLearningAgentsActionObjectElement ULearningAgentsActions::MakeAngleActionRadians(ULearningAgentsActionObject* Object, const float Angle, const float RelativeAngle, const float AngleScale, const FName Tag)
 {
 	const float LocalAngle = FMath::FindDeltaAngleRadians(RelativeAngle, Angle);
-	return MakeContinuousActionFromArrayView({ LocalAngle / FMath::Max(AngleScale, UE_SMALL_NUMBER) }, Name);
+	return MakeContinuousActionFromArrayView(Object, { LocalAngle / FMath::Max(AngleScale, UE_SMALL_NUMBER) }, Tag);
 }
 
-FLearningAgentsActionObjectElement ULearningAgentsActionObject::MakeVelocityAction(const FVector Velocity, const FTransform RelativeTransform, const float VelocityScale, const FName Name)
+FLearningAgentsActionObjectElement ULearningAgentsActions::MakeVelocityAction(ULearningAgentsActionObject* Object, const FVector Velocity, const FTransform RelativeTransform, const float VelocityScale, const FName Tag)
 {
 	const FVector LocalVelocity = RelativeTransform.InverseTransformVectorNoScale(Velocity);
 
-	return MakeContinuousActionFromArrayView({
+	return MakeContinuousActionFromArrayView(Object, {
 		(float)LocalVelocity.X / FMath::Max(VelocityScale, UE_SMALL_NUMBER),
 		(float)LocalVelocity.Y / FMath::Max(VelocityScale, UE_SMALL_NUMBER),
 		(float)LocalVelocity.Z / FMath::Max(VelocityScale, UE_SMALL_NUMBER),
-		}, Name);
+		}, Tag);
 }
 
-FLearningAgentsActionObjectElement ULearningAgentsActionObject::MakeDirectionAction(const FVector Direction, const FTransform RelativeTransform, const FName Name)
+FLearningAgentsActionObjectElement ULearningAgentsActions::MakeDirectionAction(ULearningAgentsActionObject* Object, const FVector Direction, const FTransform RelativeTransform, const FName Tag)
 {
 	const FVector LocalDirection = RelativeTransform.InverseTransformVectorNoScale(Direction).GetSafeNormal(UE_SMALL_NUMBER, FVector::ForwardVector);
 
-	return MakeContinuousActionFromArrayView({
+	return MakeContinuousActionFromArrayView(Object, {
 		(float)LocalDirection.X,
 		(float)LocalDirection.Y,
 		(float)LocalDirection.Z,
-		}, Name);
+		}, Tag);
 }
 
-FLearningAgentsActionObjectElement ULearningAgentsActionObject::MakeSpeedAction(const float Speed, const float SpeedScale, const FName Name)
+FLearningAgentsActionObjectElement ULearningAgentsActions::MakeSpeedAction(ULearningAgentsActionObject* Object, const float Speed, const float SpeedScale, const FName Tag)
 {
-	return MakeContinuousActionFromArrayView({ Speed / FMath::Max(SpeedScale, UE_SMALL_NUMBER) });
+	return MakeContinuousActionFromArrayView(Object, { Speed / FMath::Max(SpeedScale, UE_SMALL_NUMBER) });
 }
 
-bool ULearningAgentsActionObject::GetNullAction(const FLearningAgentsActionObjectElement Element, const FName Name) const
+bool ULearningAgentsActions::GetNullAction(const ULearningAgentsActionObject* Object, const FLearningAgentsActionObjectElement Element, const FName Tag)
 {
-	if (!ActionObject.IsValid(Element.ObjectElement))
+	if (!Object)
+	{
+		UE_LOG(LogLearning, Error, TEXT("GetNullAction: Object is nullptr."));
+		return false;
+	}
+
+	if (!Object->ActionObject.IsValid(Element.ObjectElement))
 	{
-		UE_LOG(LogLearning, Error, TEXT("%s: Invalid Action Object."), *GetName());
+		UE_LOG(LogLearning, Error, TEXT("GetNullAction: Invalid Action Object."));
 		return false;
 	}
 
-	if (ActionObject.GetName(Element.ObjectElement) != Name)
+	if (Object->ActionObject.GetTag(Element.ObjectElement) != Tag)
 	{
-		UE_LOG(LogLearning, Warning, TEXT("%s: Action name does not match. Action is '%s' but asked for '%s'."), *GetName(), *ActionObject.GetName(Element.ObjectElement).ToString(), *Name.ToString());
+		UE_LOG(LogLearning, Warning, TEXT("GetNullAction: Action tag does not match. Action is '%s' but asked for '%s'."), *Object->ActionObject.GetTag(Element.ObjectElement).ToString(), *Tag.ToString());
 	}
 
-	if (ActionObject.GetType(Element.ObjectElement) != UE::Learning::Action::EType::Null)
+	if (Object->ActionObject.GetType(Element.ObjectElement) != UE::Learning::Action::EType::Null)
 	{
-		UE_LOG(LogLearning, Error, TEXT("%s: Action '%s' type does not match. Action is '%s' but asked for '%s'."),
-			*GetName(),
-			*ActionObject.GetName(Element.ObjectElement).ToString(),
-			UE::Learning::Agents::Action::Private::GetActionTypeString(ActionObject.GetType(Element.ObjectElement)),
+		UE_LOG(LogLearning, Error, TEXT("GetNullAction: Action '%s' type does not match. Action is '%s' but asked for '%s'."),
+			*Object->ActionObject.GetTag(Element.ObjectElement).ToString(),
+			UE::Learning::Agents::Action::Private::GetActionTypeString(Object->ActionObject.GetType(Element.ObjectElement)),
 			UE::Learning::Agents::Action::Private::GetActionTypeString(UE::Learning::Action::EType::Null));
 		return false;
 	}
@@ -1571,39 +1670,45 @@ bool ULearningAgentsActionObject::GetNullAction(const FLearningAgentsActionObjec
 	return true;
 }
 
-bool ULearningAgentsActionObject::GetContinuousActionNum(int32& OutNum, const FLearningAgentsActionObjectElement Element, const FName Name) const
+bool ULearningAgentsActions::GetContinuousActionNum(int32& OutNum, const ULearningAgentsActionObject* Object, const FLearningAgentsActionObjectElement Element, const FName Tag)
 {
-	if (!ActionObject.IsValid(Element.ObjectElement))
+	if (!Object)
+	{
+		UE_LOG(LogLearning, Error, TEXT("GetContinuousActionNum: Object is nullptr."));
+		OutNum = 0;
+		return false;
+	}
+
+	if (!Object->ActionObject.IsValid(Element.ObjectElement))
 	{
-		UE_LOG(LogLearning, Error, TEXT("%s: Invalid Action Object."), *GetName());
+		UE_LOG(LogLearning, Error, TEXT("GetContinuousActionNum: Invalid Action Object."));
 		OutNum = 0;
 		return false;
 	}
 
-	if (ActionObject.GetName(Element.ObjectElement) != Name)
+	if (Object->ActionObject.GetTag(Element.ObjectElement) != Tag)
 	{
-		UE_LOG(LogLearning, Warning, TEXT("%s: Action name does not match. Action is '%s' but asked for '%s'."), *GetName(), *ActionObject.GetName(Element.ObjectElement).ToString(), *Name.ToString());
+		UE_LOG(LogLearning, Warning, TEXT("GetContinuousActionNum: Action tag does not match. Action is '%s' but asked for '%s'."), *Object->ActionObject.GetTag(Element.ObjectElement).ToString(), *Tag.ToString());
 	}
 
-	if (ActionObject.GetType(Element.ObjectElement) != UE::Learning::Action::EType::Continuous)
+	if (Object->ActionObject.GetType(Element.ObjectElement) != UE::Learning::Action::EType::Continuous)
 	{
-		UE_LOG(LogLearning, Error, TEXT("%s: Action '%s' type does not match. Action is '%s' but asked for '%s'."),
-			*GetName(),
-			*ActionObject.GetName(Element.ObjectElement).ToString(),
-			UE::Learning::Agents::Action::Private::GetActionTypeString(ActionObject.GetType(Element.ObjectElement)),
+		UE_LOG(LogLearning, Error, TEXT("GetContinuousActionNum: Action '%s' type does not match. Action is '%s' but asked for '%s'."),
+			*Object->ActionObject.GetTag(Element.ObjectElement).ToString(),
+			UE::Learning::Agents::Action::Private::GetActionTypeString(Object->ActionObject.GetType(Element.ObjectElement)),
 			UE::Learning::Agents::Action::Private::GetActionTypeString(UE::Learning::Action::EType::Continuous));
 		OutNum = 0;
 		return false;
 	}
 
-	OutNum = ActionObject.GetContinuous(Element.ObjectElement).Values.Num();
+	OutNum = Object->ActionObject.GetContinuous(Element.ObjectElement).Values.Num();
 	return true;
 }
 
-bool ULearningAgentsActionObject::GetContinuousAction(TArray<float>& OutValues, const FLearningAgentsActionObjectElement Element, const FName Name) const
+bool ULearningAgentsActions::GetContinuousAction(TArray<float>& OutValues, const ULearningAgentsActionObject* Object, const FLearningAgentsActionObjectElement Element, const FName Tag)
 {
 	int32 OutValueNum = 0;
-	if (!GetContinuousActionNum(OutValueNum, Element, Name))
+	if (!GetContinuousActionNum(OutValueNum, Object, Element, Tag))
 	{
 		OutValues.Empty();
 		return false;
@@ -1611,7 +1716,7 @@ bool ULearningAgentsActionObject::GetContinuousAction(TArray<float>& OutValues,
 
 	OutValues.SetNumUninitialized(OutValueNum);
 
-	if (!GetContinuousActionToArrayView(OutValues, Element, Name))
+	if (!GetContinuousActionToArrayView(OutValues, Object, Element, Tag))
 	{
 		OutValues.Empty();
 		return false;
@@ -1620,37 +1725,43 @@ bool ULearningAgentsActionObject::GetContinuousAction(TArray<float>& OutValues,
 	return true;
 }
 
-bool ULearningAgentsActionObject::GetContinuousActionToArrayView(TArrayView<float> OutValues, const FLearningAgentsActionObjectElement Element, const FName Name) const
+bool ULearningAgentsActions::GetContinuousActionToArrayView(TArrayView<float> OutValues, const ULearningAgentsActionObject* Object, const FLearningAgentsActionObjectElement Element, const FName Tag)
 {
-	if (!ActionObject.IsValid(Element.ObjectElement))
+	if (!Object)
+	{
+		UE_LOG(LogLearning, Error, TEXT("GetContinuousActionToArrayView: Object is nullptr."));
+		UE::Learning::Array::Zero<1, float>(OutValues);
+		return false;
+	}
+
+	if (!Object->ActionObject.IsValid(Element.ObjectElement))
 	{
-		UE_LOG(LogLearning, Error, TEXT("%s: Invalid Action Object."), *GetName());
+		UE_LOG(LogLearning, Error, TEXT("GetContinuousActionToArrayView: Invalid Action Object."));
 		UE::Learning::Array::Zero<1, float>(OutValues);
 		return false;
 	}
 
-	if (ActionObject.GetName(Element.ObjectElement) != Name)
+	if (Object->ActionObject.GetTag(Element.ObjectElement) != Tag)
 	{
-		UE_LOG(LogLearning, Warning, TEXT("%s: Action name does not match. Action is '%s' but asked for '%s'."), *GetName(), *ActionObject.GetName(Element.ObjectElement).ToString(), *Name.ToString());
+		UE_LOG(LogLearning, Warning, TEXT("GetContinuousActionToArrayView: Action tag does not match. Action is '%s' but asked for '%s'."), *Object->ActionObject.GetTag(Element.ObjectElement).ToString(), *Tag.ToString());
 	}
 
-	if (ActionObject.GetType(Element.ObjectElement) != UE::Learning::Action::EType::Continuous)
+	if (Object->ActionObject.GetType(Element.ObjectElement) != UE::Learning::Action::EType::Continuous)
 	{
-		UE_LOG(LogLearning, Error, TEXT("%s: Action '%s' type does not match. Action is '%s' but asked for '%s'."),
-			*GetName(),
-			*ActionObject.GetName(Element.ObjectElement).ToString(),
-			UE::Learning::Agents::Action::Private::GetActionTypeString(ActionObject.GetType(Element.ObjectElement)),
+		UE_LOG(LogLearning, Error, TEXT("GetContinuousActionToArrayView: Action '%s' type does not match. Action is '%s' but asked for '%s'."),
+			*Object->ActionObject.GetTag(Element.ObjectElement).ToString(),
+			UE::Learning::Agents::Action::Private::GetActionTypeString(Object->ActionObject.GetType(Element.ObjectElement)),
 			UE::Learning::Agents::Action::Private::GetActionTypeString(UE::Learning::Action::EType::Continuous));
 		UE::Learning::Array::Zero<1, float>(OutValues);
 		return false;
 	}
 
-	const TArrayView<const float> Values = ActionObject.GetContinuous(Element.ObjectElement).Values;
+	const TArrayView<const float> Values = Object->ActionObject.GetContinuous(Element.ObjectElement).Values;
 
 	if (Values.Num() != OutValues.Num())
 	{
-		UE_LOG(LogLearning, Error, TEXT("%s: Action '%s' size does not match. Action is '%i' values but asked for '%i'."),
-			*GetName(), *ActionObject.GetName(Element.ObjectElement).ToString(),
+		UE_LOG(LogLearning, Error, TEXT("GetContinuousActionToArrayView: Action '%s' size does not match. Action is '%i' values but asked for '%i'."),
+			*Object->ActionObject.GetTag(Element.ObjectElement).ToString(),
 			Values.Num(), OutValues.Num());
 		UE::Learning::Array::Zero<1, float>(OutValues);
 		return false;
@@ -1660,68 +1771,80 @@ bool ULearningAgentsActionObject::GetContinuousActionToArrayView(TArrayView<floa
 	return true;
 }
 
-bool ULearningAgentsActionObject::GetExclusiveDiscreteAction(int32& OutIndex, const FLearningAgentsActionObjectElement Element, const FName Name) const
+bool ULearningAgentsActions::GetExclusiveDiscreteAction(int32& OutIndex, const ULearningAgentsActionObject* Object, const FLearningAgentsActionObjectElement Element, const FName Tag)
 {
-	if (!ActionObject.IsValid(Element.ObjectElement))
+	if (!Object)
+	{
+		UE_LOG(LogLearning, Error, TEXT("GetExclusiveDiscreteAction: Object is nullptr."));
+		OutIndex = 0;
+		return false;
+	}
+
+	if (!Object->ActionObject.IsValid(Element.ObjectElement))
 	{
-		UE_LOG(LogLearning, Error, TEXT("%s: Invalid Action Object."), *GetName());
+		UE_LOG(LogLearning, Error, TEXT("GetExclusiveDiscreteAction: Invalid Action Object."));
 		OutIndex = 0;
 		return false;
 	}
 
-	if (ActionObject.GetName(Element.ObjectElement) != Name)
+	if (Object->ActionObject.GetTag(Element.ObjectElement) != Tag)
 	{
-		UE_LOG(LogLearning, Warning, TEXT("%s: Action name does not match. Action is '%s' but asked for '%s'."), *GetName(), *ActionObject.GetName(Element.ObjectElement).ToString(), *Name.ToString());
+		UE_LOG(LogLearning, Warning, TEXT("GetExclusiveDiscreteAction: Action tag does not match. Action is '%s' but asked for '%s'."), *Object->ActionObject.GetTag(Element.ObjectElement).ToString(), *Tag.ToString());
 	}
 
-	if (ActionObject.GetType(Element.ObjectElement) != UE::Learning::Action::EType::DiscreteExclusive)
+	if (Object->ActionObject.GetType(Element.ObjectElement) != UE::Learning::Action::EType::DiscreteExclusive)
 	{
-		UE_LOG(LogLearning, Error, TEXT("%s: Action '%s' type does not match. Action is '%s' but asked for '%s'."),
-			*GetName(),
-			*ActionObject.GetName(Element.ObjectElement).ToString(),
-			UE::Learning::Agents::Action::Private::GetActionTypeString(ActionObject.GetType(Element.ObjectElement)),
+		UE_LOG(LogLearning, Error, TEXT("GetExclusiveDiscreteAction: Action '%s' type does not match. Action is '%s' but asked for '%s'."),
+			*Object->ActionObject.GetTag(Element.ObjectElement).ToString(),
+			UE::Learning::Agents::Action::Private::GetActionTypeString(Object->ActionObject.GetType(Element.ObjectElement)),
 			UE::Learning::Agents::Action::Private::GetActionTypeString(UE::Learning::Action::EType::DiscreteExclusive));
 		OutIndex = 0;
 		return false;
 	}
 
-	OutIndex = ActionObject.GetDiscreteExclusive(Element.ObjectElement).DiscreteIndex;
+	OutIndex = Object->ActionObject.GetDiscreteExclusive(Element.ObjectElement).DiscreteIndex;
 	return true;
 }
 
-bool ULearningAgentsActionObject::GetInclusiveDiscreteActionNum(int32& OutNum, const FLearningAgentsActionObjectElement Element, const FName Name) const
+bool ULearningAgentsActions::GetInclusiveDiscreteActionNum(int32& OutNum, const ULearningAgentsActionObject* Object, const FLearningAgentsActionObjectElement Element, const FName Tag)
 {
-	if (!ActionObject.IsValid(Element.ObjectElement))
+	if (!Object)
+	{
+		UE_LOG(LogLearning, Error, TEXT("GetInclusiveDiscreteActionNum: Object is nullptr."));
+		OutNum = 0;
+		return false;
+	}
+
+	if (!Object->ActionObject.IsValid(Element.ObjectElement))
 	{
-		UE_LOG(LogLearning, Error, TEXT("%s: Invalid Action Object."), *GetName());
+		UE_LOG(LogLearning, Error, TEXT("GetInclusiveDiscreteActionNum: Invalid Action Object."));
 		OutNum = 0;
 		return false;
 	}
 
-	if (ActionObject.GetName(Element.ObjectElement) != Name)
+	if (Object->ActionObject.GetTag(Element.ObjectElement) != Tag)
 	{
-		UE_LOG(LogLearning, Warning, TEXT("%s: Action name does not match. Action is '%s' but asked for '%s'."), *GetName(), *ActionObject.GetName(Element.ObjectElement).ToString(), *Name.ToString());
+		UE_LOG(LogLearning, Warning, TEXT("GetInclusiveDiscreteActionNum: Action tag does not match. Action is '%s' but asked for '%s'."), *Object->ActionObject.GetTag(Element.ObjectElement).ToString(), *Tag.ToString());
 	}
 
-	if (ActionObject.GetType(Element.ObjectElement) != UE::Learning::Action::EType::DiscreteInclusive)
+	if (Object->ActionObject.GetType(Element.ObjectElement) != UE::Learning::Action::EType::DiscreteInclusive)
 	{
-		UE_LOG(LogLearning, Error, TEXT("%s: Action '%s' type does not match. Action is '%s' but asked for '%s'."),
-			*GetName(),
-			*ActionObject.GetName(Element.ObjectElement).ToString(),
-			UE::Learning::Agents::Action::Private::GetActionTypeString(ActionObject.GetType(Element.ObjectElement)),
+		UE_LOG(LogLearning, Error, TEXT("GetInclusiveDiscreteActionNum: Action '%s' type does not match. Action is '%s' but asked for '%s'."),
+			*Object->ActionObject.GetTag(Element.ObjectElement).ToString(),
+			UE::Learning::Agents::Action::Private::GetActionTypeString(Object->ActionObject.GetType(Element.ObjectElement)),
 			UE::Learning::Agents::Action::Private::GetActionTypeString(UE::Learning::Action::EType::DiscreteInclusive));
 		OutNum = 0;
 		return false;
 	}
 
-	OutNum = ActionObject.GetDiscreteInclusive(Element.ObjectElement).DiscreteIndices.Num();
+	OutNum = Object->ActionObject.GetDiscreteInclusive(Element.ObjectElement).DiscreteIndices.Num();
 	return true;
 }
 
-bool ULearningAgentsActionObject::GetInclusiveDiscreteAction(TArray<int32>& OutIndices, const FLearningAgentsActionObjectElement Element, const FName Name) const
+bool ULearningAgentsActions::GetInclusiveDiscreteAction(TArray<int32>& OutIndices, const ULearningAgentsActionObject* Object, const FLearningAgentsActionObjectElement Element, const FName Tag)
 {
 	int32 OutIndexNum = 0;
-	if (!GetInclusiveDiscreteActionNum(OutIndexNum, Element, Name))
+	if (!GetInclusiveDiscreteActionNum(OutIndexNum, Object, Element, Tag))
 	{
 		OutIndices.Empty();
 		return false;
@@ -1729,7 +1852,7 @@ bool ULearningAgentsActionObject::GetInclusiveDiscreteAction(TArray<int32>& OutI
 
 	OutIndices.SetNumUninitialized(OutIndexNum);
 
-	if (!GetInclusiveDiscreteActionToArrayView(OutIndices, Element, Name))
+	if (!GetInclusiveDiscreteActionToArrayView(OutIndices, Object, Element, Tag))
 	{
 		OutIndices.Empty();
 		return false;
@@ -1738,37 +1861,43 @@ bool ULearningAgentsActionObject::GetInclusiveDiscreteAction(TArray<int32>& OutI
 	return true;
 }
 
-bool ULearningAgentsActionObject::GetInclusiveDiscreteActionToArrayView(TArrayView<int32> OutIndices, const FLearningAgentsActionObjectElement Element, const FName Name) const
+bool ULearningAgentsActions::GetInclusiveDiscreteActionToArrayView(TArrayView<int32> OutIndices, const ULearningAgentsActionObject* Object, const FLearningAgentsActionObjectElement Element, const FName Tag)
 {
-	if (!ActionObject.IsValid(Element.ObjectElement))
+	if (!Object)
+	{
+		UE_LOG(LogLearning, Error, TEXT("GetInclusiveDiscreteActionToArrayView: Object is nullptr."));
+		UE::Learning::Array::Zero<1, int32>(OutIndices);
+		return false;
+	}
+
+	if (!Object->ActionObject.IsValid(Element.ObjectElement))
 	{
-		UE_LOG(LogLearning, Error, TEXT("%s: Invalid Action Object."), *GetName());
+		UE_LOG(LogLearning, Error, TEXT("GetInclusiveDiscreteActionToArrayView: Invalid Action Object."));
 		UE::Learning::Array::Zero<1, int32>(OutIndices);
 		return false;
 	}
 
-	if (ActionObject.GetName(Element.ObjectElement) != Name)
+	if (Object->ActionObject.GetTag(Element.ObjectElement) != Tag)
 	{
-		UE_LOG(LogLearning, Warning, TEXT("%s: Action name does not match. Action is '%s' but asked for '%s'."), *GetName(), *ActionObject.GetName(Element.ObjectElement).ToString(), *Name.ToString());
+		UE_LOG(LogLearning, Warning, TEXT("GetInclusiveDiscreteActionToArrayView: Action tag does not match. Action is '%s' but asked for '%s'."), *Object->ActionObject.GetTag(Element.ObjectElement).ToString(), *Tag.ToString());
 	}
 
-	if (ActionObject.GetType(Element.ObjectElement) != UE::Learning::Action::EType::DiscreteInclusive)
+	if (Object->ActionObject.GetType(Element.ObjectElement) != UE::Learning::Action::EType::DiscreteInclusive)
 	{
-		UE_LOG(LogLearning, Error, TEXT("%s: Action '%s' type does not match. Action is '%s' but asked for '%s'."),
-			*GetName(),
-			*ActionObject.GetName(Element.ObjectElement).ToString(),
-			UE::Learning::Agents::Action::Private::GetActionTypeString(ActionObject.GetType(Element.ObjectElement)),
+		UE_LOG(LogLearning, Error, TEXT("GetInclusiveDiscreteActionToArrayView: Action '%s' type does not match. Action is '%s' but asked for '%s'."),
+			*Object->ActionObject.GetTag(Element.ObjectElement).ToString(),
+			UE::Learning::Agents::Action::Private::GetActionTypeString(Object->ActionObject.GetType(Element.ObjectElement)),
 			UE::Learning::Agents::Action::Private::GetActionTypeString(UE::Learning::Action::EType::DiscreteInclusive));
 		UE::Learning::Array::Zero<1, int32>(OutIndices);
 		return false;
 	}
 
-	const TArrayView<const int32> Indices = ActionObject.GetDiscreteInclusive(Element.ObjectElement).DiscreteIndices;
+	const TArrayView<const int32> Indices = Object->ActionObject.GetDiscreteInclusive(Element.ObjectElement).DiscreteIndices;
 
 	if (Indices.Num() != OutIndices.Num())
 	{
-		UE_LOG(LogLearning, Error, TEXT("%s: Action '%s' size does not match. Action is '%i' elements but asked for '%i'."),
-			*GetName(), *ActionObject.GetName(Element.ObjectElement).ToString(),
+		UE_LOG(LogLearning, Error, TEXT("GetInclusiveDiscreteActionToArrayView: Action '%s' size does not match. Action is '%i' elements but asked for '%i'."),
+			*Object->ActionObject.GetTag(Element.ObjectElement).ToString(),
 			Indices.Num(), OutIndices.Num());
 		UE::Learning::Array::Zero<1, int32>(OutIndices);
 		return false;
@@ -1778,41 +1907,47 @@ bool ULearningAgentsActionObject::GetInclusiveDiscreteActionToArrayView(TArrayVi
 	return true;
 }
 
-bool ULearningAgentsActionObject::GetStructActionNum(int32& OutNum, const FLearningAgentsActionObjectElement Element, const FName Name) const
+bool ULearningAgentsActions::GetStructActionNum(int32& OutNum, const ULearningAgentsActionObject* Object, const FLearningAgentsActionObjectElement Element, const FName Tag)
 {
-	if (!ActionObject.IsValid(Element.ObjectElement))
+	if (!Object)
+	{
+		UE_LOG(LogLearning, Error, TEXT("GetStructActionNum: Object is nullptr."));
+		OutNum = 0;
+		return false;
+	}
+
+	if (!Object->ActionObject.IsValid(Element.ObjectElement))
 	{
-		UE_LOG(LogLearning, Error, TEXT("%s: Invalid Action Object."), *GetName());
+		UE_LOG(LogLearning, Error, TEXT("GetStructActionNum: Invalid Action Object."));
 		OutNum = 0;
 		return false;
 	}
 
-	if (ActionObject.GetName(Element.ObjectElement) != Name)
+	if (Object->ActionObject.GetTag(Element.ObjectElement) != Tag)
 	{
-		UE_LOG(LogLearning, Warning, TEXT("%s: Action name does not match. Action is '%s' but asked for '%s'."), *GetName(), *ActionObject.GetName(Element.ObjectElement).ToString(), *Name.ToString());
+		UE_LOG(LogLearning, Warning, TEXT("GetStructActionNum: Action tag does not match. Action is '%s' but asked for '%s'."), *Object->ActionObject.GetTag(Element.ObjectElement).ToString(), *Tag.ToString());
 	}
 
-	if (ActionObject.GetType(Element.ObjectElement) != UE::Learning::Action::EType::And)
+	if (Object->ActionObject.GetType(Element.ObjectElement) != UE::Learning::Action::EType::And)
 	{
-		UE_LOG(LogLearning, Error, TEXT("%s: Action '%s' type does not match. Action is '%s' but asked for '%s'."),
-			*GetName(),
-			*ActionObject.GetName(Element.ObjectElement).ToString(),
-			UE::Learning::Agents::Action::Private::GetActionTypeString(ActionObject.GetType(Element.ObjectElement)),
+		UE_LOG(LogLearning, Error, TEXT("GetStructActionNum: Action '%s' type does not match. Action is '%s' but asked for '%s'."),
+			*Object->ActionObject.GetTag(Element.ObjectElement).ToString(),
+			UE::Learning::Agents::Action::Private::GetActionTypeString(Object->ActionObject.GetType(Element.ObjectElement)),
 			UE::Learning::Agents::Action::Private::GetActionTypeString(UE::Learning::Action::EType::And));
 		OutNum = 0;
 		return false;
 	}
 	
-	const UE::Learning::Action::FObjectAndParameters Parameters = ActionObject.GetAnd(Element.ObjectElement);
+	const UE::Learning::Action::FObjectAndParameters Parameters = Object->ActionObject.GetAnd(Element.ObjectElement);
 
 	OutNum = Parameters.Elements.Num();
 	return true;
 }
 
-bool ULearningAgentsActionObject::GetStructAction(TMap<FName, FLearningAgentsActionObjectElement>& OutElements, const FLearningAgentsActionObjectElement Element, const FName Name) const
+bool ULearningAgentsActions::GetStructAction(TMap<FName, FLearningAgentsActionObjectElement>& OutElements, const ULearningAgentsActionObject* Object, const FLearningAgentsActionObjectElement Element, const FName Tag)
 {
 	int32 OutElementNum = 0;
-	if (!GetStructActionNum(OutElementNum, Element, Name))
+	if (!GetStructActionNum(OutElementNum, Object, Element, Tag))
 	{
 		OutElements.Empty();
 		return false;
@@ -1823,7 +1958,7 @@ bool ULearningAgentsActionObject::GetStructAction(TMap<FName, FLearningAgentsAct
 	SubElementNames.SetNumUninitialized(OutElementNum);
 	SubElements.SetNumUninitialized(OutElementNum);
 
-	if (!GetStructActionToArrayViews(SubElementNames, SubElements, Element, Name))
+	if (!GetStructActionToArrayViews(SubElementNames, SubElements, Object, Element, Tag))
 	{
 		OutElements.Empty();
 		return false;
@@ -1838,10 +1973,10 @@ bool ULearningAgentsActionObject::GetStructAction(TMap<FName, FLearningAgentsAct
 	return true;
 }
 
-bool ULearningAgentsActionObject::GetStructActionToArrays(TArray<FName>& OutElementNames, TArray<FLearningAgentsActionObjectElement>& OutElements, const FLearningAgentsActionObjectElement Element, const FName Name) const
+bool ULearningAgentsActions::GetStructActionToArrays(TArray<FName>& OutElementNames, TArray<FLearningAgentsActionObjectElement>& OutElements, const ULearningAgentsActionObject* Object, const FLearningAgentsActionObjectElement Element, const FName Tag)
 {
 	int32 OutElementNum = 0;
-	if (!GetStructActionNum(OutElementNum, Element, Name))
+	if (!GetStructActionNum(OutElementNum, Object, Element, Tag))
 	{
 		OutElementNames.Empty();
 		OutElements.Empty();
@@ -1851,7 +1986,7 @@ bool ULearningAgentsActionObject::GetStructActionToArrays(TArray<FName>& OutElem
 	OutElementNames.SetNumUninitialized(OutElementNum);
 	OutElements.SetNumUninitialized(OutElementNum);
 
-	if (!GetStructActionToArrayViews(OutElementNames, OutElements, Element, Name))
+	if (!GetStructActionToArrayViews(OutElementNames, OutElements, Object, Element, Tag))
 	{
 		OutElementNames.Empty();
 		OutElements.Empty();
@@ -1861,45 +1996,51 @@ bool ULearningAgentsActionObject::GetStructActionToArrays(TArray<FName>& OutElem
 	return true;
 }
 
-bool ULearningAgentsActionObject::GetStructActionToArrayViews(TArrayView<FName> OutElementNames, TArrayView<FLearningAgentsActionObjectElement> OutElements, const FLearningAgentsActionObjectElement Element, const FName Name) const
+bool ULearningAgentsActions::GetStructActionToArrayViews(TArrayView<FName> OutElementNames, TArrayView<FLearningAgentsActionObjectElement> OutElements, const ULearningAgentsActionObject* Object, const FLearningAgentsActionObjectElement Element, const FName Tag)
 {
-	if (!ActionObject.IsValid(Element.ObjectElement))
+	if (!Object)
+	{
+		UE_LOG(LogLearning, Error, TEXT("GetStructActionToArrayViews: Object is nullptr."));
+		UE::Learning::Array::Set<1, FName>(OutElementNames, NAME_None);
+		UE::Learning::Array::Set<1, FLearningAgentsActionObjectElement>(OutElements, FLearningAgentsActionObjectElement());
+		return false;
+	}
+
+	if (!Object->ActionObject.IsValid(Element.ObjectElement))
 	{
-		UE_LOG(LogLearning, Error, TEXT("%s: Invalid Action Object."), *GetName());
+		UE_LOG(LogLearning, Error, TEXT("GetStructActionToArrayViews: Invalid Action Object."));
 		UE::Learning::Array::Set<1, FName>(OutElementNames, NAME_None);
 		UE::Learning::Array::Set<1, FLearningAgentsActionObjectElement>(OutElements, FLearningAgentsActionObjectElement());
 		return false;
 	}
 
-	if (ActionObject.GetName(Element.ObjectElement) != Name)
+	if (Object->ActionObject.GetTag(Element.ObjectElement) != Tag)
 	{
-		UE_LOG(LogLearning, Warning, TEXT("%s: Action name does not match. Action is '%s' but asked for '%s'."), *GetName(), *ActionObject.GetName(Element.ObjectElement).ToString(), *Name.ToString());
+		UE_LOG(LogLearning, Warning, TEXT("GetStructActionToArrayViews: Action tag does not match. Action is '%s' but asked for '%s'."), *Object->ActionObject.GetTag(Element.ObjectElement).ToString(), *Tag.ToString());
 	}
 
-	if (ActionObject.GetType(Element.ObjectElement) != UE::Learning::Action::EType::And)
+	if (Object->ActionObject.GetType(Element.ObjectElement) != UE::Learning::Action::EType::And)
 	{
-		UE_LOG(LogLearning, Error, TEXT("%s: Action '%s' type does not match. Action is '%s' but asked for '%s'."),
-			*GetName(),
-			*ActionObject.GetName(Element.ObjectElement).ToString(),
-			UE::Learning::Agents::Action::Private::GetActionTypeString(ActionObject.GetType(Element.ObjectElement)),
+		UE_LOG(LogLearning, Error, TEXT("GetStructActionToArrayViews: Action '%s' type does not match. Action is '%s' but asked for '%s'."),
+			*Object->ActionObject.GetTag(Element.ObjectElement).ToString(),
+			UE::Learning::Agents::Action::Private::GetActionTypeString(Object->ActionObject.GetType(Element.ObjectElement)),
 			UE::Learning::Agents::Action::Private::GetActionTypeString(UE::Learning::Action::EType::And));
 		UE::Learning::Array::Set<1, FName>(OutElementNames, NAME_None);
 		UE::Learning::Array::Set<1, FLearningAgentsActionObjectElement>(OutElements, FLearningAgentsActionObjectElement());
 		return false;
 	}
 
-	const UE::Learning::Action::FObjectAndParameters Parameters = ActionObject.GetAnd(Element.ObjectElement);
+	const UE::Learning::Action::FObjectAndParameters Parameters = Object->ActionObject.GetAnd(Element.ObjectElement);
 
 	if (Parameters.Elements.Num() == 0)
 	{
-		UE_LOG(LogLearning, Warning, TEXT("%s: Getting zero-sized And Action."), *GetName());
+		UE_LOG(LogLearning, Warning, TEXT("GetStructActionToArrayViews: Getting zero-sized And Action."));
 	}
 
 	if (Parameters.Elements.Num() != OutElements.Num())
 	{
-		UE_LOG(LogLearning, Error, TEXT("%s: Action '%s' size does not match. Action is '%i' elements but asked for '%i'."),
-			*GetName(),
-			*ActionObject.GetName(Element.ObjectElement).ToString(),
+		UE_LOG(LogLearning, Error, TEXT("GetStructActionToArrayViews: Action '%s' size does not match. Action is '%i' elements but asked for '%i'."),
+			*Object->ActionObject.GetTag(Element.ObjectElement).ToString(),
 			Parameters.Elements.Num(), OutElements.Num());
 		UE::Learning::Array::Set<1, FName>(OutElementNames, NAME_None);
 		UE::Learning::Array::Set<1, FLearningAgentsActionObjectElement>(OutElements, FLearningAgentsActionObjectElement());
@@ -1908,9 +2049,9 @@ bool ULearningAgentsActionObject::GetStructActionToArrayViews(TArrayView<FName>
 
 	for (int32 ElementIdx = 0; ElementIdx < Parameters.Elements.Num(); ElementIdx++)
 	{
-		if (!ActionObject.IsValid(Parameters.Elements[ElementIdx]))
+		if (!Object->ActionObject.IsValid(Parameters.Elements[ElementIdx]))
 		{
-			UE_LOG(LogLearning, Error, TEXT("%s: Invalid Action Object."), *GetName());
+			UE_LOG(LogLearning, Error, TEXT("GetStructActionToArrayViews: Invalid Action Object."));
 			UE::Learning::Array::Set<1, FName>(OutElementNames, NAME_None);
 			UE::Learning::Array::Set<1, FLearningAgentsActionObjectElement>(OutElements, FLearningAgentsActionObjectElement());
 			return false;
@@ -1923,75 +2064,88 @@ bool ULearningAgentsActionObject::GetStructActionToArrayViews(TArrayView<FName>
 	return true;
 }
 
-bool ULearningAgentsActionObject::GetExclusiveUnionAction(FName& OutElementName, FLearningAgentsActionObjectElement& OutElement, const FLearningAgentsActionObjectElement Element, const FName Name) const
+bool ULearningAgentsActions::GetExclusiveUnionAction(FName& OutElementName, FLearningAgentsActionObjectElement& OutElement, const ULearningAgentsActionObject* Object, const FLearningAgentsActionObjectElement Element, const FName Tag)
 {
-	if (!ActionObject.IsValid(Element.ObjectElement))
+	if (!Object)
+	{
+		UE_LOG(LogLearning, Error, TEXT("GetExclusiveUnionAction: Object is nullptr."));
+		OutElementName = NAME_None;
+		OutElement = FLearningAgentsActionObjectElement();
+		return false;
+	}
+
+	if (!Object->ActionObject.IsValid(Element.ObjectElement))
 	{
-		UE_LOG(LogLearning, Error, TEXT("%s: Invalid Action Object."), *GetName());
+		UE_LOG(LogLearning, Error, TEXT("GetExclusiveUnionAction: Invalid Action Object."));
 		OutElementName = NAME_None;
 		OutElement = FLearningAgentsActionObjectElement();
 		return false;
 	}
 
-	if (ActionObject.GetName(Element.ObjectElement) != Name)
+	if (Object->ActionObject.GetTag(Element.ObjectElement) != Tag)
 	{
-		UE_LOG(LogLearning, Warning, TEXT("%s: Action name does not match. Action is '%s' but asked for '%s'."), *GetName(), *ActionObject.GetName(Element.ObjectElement).ToString(), *Name.ToString());
+		UE_LOG(LogLearning, Warning, TEXT("GetExclusiveUnionAction: Action tag does not match. Action is '%s' but asked for '%s'."), *Object->ActionObject.GetTag(Element.ObjectElement).ToString(), *Tag.ToString());
 	}
 
-	if (ActionObject.GetType(Element.ObjectElement) != UE::Learning::Action::EType::OrExclusive)
+	if (Object->ActionObject.GetType(Element.ObjectElement) != UE::Learning::Action::EType::OrExclusive)
 	{
-		UE_LOG(LogLearning, Error, TEXT("%s: Action '%s' type does not match. Action is '%s' but asked for '%s'."),
-			*GetName(),
-			*ActionObject.GetName(Element.ObjectElement).ToString(),
-			UE::Learning::Agents::Action::Private::GetActionTypeString(ActionObject.GetType(Element.ObjectElement)),
+		UE_LOG(LogLearning, Error, TEXT("GetExclusiveUnionAction: Action '%s' type does not match. Action is '%s' but asked for '%s'."),
+			*Object->ActionObject.GetTag(Element.ObjectElement).ToString(),
+			UE::Learning::Agents::Action::Private::GetActionTypeString(Object->ActionObject.GetType(Element.ObjectElement)),
 			UE::Learning::Agents::Action::Private::GetActionTypeString(UE::Learning::Action::EType::OrExclusive));
 		OutElementName = NAME_None;
 		OutElement = FLearningAgentsActionObjectElement();
 		return false;
 	}
 
-	const UE::Learning::Action::FObjectOrExclusiveParameters Parameters = ActionObject.GetOrExclusive(Element.ObjectElement);
+	const UE::Learning::Action::FObjectOrExclusiveParameters Parameters = Object->ActionObject.GetOrExclusive(Element.ObjectElement);
 
 	OutElementName = Parameters.ElementName;
 	OutElement = { Parameters.Element };
 	return true;
 }
 
-bool ULearningAgentsActionObject::GetInclusiveUnionActionNum(int32& OutNum, const FLearningAgentsActionObjectElement Element, const FName Name) const
+bool ULearningAgentsActions::GetInclusiveUnionActionNum(int32& OutNum, const ULearningAgentsActionObject* Object, const FLearningAgentsActionObjectElement Element, const FName Tag)
 {
-	if (!ActionObject.IsValid(Element.ObjectElement))
+	if (!Object)
+	{
+		UE_LOG(LogLearning, Error, TEXT("GetInclusiveUnionActionNum: Object is nullptr."));
+		OutNum = 0;
+		return false;
+	}
+
+	if (!Object->ActionObject.IsValid(Element.ObjectElement))
 	{
-		UE_LOG(LogLearning, Error, TEXT("%s: Invalid Action Object."), *GetName());
+		UE_LOG(LogLearning, Error, TEXT("GetInclusiveUnionActionNum: Invalid Action Object."));
 		OutNum = 0;
 		return false;
 	}
 
-	if (ActionObject.GetName(Element.ObjectElement) != Name)
+	if (Object->ActionObject.GetTag(Element.ObjectElement) != Tag)
 	{
-		UE_LOG(LogLearning, Warning, TEXT("%s: Action name does not match. Action is '%s' but asked for '%s'."), *GetName(), *ActionObject.GetName(Element.ObjectElement).ToString(), *Name.ToString());
+		UE_LOG(LogLearning, Warning, TEXT("GetInclusiveUnionActionNum: Action tag does not match. Action is '%s' but asked for '%s'."), *Object->ActionObject.GetTag(Element.ObjectElement).ToString(), *Tag.ToString());
 	}
 
-	if (ActionObject.GetType(Element.ObjectElement) != UE::Learning::Action::EType::OrInclusive)
+	if (Object->ActionObject.GetType(Element.ObjectElement) != UE::Learning::Action::EType::OrInclusive)
 	{
-		UE_LOG(LogLearning, Error, TEXT("%s: Action '%s' type does not match. Action is '%s' but asked for '%s'."),
-			*GetName(),
-			*ActionObject.GetName(Element.ObjectElement).ToString(),
-			UE::Learning::Agents::Action::Private::GetActionTypeString(ActionObject.GetType(Element.ObjectElement)),
+		UE_LOG(LogLearning, Error, TEXT("GetInclusiveUnionActionNum: Action '%s' type does not match. Action is '%s' but asked for '%s'."),
+			*Object->ActionObject.GetTag(Element.ObjectElement).ToString(),
+			UE::Learning::Agents::Action::Private::GetActionTypeString(Object->ActionObject.GetType(Element.ObjectElement)),
 			UE::Learning::Agents::Action::Private::GetActionTypeString(UE::Learning::Action::EType::OrInclusive));
 		OutNum = 0;
 		return false;
 	}
 
-	const UE::Learning::Action::FObjectOrInclusiveParameters Parameters = ActionObject.GetOrInclusive(Element.ObjectElement);
+	const UE::Learning::Action::FObjectOrInclusiveParameters Parameters = Object->ActionObject.GetOrInclusive(Element.ObjectElement);
 
 	OutNum = Parameters.Elements.Num();
 	return true;
 }
 
-bool ULearningAgentsActionObject::GetInclusiveUnionAction(TMap<FName, FLearningAgentsActionObjectElement>& OutElements, const FLearningAgentsActionObjectElement Element, const FName Name) const
+bool ULearningAgentsActions::GetInclusiveUnionAction(TMap<FName, FLearningAgentsActionObjectElement>& OutElements, const ULearningAgentsActionObject* Object, const FLearningAgentsActionObjectElement Element, const FName Tag)
 {
 	int32 OutElementNum = 0;
-	if (!GetInclusiveUnionActionNum(OutElementNum, Element, Name))
+	if (!GetInclusiveUnionActionNum(OutElementNum, Object, Element, Tag))
 	{
 		OutElements.Empty();
 		return false;
@@ -2002,7 +2156,7 @@ bool ULearningAgentsActionObject::GetInclusiveUnionAction(TMap<FName, FLearningA
 	SubElementNames.SetNumUninitialized(OutElementNum);
 	SubElements.SetNumUninitialized(OutElementNum);
 
-	if (!GetInclusiveUnionActionToArrayViews(SubElementNames, SubElements, Element, Name))
+	if (!GetInclusiveUnionActionToArrayViews(SubElementNames, SubElements, Object, Element, Tag))
 	{
 		OutElements.Empty();
 		return false;
@@ -2017,10 +2171,10 @@ bool ULearningAgentsActionObject::GetInclusiveUnionAction(TMap<FName, FLearningA
 	return true;
 }
 
-bool ULearningAgentsActionObject::GetInclusiveUnionActionToArrays(TArray<FName>& OutElementNames, TArray<FLearningAgentsActionObjectElement>& OutElements, const FLearningAgentsActionObjectElement Element, const FName Name) const
+bool ULearningAgentsActions::GetInclusiveUnionActionToArrays(TArray<FName>& OutElementNames, TArray<FLearningAgentsActionObjectElement>& OutElements, const ULearningAgentsActionObject* Object, const FLearningAgentsActionObjectElement Element, const FName Tag)
 {
 	int32 OutElementNum = 0;
-	if (!GetInclusiveUnionActionNum(OutElementNum, Element, Name))
+	if (!GetInclusiveUnionActionNum(OutElementNum, Object, Element, Tag))
 	{
 		OutElementNames.Empty();
 		OutElements.Empty();
@@ -2030,7 +2184,7 @@ bool ULearningAgentsActionObject::GetInclusiveUnionActionToArrays(TArray<FName>&
 	OutElementNames.SetNumUninitialized(OutElementNum);
 	OutElements.SetNumUninitialized(OutElementNum);
 
-	if (!GetInclusiveUnionActionToArrayViews(OutElementNames, OutElements, Element, Name))
+	if (!GetInclusiveUnionActionToArrayViews(OutElementNames, OutElements, Object, Element, Tag))
 	{
 		OutElementNames.Empty();
 		OutElements.Empty();
@@ -2040,40 +2194,46 @@ bool ULearningAgentsActionObject::GetInclusiveUnionActionToArrays(TArray<FName>&
 	return true;
 }
 
-bool ULearningAgentsActionObject::GetInclusiveUnionActionToArrayViews(TArrayView<FName> OutElementNames, TArrayView<FLearningAgentsActionObjectElement> OutElements, const FLearningAgentsActionObjectElement Element, const FName Name) const
+bool ULearningAgentsActions::GetInclusiveUnionActionToArrayViews(TArrayView<FName> OutElementNames, TArrayView<FLearningAgentsActionObjectElement> OutElements, const ULearningAgentsActionObject* Object, const FLearningAgentsActionObjectElement Element, const FName Tag)
 {
-	if (!ActionObject.IsValid(Element.ObjectElement))
+	if (!Object)
+	{
+		UE_LOG(LogLearning, Error, TEXT("GetInclusiveUnionActionToArrayViews: Object is nullptr."));
+		UE::Learning::Array::Set<1, FName>(OutElementNames, NAME_None);
+		UE::Learning::Array::Set<1, FLearningAgentsActionObjectElement>(OutElements, FLearningAgentsActionObjectElement());
+		return false;
+	}
+
+	if (!Object->ActionObject.IsValid(Element.ObjectElement))
 	{
-		UE_LOG(LogLearning, Error, TEXT("%s: Invalid Action Object."), *GetName());
+		UE_LOG(LogLearning, Error, TEXT("GetInclusiveUnionActionToArrayViews: Invalid Action Object."));
 		UE::Learning::Array::Set<1, FName>(OutElementNames, NAME_None);
 		UE::Learning::Array::Set<1, FLearningAgentsActionObjectElement>(OutElements, FLearningAgentsActionObjectElement());
 		return false;
 	}
 
-	if (ActionObject.GetName(Element.ObjectElement) != Name)
+	if (Object->ActionObject.GetTag(Element.ObjectElement) != Tag)
 	{
-		UE_LOG(LogLearning, Warning, TEXT("%s: Action name does not match. Action is '%s' but asked for '%s'."), *GetName(), *ActionObject.GetName(Element.ObjectElement).ToString(), *Name.ToString());
+		UE_LOG(LogLearning, Warning, TEXT("GetInclusiveUnionActionToArrayViews: Action tag does not match. Action is '%s' but asked for '%s'."), *Object->ActionObject.GetTag(Element.ObjectElement).ToString(), *Tag.ToString());
 	}
 
-	if (ActionObject.GetType(Element.ObjectElement) != UE::Learning::Action::EType::OrInclusive)
+	if (Object->ActionObject.GetType(Element.ObjectElement) != UE::Learning::Action::EType::OrInclusive)
 	{
-		UE_LOG(LogLearning, Error, TEXT("%s: Action '%s' type does not match. Action is '%s' but asked for '%s'."),
-			*GetName(),
-			*ActionObject.GetName(Element.ObjectElement).ToString(),
-			UE::Learning::Agents::Action::Private::GetActionTypeString(ActionObject.GetType(Element.ObjectElement)),
+		UE_LOG(LogLearning, Error, TEXT("GetInclusiveUnionActionToArrayViews: Action '%s' type does not match. Action is '%s' but asked for '%s'."),
+			*Object->ActionObject.GetTag(Element.ObjectElement).ToString(),
+			UE::Learning::Agents::Action::Private::GetActionTypeString(Object->ActionObject.GetType(Element.ObjectElement)),
 			UE::Learning::Agents::Action::Private::GetActionTypeString(UE::Learning::Action::EType::OrInclusive));
 		UE::Learning::Array::Set<1, FName>(OutElementNames, NAME_None);
 		UE::Learning::Array::Set<1, FLearningAgentsActionObjectElement>(OutElements, FLearningAgentsActionObjectElement());
 		return false;
 	}
 
-	const UE::Learning::Action::FObjectOrInclusiveParameters Parameters = ActionObject.GetOrInclusive(Element.ObjectElement);
+	const UE::Learning::Action::FObjectOrInclusiveParameters Parameters = Object->ActionObject.GetOrInclusive(Element.ObjectElement);
 
 	if (Parameters.Elements.Num() != OutElements.Num())
 	{
-		UE_LOG(LogLearning, Error, TEXT("%s: Action '%s' size does not match. Action is '%i' elements but asked for '%i'."),
-			*GetName(),
-			*ActionObject.GetName(Element.ObjectElement).ToString(),
+		UE_LOG(LogLearning, Error, TEXT("GetInclusiveUnionActionToArrayViews: Action '%s' size does not match. Action is '%i' elements but asked for '%i'."),
+			*Object->ActionObject.GetTag(Element.ObjectElement).ToString(),
 			Parameters.Elements.Num(), OutElements.Num());
 		UE::Learning::Array::Set<1, FName>(OutElementNames, NAME_None);
 		UE::Learning::Array::Set<1, FLearningAgentsActionObjectElement>(OutElements, FLearningAgentsActionObjectElement());
@@ -2082,9 +2242,9 @@ bool ULearningAgentsActionObject::GetInclusiveUnionActionToArrayViews(TArrayView
 
 	for (int32 ElementIdx = 0; ElementIdx < Parameters.Elements.Num(); ElementIdx++)
 	{
-		if (!ActionObject.IsValid(Parameters.Elements[ElementIdx]))
+		if (!Object->ActionObject.IsValid(Parameters.Elements[ElementIdx]))
 		{
-			UE_LOG(LogLearning, Error, TEXT("%s: Invalid Action Object."), *GetName());
+			UE_LOG(LogLearning, Error, TEXT("GetInclusiveUnionActionToArrayViews: Invalid Action Object."));
 			UE::Learning::Array::Set<1, FName>(OutElementNames, NAME_None);
 			UE::Learning::Array::Set<1, FLearningAgentsActionObjectElement>(OutElements, FLearningAgentsActionObjectElement());
 			return false;
@@ -2097,39 +2257,45 @@ bool ULearningAgentsActionObject::GetInclusiveUnionActionToArrayViews(TArrayView
 	return true;
 }
 
-bool ULearningAgentsActionObject::GetStaticArrayActionNum(int32& OutNum, const FLearningAgentsActionObjectElement Element, const FName Name) const
+bool ULearningAgentsActions::GetStaticArrayActionNum(int32& OutNum, const ULearningAgentsActionObject* Object, const FLearningAgentsActionObjectElement Element, const FName Tag)
 {
-	if (!ActionObject.IsValid(Element.ObjectElement))
+	if (!Object)
+	{
+		UE_LOG(LogLearning, Error, TEXT("GetStaticArrayActionNum: Object is nullptr."));
+		OutNum = 0;
+		return false;
+	}
+
+	if (!Object->ActionObject.IsValid(Element.ObjectElement))
 	{
-		UE_LOG(LogLearning, Error, TEXT("%s: Invalid Action Object."), *GetName());
+		UE_LOG(LogLearning, Error, TEXT("GetStaticArrayActionNum: Invalid Action Object."));
 		OutNum = 0;
 		return false;
 	}
 
-	if (ActionObject.GetName(Element.ObjectElement) != Name)
+	if (Object->ActionObject.GetTag(Element.ObjectElement) != Tag)
 	{
-		UE_LOG(LogLearning, Warning, TEXT("%s: Action name does not match. Action is '%s' but asked for '%s'."), *GetName(), *ActionObject.GetName(Element.ObjectElement).ToString(), *Name.ToString());
+		UE_LOG(LogLearning, Warning, TEXT("GetStaticArrayActionNum: Action tag does not match. Action is '%s' but asked for '%s'."), *Object->ActionObject.GetTag(Element.ObjectElement).ToString(), *Tag.ToString());
 	}
 
-	if (ActionObject.GetType(Element.ObjectElement) != UE::Learning::Action::EType::Array)
+	if (Object->ActionObject.GetType(Element.ObjectElement) != UE::Learning::Action::EType::Array)
 	{
-		UE_LOG(LogLearning, Error, TEXT("%s: Action '%s' type does not match. Action is '%s' but asked for '%s'."),
-			*GetName(),
-			*ActionObject.GetName(Element.ObjectElement).ToString(),
-			UE::Learning::Agents::Action::Private::GetActionTypeString(ActionObject.GetType(Element.ObjectElement)),
+		UE_LOG(LogLearning, Error, TEXT("GetStaticArrayActionNum: Action '%s' type does not match. Action is '%s' but asked for '%s'."),
+			*Object->ActionObject.GetTag(Element.ObjectElement).ToString(),
+			UE::Learning::Agents::Action::Private::GetActionTypeString(Object->ActionObject.GetType(Element.ObjectElement)),
 			UE::Learning::Agents::Action::Private::GetActionTypeString(UE::Learning::Action::EType::Array));
 		OutNum = 0;
 		return false;
 	}
 
-	OutNum = ActionObject.GetArray(Element.ObjectElement).Elements.Num();
+	OutNum = Object->ActionObject.GetArray(Element.ObjectElement).Elements.Num();
 	return true;
 }
 
-bool ULearningAgentsActionObject::GetStaticArrayAction(TArray<FLearningAgentsActionObjectElement>& OutElements, const FLearningAgentsActionObjectElement Element, const FName Name) const
+bool ULearningAgentsActions::GetStaticArrayAction(TArray<FLearningAgentsActionObjectElement>& OutElements, const ULearningAgentsActionObject* Object, const FLearningAgentsActionObjectElement Element, const FName Tag)
 {
 	int32 OutElementNum = 0;
-	if (!GetStaticArrayActionNum(OutElementNum, Element, Name))
+	if (!GetStaticArrayActionNum(OutElementNum, Object, Element, Tag))
 	{
 		OutElements.Empty();
 		return false;
@@ -2137,7 +2303,7 @@ bool ULearningAgentsActionObject::GetStaticArrayAction(TArray<FLearningAgentsAct
 
 	OutElements.SetNumUninitialized(OutElementNum);
 
-	if (!GetStaticArrayActionToArrayView(OutElements, Element, Name))
+	if (!GetStaticArrayActionToArrayView(OutElements, Object, Element, Tag))
 	{
 		OutElements.Empty();
 		return false;
@@ -2146,42 +2312,48 @@ bool ULearningAgentsActionObject::GetStaticArrayAction(TArray<FLearningAgentsAct
 	return true;
 }
 
-bool ULearningAgentsActionObject::GetStaticArrayActionToArrayView(TArrayView<FLearningAgentsActionObjectElement> OutElements, const FLearningAgentsActionObjectElement Element, const FName Name) const
+bool ULearningAgentsActions::GetStaticArrayActionToArrayView(TArrayView<FLearningAgentsActionObjectElement> OutElements, const ULearningAgentsActionObject* Object, const FLearningAgentsActionObjectElement Element, const FName Tag)
 {
-	if (!ActionObject.IsValid(Element.ObjectElement))
+	if (!Object)
 	{
-		UE_LOG(LogLearning, Error, TEXT("%s: Invalid Action Object."), *GetName());
+		UE_LOG(LogLearning, Error, TEXT("GetStaticArrayActionToArrayView: Object is nullptr."));
 		UE::Learning::Array::Set<1, FLearningAgentsActionObjectElement>(OutElements, FLearningAgentsActionObjectElement());
 		return false;
 	}
 
-	if (ActionObject.GetName(Element.ObjectElement) != Name)
+	if (!Object->ActionObject.IsValid(Element.ObjectElement))
 	{
-		UE_LOG(LogLearning, Warning, TEXT("%s: Action name does not match. Action is '%s' but asked for '%s'."), *GetName(), *ActionObject.GetName(Element.ObjectElement).ToString(), *Name.ToString());
+		UE_LOG(LogLearning, Error, TEXT("GetStaticArrayActionToArrayView: Invalid Action Object."));
+		UE::Learning::Array::Set<1, FLearningAgentsActionObjectElement>(OutElements, FLearningAgentsActionObjectElement());
+		return false;
+	}
+
+	if (Object->ActionObject.GetTag(Element.ObjectElement) != Tag)
+	{
+		UE_LOG(LogLearning, Warning, TEXT("GetStaticArrayActionToArrayView: Action tag does not match. Action is '%s' but asked for '%s'."), *Object->ActionObject.GetTag(Element.ObjectElement).ToString(), *Tag.ToString());
 	}
 
-	if (ActionObject.GetType(Element.ObjectElement) != UE::Learning::Action::EType::Array)
+	if (Object->ActionObject.GetType(Element.ObjectElement) != UE::Learning::Action::EType::Array)
 	{
-		UE_LOG(LogLearning, Error, TEXT("%s: Action '%s' type does not match. Action is '%s' but asked for '%s'."), 
-			*GetName(),
-			*ActionObject.GetName(Element.ObjectElement).ToString(),
-			UE::Learning::Agents::Action::Private::GetActionTypeString(ActionObject.GetType(Element.ObjectElement)),
+		UE_LOG(LogLearning, Error, TEXT("GetStaticArrayActionToArrayView: Action '%s' type does not match. Action is '%s' but asked for '%s'."), 
+			*Object->ActionObject.GetTag(Element.ObjectElement).ToString(),
+			UE::Learning::Agents::Action::Private::GetActionTypeString(Object->ActionObject.GetType(Element.ObjectElement)),
 			UE::Learning::Agents::Action::Private::GetActionTypeString(UE::Learning::Action::EType::Array));
 		UE::Learning::Array::Set<1, FLearningAgentsActionObjectElement>(OutElements, FLearningAgentsActionObjectElement());
 		return false;
 	}
 
-	const TArrayView<const UE::Learning::Action::FObjectElement> SubElements = ActionObject.GetArray(Element.ObjectElement).Elements;
+	const TArrayView<const UE::Learning::Action::FObjectElement> SubElements = Object->ActionObject.GetArray(Element.ObjectElement).Elements;
 
 	if (SubElements.Num() == 0)
 	{
-		UE_LOG(LogLearning, Warning, TEXT("%s: Getting zero-sized Array Action."), *GetName());
+		UE_LOG(LogLearning, Warning, TEXT("GetStaticArrayActionToArrayView: Getting zero-sized Array Action."));
 	}
 
 	if (SubElements.Num() != OutElements.Num())
 	{
-		UE_LOG(LogLearning, Error, TEXT("%s: Action '%s' size does not match. Action is '%i' elements but asked for '%i'."),
-			*GetName(), *ActionObject.GetName(Element.ObjectElement).ToString(),
+		UE_LOG(LogLearning, Error, TEXT("GetStaticArrayActionToArrayView: Action '%s' size does not match. Action is '%i' elements but asked for '%i'."),
+			*Object->ActionObject.GetTag(Element.ObjectElement).ToString(),
 			SubElements.Num(), OutElements.Num());
 		UE::Learning::Array::Set<1, FLearningAgentsActionObjectElement>(OutElements, FLearningAgentsActionObjectElement());
 		return false;
@@ -2189,9 +2361,9 @@ bool ULearningAgentsActionObject::GetStaticArrayActionToArrayView(TArrayView<FLe
 
 	for (int32 ElementIdx = 0; ElementIdx < SubElements.Num(); ElementIdx++)
 	{
-		if (!ActionObject.IsValid(SubElements[ElementIdx]))
+		if (!Object->ActionObject.IsValid(SubElements[ElementIdx]))
 		{
-			UE_LOG(LogLearning, Error, TEXT("%s: Invalid Action Object."), *GetName());
+			UE_LOG(LogLearning, Error, TEXT("GetStaticArrayActionToArrayView: Invalid Action Object."));
 			UE::Learning::Array::Set<1, FLearningAgentsActionObjectElement>(OutElements, FLearningAgentsActionObjectElement());
 			return false;
 		}
@@ -2202,11 +2374,11 @@ bool ULearningAgentsActionObject::GetStaticArrayActionToArrayView(TArrayView<FLe
 	return true;
 }
 
-bool ULearningAgentsActionObject::GetPairAction(FLearningAgentsActionObjectElement& OutKey, FLearningAgentsActionObjectElement& OutValue, const FLearningAgentsActionObjectElement Element, const FName Name) const
+bool ULearningAgentsActions::GetPairAction(FLearningAgentsActionObjectElement& OutKey, FLearningAgentsActionObjectElement& OutValue, const ULearningAgentsActionObject* Object, const FLearningAgentsActionObjectElement Element, const FName Tag)
 {
 	TStaticArray<FName, 2> OutElementNames;
 	TStaticArray<FLearningAgentsActionObjectElement, 2> OutElements;
-	if (!GetStructActionToArrayViews(OutElementNames, OutElements, Element, Name))
+	if (!GetStructActionToArrayViews(OutElementNames, OutElements, Object, Element, Tag))
 	{
 		OutKey = FLearningAgentsActionObjectElement();
 		OutValue = FLearningAgentsActionObjectElement();
@@ -2218,17 +2390,17 @@ bool ULearningAgentsActionObject::GetPairAction(FLearningAgentsActionObjectEleme
 	return true;
 }
 
-bool ULearningAgentsActionObject::GetEnumAction(uint8& OutEnumValue, const UEnum* Enum, const FLearningAgentsActionObjectElement Element, const FName Name) const
+bool ULearningAgentsActions::GetEnumAction(uint8& OutEnumValue, const ULearningAgentsActionObject* Object, const FLearningAgentsActionObjectElement Element, const UEnum* Enum, const FName Tag)
 {
 	if (!Enum)
 	{
-		UE_LOG(LogLearning, Error, TEXT("%s: Enum is nullptr."), *GetName());
+		UE_LOG(LogLearning, Error, TEXT("GetEnumAction: Enum is nullptr."));
 		OutEnumValue = 0;
 		return false;
 	}
 
 	int32 OutIndex = 0;
-	if (!GetExclusiveDiscreteAction(OutIndex, Element, Name))
+	if (!GetExclusiveDiscreteAction(OutIndex, Object, Element, Tag))
 	{
 		OutEnumValue = 0;
 		return false;
@@ -2236,7 +2408,7 @@ bool ULearningAgentsActionObject::GetEnumAction(uint8& OutEnumValue, const UEnum
 
 	if (OutIndex >= Enum->NumEnums() - 1)
 	{
-		UE_LOG(LogLearning, Error, TEXT("%s: EnumValue out of range for Enum '%s'. Expected %i or less, got %i."), *GetName(), *Enum->GetName(), Enum->NumEnums() - 1, OutIndex);
+		UE_LOG(LogLearning, Error, TEXT("GetEnumAction: EnumValue out of range for Enum '%s'. Expected %i or less, got %i."), *Enum->GetName(), Enum->NumEnums() - 1, OutIndex);
 		OutEnumValue = 0;
 		return false;
 	}
@@ -2245,7 +2417,7 @@ bool ULearningAgentsActionObject::GetEnumAction(uint8& OutEnumValue, const UEnum
 
 	if (EnumValue == INDEX_NONE)
 	{
-		UE_LOG(LogLearning, Error, TEXT("%s: Enum Value not found for index %i."), *GetName(), OutIndex);
+		UE_LOG(LogLearning, Error, TEXT("GetEnumAction: Enum Value not found for index %i."), OutIndex);
 		OutEnumValue = 0;
 		return false;
 	}
@@ -2254,24 +2426,24 @@ bool ULearningAgentsActionObject::GetEnumAction(uint8& OutEnumValue, const UEnum
 	return true;
 }
 
-bool ULearningAgentsActionObject::GetBitmaskAction(int32& OutBitmaskValue, const UEnum* Enum, const FLearningAgentsActionObjectElement Element, const FName Name) const
+bool ULearningAgentsActions::GetBitmaskAction(int32& OutBitmaskValue, const ULearningAgentsActionObject* Object, const FLearningAgentsActionObjectElement Element, const UEnum* Enum, const FName Tag)
 {
 	if (!Enum)
 	{
-		UE_LOG(LogLearning, Error, TEXT("%s: Enum is nullptr."), *GetName());
+		UE_LOG(LogLearning, Error, TEXT("GetBitmaskAction: Enum is nullptr."));
 		OutBitmaskValue = 0;
 		return false;
 	}
 
 	if (Enum->NumEnums() - 1 > 32)
 	{
-		UE_LOG(LogLearning, Error, TEXT("%s: Too many values in Enum to use as Bitmask (%i)."), *GetName(), Enum->NumEnums() - 1);
+		UE_LOG(LogLearning, Error, TEXT("GetBitmaskAction: Too many values in Enum to use as Bitmask (%i)."), Enum->NumEnums() - 1);
 		OutBitmaskValue = 0;
 		return false;
 	}
 
 	int32 EnumValueNum;
-	if (!GetInclusiveDiscreteActionNum(EnumValueNum, Element, Name))
+	if (!GetInclusiveDiscreteActionNum(EnumValueNum, Object, Element, Tag))
 	{
 		OutBitmaskValue = 0;
 		return false;
@@ -2279,14 +2451,14 @@ bool ULearningAgentsActionObject::GetBitmaskAction(int32& OutBitmaskValue, const
 
 	if (EnumValueNum > Enum->NumEnums() - 1)
 	{
-		UE_LOG(LogLearning, Error, TEXT("%s: Too many values for Enum '%s'. Expected %i or less, got %i."), *GetName(), *Enum->GetName(), Enum->NumEnums() - 1, EnumValueNum);
+		UE_LOG(LogLearning, Error, TEXT("GetBitmaskAction: Too many values for Enum '%s'. Expected %i or less, got %i."), *Enum->GetName(), Enum->NumEnums() - 1, EnumValueNum);
 		OutBitmaskValue = 0;
 		return false;
 	}
 
 	TArray<int32, TInlineAllocator<32>> OutIndices;
 	OutIndices.SetNumUninitialized(EnumValueNum);
-	if (!GetInclusiveDiscreteActionToArrayView(OutIndices, Element, Name))
+	if (!GetInclusiveDiscreteActionToArrayView(OutIndices, Object, Element, Tag))
 	{
 		OutBitmaskValue = 0;
 		return false;
@@ -2300,10 +2472,10 @@ bool ULearningAgentsActionObject::GetBitmaskAction(int32& OutBitmaskValue, const
 	return true;
 }
 
-bool ULearningAgentsActionObject::GetOptionalAction(ELearningAgentsOptionalAction& OutOption, FLearningAgentsActionObjectElement& OutElement, const FLearningAgentsActionObjectElement Element, const FName Name) const
+bool ULearningAgentsActions::GetOptionalAction(ELearningAgentsOptionalAction& OutOption, FLearningAgentsActionObjectElement& OutElement, const ULearningAgentsActionObject* Object, const FLearningAgentsActionObjectElement Element, const FName Tag)
 {
 	FName OutName = NAME_None;
-	if (!GetExclusiveUnionAction(OutName, OutElement, Element, Name))
+	if (!GetExclusiveUnionAction(OutName, OutElement, Object, Element, Tag))
 	{
 		OutOption = ELearningAgentsOptionalAction::Null;
 		return false;
@@ -2313,10 +2485,10 @@ bool ULearningAgentsActionObject::GetOptionalAction(ELearningAgentsOptionalActio
 	return true;
 }
 
-bool ULearningAgentsActionObject::GetEitherAction(ELearningAgentsEitherAction& OutEither, FLearningAgentsActionObjectElement& OutElement, const FLearningAgentsActionObjectElement Element, const FName Name) const
+bool ULearningAgentsActions::GetEitherAction(ELearningAgentsEitherAction& OutEither, FLearningAgentsActionObjectElement& OutElement, const ULearningAgentsActionObject* Object, const FLearningAgentsActionObjectElement Element, const FName Tag)
 {
 	FName OutName = NAME_None;
-	if (!GetExclusiveUnionAction(OutName, OutElement, Element, Name))
+	if (!GetExclusiveUnionAction(OutName, OutElement, Object, Element, Tag))
 	{
 		OutEither = ELearningAgentsEitherAction::A;
 		return false;
@@ -2326,40 +2498,46 @@ bool ULearningAgentsActionObject::GetEitherAction(ELearningAgentsEitherAction& O
 	return true;
 }
 
-bool ULearningAgentsActionObject::GetEncodingAction(FLearningAgentsActionObjectElement& OutElement, const FLearningAgentsActionObjectElement Element, const FName Name) const
+bool ULearningAgentsActions::GetEncodingAction(FLearningAgentsActionObjectElement& OutElement, const ULearningAgentsActionObject* Object, const FLearningAgentsActionObjectElement Element, const FName Tag)
 {
-	if (!ActionObject.IsValid(Element.ObjectElement))
+	if (!Object)
+	{
+		UE_LOG(LogLearning, Error, TEXT("GetEncodingAction: Object is nullptr."));
+		OutElement = FLearningAgentsActionObjectElement();
+		return false;
+	}
+
+	if (!Object->ActionObject.IsValid(Element.ObjectElement))
 	{
-		UE_LOG(LogLearning, Error, TEXT("%s: Invalid Action Object."), *GetName());
+		UE_LOG(LogLearning, Error, TEXT("GetEncodingAction: Invalid Action Object."));
 		OutElement = FLearningAgentsActionObjectElement();
 		return false;
 	}
 
-	if (ActionObject.GetName(Element.ObjectElement) != Name)
+	if (Object->ActionObject.GetTag(Element.ObjectElement) != Tag)
 	{
-		UE_LOG(LogLearning, Warning, TEXT("%s: Action name does not match. Action is '%s' but asked for '%s'."), *GetName(), *ActionObject.GetName(Element.ObjectElement).ToString(), *Name.ToString());
+		UE_LOG(LogLearning, Warning, TEXT("GetEncodingAction: Action tag does not match. Action is '%s' but asked for '%s'."), *Object->ActionObject.GetTag(Element.ObjectElement).ToString(), *Tag.ToString());
 	}
 
-	if (ActionObject.GetType(Element.ObjectElement) != UE::Learning::Action::EType::Encoding)
+	if (Object->ActionObject.GetType(Element.ObjectElement) != UE::Learning::Action::EType::Encoding)
 	{
-		UE_LOG(LogLearning, Error, TEXT("%s: Action '%s' type does not match. Action is '%s' but asked for '%s'."),
-			*GetName(),
-			*ActionObject.GetName(Element.ObjectElement).ToString(),
-			UE::Learning::Agents::Action::Private::GetActionTypeString(ActionObject.GetType(Element.ObjectElement)),
+		UE_LOG(LogLearning, Error, TEXT("GetEncodingAction: Action '%s' type does not match. Action is '%s' but asked for '%s'."),
+			*Object->ActionObject.GetTag(Element.ObjectElement).ToString(),
+			UE::Learning::Agents::Action::Private::GetActionTypeString(Object->ActionObject.GetType(Element.ObjectElement)),
 			UE::Learning::Agents::Action::Private::GetActionTypeString(UE::Learning::Action::EType::Encoding));
 		OutElement = FLearningAgentsActionObjectElement();
 		return false;
 	}
 
-	OutElement = { ActionObject.GetEncoding(Element.ObjectElement).Element };
+	OutElement = { Object->ActionObject.GetEncoding(Element.ObjectElement).Element };
 
 	return true;
 }
 
-bool ULearningAgentsActionObject::GetBoolAction(bool& bOutValue, const FLearningAgentsActionObjectElement Element, const FName Name) const
+bool ULearningAgentsActions::GetBoolAction(bool& bOutValue, const ULearningAgentsActionObject* Object, const FLearningAgentsActionObjectElement Element, const FName Tag)
 {
 	int32 OutIndex = 0;
-	if (!GetExclusiveDiscreteAction(OutIndex, Element, Name))
+	if (!GetExclusiveDiscreteAction(OutIndex, Object, Element, Tag))
 	{
 		bOutValue = false;
 		return false;
@@ -2369,10 +2547,10 @@ bool ULearningAgentsActionObject::GetBoolAction(bool& bOutValue, const FLearning
 	return true;
 }
 
-bool ULearningAgentsActionObject::GetFloatAction(float& OutValue, const FLearningAgentsActionObjectElement Element, const float FloatScale, const FName Name) const
+bool ULearningAgentsActions::GetFloatAction(float& OutValue, const ULearningAgentsActionObject* Object, const FLearningAgentsActionObjectElement Element, const float FloatScale, const FName Tag)
 {
 	float OutValuesData;
-	if (!GetContinuousActionToArrayView(MakeArrayView(&OutValuesData, 1), Element, Name))
+	if (!GetContinuousActionToArrayView(MakeArrayView(&OutValuesData, 1), Object, Element, Tag))
 	{
 		OutValue = 0.0f;
 		return false;
@@ -2382,10 +2560,10 @@ bool ULearningAgentsActionObject::GetFloatAction(float& OutValue, const FLearnin
 	return true;
 }
 
-bool ULearningAgentsActionObject::GetLocationAction(FVector& OutLocation, const FLearningAgentsActionObjectElement Element, const FTransform RelativeTransform, const float LocationScale, const FName Name) const
+bool ULearningAgentsActions::GetLocationAction(FVector& OutLocation, const ULearningAgentsActionObject* Object, const FLearningAgentsActionObjectElement Element, const FTransform RelativeTransform, const float LocationScale, const FName Tag)
 {
 	TStaticArray<float, 3> OutValues;
-	if (!GetContinuousActionToArrayView(OutValues, Element, Name))
+	if (!GetContinuousActionToArrayView(OutValues, Object, Element, Tag))
 	{
 		OutLocation = FVector::ZeroVector;
 		return false;
@@ -2396,10 +2574,10 @@ bool ULearningAgentsActionObject::GetLocationAction(FVector& OutLocation, const
 	return true;
 }
 
-bool ULearningAgentsActionObject::GetRotationAction(FRotator& OutRotation, const FLearningAgentsActionObjectElement Element, const FRotator RelativeRotation, const float RotationScale, const FName Name) const
+bool ULearningAgentsActions::GetRotationAction(FRotator& OutRotation, const ULearningAgentsActionObject* Object, const FLearningAgentsActionObjectElement Element, const FRotator RelativeRotation, const float RotationScale, const FName Tag)
 {
 	FQuat OutRotationQuat;
-	if (!GetRotationActionAsQuat(OutRotationQuat, Element, FQuat::MakeFromRotator(RelativeRotation), RotationScale, Name))
+	if (!GetRotationActionAsQuat(OutRotationQuat, Object, Element, FQuat::MakeFromRotator(RelativeRotation), RotationScale, Tag))
 	{
 		OutRotation = FRotator::ZeroRotator;
 		return false;
@@ -2409,10 +2587,10 @@ bool ULearningAgentsActionObject::GetRotationAction(FRotator& OutRotation, const
 	return true;
 }
 
-bool ULearningAgentsActionObject::GetRotationActionAsQuat(FQuat& OutRotation, const FLearningAgentsActionObjectElement Element, const FQuat RelativeRotation, const float RotationScale, const FName Name) const
+bool ULearningAgentsActions::GetRotationActionAsQuat(FQuat& OutRotation, const ULearningAgentsActionObject* Object, const FLearningAgentsActionObjectElement Element, const FQuat RelativeRotation, const float RotationScale, const FName Tag)
 {
 	TStaticArray<float, 3> OutValues;
-	if (!GetContinuousActionToArrayView(OutValues, Element, Name))
+	if (!GetContinuousActionToArrayView(OutValues, Object, Element, Tag))
 	{
 		OutRotation = FQuat::Identity;
 		return false;
@@ -2423,10 +2601,10 @@ bool ULearningAgentsActionObject::GetRotationActionAsQuat(FQuat& OutRotation, co
 	return true;
 }
 
-bool ULearningAgentsActionObject::GetScaleAction(FVector& OutScale, const FLearningAgentsActionObjectElement Element, const FVector RelativeScale, const float Scale, const FName Name) const
+bool ULearningAgentsActions::GetScaleAction(FVector& OutScale, const ULearningAgentsActionObject* Object, const FLearningAgentsActionObjectElement Element, const FVector RelativeScale, const float Scale, const FName Tag)
 {
 	TStaticArray<float, 3> OutValues;
-	if (!GetContinuousActionToArrayView(OutValues, Element, Name))
+	if (!GetContinuousActionToArrayView(OutValues, Object, Element, Tag))
 	{
 		OutScale = FVector::OneVector;
 		return false;
@@ -2437,32 +2615,32 @@ bool ULearningAgentsActionObject::GetScaleAction(FVector& OutScale, const FLearn
 	return true;
 }
 
-bool ULearningAgentsActionObject::GetTransformAction(FTransform& OutTransform, const FLearningAgentsActionObjectElement Element, const FTransform RelativeTransform, const float LocationScale, const float RotationScale, const float ScaleScale, const FName Name) const
+bool ULearningAgentsActions::GetTransformAction(FTransform& OutTransform, const ULearningAgentsActionObject* Object, const FLearningAgentsActionObjectElement Element, const FTransform RelativeTransform, const float LocationScale, const float RotationScale, const float ScaleScale, const FName Tag)
 {
 	TStaticArray<FName, 3> OutElementNames;
 	TStaticArray<FLearningAgentsActionObjectElement, 3> OutElements;
-	if (!GetStructActionToArrayViews(OutElementNames, OutElements, Element, Name))
+	if (!GetStructActionToArrayViews(OutElementNames, OutElements, Object, Element, Tag))
 	{
 		OutTransform = FTransform::Identity;
 		return false;
 	}
 
 	FVector OutLocation;
-	if (!GetLocationAction(OutLocation, OutElements[MakeArrayView(OutElementNames).Find(TEXT("Location"))], RelativeTransform, LocationScale))
+	if (!GetLocationAction(OutLocation, Object, OutElements[MakeArrayView(OutElementNames).Find(TEXT("Location"))], RelativeTransform, LocationScale))
 	{
 		OutTransform = FTransform::Identity;
 		return false;
 	}
 
 	FQuat OutRotation;
-	if (!GetRotationActionAsQuat(OutRotation, OutElements[MakeArrayView(OutElementNames).Find(TEXT("Rotation"))], RelativeTransform.GetRotation(), RotationScale))
+	if (!GetRotationActionAsQuat(OutRotation, Object, OutElements[MakeArrayView(OutElementNames).Find(TEXT("Rotation"))], RelativeTransform.GetRotation(), RotationScale))
 	{
 		OutTransform = FTransform::Identity;
 		return false;
 	}
 
 	FVector OutScale;
-	if (!GetScaleAction(OutScale, OutElements[MakeArrayView(OutElementNames).Find(TEXT("Scale"))], RelativeTransform.GetScale3D(), ScaleScale))
+	if (!GetScaleAction(OutScale, Object, OutElements[MakeArrayView(OutElementNames).Find(TEXT("Scale"))], RelativeTransform.GetScale3D(), ScaleScale))
 	{
 		OutTransform = FTransform::Identity;
 		return false;
@@ -2472,9 +2650,9 @@ bool ULearningAgentsActionObject::GetTransformAction(FTransform& OutTransform, c
 	return true;
 }
 
-bool ULearningAgentsActionObject::GetAngleAction(float& OutAngle, const FLearningAgentsActionObjectElement Element, const float RelativeAngle, const float AngleScale, const FName Name) const
+bool ULearningAgentsActions::GetAngleAction(float& OutAngle, const ULearningAgentsActionObject* Object, const FLearningAgentsActionObjectElement Element, const float RelativeAngle, const float AngleScale, const FName Tag)
 {
-	if (GetAngleActionRadians(OutAngle, Element, FMath::DegreesToRadians(RelativeAngle), FMath::DegreesToRadians(AngleScale), Name))
+	if (GetAngleActionRadians(OutAngle, Object, Element, FMath::DegreesToRadians(RelativeAngle), FMath::DegreesToRadians(AngleScale), Tag))
 	{
 		OutAngle = FMath::RadiansToDegrees(OutAngle);
 		return true;
@@ -2486,9 +2664,9 @@ bool ULearningAgentsActionObject::GetAngleAction(float& OutAngle, const FLearnin
 	}
 }
 
-bool ULearningAgentsActionObject::GetAngleActionRadians(float& OutAngle, const FLearningAgentsActionObjectElement Element, const float RelativeAngle, const float AngleScale, const FName Name) const
+bool ULearningAgentsActions::GetAngleActionRadians(float& OutAngle, const ULearningAgentsActionObject* Object, const FLearningAgentsActionObjectElement Element, const float RelativeAngle, const float AngleScale, const FName Tag)
 {
-	if (!GetContinuousActionToArrayView(MakeArrayView(&OutAngle, 1), Element, Name))
+	if (!GetContinuousActionToArrayView(MakeArrayView(&OutAngle, 1), Object, Element, Tag))
 	{
 		OutAngle = 0.0f;
 		return false;
@@ -2498,10 +2676,10 @@ bool ULearningAgentsActionObject::GetAngleActionRadians(float& OutAngle, const F
 	return true;
 }
 
-bool ULearningAgentsActionObject::GetVelocityAction(FVector& OutVelocity, const FLearningAgentsActionObjectElement Element, const FTransform RelativeTransform, const float VelocityScale, const FName Name) const
+bool ULearningAgentsActions::GetVelocityAction(FVector& OutVelocity, const ULearningAgentsActionObject* Object, const FLearningAgentsActionObjectElement Element, const FTransform RelativeTransform, const float VelocityScale, const FName Tag)
 {
 	TStaticArray<float, 3> OutValues;
-	if (!GetContinuousActionToArrayView(OutValues, Element, Name))
+	if (!GetContinuousActionToArrayView(OutValues, Object, Element, Tag))
 	{
 		OutVelocity = FVector::OneVector;
 		return false;
@@ -2511,10 +2689,10 @@ bool ULearningAgentsActionObject::GetVelocityAction(FVector& OutVelocity, const
 	return true;
 }
 
-bool ULearningAgentsActionObject::GetDirectionAction(FVector& OutDirection, const FLearningAgentsActionObjectElement Element, const FTransform RelativeTransform, const FName Name) const
+bool ULearningAgentsActions::GetDirectionAction(FVector& OutDirection, const ULearningAgentsActionObject* Object, const FLearningAgentsActionObjectElement Element, const FTransform RelativeTransform, const FName Tag)
 {
 	TStaticArray<float, 3> OutValues;
-	if (!GetContinuousActionToArrayView(OutValues, Element, Name))
+	if (!GetContinuousActionToArrayView(OutValues, Object, Element, Tag))
 	{
 		OutDirection = FVector::ForwardVector;
 		return false;
@@ -2524,9 +2702,9 @@ bool ULearningAgentsActionObject::GetDirectionAction(FVector& OutDirection, cons
 	return true;
 }
 
-bool ULearningAgentsActionObject::GetSpeedAction(float& OutSpeed, const FLearningAgentsActionObjectElement Element, const float SpeedScale, const FName Name) const
+bool ULearningAgentsActions::GetSpeedAction(float& OutSpeed, const ULearningAgentsActionObject* Object, const FLearningAgentsActionObjectElement Element, const float SpeedScale, const FName Tag)
 {
-	if (!GetContinuousActionToArrayView(MakeArrayView(&OutSpeed, 1), Element, Name))
+	if (!GetContinuousActionToArrayView(MakeArrayView(&OutSpeed, 1), Object, Element, Tag))
 	{
 		OutSpeed = 0.0f;
 		return false;
