LearningAgents: Refactored to make Observation and Action functions static and use "AdvancedDisplay" markup. Added Visual Logging and comments for Completions.

diff --git a/Engine/Plugins/Experimental/LearningAgents/Source/LearningAgents/Private/LearningAgentsController.cpp b/Engine/Plugins/Experimental/LearningAgents/Source/LearningAgents/Private/LearningAgentsController.cpp
--- a/Engine/Plugins/Experimental/LearningAgents/Source/LearningAgents/Private/LearningAgentsController.cpp
+++ b/Engine/Plugins/Experimental/LearningAgents/Source/LearningAgents/Private/LearningAgentsController.cpp
@@ -130,7 +130,7 @@ void ULearningAgentsController::EvaluateController()
 
 	// Run EvaluateAgentControllers Callback
 
-	Interactor->ActionObject->GetActionObject().Reset();
+	Interactor->ActionObject->ActionObject.Reset();
 	Interactor->ActionObjectElements.Empty(Manager->GetMaxAgentNum());
 	EvaluateAgentControllers(
 		Interactor->ActionObjectElements, 
@@ -149,16 +149,17 @@ void ULearningAgentsController::EvaluateController()
 
 	for (int32 AgentIdx = 0; AgentIdx < Manager->GetAgentNum(); AgentIdx++)
 	{
-		if (Interactor->ActionSchema->ValidateObjectMatchesSchema(
+		if (ULearningAgentsActions::ValidateObjectMatchesSchema(
+			Interactor->ActionSchema,
 			Interactor->ActionSchemaElement, 
 			Interactor->ActionObject, 
 			Interactor->ActionObjectElements[AgentIdx]))
 		{
 			UE::Learning::Action::SetVectorFromObject(
 				Interactor->ActionVectors[Manager->GetAllAgentSet()[AgentIdx]],
-				Interactor->ActionSchema->GetActionSchema(),
+				Interactor->ActionSchema->ActionSchema,
 				Interactor->ActionSchemaElement.SchemaElement,
-				Interactor->ActionObject->GetActionObject(),
+				Interactor->ActionObject->ActionObject,
 				Interactor->ActionObjectElements[AgentIdx].ObjectElement);
 
 			Interactor->ActionVectorIteration[Manager->GetAllAgentSet()[AgentIdx]]++;
diff --git a/Engine/Plugins/Experimental/LearningAgents/Source/LearningAgents/Private/LearningAgentsInteractor.cpp b/Engine/Plugins/Experimental/LearningAgents/Source/LearningAgents/Private/LearningAgentsInteractor.cpp
--- a/Engine/Plugins/Experimental/LearningAgents/Source/LearningAgents/Private/LearningAgentsInteractor.cpp
+++ b/Engine/Plugins/Experimental/LearningAgents/Source/LearningAgents/Private/LearningAgentsInteractor.cpp
@@ -67,13 +67,13 @@ void ULearningAgentsInteractor::SetupInteractor(ULearningAgentsManager* InManage
 	ObservationSchema = NewObject<ULearningAgentsObservationSchema>(this, ObservationSchemaUniqueName);
 	SpecifyAgentObservation(ObservationSchemaElement, ObservationSchema);
 
-	if (!ObservationSchema->GetObservationSchema().IsValid(ObservationSchemaElement.SchemaElement))
+	if (!ObservationSchema->ObservationSchema.IsValid(ObservationSchemaElement.SchemaElement))
 	{
 		UE_LOG(LogLearning, Error, TEXT("%s: Invalid observation provided to Interactor during SpecifyObservations."), *GetName());
 		return;
 	}
 
-	const int32 ObservationVectorSize = ObservationSchema->GetObservationSchema().GetObservationVectorSize(ObservationSchemaElement.SchemaElement);
+	const int32 ObservationVectorSize = ObservationSchema->ObservationSchema.GetObservationVectorSize(ObservationSchemaElement.SchemaElement);
 
 	if (ObservationVectorSize == 0)
 	{
@@ -81,7 +81,7 @@ void ULearningAgentsInteractor::SetupInteractor(ULearningAgentsManager* InManage
 		return;
 	}
 
-	const int32 ObservationEncodedVectorSize = ObservationSchema->GetObservationSchema().GetEncodedVectorSize(ObservationSchemaElement.SchemaElement);
+	const int32 ObservationEncodedVectorSize = ObservationSchema->ObservationSchema.GetEncodedVectorSize(ObservationSchemaElement.SchemaElement);
 
 	if (ObservationEncodedVectorSize == 0)
 	{
@@ -90,7 +90,7 @@ void ULearningAgentsInteractor::SetupInteractor(ULearningAgentsManager* InManage
 	}
 
 	ObservationVectors.SetNumUninitialized({ Manager->GetMaxAgentNum(), ObservationVectorSize  });
-	ObservationCompatibilityHash = UE::Learning::Observation::GetSchemaObjectsCompatibilityHash(ObservationSchema->GetObservationSchema(), ObservationSchemaElement.SchemaElement);
+	ObservationCompatibilityHash = UE::Learning::Observation::GetSchemaObjectsCompatibilityHash(ObservationSchema->ObservationSchema, ObservationSchemaElement.SchemaElement);
 
 	const FName ObservationObjectUniqueName = MakeUniqueObjectName(this, ULearningAgentsObservationObject::StaticClass(), TEXT("ObservationObject"), EUniqueObjectNameOptions::GloballyUnique);
 
@@ -104,13 +104,13 @@ void ULearningAgentsInteractor::SetupInteractor(ULearningAgentsManager* InManage
 	ActionSchema = NewObject<ULearningAgentsActionSchema>(this, ActionSchemaUniqueName);
 	SpecifyAgentAction(ActionSchemaElement, ActionSchema);
 
-	if (!ActionSchema->GetActionSchema().IsValid(ActionSchemaElement.SchemaElement))
+	if (!ActionSchema->ActionSchema.IsValid(ActionSchemaElement.SchemaElement))
 	{
 		UE_LOG(LogLearning, Error, TEXT("%s: Invalid action provided to Interactor during SpecifyActions."), *GetName());
 		return;
 	}
 
-	const int32 ActionVectorSize = ActionSchema->GetActionSchema().GetActionVectorSize(ActionSchemaElement.SchemaElement);
+	const int32 ActionVectorSize = ActionSchema->ActionSchema.GetActionVectorSize(ActionSchemaElement.SchemaElement);
 
 	if (ActionVectorSize == 0)
 	{
@@ -118,7 +118,7 @@ void ULearningAgentsInteractor::SetupInteractor(ULearningAgentsManager* InManage
 		return;
 	}
 
-	const int32 ActionEncodedVectorSize = ActionSchema->GetActionSchema().GetEncodedVectorSize(ActionSchemaElement.SchemaElement);
+	const int32 ActionEncodedVectorSize = ActionSchema->ActionSchema.GetEncodedVectorSize(ActionSchemaElement.SchemaElement);
 
 	if (ActionEncodedVectorSize == 0)
 	{
@@ -126,10 +126,10 @@ void ULearningAgentsInteractor::SetupInteractor(ULearningAgentsManager* InManage
 		return;
 	}
 
-	const int32 ActionDistributionVectorSize = ActionSchema->GetActionSchema().GetActionDistributionVectorSize(ActionSchemaElement.SchemaElement);
+	const int32 ActionDistributionVectorSize = ActionSchema->ActionSchema.GetActionDistributionVectorSize(ActionSchemaElement.SchemaElement);
 
 	ActionVectors.SetNumUninitialized({ Manager->GetMaxAgentNum(), ActionVectorSize });
-	ActionCompatibilityHash = UE::Learning::Action::GetSchemaObjectsCompatibilityHash(ActionSchema->GetActionSchema(), ActionSchemaElement.SchemaElement);
+	ActionCompatibilityHash = UE::Learning::Action::GetSchemaObjectsCompatibilityHash(ActionSchema->ActionSchema, ActionSchemaElement.SchemaElement);
 
 	const FName ActionObjectUniqueName = MakeUniqueObjectName(this, ULearningAgentsActionObject::StaticClass(), TEXT("ActionObject"), EUniqueObjectNameOptions::GloballyUnique);
 
@@ -253,7 +253,7 @@ void ULearningAgentsInteractor::GatherObservations(const UE::Learning::FIndexSet
 		ValidAgentIds.Add(AgentId);
 	}
 
-	ObservationObject->GetObservationObject().Reset();
+	ObservationObject->ObservationObject.Reset();
 	ObservationObjectElements.Empty(Manager->GetMaxAgentNum());
 	GatherAgentObservations(ObservationObjectElements, ObservationObject, ValidAgentIds);
 
@@ -267,13 +267,13 @@ void ULearningAgentsInteractor::GatherObservations(const UE::Learning::FIndexSet
 
 	for (int32 AgentIdx = 0; AgentIdx < AgentSet.Num(); AgentIdx++)
 	{
-		if (ObservationSchema->ValidateObjectMatchesSchema(ObservationSchemaElement, ObservationObject, ObservationObjectElements[AgentIdx]))
+		if (ULearningAgentsObservations::ValidateObjectMatchesSchema(ObservationSchema, ObservationSchemaElement, ObservationObject, ObservationObjectElements[AgentIdx]))
 		{
 			UE::Learning::Observation::SetVectorFromObject(
 				ObservationVectors[AgentSet[AgentIdx]],
-				ObservationSchema->GetObservationSchema(),
+				ObservationSchema->ObservationSchema,
 				ObservationSchemaElement.SchemaElement,
-				ObservationObject->GetObservationObject(),
+				ObservationObject->ObservationObject,
 				ObservationObjectElements[AgentIdx].ObjectElement);
 
 			if (bIncrementIteration)
@@ -314,7 +314,7 @@ void ULearningAgentsInteractor::PerformActions(const UE::Learning::FIndexSet Age
 
 	// Generate Action Objects
 
-	ActionObject->GetActionObject().Reset();
+	ActionObject->ActionObject.Reset();
 	ActionObjectElements.Empty(Manager->GetMaxAgentNum());
 
 	for (int32 AgentIdx = 0; AgentIdx < ValidAgentSet.Num(); AgentIdx++)
@@ -322,9 +322,9 @@ void ULearningAgentsInteractor::PerformActions(const UE::Learning::FIndexSet Age
 		FLearningAgentsActionObjectElement ActionObjectElement;
 
 		UE::Learning::Action::GetObjectFromVector(
-			ActionObject->GetActionObject(),
+			ActionObject->ActionObject,
 			ActionObjectElement.ObjectElement,
-			ActionSchema->GetActionSchema(),
+			ActionSchema->ActionSchema,
 			ActionSchemaElement.SchemaElement,
 			ActionVectors[ValidAgentSet[AgentIdx]]);
 
@@ -521,7 +521,7 @@ int32 ULearningAgentsInteractor::GetObservationVectorSize() const
 		return 0;
 	}
 
-	return ObservationSchema->GetObservationSchema().GetObservationVectorSize(ObservationSchemaElement.SchemaElement);
+	return ObservationSchema->ObservationSchema.GetObservationVectorSize(ObservationSchemaElement.SchemaElement);
 }
 
 int32 ULearningAgentsInteractor::GetObservationEncodedVectorSize() const
@@ -532,7 +532,7 @@ int32 ULearningAgentsInteractor::GetObservationEncodedVectorSize() const
 		return 0;
 	}
 
-	return ObservationSchema->GetObservationSchema().GetEncodedVectorSize(ObservationSchemaElement.SchemaElement);
+	return ObservationSchema->ObservationSchema.GetEncodedVectorSize(ObservationSchemaElement.SchemaElement);
 }
 
 int32 ULearningAgentsInteractor::GetActionVectorSize() const
@@ -543,7 +543,7 @@ int32 ULearningAgentsInteractor::GetActionVectorSize() const
 		return 0;
 	}
 
-	return ActionSchema->GetActionSchema().GetActionVectorSize(ActionSchemaElement.SchemaElement);
+	return ActionSchema->ActionSchema.GetActionVectorSize(ActionSchemaElement.SchemaElement);
 }
 
 int32 ULearningAgentsInteractor::GetActionDistributionVectorSize() const
@@ -554,7 +554,7 @@ int32 ULearningAgentsInteractor::GetActionDistributionVectorSize() const
 		return 0;
 	}
 
-	return ActionSchema->GetActionSchema().GetActionDistributionVectorSize(ActionSchemaElement.SchemaElement);
+	return ActionSchema->ActionSchema.GetActionDistributionVectorSize(ActionSchemaElement.SchemaElement);
 }
 
 int32 ULearningAgentsInteractor::GetActionEncodedVectorSize() const
@@ -565,12 +565,12 @@ int32 ULearningAgentsInteractor::GetActionEncodedVectorSize() const
 		return 0;
 	}
 
-	return ActionSchema->GetActionSchema().GetEncodedVectorSize(ActionSchemaElement.SchemaElement);
+	return ActionSchema->ActionSchema.GetEncodedVectorSize(ActionSchemaElement.SchemaElement);
 }
 
 const UE::Learning::Observation::FSchema& ULearningAgentsInteractor::GetObservationSchema() const
 {
-	return ObservationSchema->GetObservationSchema();
+	return ObservationSchema->ObservationSchema;
 }
 
 UE::Learning::Observation::FSchemaElement ULearningAgentsInteractor::GetObservationSchemaElement() const
@@ -580,7 +580,7 @@ UE::Learning::Observation::FSchemaElement ULearningAgentsInteractor::GetObservat
 
 const UE::Learning::Action::FSchema& ULearningAgentsInteractor::GetActionSchema() const
 {
-	return ActionSchema->GetActionSchema();
+	return ActionSchema->ActionSchema;
 }
 
 UE::Learning::Action::FSchemaElement ULearningAgentsInteractor::GetActionSchemaElement() const
