LearningAgents: Refactored to make Observation and Action functions static and use "AdvancedDisplay" markup. Added Visual Logging and comments for Completions.

diff --git a/Engine/Plugins/Experimental/LearningAgents/Source/LearningAgents/Private/LearningAgentsObservations.cpp b/Engine/Plugins/Experimental/LearningAgents/Source/LearningAgents/Private/LearningAgentsObservations.cpp
--- a/Engine/Plugins/Experimental/LearningAgents/Source/LearningAgents/Private/LearningAgentsObservations.cpp
+++ b/Engine/Plugins/Experimental/LearningAgents/Source/LearningAgents/Private/LearningAgentsObservations.cpp
@@ -24,11 +24,6 @@ uint32 GetTypeHash(const FLearningAgentsObservationObjectElement& Element)
 	return (uint32)Element.ObjectElement.Index;
 }
 
-const UE::Learning::Observation::FSchema& ULearningAgentsObservationSchema::GetObservationSchema() const
-{
-	return ObservationSchema;
-}
-
 namespace UE::Learning::Agents::Observation::Private
 {
 	static inline bool ContainsDuplicates(const TArrayView<const int32> Indices)
@@ -67,32 +62,31 @@ namespace UE::Learning::Agents::Observation::Private
 		const Learning::Observation::FSchema& Schema,
 		const Learning::Observation::FSchemaElement SchemaElement,
 		const Learning::Observation::FObject& Object,
-		const Learning::Observation::FObjectElement ObjectElement,
-		const FString& ObjectName)
+		const Learning::Observation::FObjectElement ObjectElement)
 	{
 		// Check Elements are Valid
 
 		if (!Schema.IsValid(SchemaElement))
 		{
-			UE_LOG(LogLearning, Error, TEXT("%s: Invalid Observation Schema Element."), *ObjectName);
+			UE_LOG(LogLearning, Error, TEXT("ValidateObjectMatchesSchema: Invalid Observation Schema Element."));
 			return false;
 		}
 
 		if (!Object.IsValid(ObjectElement))
 		{
-			UE_LOG(LogLearning, Error, TEXT("%s: Invalid Observation Object Element."), *ObjectName);
+			UE_LOG(LogLearning, Error, TEXT("ValidateObjectMatchesSchema: Invalid Observation Object Element."));
 			return false;
 		}
 
 		// Check Names Match
 
-		const FName ObservationSchemaElementName = Schema.GetName(SchemaElement);
-		const FName ObservationObjectElementName = Object.GetName(ObjectElement);
+		const FName ObservationSchemaElementTag = Schema.GetTag(SchemaElement);
+		const FName ObservationObjectElementTag = Object.GetTag(ObjectElement);
 
-		if (ObservationSchemaElementName != ObservationObjectElementName)
+		if (ObservationSchemaElementTag != ObservationObjectElementTag)
 		{
-			UE_LOG(LogLearning, Warning, TEXT("%s: Observation name does not match Schema. Expected '%s', got '%s'."),
-				*ObjectName, *ObservationSchemaElementName.ToString(), *ObservationObjectElementName.ToString());
+			UE_LOG(LogLearning, Warning, TEXT("ValidateObjectMatchesSchema: Observation tag does not match Schema. Expected '%s', got '%s'."),
+				*ObservationSchemaElementTag.ToString(), *ObservationObjectElementTag.ToString());
 		}
 
 		// Check Types Match
@@ -102,9 +96,8 @@ namespace UE::Learning::Agents::Observation::Private
 
 		if (ObservationSchemaElementType != ObservationObjectElementType)
 		{
-			UE_LOG(LogLearning, Error, TEXT("%s: Observation '%s' type does not match Schema. Expected type '%s', got type '%s'."),
-				*ObjectName,
-				*ObservationSchemaElementName.ToString(),
+			UE_LOG(LogLearning, Error, TEXT("ValidateObjectMatchesSchema: Observation '%s' type does not match Schema. Expected type '%s', got type '%s'."),
+				*ObservationSchemaElementTag.ToString(),
 				GetObservationTypeString(ObservationSchemaElementType),
 				GetObservationTypeString(ObservationObjectElementType));
 			return false;
@@ -123,9 +116,8 @@ namespace UE::Learning::Agents::Observation::Private
 
 			if (SchemaElementSize != ObjectElementSize)
 			{
-				UE_LOG(LogLearning, Error, TEXT("%s: Observation '%s' size does not match Schema. Expected '%i', got '%i'."),
-					*ObjectName,
-					*ObservationSchemaElementName.ToString(),
+				UE_LOG(LogLearning, Error, TEXT("ValidateObjectMatchesSchema: Observation '%s' size does not match Schema. Expected '%i', got '%i'."),
+					*ObservationSchemaElementTag.ToString(),
 					SchemaElementSize,
 					ObjectElementSize);
 				return false;
@@ -143,9 +135,8 @@ namespace UE::Learning::Agents::Observation::Private
 
 			if (SchemaParameters.Elements.Num() != ObjectParameters.Elements.Num())
 			{
-				UE_LOG(LogLearning, Error, TEXT("%s: Observation '%s' number of sub-elements does not match Schema. Expected '%i', got '%i'."),
-					*ObjectName,
-					*ObservationSchemaElementName.ToString(),
+				UE_LOG(LogLearning, Error, TEXT("ValidateObjectMatchesSchema: Observation '%s' number of sub-elements does not match Schema. Expected '%i', got '%i'."),
+					*ObservationSchemaElementTag.ToString(),
 					SchemaParameters.Elements.Num(),
 					ObjectParameters.Elements.Num());
 				return false;
@@ -157,9 +148,8 @@ namespace UE::Learning::Agents::Observation::Private
 
 				if (ObjectElementIdx == INDEX_NONE)
 				{
-					UE_LOG(LogLearning, Error, TEXT("%s: Observation '%s' does not include '%s' observation required by Schema."),
-						*ObjectName,
-						*ObservationSchemaElementName.ToString(),
+					UE_LOG(LogLearning, Error, TEXT("ValidateObjectMatchesSchema: Observation '%s' does not include '%s' observation required by Schema."),
+						*ObservationSchemaElementTag.ToString(),
 						*SchemaParameters.ElementNames[SchemaElementIdx].ToString());
 					return false;
 				}
@@ -168,8 +158,7 @@ namespace UE::Learning::Agents::Observation::Private
 					Schema,
 					SchemaParameters.Elements[SchemaElementIdx],
 					Object,
-					ObjectParameters.Elements[ObjectElementIdx],
-					ObjectName))
+					ObjectParameters.Elements[ObjectElementIdx]))
 				{
 					return false;
 				}
@@ -188,9 +177,8 @@ namespace UE::Learning::Agents::Observation::Private
 
 			if (SchemaSubElementIdx == INDEX_NONE)
 			{
-				UE_LOG(LogLearning, Error, TEXT("%s: Observation '%s' Schema does not include '%s' observation."),
-					*ObjectName,
-					*ObservationSchemaElementName.ToString(),
+				UE_LOG(LogLearning, Error, TEXT("ValidateObjectMatchesSchema: Observation '%s' Schema does not include '%s' observation."),
+					*ObservationSchemaElementTag.ToString(),
 					*ObjectParameters.ElementName.ToString());
 				return false;
 			}
@@ -199,8 +187,7 @@ namespace UE::Learning::Agents::Observation::Private
 				Schema,
 				SchemaParameters.Elements[SchemaSubElementIdx],
 				Object,
-				ObjectParameters.Element,
-				ObjectName);
+				ObjectParameters.Element);
 		}
 
 		case Learning::Observation::EType::OrInclusive:
@@ -210,9 +197,8 @@ namespace UE::Learning::Agents::Observation::Private
 
 			if (ObjectParameters.Elements.Num() > SchemaParameters.Elements.Num())
 			{
-				UE_LOG(LogLearning, Error, TEXT("%s: Observation '%s' too many sub-observations provided. Expected at most '%i', got '%i'."),
-					*ObjectName,
-					*ObservationSchemaElementName.ToString(),
+				UE_LOG(LogLearning, Error, TEXT("ValidateObjectMatchesSchema: Observation '%s' too many sub-observations provided. Expected at most '%i', got '%i'."),
+					*ObservationSchemaElementTag.ToString(),
 					SchemaParameters.Elements.Num(),
 					ObjectParameters.Elements.Num());
 				return false;
@@ -224,9 +210,8 @@ namespace UE::Learning::Agents::Observation::Private
 
 				if (SchemaSubElementIdx == INDEX_NONE)
 				{
-					UE_LOG(LogLearning, Error, TEXT("%s: Observation '%s' Schema does not include '%s' observation."),
-						*ObjectName,
-						*ObservationSchemaElementName.ToString(),
+					UE_LOG(LogLearning, Error, TEXT("ValidateObjectMatchesSchema: Observation '%s' Schema does not include '%s' observation."),
+						*ObservationSchemaElementTag.ToString(),
 						*ObjectParameters.ElementNames[ObjectSubElementIdx].ToString());
 					return false;
 				}
@@ -235,8 +220,7 @@ namespace UE::Learning::Agents::Observation::Private
 					Schema,
 					SchemaParameters.Elements[SchemaSubElementIdx],
 					Object,
-					ObjectParameters.Elements[ObjectSubElementIdx],
-					ObjectName))
+					ObjectParameters.Elements[ObjectSubElementIdx]))
 				{
 					return false;
 				}
@@ -252,9 +236,8 @@ namespace UE::Learning::Agents::Observation::Private
 
 			if (ObjectParameters.Elements.Num() != SchemaParameters.Num)
 			{
-				UE_LOG(LogLearning, Error, TEXT("%s: Observation '%s' array incorrect size. Expected '%i' elements, got '%i'."),
-					*ObjectName,
-					*ObservationSchemaElementName.ToString(),
+				UE_LOG(LogLearning, Error, TEXT("ValidateObjectMatchesSchema: Observation '%s' array incorrect size. Expected '%i' elements, got '%i'."),
+					*ObservationSchemaElementTag.ToString(),
 					SchemaParameters.Num,
 					ObjectParameters.Elements.Num());
 				return false;
@@ -266,8 +249,7 @@ namespace UE::Learning::Agents::Observation::Private
 					Schema,
 					SchemaParameters.Element,
 					Object,
-					ObjectParameters.Elements[ElementIdx],
-					ObjectName))
+					ObjectParameters.Elements[ElementIdx]))
 				{
 					return false;
 				}
@@ -283,9 +265,8 @@ namespace UE::Learning::Agents::Observation::Private
 
 			if (ObjectParameters.Elements.Num() > SchemaParameters.MaxNum)
 			{
-				UE_LOG(LogLearning, Error, TEXT("%s: Observation '%s' set too large. Expected at most '%i' elements, got '%i'."),
-					*ObjectName,
-					*ObservationSchemaElementName.ToString(),
+				UE_LOG(LogLearning, Error, TEXT("ValidateObjectMatchesSchema: Observation '%s' set too large. Expected at most '%i' elements, got '%i'."),
+					*ObservationSchemaElementTag.ToString(),
 					SchemaParameters.MaxNum,
 					ObjectParameters.Elements.Num());
 				return false;
@@ -297,8 +278,7 @@ namespace UE::Learning::Agents::Observation::Private
 					Schema,
 					SchemaParameters.Element,
 					Object,
-					ObjectParameters.Elements[ElementIdx],
-					ObjectName))
+					ObjectParameters.Elements[ElementIdx]))
 				{
 					return false;
 				}
@@ -316,8 +296,7 @@ namespace UE::Learning::Agents::Observation::Private
 				Schema,
 				SchemaParameters.Element,
 				Object,
-				ObjectParameters.Element,
-				ObjectName);
+				ObjectParameters.Element);
 		}
 
 		default:
@@ -332,23 +311,22 @@ namespace UE::Learning::Agents::Observation::Private
 		const UE::Learning::Observation::FObject& Object, 
 		const UE::Learning::Observation::FObjectElement ObjectElement,
 		const FString& Indentation,
-		const FString& Prefix,
-		const FString& ObjectName)
+		const FString& Prefix)
 	{
 		if (!Object.IsValid(ObjectElement))
 		{
-			UE_LOG(LogLearning, Error, TEXT("%s: Invalid Observation Object Element."), *ObjectName);
+			UE_LOG(LogLearning, Error, TEXT("LogObservation: Invalid Observation Object Element."));
 			return;
 		}
 
 		const UE::Learning::Observation::EType Type = Object.GetType(ObjectElement);
-		const FName Name = Object.GetName(ObjectElement);
+		const FName Tag = Object.GetTag(ObjectElement);
 
 		switch (Type)
 		{
 		case UE::Learning::Observation::EType::Null:
 		{
-			UE_LOG(LogLearning, Display, TEXT("%s%s \"%s\" (%s)"), *Indentation, *Prefix, *Name.ToString(), GetObservationTypeString(Type));
+			UE_LOG(LogLearning, Display, TEXT("%s%s \"%s\" (%s)"), *Indentation, *Prefix, *Tag.ToString(), GetObservationTypeString(Type));
 			return;
 		}
 
@@ -356,7 +334,7 @@ namespace UE::Learning::Agents::Observation::Private
 		{
 			const UE::Learning::Observation::FObjectContinuousParameters Parameters = Object.GetContinuous(ObjectElement);
 
-			UE_LOG(LogLearning, Display, TEXT("%s%s \"%s\" (%s) %s"), *Indentation, *Prefix, *Name.ToString(), GetObservationTypeString(Type), *UE::Learning::Array::FormatFloat(Parameters.Values));
+			UE_LOG(LogLearning, Display, TEXT("%s%s \"%s\" (%s) %s"), *Indentation, *Prefix, *Tag.ToString(), GetObservationTypeString(Type), *UE::Learning::Array::FormatFloat(Parameters.Values));
 			return;
 		}
 
@@ -364,10 +342,10 @@ namespace UE::Learning::Agents::Observation::Private
 		{
 			const UE::Learning::Observation::FObjectAndParameters Parameters = Object.GetAnd(ObjectElement);
 
-			UE_LOG(LogLearning, Display, TEXT("%s%s \"%s\" (%s)"), *Indentation, *Prefix, *Name.ToString(), GetObservationTypeString(Type));
+			UE_LOG(LogLearning, Display, TEXT("%s%s \"%s\" (%s)"), *Indentation, *Prefix, *Tag.ToString(), GetObservationTypeString(Type));
 			for (int32 SubElementIdx = 0; SubElementIdx < Parameters.Elements.Num(); SubElementIdx++)
 			{
-				LogObservation(Object, Parameters.Elements[SubElementIdx], *(Indentation + TEXT("    ")), FString::Printf(TEXT("| \"%s\": "), *Parameters.ElementNames[SubElementIdx].ToString()), ObjectName);
+				LogObservation(Object, Parameters.Elements[SubElementIdx], *(Indentation + TEXT("    ")), FString::Printf(TEXT("| \"%s\": "), *Parameters.ElementNames[SubElementIdx].ToString()));
 			}
 
 			return;
@@ -377,8 +355,8 @@ namespace UE::Learning::Agents::Observation::Private
 		{
 			const UE::Learning::Observation::FObjectOrExclusiveParameters Parameters = Object.GetOrExclusive(ObjectElement);
 
-			UE_LOG(LogLearning, Display, TEXT("%s%s \"%s\" (%s)"), *Indentation, *Prefix, *Name.ToString(), GetObservationTypeString(Type));
-			LogObservation(Object, Parameters.Element, *(Indentation + TEXT("    ")), FString::Printf(TEXT("| \"%s\": "), *Parameters.ElementName.ToString()), ObjectName);
+			UE_LOG(LogLearning, Display, TEXT("%s%s \"%s\" (%s)"), *Indentation, *Prefix, *Tag.ToString(), GetObservationTypeString(Type));
+			LogObservation(Object, Parameters.Element, *(Indentation + TEXT("    ")), FString::Printf(TEXT("| \"%s\": "), *Parameters.ElementName.ToString()));
 
 			return;
 		}
@@ -387,10 +365,10 @@ namespace UE::Learning::Agents::Observation::Private
 		{
 			const UE::Learning::Observation::FObjectOrInclusiveParameters Parameters = Object.GetOrInclusive(ObjectElement);
 
-			UE_LOG(LogLearning, Display, TEXT("%s%s \"%s\" (%s)"), *Indentation, *Prefix, *Name.ToString(), GetObservationTypeString(Type));
+			UE_LOG(LogLearning, Display, TEXT("%s%s \"%s\" (%s)"), *Indentation, *Prefix, *Tag.ToString(), GetObservationTypeString(Type));
 			for (int32 SubElementIdx = 0; SubElementIdx < Parameters.Elements.Num(); SubElementIdx++)
 			{
-				LogObservation(Object, Parameters.Elements[SubElementIdx], *(Indentation + TEXT("    ")), FString::Printf(TEXT("| \"%s\": "), *Parameters.ElementNames[SubElementIdx].ToString()), ObjectName);
+				LogObservation(Object, Parameters.Elements[SubElementIdx], *(Indentation + TEXT("    ")), FString::Printf(TEXT("| \"%s\": "), *Parameters.ElementNames[SubElementIdx].ToString()));
 			}
 
 			return;
@@ -400,10 +378,10 @@ namespace UE::Learning::Agents::Observation::Private
 		{
 			const UE::Learning::Observation::FObjectArrayParameters Parameters = Object.GetArray(ObjectElement);
 
-			UE_LOG(LogLearning, Display, TEXT("%s%s \"%s\" (%s)"), *Indentation, *Prefix, *Name.ToString(), GetObservationTypeString(Type));
+			UE_LOG(LogLearning, Display, TEXT("%s%s \"%s\" (%s)"), *Indentation, *Prefix, *Tag.ToString(), GetObservationTypeString(Type));
 			for (int32 SubElementIdx = 0; SubElementIdx < Parameters.Elements.Num(); SubElementIdx++)
 			{
-				LogObservation(Object, Parameters.Elements[SubElementIdx], *(Indentation + TEXT("    ")), FString::Printf(TEXT("| %3i:"), SubElementIdx), ObjectName);
+				LogObservation(Object, Parameters.Elements[SubElementIdx], *(Indentation + TEXT("    ")), FString::Printf(TEXT("| %3i:"), SubElementIdx));
 			}
 
 			return;
@@ -413,10 +391,10 @@ namespace UE::Learning::Agents::Observation::Private
 		{
 			const UE::Learning::Observation::FObjectSetParameters Parameters = Object.GetSet(ObjectElement);
 
-			UE_LOG(LogLearning, Display, TEXT("%s%s \"%s\" (%s)"), *Indentation, *Prefix, *Name.ToString(), GetObservationTypeString(Type));
+			UE_LOG(LogLearning, Display, TEXT("%s%s \"%s\" (%s)"), *Indentation, *Prefix, *Tag.ToString(), GetObservationTypeString(Type));
 			for (int32 SubElementIdx = 0; SubElementIdx < Parameters.Elements.Num(); SubElementIdx++)
 			{
-				LogObservation(Object, Parameters.Elements[SubElementIdx], *(Indentation + TEXT("    ")), FString::Printf(TEXT("| %3i:"), SubElementIdx), ObjectName);
+				LogObservation(Object, Parameters.Elements[SubElementIdx], *(Indentation + TEXT("    ")), FString::Printf(TEXT("| %3i:"), SubElementIdx));
 			}
 
 			return;
@@ -426,8 +404,8 @@ namespace UE::Learning::Agents::Observation::Private
 		{
 			const UE::Learning::Observation::FObjectEncodingParameters Parameters = Object.GetEncoding(ObjectElement);
 
-			UE_LOG(LogLearning, Display, TEXT("%s%s \"%s\" (%s)"), *Indentation, *Prefix, *Name.ToString(), GetObservationTypeString(Type));
-			LogObservation(Object, Parameters.Element, *(Indentation + TEXT("    ")), TEXT("|"), ObjectName);
+			UE_LOG(LogLearning, Display, TEXT("%s%s \"%s\" (%s)"), *Indentation, *Prefix, *Tag.ToString(), GetObservationTypeString(Type));
+			LogObservation(Object, Parameters.Element, *(Indentation + TEXT("    ")), TEXT("|"));
 
 			return;
 		}
@@ -462,7 +440,7 @@ namespace UE::Learning::Agents::Observation::Private
 	}
 }
 
-FTransform ULearningAgentsObservationFunctions::ProjectTransformOntoGroundPlane(const FTransform Transform, const FVector LocalForwardVector, const float GroundPlaneHeight)
+FTransform ULearningAgentsObservations::ProjectTransformOntoGroundPlane(const FTransform Transform, const FVector LocalForwardVector, const float GroundPlaneHeight)
 {
 	FVector Position = Transform.GetLocation();
 	Position.Z = GroundPlaneHeight;
@@ -472,70 +450,94 @@ FTransform ULearningAgentsObservationFunctions::ProjectTransformOntoGroundPlane(
 	return FTransform(FQuat::FindBetweenNormals(FVector::ForwardVector, Direction), Position, Transform.GetScale3D());
 }
 
-UEnum* ULearningAgentsObservationFunctions::FindEnumByName(const FString& Name)
+UEnum* ULearningAgentsObservations::FindEnumByName(const FString& Name)
 {
 	return FindObject<UEnum>(nullptr, *Name);
 }
 
-bool ULearningAgentsObservationSchema::ValidateObjectMatchesSchema(
+bool ULearningAgentsObservations::ValidateObjectMatchesSchema(
+	const ULearningAgentsObservationSchema* Schema,
 	const FLearningAgentsObservationSchemaElement SchemaElement,
 	const ULearningAgentsObservationObject* Object,
-	const FLearningAgentsObservationObjectElement ObjectElement) const
+	const FLearningAgentsObservationObjectElement ObjectElement)
 {
+	if (!Schema)
+	{
+		UE_LOG(LogLearning, Error, TEXT("ValidateObjectMatchesSchema: Schema is nullptr."));
+		return false;
+	}
+
+	if (!Object)
+	{
+		UE_LOG(LogLearning, Error, TEXT("ValidateObjectMatchesSchema: Object is nullptr."));
+		return false;
+	}
+
 	return UE::Learning::Agents::Observation::Private::ValidateObjectMatchesSchema(
-		ObservationSchema,
+		Schema->ObservationSchema,
 		SchemaElement.SchemaElement,
-		Object->GetObservationObject(),
-		ObjectElement.ObjectElement,
-		GetName());
+		Object->ObservationObject,
+		ObjectElement.ObjectElement);
 }
 
-FLearningAgentsObservationSchemaElement ULearningAgentsObservationSchema::SpecifyNullObservation(const FName Name)
+FLearningAgentsObservationSchemaElement ULearningAgentsObservations::SpecifyNullObservation(ULearningAgentsObservationSchema* Schema, const FName Tag)
 {
-	return { ObservationSchema.CreateNull(Name) };
+	if (!Schema)
+	{
+		UE_LOG(LogLearning, Error, TEXT("SpecifyNullObservation: Schema is nullptr."));
+		return FLearningAgentsObservationSchemaElement();
+	}
+
+	return { Schema->ObservationSchema.CreateNull(Tag)};
 }
 
-FLearningAgentsObservationSchemaElement ULearningAgentsObservationSchema::SpecifyContinuousObservation(const int32 Size, const FName Name)
+FLearningAgentsObservationSchemaElement ULearningAgentsObservations::SpecifyContinuousObservation(ULearningAgentsObservationSchema* Schema, const int32 Size, const FName Tag)
 {
+	if (!Schema)
+	{
+		UE_LOG(LogLearning, Error, TEXT("SpecifyContinuousObservation: Schema is nullptr."));
+		return FLearningAgentsObservationSchemaElement();
+	}
+
 	if (Size < 0)
 	{
-		UE_LOG(LogLearning, Error, TEXT("%s: Invalid Continuous Observation Size '%i'."), *GetName(), Size);
+		UE_LOG(LogLearning, Error, TEXT("SpecifyContinuousObservation: Invalid Continuous Observation Size '%i'."), Size);
 		return FLearningAgentsObservationSchemaElement();
 	}
 	
 	if (Size == 0)
 	{
-		UE_LOG(LogLearning, Warning, TEXT("%s: Specifying zero-sized Continuous Observation."), *GetName());
+		UE_LOG(LogLearning, Warning, TEXT("SpecifyContinuousObservation: Specifying zero-sized Continuous Observation."));
 	}
 
-	return { ObservationSchema.CreateContinuous({ Size }, Name) };
+	return { Schema->ObservationSchema.CreateContinuous({ Size }, Tag) };
 }
 
-FLearningAgentsObservationSchemaElement ULearningAgentsObservationSchema::SpecifyExclusiveDiscreteObservation(const int32 Size, const FName Name)
+FLearningAgentsObservationSchemaElement ULearningAgentsObservations::SpecifyExclusiveDiscreteObservation(ULearningAgentsObservationSchema* Schema, const int32 Size, const FName Tag)
 {
-	return SpecifyContinuousObservation(Size, Name);
+	return SpecifyContinuousObservation(Schema, Size, Tag);
 }
 
-FLearningAgentsObservationSchemaElement ULearningAgentsObservationSchema::SpecifyInclusiveDiscreteObservation(const int32 Size, const FName Name)
+FLearningAgentsObservationSchemaElement ULearningAgentsObservations::SpecifyInclusiveDiscreteObservation(ULearningAgentsObservationSchema* Schema, const int32 Size, const FName Tag)
 {
-	return SpecifyContinuousObservation(Size, Name);
+	return SpecifyContinuousObservation(Schema, Size, Tag);
 }
 
-FLearningAgentsObservationSchemaElement ULearningAgentsObservationSchema::SpecifyIndexObservation(const int32 Size, const FName Name)
+FLearningAgentsObservationSchemaElement ULearningAgentsObservations::SpecifyIndexObservation(ULearningAgentsObservationSchema* Schema, const int32 Size, const FName Tag)
 {
-	return SpecifyContinuousObservation(Size, Name);
+	return SpecifyContinuousObservation(Schema, Size, Tag);
 }
 
-FLearningAgentsObservationSchemaElement ULearningAgentsObservationSchema::SpecifyCountObservation(const FName Name)
+FLearningAgentsObservationSchemaElement ULearningAgentsObservations::SpecifyCountObservation(ULearningAgentsObservationSchema* Schema, const FName Tag)
 {
-	return SpecifyContinuousObservation(1, Name);
+	return SpecifyContinuousObservation(Schema, 1, Tag);
 }
 
-FLearningAgentsObservationSchemaElement ULearningAgentsObservationSchema::SpecifyStructObservation(const TMap<FName, FLearningAgentsObservationSchemaElement>& Elements, const FName Name)
+FLearningAgentsObservationSchemaElement ULearningAgentsObservations::SpecifyStructObservation(ULearningAgentsObservationSchema* Schema, const TMap<FName, FLearningAgentsObservationSchemaElement>& Elements, const FName Tag)
 {
 	if (Elements.Num() == 0)
 	{
-		UE_LOG(LogLearning, Warning, TEXT("%s: Specifying zero-sized Struct Observation."), *GetName());
+		UE_LOG(LogLearning, Warning, TEXT("SpecifyStructObservation: Specifying zero-sized Struct Observation."));
 	}
 
 	const int32 SubElementNum = Elements.Num();
@@ -573,30 +575,36 @@ FLearningAgentsObservationSchemaElement ULearningAgentsObservationSchema::Specif
 		SortedSubElements[Idx] = SubElements[SubElementIndices[Idx]];
 	}
 
-	return SpecifyStructObservationFromArrayViews(SortedSubElementNames, SortedSubElements, Name);
+	return SpecifyStructObservationFromArrayViews(Schema, SortedSubElementNames, SortedSubElements, Tag);
 }
 
-FLearningAgentsObservationSchemaElement ULearningAgentsObservationSchema::SpecifyStructObservationFromArrays(const TArray<FName>& ElementNames, const TArray<FLearningAgentsObservationSchemaElement>& Elements, const FName Name)
+FLearningAgentsObservationSchemaElement ULearningAgentsObservations::SpecifyStructObservationFromArrays(ULearningAgentsObservationSchema* Schema, const TArray<FName>& ElementNames, const TArray<FLearningAgentsObservationSchemaElement>& Elements, const FName Tag)
 {
-	return SpecifyStructObservationFromArrayViews(ElementNames, Elements, Name);
+	return SpecifyStructObservationFromArrayViews(Schema, ElementNames, Elements, Tag);
 }
 
-FLearningAgentsObservationSchemaElement ULearningAgentsObservationSchema::SpecifyStructObservationFromArrayViews(const TArrayView<const FName> ElementNames, const TArrayView<const FLearningAgentsObservationSchemaElement> Elements, const FName Name)
+FLearningAgentsObservationSchemaElement ULearningAgentsObservations::SpecifyStructObservationFromArrayViews(ULearningAgentsObservationSchema* Schema, const TArrayView<const FName> ElementNames, const TArrayView<const FLearningAgentsObservationSchemaElement> Elements, const FName Tag)
 {
+	if (!Schema)
+	{
+		UE_LOG(LogLearning, Error, TEXT("SpecifyStructObservationFromArrayViews: Schema is nullptr."));
+		return FLearningAgentsObservationSchemaElement();
+	}
+
 	if (Elements.Num() == 0)
 	{
-		UE_LOG(LogLearning, Warning, TEXT("%s: Specifying zero-sized Struct Observation."), *GetName());
+		UE_LOG(LogLearning, Warning, TEXT("SpecifyStructObservationFromArrayViews: Specifying zero-sized Struct Observation."));
 	}
 
 	if (Elements.Num() != ElementNames.Num())
 	{
-		UE_LOG(LogLearning, Error, TEXT("%s: Number of elements (%i) must match number of names (%i)."), *GetName(), Elements.Num(), ElementNames.Num());
+		UE_LOG(LogLearning, Error, TEXT("SpecifyStructObservationFromArrayViews: Number of elements (%i) must match number of names (%i)."), Elements.Num(), ElementNames.Num());
 		return FLearningAgentsObservationSchemaElement();
 	}
 
 	if (UE::Learning::Agents::Observation::Private::ContainsDuplicates(ElementNames))
 	{
-		UE_LOG(LogLearning, Error, TEXT("%s: Element Names contain duplicates."), *GetName());
+		UE_LOG(LogLearning, Error, TEXT("SpecifyStructObservationFromArrayViews: Element Names contain duplicates."));
 		return FLearningAgentsObservationSchemaElement();
 	}
 
@@ -605,29 +613,29 @@ FLearningAgentsObservationSchemaElement ULearningAgentsObservationSchema::Specif
 
 	for (const FLearningAgentsObservationSchemaElement& Element : Elements)
 	{
-		if (!ObservationSchema.IsValid(Element.SchemaElement))
+		if (!Schema->ObservationSchema.IsValid(Element.SchemaElement))
 		{
-			UE_LOG(LogLearning, Error, TEXT("%s: Invalid Observation Object."), *GetName());
+			UE_LOG(LogLearning, Error, TEXT("SpecifyStructObservationFromArrayViews: Invalid Observation Object."));
 			return FLearningAgentsObservationSchemaElement();
 		}
 
 		SubElements.Add(Element.SchemaElement);
 	}
 
-	return { ObservationSchema.CreateAnd({ ElementNames, SubElements }, Name) };
+	return { Schema->ObservationSchema.CreateAnd({ ElementNames, SubElements }, Tag) };
 }
 
-FLearningAgentsObservationSchemaElement ULearningAgentsObservationSchema::SpecifyExclusiveUnionObservation(const TMap<FName, FLearningAgentsObservationSchemaElement>& Elements, const int32 EncodingSize, const FName Name)
+FLearningAgentsObservationSchemaElement ULearningAgentsObservations::SpecifyExclusiveUnionObservation(ULearningAgentsObservationSchema* Schema, const TMap<FName, FLearningAgentsObservationSchemaElement>& Elements, const int32 EncodingSize, const FName Tag)
 {
 	if (EncodingSize < 1)
 	{
-		UE_LOG(LogLearning, Error, TEXT("%s: Invalid Observation EncodingSize '%i' - must be greater than zero."), *GetName(), EncodingSize);
+		UE_LOG(LogLearning, Error, TEXT("SpecifyExclusiveUnionObservation: Invalid Observation EncodingSize '%i' - must be greater than zero."), EncodingSize);
 		return FLearningAgentsObservationSchemaElement();
 	}
 
 	if (Elements.Num() == 0)
 	{
-		UE_LOG(LogLearning, Warning, TEXT("%s: Specifying zero-sized Exclusive Union Observation."), *GetName());
+		UE_LOG(LogLearning, Warning, TEXT("SpecifyExclusiveUnionObservation: Specifying zero-sized Exclusive Union Observation."));
 	}
 
 	const int32 SubElementNum = Elements.Num();
@@ -665,36 +673,42 @@ FLearningAgentsObservationSchemaElement ULearningAgentsObservationSchema::Specif
 		SortedSubElements[Idx] = SubElements[SubElementIndices[Idx]];
 	}
 
-	return SpecifyExclusiveUnionObservationFromArrayViews(SortedSubElementNames, SortedSubElements, EncodingSize, Name);
+	return SpecifyExclusiveUnionObservationFromArrayViews(Schema, SortedSubElementNames, SortedSubElements, EncodingSize, Tag);
 }
 
-FLearningAgentsObservationSchemaElement ULearningAgentsObservationSchema::SpecifyExclusiveUnionObservationFromArrays(const TArray<FName>& ElementNames, const TArray<FLearningAgentsObservationSchemaElement>& Elements, const int32 EncodingSize, const FName Name)
+FLearningAgentsObservationSchemaElement ULearningAgentsObservations::SpecifyExclusiveUnionObservationFromArrays(ULearningAgentsObservationSchema* Schema, const TArray<FName>& ElementNames, const TArray<FLearningAgentsObservationSchemaElement>& Elements, const int32 EncodingSize, const FName Tag)
 {
-	return SpecifyExclusiveUnionObservationFromArrayViews(ElementNames, Elements, EncodingSize, Name);
+	return SpecifyExclusiveUnionObservationFromArrayViews(Schema, ElementNames, Elements, EncodingSize, Tag);
 }
 
-FLearningAgentsObservationSchemaElement ULearningAgentsObservationSchema::SpecifyExclusiveUnionObservationFromArrayViews(const TArrayView<const FName> ElementNames, const TArrayView<const FLearningAgentsObservationSchemaElement> Elements, const int32 EncodingSize, const FName Name)
+FLearningAgentsObservationSchemaElement ULearningAgentsObservations::SpecifyExclusiveUnionObservationFromArrayViews(ULearningAgentsObservationSchema* Schema, const TArrayView<const FName> ElementNames, const TArrayView<const FLearningAgentsObservationSchemaElement> Elements, const int32 EncodingSize, const FName Tag)
 {
+	if (!Schema)
+	{
+		UE_LOG(LogLearning, Error, TEXT("SpecifyExclusiveUnionObservationFromArrayViews: Schema is nullptr."));
+		return FLearningAgentsObservationSchemaElement();
+	}
+
 	if (EncodingSize < 1)
 	{
-		UE_LOG(LogLearning, Error, TEXT("%s: Invalid Observation EncodingSize '%i' - must be greater than zero."), *GetName(), EncodingSize);
+		UE_LOG(LogLearning, Error, TEXT("SpecifyExclusiveUnionObservationFromArrayViews: Invalid Observation EncodingSize '%i' - must be greater than zero."), EncodingSize);
 		return FLearningAgentsObservationSchemaElement();
 	}
 
 	if (Elements.Num() == 0)
 	{
-		UE_LOG(LogLearning, Warning, TEXT("%s: Specifying zero-sized Exclusive Union Observation."), *GetName());
+		UE_LOG(LogLearning, Warning, TEXT("SpecifyExclusiveUnionObservationFromArrayViews: Specifying zero-sized Exclusive Union Observation."));
 	}
 
 	if (Elements.Num() != ElementNames.Num())
 	{
-		UE_LOG(LogLearning, Error, TEXT("%s: Number of elements (%i) must match number of names (%i)."), *GetName(), Elements.Num(), ElementNames.Num());
+		UE_LOG(LogLearning, Error, TEXT("SpecifyExclusiveUnionObservationFromArrayViews: Number of elements (%i) must match number of names (%i)."), Elements.Num(), ElementNames.Num());
 		return FLearningAgentsObservationSchemaElement();
 	}
 
 	if (UE::Learning::Agents::Observation::Private::ContainsDuplicates(ElementNames))
 	{
-		UE_LOG(LogLearning, Error, TEXT("%s: Element Names contain duplicates."), *GetName());
+		UE_LOG(LogLearning, Error, TEXT("SpecifyExclusiveUnionObservationFromArrayViews: Element Names contain duplicates."));
 		return FLearningAgentsObservationSchemaElement();
 	}
 
@@ -703,29 +717,29 @@ FLearningAgentsObservationSchemaElement ULearningAgentsObservationSchema::Specif
 
 	for (const FLearningAgentsObservationSchemaElement& Element : Elements)
 	{
-		if (!ObservationSchema.IsValid(Element.SchemaElement))
+		if (!Schema->ObservationSchema.IsValid(Element.SchemaElement))
 		{
-			UE_LOG(LogLearning, Error, TEXT("%s: Invalid Observation Object."), *GetName());
+			UE_LOG(LogLearning, Error, TEXT("SpecifyExclusiveUnionObservationFromArrayViews: Invalid Observation Object."));
 			return FLearningAgentsObservationSchemaElement();
 		}
 
 		SubElements.Add(Element.SchemaElement);
 	}
 
-	return { ObservationSchema.CreateOrExclusive({ ElementNames, SubElements, EncodingSize }, Name) };
+	return { Schema->ObservationSchema.CreateOrExclusive({ ElementNames, SubElements, EncodingSize }, Tag) };
 }
 
-FLearningAgentsObservationSchemaElement ULearningAgentsObservationSchema::SpecifyInclusiveUnionObservation(const TMap<FName, FLearningAgentsObservationSchemaElement>& Elements, const int32 AttentionEncodingSize, const int32 AttentionHeadNum, const int32 ValueEncodingSize, const FName Name)
+FLearningAgentsObservationSchemaElement ULearningAgentsObservations::SpecifyInclusiveUnionObservation(ULearningAgentsObservationSchema* Schema, const TMap<FName, FLearningAgentsObservationSchemaElement>& Elements, const int32 AttentionEncodingSize, const int32 AttentionHeadNum, const int32 ValueEncodingSize, const FName Tag)
 {
 	if (AttentionEncodingSize < 1 || AttentionHeadNum < 1 || ValueEncodingSize < 1)
 	{
-		UE_LOG(LogLearning, Error, TEXT("%s: Invalid Observation Parameters: AttentionEncodingSize: %i, AttentionHeadNum: %i, ValueEncodingSize: %i - must be greater than zero."), *GetName(), AttentionEncodingSize, AttentionHeadNum, ValueEncodingSize);
+		UE_LOG(LogLearning, Error, TEXT("SpecifyInclusiveUnionObservation: Invalid Observation Parameters: AttentionEncodingSize: %i, AttentionHeadNum: %i, ValueEncodingSize: %i - must be greater than zero."), AttentionEncodingSize, AttentionHeadNum, ValueEncodingSize);
 		return FLearningAgentsObservationSchemaElement();
 	}
 
 	if (Elements.Num() == 0)
 	{
-		UE_LOG(LogLearning, Warning, TEXT("%s: Specifying zero-sized Inclusive Union Observation."), *GetName());
+		UE_LOG(LogLearning, Warning, TEXT("SpecifyInclusiveUnionObservation: Specifying zero-sized Inclusive Union Observation."));
 	}
 
 	const int32 SubElementNum = Elements.Num();
@@ -763,36 +777,42 @@ FLearningAgentsObservationSchemaElement ULearningAgentsObservationSchema::Specif
 		SortedSubElements[Idx] = SubElements[SubElementIndices[Idx]];
 	}
 
-	return SpecifyInclusiveUnionObservationFromArrayViews(SortedSubElementNames, SortedSubElements, AttentionEncodingSize, AttentionHeadNum, ValueEncodingSize, Name);
+	return SpecifyInclusiveUnionObservationFromArrayViews(Schema, SortedSubElementNames, SortedSubElements, AttentionEncodingSize, AttentionHeadNum, ValueEncodingSize, Tag);
 }
 
-FLearningAgentsObservationSchemaElement ULearningAgentsObservationSchema::SpecifyInclusiveUnionObservationFromArrays(const TArray<FName>& ElementNames, const TArray<FLearningAgentsObservationSchemaElement>& Elements, const int32 AttentionEncodingSize, const int32 AttentionHeadNum, const int32 ValueEncodingSize, const FName Name)
+FLearningAgentsObservationSchemaElement ULearningAgentsObservations::SpecifyInclusiveUnionObservationFromArrays(ULearningAgentsObservationSchema* Schema, const TArray<FName>& ElementNames, const TArray<FLearningAgentsObservationSchemaElement>& Elements, const int32 AttentionEncodingSize, const int32 AttentionHeadNum, const int32 ValueEncodingSize, const FName Tag)
 {
-	return SpecifyInclusiveUnionObservationFromArrayViews(ElementNames, Elements, AttentionEncodingSize, AttentionHeadNum, ValueEncodingSize, Name);
+	return SpecifyInclusiveUnionObservationFromArrayViews(Schema, ElementNames, Elements, AttentionEncodingSize, AttentionHeadNum, ValueEncodingSize, Tag);
 }
 
-FLearningAgentsObservationSchemaElement ULearningAgentsObservationSchema::SpecifyInclusiveUnionObservationFromArrayViews(const TArrayView<const FName> ElementNames, const TArrayView<const FLearningAgentsObservationSchemaElement> Elements, const int32 AttentionEncodingSize, const int32 AttentionHeadNum, const int32 ValueEncodingSize, const FName Name)
+FLearningAgentsObservationSchemaElement ULearningAgentsObservations::SpecifyInclusiveUnionObservationFromArrayViews(ULearningAgentsObservationSchema* Schema, const TArrayView<const FName> ElementNames, const TArrayView<const FLearningAgentsObservationSchemaElement> Elements, const int32 AttentionEncodingSize, const int32 AttentionHeadNum, const int32 ValueEncodingSize, const FName Tag)
 {
+	if (!Schema)
+	{
+		UE_LOG(LogLearning, Error, TEXT("SpecifyInclusiveUnionObservationFromArrayViews: Schema is nullptr."));
+		return FLearningAgentsObservationSchemaElement();
+	}
+
 	if (AttentionEncodingSize < 1 || AttentionHeadNum < 1 || ValueEncodingSize < 1)
 	{
-		UE_LOG(LogLearning, Error, TEXT("%s: Invalid Observation Parameters: AttentionEncodingSize: %i, AttentionHeadNum: %i, ValueEncodingSize: %i - must be greater than zero."), *GetName(), AttentionEncodingSize, AttentionHeadNum, ValueEncodingSize);
+		UE_LOG(LogLearning, Error, TEXT("SpecifyInclusiveUnionObservationFromArrayViews: Invalid Observation Parameters: AttentionEncodingSize: %i, AttentionHeadNum: %i, ValueEncodingSize: %i - must be greater than zero."), AttentionEncodingSize, AttentionHeadNum, ValueEncodingSize);
 		return FLearningAgentsObservationSchemaElement();
 	}
 
 	if (Elements.Num() == 0)
 	{
-		UE_LOG(LogLearning, Warning, TEXT("%s: Specifying zero-sized Inclusive Union Observation."), *GetName());
+		UE_LOG(LogLearning, Warning, TEXT("SpecifyInclusiveUnionObservationFromArrayViews: Specifying zero-sized Inclusive Union Observation."));
 	}
 
 	if (Elements.Num() != ElementNames.Num())
 	{
-		UE_LOG(LogLearning, Error, TEXT("%s: Number of elements (%i) must match number of names (%i)."), *GetName(), Elements.Num(), ElementNames.Num());
+		UE_LOG(LogLearning, Error, TEXT("SpecifyInclusiveUnionObservationFromArrayViews: Number of elements (%i) must match number of names (%i)."), Elements.Num(), ElementNames.Num());
 		return FLearningAgentsObservationSchemaElement();
 	}
 
 	if (UE::Learning::Agents::Observation::Private::ContainsDuplicates(ElementNames))
 	{
-		UE_LOG(LogLearning, Error, TEXT("%s: Element Names contain duplicates."), *GetName());
+		UE_LOG(LogLearning, Error, TEXT("SpecifyInclusiveUnionObservationFromArrayViews: Element Names contain duplicates."));
 		return FLearningAgentsObservationSchemaElement();
 	}
 
@@ -801,310 +821,336 @@ FLearningAgentsObservationSchemaElement ULearningAgentsObservationSchema::Specif
 
 	for (const FLearningAgentsObservationSchemaElement& Element : Elements)
 	{
-		if (!ObservationSchema.IsValid(Element.SchemaElement))
+		if (!Schema->ObservationSchema.IsValid(Element.SchemaElement))
 		{
-			UE_LOG(LogLearning, Error, TEXT("%s: Invalid Observation Object."), *GetName());
+			UE_LOG(LogLearning, Error, TEXT("SpecifyInclusiveUnionObservationFromArrayViews: Invalid Observation Object."));
 			return FLearningAgentsObservationSchemaElement();
 		}
 
 		SubElements.Add(Element.SchemaElement);
 	}
 
-	return { ObservationSchema.CreateOrInclusive({ ElementNames, SubElements, AttentionEncodingSize, AttentionHeadNum, ValueEncodingSize }, Name) };
+	return { Schema->ObservationSchema.CreateOrInclusive({ ElementNames, SubElements, AttentionEncodingSize, AttentionHeadNum, ValueEncodingSize }, Tag) };
 }
 
-FLearningAgentsObservationSchemaElement ULearningAgentsObservationSchema::SpecifyStaticArrayObservation(const FLearningAgentsObservationSchemaElement Element, const int32 Num, const FName Name)
+FLearningAgentsObservationSchemaElement ULearningAgentsObservations::SpecifyStaticArrayObservation(ULearningAgentsObservationSchema* Schema, const FLearningAgentsObservationSchemaElement Element, const int32 Num, const FName Tag)
 {
+	if (!Schema)
+	{
+		UE_LOG(LogLearning, Error, TEXT("SpecifyStaticArrayObservation: Schema is nullptr."));
+		return FLearningAgentsObservationSchemaElement();
+	}
+
 	if (Num < 0)
 	{
-		UE_LOG(LogLearning, Error, TEXT("%s: Invalid Observation Array Num %i."), *GetName(), Num);
+		UE_LOG(LogLearning, Error, TEXT("SpecifyStaticArrayObservation: Invalid Observation Array Num %i."), Num);
 		return FLearningAgentsObservationSchemaElement();
 	}
 
 	if (Num == 0)
 	{
-		UE_LOG(LogLearning, Warning, TEXT("%s: Specifying zero-sized Static Array Observation."), *GetName());
+		UE_LOG(LogLearning, Warning, TEXT("SpecifyStaticArrayObservation: Specifying zero-sized Static Array Observation."));
 	}
 
-	if (!ObservationSchema.IsValid(Element.SchemaElement))
+	if (!Schema->ObservationSchema.IsValid(Element.SchemaElement))
 	{
-		UE_LOG(LogLearning, Error, TEXT("%s: Invalid Observation Object."), *GetName());
+		UE_LOG(LogLearning, Error, TEXT("SpecifyStaticArrayObservation: Invalid Observation Object."));
 		return FLearningAgentsObservationSchemaElement();
 	}
 
-	return { ObservationSchema.CreateArray({ Element.SchemaElement, Num }, Name) };
+	return { Schema->ObservationSchema.CreateArray({ Element.SchemaElement, Num }, Tag) };
 }
 
-FLearningAgentsObservationSchemaElement ULearningAgentsObservationSchema::SpecifySetObservation(const FLearningAgentsObservationSchemaElement Element, const int32 MaxNum, const int32 AttentionEncodingSize, const int32 AttentionHeadNum, const int32 ValueEncodingSize, const FName Name)
+FLearningAgentsObservationSchemaElement ULearningAgentsObservations::SpecifySetObservation(ULearningAgentsObservationSchema* Schema, const FLearningAgentsObservationSchemaElement Element, const int32 MaxNum, const int32 AttentionEncodingSize, const int32 AttentionHeadNum, const int32 ValueEncodingSize, const FName Tag)
 {
+	if (!Schema)
+	{
+		UE_LOG(LogLearning, Error, TEXT("SpecifySetObservation: Schema is nullptr."));
+		return FLearningAgentsObservationSchemaElement();
+	}
+
 	if (MaxNum < 0)
 	{
-		UE_LOG(LogLearning, Error, TEXT("%s: Invalid Observation Set MaxNum %i - must be greater than or equal to zero."), *GetName(), MaxNum);
+		UE_LOG(LogLearning, Error, TEXT("SpecifySetObservation: Invalid Observation Set MaxNum %i - must be greater than or equal to zero."), MaxNum);
 		return FLearningAgentsObservationSchemaElement();
 	}
 
 	if (MaxNum == 0)
 	{
-		UE_LOG(LogLearning, Warning, TEXT("%s: Specifying zero-sized Set Observation."), *GetName());
+		UE_LOG(LogLearning, Warning, TEXT("SpecifySetObservation: Specifying zero-sized Set Observation."));
 	}
 
 	if (AttentionEncodingSize < 1 || AttentionHeadNum < 1 || ValueEncodingSize < 1)
 	{
-		UE_LOG(LogLearning, Error, TEXT("%s: Invalid Observation Parameters: AttentionEncodingSize: %i, AttentionHeadNum: %i, ValueEncodingSize: %i - must be greater than zero."), *GetName(), AttentionEncodingSize, AttentionHeadNum, ValueEncodingSize);
+		UE_LOG(LogLearning, Error, TEXT("SpecifySetObservation: Invalid Observation Parameters: AttentionEncodingSize: %i, AttentionHeadNum: %i, ValueEncodingSize: %i - must be greater than zero."), AttentionEncodingSize, AttentionHeadNum, ValueEncodingSize);
 		return FLearningAgentsObservationSchemaElement();
 	}
 
-	if (!ObservationSchema.IsValid(Element.SchemaElement))
+	if (!Schema->ObservationSchema.IsValid(Element.SchemaElement))
 	{
-		UE_LOG(LogLearning, Error, TEXT("%s: Invalid Observation Object."), *GetName());
+		UE_LOG(LogLearning, Error, TEXT("SpecifySetObservation: Invalid Observation Object."));
 		return FLearningAgentsObservationSchemaElement();
 	}
 
-	return { ObservationSchema.CreateSet({ Element.SchemaElement, MaxNum, AttentionEncodingSize, AttentionHeadNum, ValueEncodingSize }, Name) };
+	return { Schema->ObservationSchema.CreateSet({ Element.SchemaElement, MaxNum, AttentionEncodingSize, AttentionHeadNum, ValueEncodingSize }, Tag) };
 }
 
-FLearningAgentsObservationSchemaElement ULearningAgentsObservationSchema::SpecifyPairObservation(const FLearningAgentsObservationSchemaElement Element0, const FLearningAgentsObservationSchemaElement Element1, const FName Name)
+FLearningAgentsObservationSchemaElement ULearningAgentsObservations::SpecifyPairObservation(ULearningAgentsObservationSchema* Schema, const FLearningAgentsObservationSchemaElement Element0, const FLearningAgentsObservationSchemaElement Element1, const FName Tag)
 {
-	return SpecifyStructObservationFromArrayViews({ TEXT("Key"), TEXT("Value") }, { Element0, Element1 }, Name);
+	return SpecifyStructObservationFromArrayViews(Schema, { TEXT("Key"), TEXT("Value") }, { Element0, Element1 }, Tag);
 }
 
-FLearningAgentsObservationSchemaElement ULearningAgentsObservationSchema::SpecifyArrayObservation(const FLearningAgentsObservationSchemaElement Element, const int32 MaxNum, const int32 AttentionEncodingSize, const int32 AttentionHeadNum, const int32 ValueEncodingSize, const FName Name)
+FLearningAgentsObservationSchemaElement ULearningAgentsObservations::SpecifyArrayObservation(ULearningAgentsObservationSchema* Schema, const FLearningAgentsObservationSchemaElement Element, const int32 MaxNum, const int32 AttentionEncodingSize, const int32 AttentionHeadNum, const int32 ValueEncodingSize, const FName Tag)
 {
-	return SpecifySetObservation(SpecifyPairObservation(SpecifyIndexObservation(MaxNum), Element), MaxNum, AttentionEncodingSize, AttentionHeadNum, ValueEncodingSize);
+	return SpecifySetObservation(Schema, SpecifyPairObservation(Schema, SpecifyIndexObservation(Schema, MaxNum), Element), MaxNum, AttentionEncodingSize, AttentionHeadNum, ValueEncodingSize);
 }
 
-FLearningAgentsObservationSchemaElement ULearningAgentsObservationSchema::SpecifyMapObservation(const FLearningAgentsObservationSchemaElement KeyElement, const FLearningAgentsObservationSchemaElement ValueElement, const int32 MaxNum, const int32 AttentionEncodingSize, const int32 AttentionHeadNum, const int32 ValueEncodingSize, const FName Name)
+FLearningAgentsObservationSchemaElement ULearningAgentsObservations::SpecifyMapObservation(ULearningAgentsObservationSchema* Schema, const FLearningAgentsObservationSchemaElement KeyElement, const FLearningAgentsObservationSchemaElement ValueElement, const int32 MaxNum, const int32 AttentionEncodingSize, const int32 AttentionHeadNum, const int32 ValueEncodingSize, const FName Tag)
 {
-	return SpecifySetObservation(SpecifyPairObservation(KeyElement, ValueElement), MaxNum, AttentionEncodingSize, AttentionHeadNum, ValueEncodingSize);
+	return SpecifySetObservation(Schema, SpecifyPairObservation(Schema, KeyElement, ValueElement), MaxNum, AttentionEncodingSize, AttentionHeadNum, ValueEncodingSize);
 }
 
-FLearningAgentsObservationSchemaElement ULearningAgentsObservationSchema::SpecifyEnumObservation(const UEnum* Enum, const FName Name)
+FLearningAgentsObservationSchemaElement ULearningAgentsObservations::SpecifyEnumObservation(ULearningAgentsObservationSchema* Schema, const UEnum* Enum, const FName Tag)
 {
 	if (!Enum)
 	{
-		UE_LOG(LogLearning, Error, TEXT("%s: Enum is nullptr."), *GetName());
+		UE_LOG(LogLearning, Error, TEXT("SpecifyEnumObservation: Enum is nullptr."));
 		return FLearningAgentsObservationSchemaElement();
 	}
 
-	return SpecifyContinuousObservation(Enum->NumEnums() - 1, Name);
+	return SpecifyContinuousObservation(Schema, Enum->NumEnums() - 1, Tag);
 }
 
-FLearningAgentsObservationSchemaElement ULearningAgentsObservationSchema::SpecifyBitmaskObservation(const UEnum* Enum, const FName Name)
+FLearningAgentsObservationSchemaElement ULearningAgentsObservations::SpecifyBitmaskObservation(ULearningAgentsObservationSchema* Schema, const UEnum* Enum, const FName Tag)
 {
 	if (!Enum)
 	{
-		UE_LOG(LogLearning, Error, TEXT("%s: Enum is nullptr."), *GetName());
+		UE_LOG(LogLearning, Error, TEXT("SpecifyBitmaskObservation: Enum is nullptr."));
 		return FLearningAgentsObservationSchemaElement();
 	}
 
 	if (Enum->NumEnums() - 1 > 32)
 	{
-		UE_LOG(LogLearning, Error, TEXT("%s: Too many values in Enum to use as Bitmask (%i)."), *GetName(), Enum->NumEnums() - 1);
+		UE_LOG(LogLearning, Error, TEXT("SpecifyBitmaskObservation: Too many values in Enum to use as Bitmask (%i)."), Enum->NumEnums() - 1);
 		return FLearningAgentsObservationSchemaElement();
 	}
 
-	return SpecifyContinuousObservation(Enum->NumEnums() - 1, Name);
+	return SpecifyContinuousObservation(Schema, Enum->NumEnums() - 1, Tag);
 }
 
-FLearningAgentsObservationSchemaElement ULearningAgentsObservationSchema::SpecifyOptionalObservation(const FLearningAgentsObservationSchemaElement Element, const int32 EncodingSize, const FName Name)
+FLearningAgentsObservationSchemaElement ULearningAgentsObservations::SpecifyOptionalObservation(ULearningAgentsObservationSchema* Schema, const FLearningAgentsObservationSchemaElement Element, const int32 EncodingSize, const FName Tag)
 {
-	return SpecifyExclusiveUnionObservationFromArrayViews({ TEXT("Null"), TEXT("Valid") }, { SpecifyNullObservation(), Element }, EncodingSize, Name);
+	return SpecifyExclusiveUnionObservationFromArrayViews(Schema, { TEXT("Null"), TEXT("Valid") }, { SpecifyNullObservation(Schema), Element }, EncodingSize, Tag);
 }
 
-FLearningAgentsObservationSchemaElement ULearningAgentsObservationSchema::SpecifyEitherObservation(const FLearningAgentsObservationSchemaElement A, const FLearningAgentsObservationSchemaElement B, const int32 EncodingSize, const FName Name)
+FLearningAgentsObservationSchemaElement ULearningAgentsObservations::SpecifyEitherObservation(ULearningAgentsObservationSchema* Schema, const FLearningAgentsObservationSchemaElement A, const FLearningAgentsObservationSchemaElement B, const int32 EncodingSize, const FName Tag)
 {
-	return SpecifyExclusiveUnionObservationFromArrayViews({ TEXT("A"), TEXT("B") }, { A, B }, EncodingSize, Name);
+	return SpecifyExclusiveUnionObservationFromArrayViews(Schema, { TEXT("A"), TEXT("B") }, { A, B }, EncodingSize, Tag);
 }
 
-FLearningAgentsObservationSchemaElement ULearningAgentsObservationSchema::SpecifyEncodingObservation(const FLearningAgentsObservationSchemaElement Element, const int32 EncodingSize, const int32 LayerNum, const ELearningAgentsActivationFunction ActivationFunction, const FName Name)
+FLearningAgentsObservationSchemaElement ULearningAgentsObservations::SpecifyEncodingObservation(ULearningAgentsObservationSchema* Schema, const FLearningAgentsObservationSchemaElement Element, const int32 EncodingSize, const int32 LayerNum, const ELearningAgentsActivationFunction ActivationFunction, const FName Tag)
 {
+	if (!Schema)
+	{
+		UE_LOG(LogLearning, Error, TEXT("SpecifyEncodingObservation: Schema is nullptr."));
+		return FLearningAgentsObservationSchemaElement();
+	}
+
 	if (EncodingSize < 1)
 	{
-		UE_LOG(LogLearning, Error, TEXT("%s: Invalid Observation EncodingSize '%i' - must be greater than zero."), *GetName(), EncodingSize);
+		UE_LOG(LogLearning, Error, TEXT("SpecifyEncodingObservation: Invalid Observation EncodingSize '%i' - must be greater than zero."), EncodingSize);
 		return FLearningAgentsObservationSchemaElement();
 	}
 
 	if (LayerNum < 1)
 	{
-		UE_LOG(LogLearning, Error, TEXT("%s: Invalid Observation LayerNum '%i' - must be greater than zero."), *GetName(), LayerNum);
+		UE_LOG(LogLearning, Error, TEXT("SpecifyEncodingObservation: Invalid Observation LayerNum '%i' - must be greater than zero."), LayerNum);
 		return FLearningAgentsObservationSchemaElement();
 	}
 
-	if (!ObservationSchema.IsValid(Element.SchemaElement))
+	if (!Schema->ObservationSchema.IsValid(Element.SchemaElement))
 	{
-		UE_LOG(LogLearning, Error, TEXT("%s: Invalid Observation Object."), *GetName());
+		UE_LOG(LogLearning, Error, TEXT("SpecifyEncodingObservation: Invalid Observation Object."));
 		return FLearningAgentsObservationSchemaElement();
 	}
 
-	return { ObservationSchema.CreateEncoding({ Element.SchemaElement, EncodingSize, LayerNum, UE::Learning::Agents::Observation::Private::GetEncodingActivationFunction(ActivationFunction) }, Name)};
+	return { Schema->ObservationSchema.CreateEncoding({ Element.SchemaElement, EncodingSize, LayerNum, UE::Learning::Agents::Observation::Private::GetEncodingActivationFunction(ActivationFunction) }, Tag)};
 }
 
-FLearningAgentsObservationSchemaElement ULearningAgentsObservationSchema::SpecifyBoolObservation(const FName Name)
+FLearningAgentsObservationSchemaElement ULearningAgentsObservations::SpecifyBoolObservation(ULearningAgentsObservationSchema* Schema, const FName Tag)
 {
-	return SpecifyContinuousObservation(1, Name);
+	return SpecifyContinuousObservation(Schema, 1, Tag);
 }
 
-FLearningAgentsObservationSchemaElement ULearningAgentsObservationSchema::SpecifyFloatObservation(const FName Name)
+FLearningAgentsObservationSchemaElement ULearningAgentsObservations::SpecifyFloatObservation(ULearningAgentsObservationSchema* Schema, const FName Tag)
 {
-	return SpecifyContinuousObservation(1, Name);
+	return SpecifyContinuousObservation(Schema, 1, Tag);
 }
 
-FLearningAgentsObservationSchemaElement ULearningAgentsObservationSchema::SpecifyLocationObservation(const FName Name)
+FLearningAgentsObservationSchemaElement ULearningAgentsObservations::SpecifyLocationObservation(ULearningAgentsObservationSchema* Schema, const FName Tag)
 {
-	return SpecifyContinuousObservation(3, Name);
+	return SpecifyContinuousObservation(Schema, 3, Tag);
 }
 
-FLearningAgentsObservationSchemaElement ULearningAgentsObservationSchema::SpecifyRotationObservation(const FName Name)
+FLearningAgentsObservationSchemaElement ULearningAgentsObservations::SpecifyRotationObservation(ULearningAgentsObservationSchema* Schema, const FName Tag)
 {
-	return SpecifyContinuousObservation(6, Name);
+	return SpecifyContinuousObservation(Schema, 6, Tag);
 }
 
-FLearningAgentsObservationSchemaElement ULearningAgentsObservationSchema::SpecifyScaleObservation(const FName Name)
+FLearningAgentsObservationSchemaElement ULearningAgentsObservations::SpecifyScaleObservation(ULearningAgentsObservationSchema* Schema, const FName Tag)
 {
-	return SpecifyContinuousObservation(3, Name);
+	return SpecifyContinuousObservation(Schema, 3, Tag);
 }
 
-FLearningAgentsObservationSchemaElement ULearningAgentsObservationSchema::SpecifyTransformObservation(const FName Name)
+FLearningAgentsObservationSchemaElement ULearningAgentsObservations::SpecifyTransformObservation(ULearningAgentsObservationSchema* Schema, const FName Tag)
 {
-	return SpecifyStructObservationFromArrayViews(
+	return SpecifyStructObservationFromArrayViews(Schema,
 		{
 			TEXT("Location"),
 			TEXT("Rotation"),
 			TEXT("Scale")
 		},
 		{
-			SpecifyLocationObservation(),
-			SpecifyRotationObservation(),
-			SpecifyScaleObservation()
+			SpecifyLocationObservation(Schema),
+			SpecifyRotationObservation(Schema),
+			SpecifyScaleObservation(Schema)
 		}, 
-		Name);
+		Tag);
 }
 
-FLearningAgentsObservationSchemaElement ULearningAgentsObservationSchema::SpecifyAngleObservation(const FName Name)
+FLearningAgentsObservationSchemaElement ULearningAgentsObservations::SpecifyAngleObservation(ULearningAgentsObservationSchema* Schema, const FName Tag)
 {
-	return SpecifyContinuousObservation(2, Name);
+	return SpecifyContinuousObservation(Schema, 2, Tag);
 }
 
-FLearningAgentsObservationSchemaElement ULearningAgentsObservationSchema::SpecifyVelocityObservation(const FName Name)
+FLearningAgentsObservationSchemaElement ULearningAgentsObservations::SpecifyVelocityObservation(ULearningAgentsObservationSchema* Schema, const FName Tag)
 {
-	return SpecifyContinuousObservation(3, Name);
+	return SpecifyContinuousObservation(Schema, 3, Tag);
 }
 
-FLearningAgentsObservationSchemaElement ULearningAgentsObservationSchema::SpecifyDirectionObservation(const FName Name)
+FLearningAgentsObservationSchemaElement ULearningAgentsObservations::SpecifyDirectionObservation(ULearningAgentsObservationSchema* Schema, const FName Tag)
 {
-	return SpecifyContinuousObservation(3, Name);
+	return SpecifyContinuousObservation(Schema, 3, Tag);
 }
 
-FLearningAgentsObservationSchemaElement ULearningAgentsObservationSchema::SpecifyLocationAlongSplineObservation(const FName Name)
+FLearningAgentsObservationSchemaElement ULearningAgentsObservations::SpecifyLocationAlongSplineObservation(ULearningAgentsObservationSchema* Schema, const FName Tag)
 {
-	return SpecifyLocationObservation(Name);
+	return SpecifyLocationObservation(Schema, Tag);
 }
 
-FLearningAgentsObservationSchemaElement ULearningAgentsObservationSchema::SpecifyProportionAlongSplineObservation(const FName Name)
+FLearningAgentsObservationSchemaElement ULearningAgentsObservations::SpecifyProportionAlongSplineObservation(ULearningAgentsObservationSchema* Schema, const FName Tag)
 {
-	return SpecifyExclusiveUnionObservationFromArrayViews(
+	return SpecifyExclusiveUnionObservationFromArrayViews(Schema,
 		{
 			TEXT("Angle"),
 			TEXT("Proportion")
 		},
 		{
-			SpecifyAngleObservation(),
-			SpecifyFloatObservation()
+			SpecifyAngleObservation(Schema),
+			SpecifyFloatObservation(Schema)
 		},
 		8,
-		Name);
+		Tag);
 }
 
-FLearningAgentsObservationSchemaElement ULearningAgentsObservationSchema::SpecifyDirectionAlongSplineObservation(const FName Name)
+FLearningAgentsObservationSchemaElement ULearningAgentsObservations::SpecifyDirectionAlongSplineObservation(ULearningAgentsObservationSchema* Schema, const FName Tag)
 {
-	return SpecifyDirectionObservation(Name);
+	return SpecifyDirectionObservation(Schema, Tag);
 }
 
-FLearningAgentsObservationSchemaElement ULearningAgentsObservationSchema::SpecifyPropertiesAlongSplineObservation(const FName Name)
+FLearningAgentsObservationSchemaElement ULearningAgentsObservations::SpecifyPropertiesAlongSplineObservation(ULearningAgentsObservationSchema* Schema, const FName Tag)
 {
-	return SpecifyStructObservationFromArrayViews(
+	return SpecifyStructObservationFromArrayViews(Schema,
 		{
 			TEXT("Location"),
 			TEXT("Proportion"),
 			TEXT("Direction")
 		},
 		{
-			SpecifyLocationAlongSplineObservation(),
-			SpecifyProportionAlongSplineObservation(),
-			SpecifyDirectionAlongSplineObservation(),
+			SpecifyLocationAlongSplineObservation(Schema),
+			SpecifyProportionAlongSplineObservation(Schema),
+			SpecifyDirectionAlongSplineObservation(Schema),
 		},
-		Name);
+		Tag);
 }
 
-FLearningAgentsObservationSchemaElement ULearningAgentsObservationSchema::SpecifyProportionAlongRayObservation(const FName Name)
+FLearningAgentsObservationSchemaElement ULearningAgentsObservations::SpecifyProportionAlongRayObservation(ULearningAgentsObservationSchema* Schema, const FName Tag)
 {
-	return SpecifyFloatObservation(Name);
+	return SpecifyFloatObservation(Schema, Tag);
 }
 
-FLearningAgentsObservationSchemaElement ULearningAgentsObservationSchema::SpecifyProportionAlongRaysObservation(const int32 Num, const FName Name)
+FLearningAgentsObservationSchemaElement ULearningAgentsObservations::SpecifyProportionAlongRaysObservation(ULearningAgentsObservationSchema* Schema, const int32 Num, const FName Tag)
 {
-	return SpecifyStaticArrayObservation(SpecifyProportionAlongRayObservation(), Num, Name);
+	return SpecifyStaticArrayObservation(Schema, SpecifyProportionAlongRayObservation(Schema), Num, Tag);
 }
 
-const UE::Learning::Observation::FObject& ULearningAgentsObservationObject::GetObservationObject() const
+void ULearningAgentsObservations::LogObservation(const ULearningAgentsObservationObject* Object, const FLearningAgentsObservationObjectElement Element)
 {
-	return ObservationObject;
-}
+	if (!Object)
+	{
+		UE_LOG(LogLearning, Error, TEXT("LogObservation: Object is nullptr."));
+		return;
+	}
 
-UE::Learning::Observation::FObject& ULearningAgentsObservationObject::GetObservationObject()
-{
-	return ObservationObject;
+	UE::Learning::Agents::Observation::Private::LogObservation(Object->ObservationObject, Element.ObjectElement, TEXT(""), TEXT(""));
 }
 
-void ULearningAgentsObservationObject::LogObservation(const FLearningAgentsObservationObjectElement Element)
+FLearningAgentsObservationObjectElement ULearningAgentsObservations::MakeNullObservation(ULearningAgentsObservationObject* Object, const FName Tag)
 {
-	UE::Learning::Agents::Observation::Private::LogObservation(GetObservationObject(), Element.ObjectElement, TEXT(""), TEXT(""), GetName());
-}
+	if (!Object)
+	{
+		UE_LOG(LogLearning, Error, TEXT("MakeNullObservation: Object is nullptr."));
+		return FLearningAgentsObservationObjectElement();
+	}
 
-FLearningAgentsObservationObjectElement ULearningAgentsObservationObject::MakeNullObservation(const FName Name)
-{
-	return { ObservationObject.CreateNull(Name) };
+	return { Object->ObservationObject.CreateNull(Tag) };
 }
 
-FLearningAgentsObservationObjectElement ULearningAgentsObservationObject::MakeContinuousObservation(const TArray<float>& Values, const FName Name)
+FLearningAgentsObservationObjectElement ULearningAgentsObservations::MakeContinuousObservation(ULearningAgentsObservationObject* Object, const TArray<float>& Values, const FName Tag)
 {
-	return MakeContinuousObservationFromArrayView(Values, Name);
+	return MakeContinuousObservationFromArrayView(Object, Values, Tag);
 }
 
-FLearningAgentsObservationObjectElement ULearningAgentsObservationObject::MakeContinuousObservationFromArrayView(const TArrayView<const float> Values, const FName Name)
+FLearningAgentsObservationObjectElement ULearningAgentsObservations::MakeContinuousObservationFromArrayView(ULearningAgentsObservationObject* Object, const TArrayView<const float> Values, const FName Tag)
 {
+	if (!Object)
+	{
+		UE_LOG(LogLearning, Error, TEXT("MakeContinuousObservationFromArrayView: Object is nullptr."));
+		return FLearningAgentsObservationObjectElement();
+	}
+
 	if (Values.Num() == 0)
 	{
-		UE_LOG(LogLearning, Warning, TEXT("%s: Creating zero-sized Continuous Observation."), *GetName());
+		UE_LOG(LogLearning, Warning, TEXT("MakeContinuousObservationFromArrayView: Creating zero-sized Continuous Observation."));
 	}
 
-	return { ObservationObject.CreateContinuous({ Values }, Name) };
+	return { Object->ObservationObject.CreateContinuous({ Values }, Tag) };
 }
 
-FLearningAgentsObservationObjectElement ULearningAgentsObservationObject::MakeExclusiveDiscreteObservation(const int32 DiscreteIndex, const int32 Size, const FName Name)
+FLearningAgentsObservationObjectElement ULearningAgentsObservations::MakeExclusiveDiscreteObservation(ULearningAgentsObservationObject* Object, const int32 DiscreteIndex, const int32 Size, const FName Tag)
 {
 	if (DiscreteIndex < 0 || DiscreteIndex >= Size)
 	{
-		UE_LOG(LogLearning, Error, TEXT("%s: Discrete index out of range: Got %i, expected <= %i."), *GetName(), DiscreteIndex, Size);
+		UE_LOG(LogLearning, Error, TEXT("MakeExclusiveDiscreteObservation: Discrete index out of range: Got %i, expected <= %i."), DiscreteIndex, Size);
 		return FLearningAgentsObservationObjectElement();
 	}
 
 	TArray<float, TInlineAllocator<32>> Values;
 	Values.Init(0.0f, Size);
 	Values[DiscreteIndex] = 1.0f;
-	return MakeContinuousObservationFromArrayView(Values, Name);
+	return MakeContinuousObservationFromArrayView(Object, Values, Tag);
 }
 
-FLearningAgentsObservationObjectElement ULearningAgentsObservationObject::MakeInclusiveDiscreteObservation(const TArray<int32>& DiscreteIndices, const int32 Size, const FName Name)
+FLearningAgentsObservationObjectElement ULearningAgentsObservations::MakeInclusiveDiscreteObservation(ULearningAgentsObservationObject* Object, const TArray<int32>& DiscreteIndices, const int32 Size, const FName Tag)
 {
-	return MakeInclusiveDiscreteObservationFromArrayView(DiscreteIndices, Size, Name);
+	return MakeInclusiveDiscreteObservationFromArrayView(Object, DiscreteIndices, Size, Tag);
 }
 
-FLearningAgentsObservationObjectElement ULearningAgentsObservationObject::MakeInclusiveDiscreteObservationFromArrayView(const TArrayView<const int32> DiscreteIndices, const int32 Size, const FName Name)
+FLearningAgentsObservationObjectElement ULearningAgentsObservations::MakeInclusiveDiscreteObservationFromArrayView(ULearningAgentsObservationObject* Object, const TArrayView<const int32> DiscreteIndices, const int32 Size, const FName Tag)
 {
 	if (UE::Learning::Agents::Observation::Private::ContainsDuplicates(DiscreteIndices))
 	{
-		UE_LOG(LogLearning, Error, TEXT("%s: Indices contain duplicates."), *GetName());
+		UE_LOG(LogLearning, Error, TEXT("MakeInclusiveDiscreteObservationFromArrayView: Indices contain duplicates."));
 		return FLearningAgentsObservationObjectElement();
 	}
 
@@ -1115,21 +1161,21 @@ FLearningAgentsObservationObjectElement ULearningAgentsObservationObject::MakeIn
 	{
 		if (DiscreteIndices[Idx] < 0 || DiscreteIndices[Idx] >= Size)
 		{
-			UE_LOG(LogLearning, Error, TEXT("%s: Discrete index out of range: Got %i, expected <= %i."), *GetName(), DiscreteIndices[Idx], Size);
+			UE_LOG(LogLearning, Error, TEXT("MakeInclusiveDiscreteObservationFromArrayView: Discrete index out of range: Got %i, expected <= %i."), DiscreteIndices[Idx], Size);
 			return FLearningAgentsObservationObjectElement();
 		}
 
 		Values[DiscreteIndices[Idx]] = 1.0f;
 	}
 
-	return MakeContinuousObservationFromArrayView(Values, Name);
+	return MakeContinuousObservationFromArrayView(Object, Values, Tag);
 }
 
-FLearningAgentsObservationObjectElement ULearningAgentsObservationObject::MakeIndexObservation(const int32 Index, const int32 Size, const FName Name)
+FLearningAgentsObservationObjectElement ULearningAgentsObservations::MakeIndexObservation(ULearningAgentsObservationObject* Object, const int32 Index, const int32 Size, const FName Tag)
 {
 	if (Index < 0 || Index >= Size)
 	{
-		UE_LOG(LogLearning, Error, TEXT("%s: Discrete index out of range: Got %i, expected <= %i."), *GetName(), Index, Size);
+		UE_LOG(LogLearning, Error, TEXT("MakeIndexObservation: Discrete index out of range: Got %i, expected <= %i."), Index, Size);
 		return FLearningAgentsObservationObjectElement();
 	}
 
@@ -1141,25 +1187,25 @@ FLearningAgentsObservationObjectElement ULearningAgentsObservationObject::MakeIn
 		Values[Idx] = 1.0f;
 	}
 
-	return MakeContinuousObservationFromArrayView(Values, Name);
+	return MakeContinuousObservationFromArrayView(Object, Values, Tag);
 }
 
-FLearningAgentsObservationObjectElement ULearningAgentsObservationObject::MakeCountObservation(const int32 Num, const int32 MaxNum, const FName Name)
+FLearningAgentsObservationObjectElement ULearningAgentsObservations::MakeCountObservation(ULearningAgentsObservationObject* Object, const int32 Num, const int32 MaxNum, const FName Tag)
 {
 	if (MaxNum == 0)
 	{
-		UE_LOG(LogLearning, Error, TEXT("%s: MaxNum must not be zero for Count Observation."), *GetName());
+		UE_LOG(LogLearning, Error, TEXT("MakeCountObservation: MaxNum must not be zero for Count Observation."));
 		return FLearningAgentsObservationObjectElement();
 	}
 
-	return MakeContinuousObservationFromArrayView({ (float)Num / (float)MaxNum }, Name);
+	return MakeContinuousObservationFromArrayView(Object, { (float)Num / (float)MaxNum }, Tag);
 }
 
-FLearningAgentsObservationObjectElement ULearningAgentsObservationObject::MakeStructObservation(const TMap<FName, FLearningAgentsObservationObjectElement>& Elements, const FName Name)
+FLearningAgentsObservationObjectElement ULearningAgentsObservations::MakeStructObservation(ULearningAgentsObservationObject* Object, const TMap<FName, FLearningAgentsObservationObjectElement>& Elements, const FName Tag)
 {
 	if (Elements.Num() == 0)
 	{
-		UE_LOG(LogLearning, Warning, TEXT("%s: Creating zero-sized Struct Observation."), *GetName());
+		UE_LOG(LogLearning, Warning, TEXT("MakeStructObservation: Creating zero-sized Struct Observation."));
 	}
 
 	const int32 SubElementNum = Elements.Num();
@@ -1175,30 +1221,36 @@ FLearningAgentsObservationObjectElement ULearningAgentsObservationObject::MakeSt
 		SubElements.Add(Element.Value);
 	}
 
-	return MakeStructObservationFromArrayViews(SubElementNames, SubElements, Name);
+	return MakeStructObservationFromArrayViews(Object, SubElementNames, SubElements, Tag);
 }
 
-FLearningAgentsObservationObjectElement ULearningAgentsObservationObject::MakeStructObservationFromArrays(const TArray<FName>& ElementNames, const TArray<FLearningAgentsObservationObjectElement>& Elements, const FName Name)
+FLearningAgentsObservationObjectElement ULearningAgentsObservations::MakeStructObservationFromArrays(ULearningAgentsObservationObject* Object, const TArray<FName>& ElementNames, const TArray<FLearningAgentsObservationObjectElement>& Elements, const FName Tag)
 {
-	return MakeStructObservationFromArrayViews(ElementNames, Elements, Name);
+	return MakeStructObservationFromArrayViews(Object, ElementNames, Elements, Tag);
 }
 
-FLearningAgentsObservationObjectElement ULearningAgentsObservationObject::MakeStructObservationFromArrayViews(const TArrayView<const FName> ElementNames, const TArrayView<const FLearningAgentsObservationObjectElement> Elements, const FName Name)
+FLearningAgentsObservationObjectElement ULearningAgentsObservations::MakeStructObservationFromArrayViews(ULearningAgentsObservationObject* Object, const TArrayView<const FName> ElementNames, const TArrayView<const FLearningAgentsObservationObjectElement> Elements, const FName Tag)
 {
+	if (!Object)
+	{
+		UE_LOG(LogLearning, Error, TEXT("MakeStructObservationFromArrayViews: Object is nullptr."));
+		return FLearningAgentsObservationObjectElement();
+	}
+
 	if (Elements.Num() == 0)
 	{
-		UE_LOG(LogLearning, Warning, TEXT("%s: Creating zero-sized Struct Observation."), *GetName());
+		UE_LOG(LogLearning, Warning, TEXT("MakeStructObservationFromArrayViews: Creating zero-sized Struct Observation."));
 	}
 
 	if (Elements.Num() != ElementNames.Num())
 	{
-		UE_LOG(LogLearning, Error, TEXT("%s: Number of elements (%i) must match number of names (%i)."), *GetName(), Elements.Num(), ElementNames.Num());
+		UE_LOG(LogLearning, Error, TEXT("MakeStructObservationFromArrayViews: Number of elements (%i) must match number of names (%i)."), Elements.Num(), ElementNames.Num());
 		return FLearningAgentsObservationObjectElement();
 	}
 
 	if (UE::Learning::Agents::Observation::Private::ContainsDuplicates(ElementNames))
 	{
-		UE_LOG(LogLearning, Error, TEXT("%s: Element Names contain duplicates."), *GetName());
+		UE_LOG(LogLearning, Error, TEXT("MakeStructObservationFromArrayViews: Element Names contain duplicates."));
 		return FLearningAgentsObservationObjectElement();
 	}
 
@@ -1207,30 +1259,36 @@ FLearningAgentsObservationObjectElement ULearningAgentsObservationObject::MakeSt
 
 	for (const FLearningAgentsObservationObjectElement& Element : Elements)
 	{
-		if (!ObservationObject.IsValid(Element.ObjectElement))
+		if (!Object->ObservationObject.IsValid(Element.ObjectElement))
 		{
-			UE_LOG(LogLearning, Error, TEXT("%s: Invalid Observation Object."), *GetName());
+			UE_LOG(LogLearning, Error, TEXT("MakeStructObservationFromArrayViews: Invalid Observation Object."));
 			return FLearningAgentsObservationObjectElement();
 		}
 
 		SubElements.Add(Element.ObjectElement);
 	}
 
-	return { ObservationObject.CreateAnd({ ElementNames, SubElements }, Name)};
+	return { Object->ObservationObject.CreateAnd({ ElementNames, SubElements }, Tag)};
 }
 
-FLearningAgentsObservationObjectElement ULearningAgentsObservationObject::MakeExclusiveUnionObservation(const FName ElementName, const FLearningAgentsObservationObjectElement Element, const FName Name)
+FLearningAgentsObservationObjectElement ULearningAgentsObservations::MakeExclusiveUnionObservation(ULearningAgentsObservationObject* Object, const FName ElementName, const FLearningAgentsObservationObjectElement Element, const FName Tag)
 {
-	if (!ObservationObject.IsValid(Element.ObjectElement))
+	if (!Object)
 	{
-		UE_LOG(LogLearning, Error, TEXT("%s: Invalid Observation Object."), *GetName());
+		UE_LOG(LogLearning, Error, TEXT("MakeExclusiveUnionObservation: Object is nullptr."));
 		return FLearningAgentsObservationObjectElement();
 	}
 
-	return { ObservationObject.CreateOrExclusive({ ElementName, Element.ObjectElement }, Name) };
+	if (!Object->ObservationObject.IsValid(Element.ObjectElement))
+	{
+		UE_LOG(LogLearning, Error, TEXT("MakeExclusiveUnionObservation: Invalid Observation Object."));
+		return FLearningAgentsObservationObjectElement();
+	}
+
+	return { Object->ObservationObject.CreateOrExclusive({ ElementName, Element.ObjectElement }, Tag) };
 }
 
-FLearningAgentsObservationObjectElement ULearningAgentsObservationObject::MakeInclusiveUnionObservation(const TMap<FName, FLearningAgentsObservationObjectElement>& Elements, const FName Name)
+FLearningAgentsObservationObjectElement ULearningAgentsObservations::MakeInclusiveUnionObservation(ULearningAgentsObservationObject* Object, const TMap<FName, FLearningAgentsObservationObjectElement>& Elements, const FName Tag)
 {
 	const int32 SubElementNum = Elements.Num();
 
@@ -1245,25 +1303,31 @@ FLearningAgentsObservationObjectElement ULearningAgentsObservationObject::MakeIn
 		SubElements.Add(Element.Value);
 	}
 
-	return MakeInclusiveUnionObservationFromArrayViews(SubElementNames, SubElements, Name);
+	return MakeInclusiveUnionObservationFromArrayViews(Object, SubElementNames, SubElements, Tag);
 }
 
-FLearningAgentsObservationObjectElement ULearningAgentsObservationObject::MakeInclusiveUnionObservationFromArrays(const TArray<FName>& ElementNames, const TArray<FLearningAgentsObservationObjectElement>& Elements, const FName Name)
+FLearningAgentsObservationObjectElement ULearningAgentsObservations::MakeInclusiveUnionObservationFromArrays(ULearningAgentsObservationObject* Object, const TArray<FName>& ElementNames, const TArray<FLearningAgentsObservationObjectElement>& Elements, const FName Tag)
 {
-	return MakeInclusiveUnionObservationFromArrayViews(ElementNames, Elements, Name);
+	return MakeInclusiveUnionObservationFromArrayViews(Object, ElementNames, Elements, Tag);
 }
 
-FLearningAgentsObservationObjectElement ULearningAgentsObservationObject::MakeInclusiveUnionObservationFromArrayViews(const TArrayView<const FName> ElementNames, const TArrayView<const FLearningAgentsObservationObjectElement> Elements, const FName Name)
+FLearningAgentsObservationObjectElement ULearningAgentsObservations::MakeInclusiveUnionObservationFromArrayViews(ULearningAgentsObservationObject* Object, const TArrayView<const FName> ElementNames, const TArrayView<const FLearningAgentsObservationObjectElement> Elements, const FName Tag)
 {
+	if (!Object)
+	{
+		UE_LOG(LogLearning, Error, TEXT("MakeInclusiveUnionObservationFromArrayViews: Object is nullptr."));
+		return FLearningAgentsObservationObjectElement();
+	}
+
 	if (Elements.Num() != ElementNames.Num())
 	{
-		UE_LOG(LogLearning, Error, TEXT("%s: Number of elements (%i) must match number of names (%i)."), *GetName(), Elements.Num(), ElementNames.Num());
+		UE_LOG(LogLearning, Error, TEXT("MakeInclusiveUnionObservationFromArrayViews: Number of elements (%i) must match number of names (%i)."), Elements.Num(), ElementNames.Num());
 		return FLearningAgentsObservationObjectElement();
 	}
 
 	if (UE::Learning::Agents::Observation::Private::ContainsDuplicates(ElementNames))
 	{
-		UE_LOG(LogLearning, Error, TEXT("%s: Element Names contain duplicates."), *GetName());
+		UE_LOG(LogLearning, Error, TEXT("MakeInclusiveUnionObservationFromArrayViews: Element Names contain duplicates."));
 		return FLearningAgentsObservationObjectElement();
 	}
 
@@ -1272,28 +1336,34 @@ FLearningAgentsObservationObjectElement ULearningAgentsObservationObject::MakeIn
 
 	for (const FLearningAgentsObservationObjectElement& Element : Elements)
 	{
-		if (!ObservationObject.IsValid(Element.ObjectElement))
+		if (!Object->ObservationObject.IsValid(Element.ObjectElement))
 		{
-			UE_LOG(LogLearning, Error, TEXT("%s: Invalid Observation Object."), *GetName());
+			UE_LOG(LogLearning, Error, TEXT("MakeInclusiveUnionObservationFromArrayViews: Invalid Observation Object."));
 			return FLearningAgentsObservationObjectElement();
 		}
 
 		SubElements.Add(Element.ObjectElement);
 	}
 
-	return { ObservationObject.CreateOrInclusive({ ElementNames, SubElements }, Name) };
+	return { Object->ObservationObject.CreateOrInclusive({ ElementNames, SubElements }, Tag) };
 }
 
-FLearningAgentsObservationObjectElement ULearningAgentsObservationObject::MakeStaticArrayObservation(const TArray<FLearningAgentsObservationObjectElement>& Elements, const FName Name)
+FLearningAgentsObservationObjectElement ULearningAgentsObservations::MakeStaticArrayObservation(ULearningAgentsObservationObject* Object, const TArray<FLearningAgentsObservationObjectElement>& Elements, const FName Tag)
 {
-	return MakeStaticArrayObservationFromArrayView(Elements, Name);
+	return MakeStaticArrayObservationFromArrayView(Object, Elements, Tag);
 }
 
-FLearningAgentsObservationObjectElement ULearningAgentsObservationObject::MakeStaticArrayObservationFromArrayView(const TArrayView<const FLearningAgentsObservationObjectElement> Elements, const FName Name)
+FLearningAgentsObservationObjectElement ULearningAgentsObservations::MakeStaticArrayObservationFromArrayView(ULearningAgentsObservationObject* Object, const TArrayView<const FLearningAgentsObservationObjectElement> Elements, const FName Tag)
 {
+	if (!Object)
+	{
+		UE_LOG(LogLearning, Error, TEXT("MakeStaticArrayObservationFromArrayView: Object is nullptr."));
+		return FLearningAgentsObservationObjectElement();
+	}
+
 	if (Elements.Num() == 0)
 	{
-		UE_LOG(LogLearning, Warning, TEXT("%s: Creating zero-sized Static Array Observation."), *GetName());
+		UE_LOG(LogLearning, Warning, TEXT("MakeStaticArrayObservationFromArrayView: Creating zero-sized Static Array Observation."));
 	}
 
 	TArray<UE::Learning::Observation::FObjectElement, TInlineAllocator<16>> SubElements;
@@ -1301,73 +1371,91 @@ FLearningAgentsObservationObjectElement ULearningAgentsObservationObject::MakeSt
 
 	for (const FLearningAgentsObservationObjectElement& Element : Elements)
 	{
-		if (!ObservationObject.IsValid(Element.ObjectElement))
+		if (!Object->ObservationObject.IsValid(Element.ObjectElement))
 		{
-			UE_LOG(LogLearning, Error, TEXT("%s: Invalid Observation Object."), *GetName());
+			UE_LOG(LogLearning, Error, TEXT("MakeStaticArrayObservationFromArrayView: Invalid Observation Object."));
 			return FLearningAgentsObservationObjectElement();
 		}
 
 		SubElements.Add(Element.ObjectElement);
 	}
 
-	return { ObservationObject.CreateArray({ SubElements }, Name) };
+	return { Object->ObservationObject.CreateArray({ SubElements }, Tag) };
 }
 
-FLearningAgentsObservationObjectElement ULearningAgentsObservationObject::MakeSetObservation(const TSet<FLearningAgentsObservationObjectElement>& Elements, const FName Name)
+FLearningAgentsObservationObjectElement ULearningAgentsObservations::MakeSetObservation(ULearningAgentsObservationObject* Object, const TSet<FLearningAgentsObservationObjectElement>& Elements, const FName Tag)
 {
+	if (!Object)
+	{
+		UE_LOG(LogLearning, Error, TEXT("MakeSetObservation: Object is nullptr."));
+		return FLearningAgentsObservationObjectElement();
+	}
+
 	TArray<UE::Learning::Observation::FObjectElement, TInlineAllocator<16>> SubElements;
 	SubElements.Empty(Elements.Num());
 
 	for (const FLearningAgentsObservationObjectElement& Element : Elements)
 	{
-		if (!ObservationObject.IsValid(Element.ObjectElement))
+		if (!Object->ObservationObject.IsValid(Element.ObjectElement))
 		{
-			UE_LOG(LogLearning, Error, TEXT("%s: Invalid Observation Object."), *GetName());
+			UE_LOG(LogLearning, Error, TEXT("MakeSetObservation: Invalid Observation Object."));
 			return FLearningAgentsObservationObjectElement();
 		}
 
 		SubElements.Add(Element.ObjectElement);
 	}
 
-	return { ObservationObject.CreateSet({ SubElements }, Name) };
+	return { Object->ObservationObject.CreateSet({ SubElements }, Tag) };
 }
 
-FLearningAgentsObservationObjectElement ULearningAgentsObservationObject::MakeSetObservationFromArray(const TArray<FLearningAgentsObservationObjectElement>& Elements, const FName Name)
+FLearningAgentsObservationObjectElement ULearningAgentsObservations::MakeSetObservationFromArray(ULearningAgentsObservationObject* Object, const TArray<FLearningAgentsObservationObjectElement>& Elements, const FName Tag)
 {
-	return MakeSetObservationFromArrayView(Elements, Name);
+	return MakeSetObservationFromArrayView(Object, Elements, Tag);
 }
 
-FLearningAgentsObservationObjectElement ULearningAgentsObservationObject::MakeSetObservationFromArrayView(const TArrayView<const FLearningAgentsObservationObjectElement> Elements, const FName Name)
+FLearningAgentsObservationObjectElement ULearningAgentsObservations::MakeSetObservationFromArrayView(ULearningAgentsObservationObject* Object, const TArrayView<const FLearningAgentsObservationObjectElement> Elements, const FName Tag)
 {
+	if (!Object)
+	{
+		UE_LOG(LogLearning, Error, TEXT("MakeSetObservationFromArrayView: Object is nullptr."));
+		return FLearningAgentsObservationObjectElement();
+	}
+
 	TArray<UE::Learning::Observation::FObjectElement, TInlineAllocator<16>> SubElements;
 	SubElements.Empty(Elements.Num());
 
 	for (const FLearningAgentsObservationObjectElement& Element : Elements)
 	{
-		if (!ObservationObject.IsValid(Element.ObjectElement))
+		if (!Object->ObservationObject.IsValid(Element.ObjectElement))
 		{
-			UE_LOG(LogLearning, Error, TEXT("%s: Invalid Observation Object."), *GetName());
+			UE_LOG(LogLearning, Error, TEXT("MakeSetObservationFromArrayView: Invalid Observation Object."));
 			return FLearningAgentsObservationObjectElement();
 		}
 
 		SubElements.Add(Element.ObjectElement);
 	}
 
-	return { ObservationObject.CreateSet({ SubElements }, Name) };
+	return { Object->ObservationObject.CreateSet({ SubElements }, Tag) };
 }
 
-FLearningAgentsObservationObjectElement ULearningAgentsObservationObject::MakePairObservation(const FLearningAgentsObservationObjectElement Key, const FLearningAgentsObservationObjectElement Value, const FName Name)
+FLearningAgentsObservationObjectElement ULearningAgentsObservations::MakePairObservation(ULearningAgentsObservationObject* Object, const FLearningAgentsObservationObjectElement Key, const FLearningAgentsObservationObjectElement Value, const FName Tag)
 {
-	return MakeStructObservationFromArrayViews({ TEXT("Key"), TEXT("Value") }, { Key, Value }, Name);
+	return MakeStructObservationFromArrayViews(Object, { TEXT("Key"), TEXT("Value") }, { Key, Value }, Tag);
 }
 
-FLearningAgentsObservationObjectElement ULearningAgentsObservationObject::MakeArrayObservation(const TArray<FLearningAgentsObservationObjectElement>& Elements, const FName Name)
+FLearningAgentsObservationObjectElement ULearningAgentsObservations::MakeArrayObservation(ULearningAgentsObservationObject* Object, const TArray<FLearningAgentsObservationObjectElement>& Elements, const FName Tag)
 {
-	return MakeArrayObservationFromArrayView(Elements, Name);
+	return MakeArrayObservationFromArrayView(Object, Elements, Tag);
 }
 
-FLearningAgentsObservationObjectElement ULearningAgentsObservationObject::MakeArrayObservationFromArrayView(const TArrayView<const FLearningAgentsObservationObjectElement> Elements, const FName Name)
+FLearningAgentsObservationObjectElement ULearningAgentsObservations::MakeArrayObservationFromArrayView(ULearningAgentsObservationObject* Object, const TArrayView<const FLearningAgentsObservationObjectElement> Elements, const FName Tag)
 {
+	if (!Object)
+	{
+		UE_LOG(LogLearning, Error, TEXT("MakeArrayObservationFromArrayView: Object is nullptr."));
+		return FLearningAgentsObservationObjectElement();
+	}
+
 	TArray<FLearningAgentsObservationObjectElement, TInlineAllocator<16>> SubElements;
 	SubElements.Empty(Elements.Num());
 
@@ -1375,47 +1463,59 @@ FLearningAgentsObservationObjectElement ULearningAgentsObservationObject::MakeAr
 	{
 		const FLearningAgentsObservationObjectElement Element = Elements[ElementIdx];
 
-		if (!ObservationObject.IsValid(Element.ObjectElement))
+		if (!Object->ObservationObject.IsValid(Element.ObjectElement))
 		{
-			UE_LOG(LogLearning, Error, TEXT("%s: Invalid Observation Object."), *GetName());
+			UE_LOG(LogLearning, Error, TEXT("MakeArrayObservationFromArrayView: Invalid Observation Object."));
 			return FLearningAgentsObservationObjectElement();
 		}
 
-		SubElements.Add(MakePairObservation(MakeIndexObservation(ElementIdx, Elements.Num()), Element));
+		SubElements.Add(MakePairObservation(Object, MakeIndexObservation(Object, ElementIdx, Elements.Num()), Element));
 	}
 
-	return MakeSetObservationFromArrayView(SubElements);
+	return MakeSetObservationFromArrayView(Object, SubElements);
 }
 
-FLearningAgentsObservationObjectElement ULearningAgentsObservationObject::MakeMapObservation(const TMap<FLearningAgentsObservationObjectElement, FLearningAgentsObservationObjectElement>& Map, const FName Name)
+FLearningAgentsObservationObjectElement ULearningAgentsObservations::MakeMapObservation(ULearningAgentsObservationObject* Object, const TMap<FLearningAgentsObservationObjectElement, FLearningAgentsObservationObjectElement>& Map, const FName Tag)
 {
+	if (!Object)
+	{
+		UE_LOG(LogLearning, Error, TEXT("MakeMapObservation: Object is nullptr."));
+		return FLearningAgentsObservationObjectElement();
+	}
+
 	TArray<FLearningAgentsObservationObjectElement, TInlineAllocator<16>> SubElements;
 	SubElements.Empty(Map.Num());
 
 	for (TPair<FLearningAgentsObservationObjectElement, FLearningAgentsObservationObjectElement> Item : Map)
 	{
-		if (!ObservationObject.IsValid(Item.Key.ObjectElement) || !ObservationObject.IsValid(Item.Value.ObjectElement))
+		if (!Object->ObservationObject.IsValid(Item.Key.ObjectElement) || !Object->ObservationObject.IsValid(Item.Value.ObjectElement))
 		{
-			UE_LOG(LogLearning, Error, TEXT("%s: Invalid Observation Object."), *GetName());
+			UE_LOG(LogLearning, Error, TEXT("MakeMapObservation: Invalid Observation Object."));
 			return FLearningAgentsObservationObjectElement();
 		}
 
-		SubElements.Add(MakePairObservation(Item.Key, Item.Value));
+		SubElements.Add(MakePairObservation(Object, Item.Key, Item.Value));
 	}
 
-	return MakeSetObservationFromArrayView(SubElements);
+	return MakeSetObservationFromArrayView(Object, SubElements);
 }
 
-FLearningAgentsObservationObjectElement ULearningAgentsObservationObject::MakeMapObservationFromArrays(const TArray<FLearningAgentsObservationObjectElement>& Keys, const TArray<FLearningAgentsObservationObjectElement>& Values, const FName Name)
+FLearningAgentsObservationObjectElement ULearningAgentsObservations::MakeMapObservationFromArrays(ULearningAgentsObservationObject* Object, const TArray<FLearningAgentsObservationObjectElement>& Keys, const TArray<FLearningAgentsObservationObjectElement>& Values, const FName Tag)
 {
-	return MakeMapObservationFromArrayViews(Keys, Values, Name);
+	return MakeMapObservationFromArrayViews(Object, Keys, Values, Tag);
 }
 
-FLearningAgentsObservationObjectElement ULearningAgentsObservationObject::MakeMapObservationFromArrayViews(const TArrayView<const FLearningAgentsObservationObjectElement> Keys, const TArrayView<const FLearningAgentsObservationObjectElement> Values, const FName Name)
+FLearningAgentsObservationObjectElement ULearningAgentsObservations::MakeMapObservationFromArrayViews(ULearningAgentsObservationObject* Object, const TArrayView<const FLearningAgentsObservationObjectElement> Keys, const TArrayView<const FLearningAgentsObservationObjectElement> Values, const FName Tag)
 {
+	if (!Object)
+	{
+		UE_LOG(LogLearning, Error, TEXT("MakeMapObservationFromArrayViews: Object is nullptr."));
+		return FLearningAgentsObservationObjectElement();
+	}
+
 	if (Keys.Num() != Values.Num())
 	{
-		UE_LOG(LogLearning, Error, TEXT("%s: Number of keys (%i) must match number of values (%i)."), *GetName(), Keys.Num(), Values.Num());
+		UE_LOG(LogLearning, Error, TEXT("MakeMapObservationFromArrayViews: Number of keys (%i) must match number of values (%i)."), Keys.Num(), Values.Num());
 		return FLearningAgentsObservationObjectElement();
 	}
 
@@ -1424,23 +1524,24 @@ FLearningAgentsObservationObjectElement ULearningAgentsObservationObject::MakeMa
 
 	for (int32 ElementIdx = 0; ElementIdx < Keys.Num(); ElementIdx++)
 	{
-		if(!ObservationObject.IsValid(Keys[ElementIdx].ObjectElement) || !ObservationObject.IsValid(Values[ElementIdx].ObjectElement))
+		if (!Object->ObservationObject.IsValid(Keys[ElementIdx].ObjectElement) || 
+			!Object->ObservationObject.IsValid(Values[ElementIdx].ObjectElement))
 		{
-			UE_LOG(LogLearning, Error, TEXT("%s: Invalid Observation Object."), *GetName());
+			UE_LOG(LogLearning, Error, TEXT("MakeMapObservationFromArrayViews: Invalid Observation Object."));
 			return FLearningAgentsObservationObjectElement();
 		}
 
-		SubElements.Add(MakePairObservation(Keys[ElementIdx], Values[ElementIdx]));
+		SubElements.Add(MakePairObservation(Object, Keys[ElementIdx], Values[ElementIdx]));
 	}
 
-	return MakeSetObservationFromArrayView(SubElements);
+	return MakeSetObservationFromArrayView(Object, SubElements);
 }
 
-FLearningAgentsObservationObjectElement ULearningAgentsObservationObject::MakeEnumObservation(const UEnum* Enum, const uint8 EnumValue, const FName Name)
+FLearningAgentsObservationObjectElement ULearningAgentsObservations::MakeEnumObservation(ULearningAgentsObservationObject* Object, const UEnum* Enum, const uint8 EnumValue, const FName Tag)
 {
 	if (!Enum)
 	{
-		UE_LOG(LogLearning, Error, TEXT("%s: Enum is nullptr."), *GetName());
+		UE_LOG(LogLearning, Error, TEXT("MakeEnumObservation: Enum is nullptr."));
 		return FLearningAgentsObservationObjectElement();
 	}
 
@@ -1448,7 +1549,7 @@ FLearningAgentsObservationObjectElement ULearningAgentsObservationObject::MakeEn
 
 	if (EnumValueIndex == INDEX_NONE || EnumValueIndex < 0 || EnumValueIndex >= Enum->NumEnums() - 1)
 	{
-		UE_LOG(LogLearning, Error, TEXT("%s: EnumValue %i not valid for Enum '%s'."), *GetName(), EnumValue , *Enum->GetName());
+		UE_LOG(LogLearning, Error, TEXT("MakeEnumObservation: EnumValue %i not valid for Enum '%s'."), EnumValue , *Enum->GetName());
 		return FLearningAgentsObservationObjectElement();
 	}
 
@@ -1456,20 +1557,20 @@ FLearningAgentsObservationObjectElement ULearningAgentsObservationObject::MakeEn
 	OneHot.Init(0.0f, Enum->NumEnums() - 1);
 	OneHot[EnumValueIndex] = 1.0f;
 
-	return MakeContinuousObservationFromArrayView(OneHot, Name);
+	return MakeContinuousObservationFromArrayView(Object, OneHot, Tag);
 }
 
-FLearningAgentsObservationObjectElement ULearningAgentsObservationObject::MakeBitmaskObservation(const UEnum* Enum, const int32 BitmaskValue, const FName Name)
+FLearningAgentsObservationObjectElement ULearningAgentsObservations::MakeBitmaskObservation(ULearningAgentsObservationObject* Object, const UEnum* Enum, const int32 BitmaskValue, const FName Tag)
 {
 	if (!Enum)
 	{
-		UE_LOG(LogLearning, Error, TEXT("%s: Enum is nullptr."), *GetName());
+		UE_LOG(LogLearning, Error, TEXT("MakeBitmaskObservation: Enum is nullptr."));
 		return FLearningAgentsObservationObjectElement();
 	}
 
 	if (Enum->NumEnums() - 1 > 32)
 	{
-		UE_LOG(LogLearning, Error, TEXT("%s: Too many values in Enum to use as Bitmask (%i)."), *GetName(), Enum->NumEnums() - 1);
+		UE_LOG(LogLearning, Error, TEXT("MakeBitmaskObservation: Too many values in Enum to use as Bitmask (%i)."), Enum->NumEnums() - 1);
 		return FLearningAgentsObservationObjectElement();
 	}
 
@@ -1484,56 +1585,64 @@ FLearningAgentsObservationObjectElement ULearningAgentsObservationObject::MakeBi
 		}
 	}
 
-	return MakeContinuousObservationFromArrayView(OneHot, Name);
+	return MakeContinuousObservationFromArrayView(Object, OneHot, Tag);
 }
 
-FLearningAgentsObservationObjectElement ULearningAgentsObservationObject::MakeOptionalObservation(const FLearningAgentsObservationObjectElement Element, const ELearningAgentsOptionalObservation Option, const FName Name)
+FLearningAgentsObservationObjectElement ULearningAgentsObservations::MakeOptionalObservation(ULearningAgentsObservationObject* Object, const FLearningAgentsObservationObjectElement Element, const ELearningAgentsOptionalObservation Option, const FName Tag)
 {
 	return MakeExclusiveUnionObservation(
+		Object,
 		Option == ELearningAgentsOptionalObservation::Null ? TEXT("Null") : TEXT("Valid"),
-		Option == ELearningAgentsOptionalObservation::Null ? MakeNullObservation() : Element,
-		Name);
+		Option == ELearningAgentsOptionalObservation::Null ? MakeNullObservation(Object) : Element,
+		Tag);
 }
 
-FLearningAgentsObservationObjectElement ULearningAgentsObservationObject::MakeOptionalNullObservation(const FName Name)
+FLearningAgentsObservationObjectElement ULearningAgentsObservations::MakeOptionalNullObservation(ULearningAgentsObservationObject* Object, const FName Tag)
 {
-	return MakeExclusiveUnionObservation(TEXT("Null"), MakeNullObservation(), Name);
+	return MakeExclusiveUnionObservation(Object, TEXT("Null"), MakeNullObservation(Object), Tag);
 }
 
-FLearningAgentsObservationObjectElement ULearningAgentsObservationObject::MakeOptionalValidObservation(const FLearningAgentsObservationObjectElement Element, const FName Name)
+FLearningAgentsObservationObjectElement ULearningAgentsObservations::MakeOptionalValidObservation(ULearningAgentsObservationObject* Object, const FLearningAgentsObservationObjectElement Element, const FName Tag)
 {
-	return MakeExclusiveUnionObservation(TEXT("Valid"), Element, Name);
+	return MakeExclusiveUnionObservation(Object, TEXT("Valid"), Element, Tag);
 }
 
-FLearningAgentsObservationObjectElement ULearningAgentsObservationObject::MakeEitherObservation(const FLearningAgentsObservationObjectElement Element, const ELearningAgentsEitherObservation Either, const FName Name)
+FLearningAgentsObservationObjectElement ULearningAgentsObservations::MakeEitherObservation(ULearningAgentsObservationObject* Object, const FLearningAgentsObservationObjectElement Element, const ELearningAgentsEitherObservation Either, const FName Tag)
 {
-	return MakeExclusiveUnionObservation(Either == ELearningAgentsEitherObservation::A ? TEXT("A") : TEXT("B"), Element, Name);
+	return MakeExclusiveUnionObservation(Object, Either == ELearningAgentsEitherObservation::A ? TEXT("A") : TEXT("B"), Element, Tag);
 }
 
-FLearningAgentsObservationObjectElement ULearningAgentsObservationObject::MakeEitherAObservation(const FLearningAgentsObservationObjectElement A, const FName Name)
+FLearningAgentsObservationObjectElement ULearningAgentsObservations::MakeEitherAObservation(ULearningAgentsObservationObject* Object, const FLearningAgentsObservationObjectElement A, const FName Tag)
 {
-	return MakeExclusiveUnionObservation(TEXT("A"), A, Name);
+	return MakeExclusiveUnionObservation(Object, TEXT("A"), A, Tag);
 }
 
-FLearningAgentsObservationObjectElement ULearningAgentsObservationObject::MakeEitherBObservation(const FLearningAgentsObservationObjectElement B, const FName Name)
+FLearningAgentsObservationObjectElement ULearningAgentsObservations::MakeEitherBObservation(ULearningAgentsObservationObject* Object, const FLearningAgentsObservationObjectElement B, const FName Tag)
 {
-	return MakeExclusiveUnionObservation(TEXT("B"), B, Name);
+	return MakeExclusiveUnionObservation(Object, TEXT("B"), B, Tag);
 }
 
-FLearningAgentsObservationObjectElement ULearningAgentsObservationObject::MakeEncodingObservation(const FLearningAgentsObservationObjectElement Element, const FName Name)
+FLearningAgentsObservationObjectElement ULearningAgentsObservations::MakeEncodingObservation(ULearningAgentsObservationObject* Object, const FLearningAgentsObservationObjectElement Element, const FName Tag)
 {
-	if (!ObservationObject.IsValid(Element.ObjectElement))
+	if (!Object)
+	{
+		UE_LOG(LogLearning, Error, TEXT("MakeEncodingObservation: Object is nullptr."));
+		return FLearningAgentsObservationObjectElement();
+	}
+
+	if (!Object->ObservationObject.IsValid(Element.ObjectElement))
 	{
-		UE_LOG(LogLearning, Error, TEXT("%s: Invalid Observation Object."), *GetName());
+		UE_LOG(LogLearning, Error, TEXT("MakeEncodingObservation: Invalid Observation Object."));
 		return FLearningAgentsObservationObjectElement();
 	}
 
-	return { ObservationObject.CreateEncoding({ Element.ObjectElement }, Name) };
+	return { Object->ObservationObject.CreateEncoding({ Element.ObjectElement }, Tag) };
 }
 
-FLearningAgentsObservationObjectElement ULearningAgentsObservationObject::MakeBoolObservation(
+FLearningAgentsObservationObjectElement ULearningAgentsObservations::MakeBoolObservation(
+	ULearningAgentsObservationObject* Object,
 	const bool bValue,
-	const FName Name,
+	const FName Tag,
 	const bool bVisualLoggerEnabled,
 	ULearningAgentsManagerListener* VisualLoggerListener,
 	const int32 VisualLoggerAgentId,
@@ -1545,26 +1654,27 @@ FLearningAgentsObservationObjectElement ULearningAgentsObservationObject::MakeBo
 #if UE_LEARNING_AGENTS_ENABLE_VISUAL_LOG
 	if (bVisualLoggerEnabled && VisualLoggerListener)
 	{
-		const ULearningAgentsVisualLoggerObject* VisualLoggerObject = VisualLoggerListener->GetOrAddVisualLoggerObject(Name);
+		const ULearningAgentsVisualLoggerObject* VisualLoggerObject = VisualLoggerListener->GetOrAddVisualLoggerObject(Tag);
 
 		UE_LEARNING_AGENTS_VLOG_STRING(VisualLoggerObject, LogLearning, Display, VisualLoggerLocation,
 			VisualLoggerColor.ToFColor(true),
-			TEXT("Listener: %s\nName: %s\nAgent Id: % 3i\nValue: [%s]\nEncoded: [% 6.2f]"),
+			TEXT("Listener: %s\nTag: %s\nAgent Id: % 3i\nValue: [%s]\nEncoded: [% 6.2f]"),
 			*VisualLoggerListener->GetName(),
-			*Name.ToString(),
+			*Tag.ToString(),
 			VisualLoggerAgentId,
 			bValue ? TEXT("true") : TEXT("false"),
 			EncodedBool);
 	}
 #endif
 
-	return MakeContinuousObservationFromArrayView({ EncodedBool }, Name);
+	return MakeContinuousObservationFromArrayView(Object, { EncodedBool }, Tag);
 }
 
-FLearningAgentsObservationObjectElement ULearningAgentsObservationObject::MakeFloatObservation(
+FLearningAgentsObservationObjectElement ULearningAgentsObservations::MakeFloatObservation(
+	ULearningAgentsObservationObject* Object,
 	const float Value,
 	const float FloatScale,
-	const FName Name,
+	const FName Tag,
 	const bool bVisualLoggerEnabled,
 	ULearningAgentsManagerListener* VisualLoggerListener,
 	const int32 VisualLoggerAgentId,
@@ -1576,13 +1686,13 @@ FLearningAgentsObservationObjectElement ULearningAgentsObservationObject::MakeFl
 #if UE_LEARNING_AGENTS_ENABLE_VISUAL_LOG
 	if (bVisualLoggerEnabled && VisualLoggerListener)
 	{
-		const ULearningAgentsVisualLoggerObject* VisualLoggerObject = VisualLoggerListener->GetOrAddVisualLoggerObject(Name);
+		const ULearningAgentsVisualLoggerObject* VisualLoggerObject = VisualLoggerListener->GetOrAddVisualLoggerObject(Tag);
 
 		UE_LEARNING_AGENTS_VLOG_STRING(VisualLoggerObject, LogLearning, Display, VisualLoggerLocation,
 			VisualLoggerColor.ToFColor(true),
-			TEXT("Listener: %s\nName: %s\nAgent Id: % 3i\nValue: [% 6.1f]\nScale: [% 6.2f]\nEncoded: [% 6.2f]"),
+			TEXT("Listener: %s\nTag: %s\nAgent Id: % 3i\nValue: [% 6.1f]\nScale: [% 6.2f]\nEncoded: [% 6.2f]"),
 			*VisualLoggerListener->GetName(),
-			*Name.ToString(),
+			*Tag.ToString(),
 			VisualLoggerAgentId,
 			Value,
 			FloatScale,
@@ -1590,14 +1700,15 @@ FLearningAgentsObservationObjectElement ULearningAgentsObservationObject::MakeFl
 	}
 #endif
 
-	return MakeContinuousObservationFromArrayView({ EncodedValue }, Name);
+	return MakeContinuousObservationFromArrayView(Object, { EncodedValue }, Tag);
 }
 
-FLearningAgentsObservationObjectElement ULearningAgentsObservationObject::MakeLocationObservation(
+FLearningAgentsObservationObjectElement ULearningAgentsObservations::MakeLocationObservation(
+	ULearningAgentsObservationObject* Object,
 	const FVector Location,
 	const FTransform RelativeTransform,
 	const float LocationScale,
-	const FName Name,
+	const FName Tag,
 	const bool bVisualLoggerEnabled,
 	ULearningAgentsManagerListener* VisualLoggerListener,
 	const int32 VisualLoggerAgentId,
@@ -1613,7 +1724,7 @@ FLearningAgentsObservationObjectElement ULearningAgentsObservationObject::MakeLo
 #if UE_LEARNING_AGENTS_ENABLE_VISUAL_LOG
 	if (bVisualLoggerEnabled && VisualLoggerListener)
 	{
-		const ULearningAgentsVisualLoggerObject* VisualLoggerObject = VisualLoggerListener->GetOrAddVisualLoggerObject(Name);
+		const ULearningAgentsVisualLoggerObject* VisualLoggerObject = VisualLoggerListener->GetOrAddVisualLoggerObject(Tag);
 
 		UE_LEARNING_AGENTS_VLOG_LOCATION(VisualLoggerObject, LogLearning, Display,
 			Location,
@@ -1635,9 +1746,9 @@ FLearningAgentsObservationObjectElement ULearningAgentsObservationObject::MakeLo
 
 		UE_LEARNING_AGENTS_VLOG_STRING(VisualLoggerObject, LogLearning, Display, VisualLoggerLocation,
 			VisualLoggerColor.ToFColor(true),
-			TEXT("Listener: %s\nName: %s\nAgent Id: % 3i\nLocation: [% 6.1f % 6.1f % 6.1f]\nLocal Location: [% 6.1f % 6.1f % 6.1f]\nScale: [% 6.2f]\nEncoded: [% 6.2f % 6.2f % 6.2f]"),
+			TEXT("Listener: %s\nTag: %s\nAgent Id: % 3i\nLocation: [% 6.1f % 6.1f % 6.1f]\nLocal Location: [% 6.1f % 6.1f % 6.1f]\nScale: [% 6.2f]\nEncoded: [% 6.2f % 6.2f % 6.2f]"),
 			*VisualLoggerListener->GetName(),
-			*Name.ToString(),
+			*Tag.ToString(),
 			VisualLoggerAgentId,
 			Location.X, Location.Y, Location.Z,
 			LocalLocation.X, LocalLocation.Y, LocalLocation.Z,
@@ -1646,17 +1757,18 @@ FLearningAgentsObservationObjectElement ULearningAgentsObservationObject::MakeLo
 	}
 #endif
 
-	return MakeContinuousObservationFromArrayView({
+	return MakeContinuousObservationFromArrayView(Object, {
 		(float)EncodedLocation.X,
 		(float)EncodedLocation.Y,
 		(float)EncodedLocation.Z,
-		}, Name);
+		}, Tag);
 }
 
-FLearningAgentsObservationObjectElement ULearningAgentsObservationObject::MakeRotationObservation(
+FLearningAgentsObservationObjectElement ULearningAgentsObservations::MakeRotationObservation(
+	ULearningAgentsObservationObject* Object,
 	const FRotator Rotation,
 	const FRotator RelativeRotation,
-	const FName Name,
+	const FName Tag,
 	const bool bVisualLoggerEnabled,
 	ULearningAgentsManagerListener* VisualLoggerListener,
 	const int32 VisualLoggerAgentId,
@@ -1667,14 +1779,16 @@ FLearningAgentsObservationObjectElement ULearningAgentsObservationObject::MakeRo
 	// Visual logging is handled in the MakeRotationObservationFromQuat
 
 	return MakeRotationObservationFromQuat(
-		FQuat::MakeFromRotator(Rotation), FQuat::MakeFromRotator(RelativeRotation), Name,
+		Object,
+		FQuat::MakeFromRotator(Rotation), FQuat::MakeFromRotator(RelativeRotation), Tag,
 		bVisualLoggerEnabled, VisualLoggerListener, VisualLoggerAgentId, VisualLoggerRotationLocation, VisualLoggerLocation, VisualLoggerColor);
 }
 
-FLearningAgentsObservationObjectElement ULearningAgentsObservationObject::MakeRotationObservationFromQuat(
+FLearningAgentsObservationObjectElement ULearningAgentsObservations::MakeRotationObservationFromQuat(
+	ULearningAgentsObservationObject* Object,
 	const FQuat Rotation,
 	const FQuat RelativeRotation,
-	const FName Name,
+	const FName Tag,
 	const bool bVisualLoggerEnabled,
 	ULearningAgentsManagerListener* VisualLoggerListener,
 	const int32 VisualLoggerAgentId,
@@ -1689,7 +1803,7 @@ FLearningAgentsObservationObjectElement ULearningAgentsObservationObject::MakeRo
 #if UE_LEARNING_AGENTS_ENABLE_VISUAL_LOG
 	if (bVisualLoggerEnabled && VisualLoggerListener)
 	{
-		const ULearningAgentsVisualLoggerObject* VisualLoggerObject = VisualLoggerListener->GetOrAddVisualLoggerObject(Name);
+		const ULearningAgentsVisualLoggerObject* VisualLoggerObject = VisualLoggerListener->GetOrAddVisualLoggerObject(Tag);
 
 		UE_LEARNING_AGENTS_VLOG_TRANSFORM(VisualLoggerObject, LogLearning, Display,
 			VisualLoggerRotationLocation,
@@ -1699,9 +1813,9 @@ FLearningAgentsObservationObjectElement ULearningAgentsObservationObject::MakeRo
 
 		UE_LEARNING_AGENTS_VLOG_STRING(VisualLoggerObject, LogLearning, Display, VisualLoggerLocation,
 			VisualLoggerColor.ToFColor(true),
-			TEXT("Listener: %s\nName: %s\nAgent Id: % 3i\nRotation: [% 6.1f % 6.1f % 6.1f % 6.1f]\nLocal Rotation: [% 6.1f % 6.1f % 6.1f % 6.1f]\nEncoded Forward: [% 6.2f % 6.2f % 6.2f]\nEncoded Right: [% 6.2f % 6.2f % 6.2f]"),
+			TEXT("Listener: %s\nTag: %s\nAgent Id: % 3i\nRotation: [% 6.1f % 6.1f % 6.1f % 6.1f]\nLocal Rotation: [% 6.1f % 6.1f % 6.1f % 6.1f]\nEncoded Forward: [% 6.2f % 6.2f % 6.2f]\nEncoded Right: [% 6.2f % 6.2f % 6.2f]"),
 			*VisualLoggerListener->GetName(),
-			*Name.ToString(),
+			*Tag.ToString(),
 			VisualLoggerAgentId,
 			Rotation.W, Rotation.X, Rotation.Y, Rotation.Z,
 			LocalRotation.W, LocalRotation.X, LocalRotation.Y, LocalRotation.Z,
@@ -1710,20 +1824,21 @@ FLearningAgentsObservationObjectElement ULearningAgentsObservationObject::MakeRo
 	}
 #endif
 
-	return MakeContinuousObservationFromArrayView({
+	return MakeContinuousObservationFromArrayView(Object, {
 		(float)LocalAxisForward.X,
 		(float)LocalAxisForward.Y,
 		(float)LocalAxisForward.Z,
 		(float)LocalAxisRight.X,
 		(float)LocalAxisRight.Y,
 		(float)LocalAxisRight.Z,
-		}, Name);
+		}, Tag);
 }
 
-FLearningAgentsObservationObjectElement ULearningAgentsObservationObject::MakeScaleObservation(
+FLearningAgentsObservationObjectElement ULearningAgentsObservations::MakeScaleObservation(
+	ULearningAgentsObservationObject* Object,
 	const FVector Scale,
 	const FVector RelativeScale,
-	const FName Name,
+	const FName Tag,
 	const bool bVisualLoggerEnabled,
 	ULearningAgentsManagerListener* VisualLoggerListener,
 	const int32 VisualLoggerAgentId,
@@ -1738,31 +1853,32 @@ FLearningAgentsObservationObjectElement ULearningAgentsObservationObject::MakeSc
 #if UE_LEARNING_AGENTS_ENABLE_VISUAL_LOG
 	if (bVisualLoggerEnabled && VisualLoggerListener)
 	{
-		const ULearningAgentsVisualLoggerObject* VisualLoggerObject = VisualLoggerListener->GetOrAddVisualLoggerObject(Name);
+		const ULearningAgentsVisualLoggerObject* VisualLoggerObject = VisualLoggerListener->GetOrAddVisualLoggerObject(Tag);
 
 		UE_LEARNING_AGENTS_VLOG_STRING(VisualLoggerObject, LogLearning, Display, VisualLoggerLocation,
 			VisualLoggerColor.ToFColor(true),
-			TEXT("Listener: %s\nName: %s\nAgent Id: % 3i\nScale: [% 6.1f % 6.1f % 6.1f]\nEncoded: [% 6.2f % 6.2f % 6.2f]"),
+			TEXT("Listener: %s\nTag: %s\nAgent Id: % 3i\nScale: [% 6.1f % 6.1f % 6.1f]\nEncoded: [% 6.2f % 6.2f % 6.2f]"),
 			*VisualLoggerListener->GetName(),
-			*Name.ToString(),
+			*Tag.ToString(),
 			VisualLoggerAgentId,
 			Scale.X, Scale.Y, Scale.Z,
 			LocalLogScale.X, LocalLogScale.Y, LocalLogScale.Z);
 	}
 #endif
 
-	return MakeContinuousObservationFromArrayView({
+	return MakeContinuousObservationFromArrayView(Object, {
 		(float)LocalLogScale.X,
 		(float)LocalLogScale.Y,
 		(float)LocalLogScale.Z,
-		}, Name);
+		}, Tag);
 }
 
-FLearningAgentsObservationObjectElement ULearningAgentsObservationObject::MakeTransformObservation(
+FLearningAgentsObservationObjectElement ULearningAgentsObservations::MakeTransformObservation(
+	ULearningAgentsObservationObject* Object,
 	const FTransform Transform,
 	const FTransform RelativeTransform,
 	const float LocationScale,
-	const FName Name,
+	const FName Tag,
 	const bool bVisualLoggerEnabled,
 	ULearningAgentsManagerListener* VisualLoggerListener,
 	const int32 VisualLoggerAgentId,
@@ -1774,7 +1890,7 @@ FLearningAgentsObservationObjectElement ULearningAgentsObservationObject::MakeTr
 #if UE_LEARNING_AGENTS_ENABLE_VISUAL_LOG
 	if (bVisualLoggerEnabled && VisualLoggerListener)
 	{
-		const ULearningAgentsVisualLoggerObject* VisualLoggerObject = VisualLoggerListener->GetOrAddVisualLoggerObject(Name);
+		const ULearningAgentsVisualLoggerObject* VisualLoggerObject = VisualLoggerListener->GetOrAddVisualLoggerObject(Tag);
 
 		const FVector LocalLocation = LocalTransform.GetLocation();
 		const FRotator LocalRotation = LocalTransform.Rotator();
@@ -1792,9 +1908,9 @@ FLearningAgentsObservationObjectElement ULearningAgentsObservationObject::MakeTr
 
 		UE_LEARNING_AGENTS_VLOG_STRING(VisualLoggerObject, LogLearning, Display, VisualLoggerLocation,
 			VisualLoggerColor.ToFColor(true),
-			TEXT("Listener: %s\nName: %s\nAgent Id: % 3i\nLocation: [% 6.1f % 6.1f % 6.1f]\nLocal Location: [% 6.1f % 6.1f % 6.1f]\nRotation: [% 6.1f % 6.1f % 6.1f]\nLocal Rotation: [% 6.1f % 6.1f % 6.1f]\nScale: [% 6.1f % 6.1f % 6.1f]\nLocal Scale: [% 6.1f % 6.1f % 6.1f]\nLocation Scale: [% 6.1f]"),
+			TEXT("Listener: %s\nTag: %s\nAgent Id: % 3i\nLocation: [% 6.1f % 6.1f % 6.1f]\nLocal Location: [% 6.1f % 6.1f % 6.1f]\nRotation: [% 6.1f % 6.1f % 6.1f]\nLocal Rotation: [% 6.1f % 6.1f % 6.1f]\nScale: [% 6.1f % 6.1f % 6.1f]\nLocal Scale: [% 6.1f % 6.1f % 6.1f]\nLocation Scale: [% 6.1f]"),
 			*VisualLoggerListener->GetName(),
-			*Name.ToString(),
+			*Tag.ToString(),
 			VisualLoggerAgentId,
 			Location.X, Location.Y, Location.Z,
 			LocalLocation.X, LocalLocation.Y, LocalLocation.Z,
@@ -1806,50 +1922,51 @@ FLearningAgentsObservationObjectElement ULearningAgentsObservationObject::MakeTr
 	}
 #endif
 
-	return MakeStructObservationFromArrayViews(
+	return MakeStructObservationFromArrayViews(Object,
 		{
 			TEXT("Location"),
 			TEXT("Rotation"),
 			TEXT("Scale")
 		},
 		{
-			MakeLocationObservation(LocalTransform.GetLocation(), FTransform::Identity, LocationScale),
-			MakeRotationObservationFromQuat(LocalTransform.GetRotation(), FQuat::Identity),
-			MakeScaleObservation(LocalTransform.GetScale3D(), FVector::OneVector)
+			MakeLocationObservation(Object, LocalTransform.GetLocation(), FTransform::Identity, LocationScale),
+			MakeRotationObservationFromQuat(Object, LocalTransform.GetRotation(), FQuat::Identity),
+			MakeScaleObservation(Object, LocalTransform.GetScale3D(), FVector::OneVector)
 		},
-		Name);
+		Tag);
 }
 
-FLearningAgentsObservationObjectElement ULearningAgentsObservationObject::MakeAngleObservation(const float Angle, const float RelativeAngle, const FName Name)
+FLearningAgentsObservationObjectElement ULearningAgentsObservations::MakeAngleObservation(ULearningAgentsObservationObject* Object, const float Angle, const float RelativeAngle, const FName Tag)
 {
-	return MakeAngleObservationRadians(FMath::DegreesToRadians(Angle), FMath::DegreesToRadians(RelativeAngle), Name);
+	return MakeAngleObservationRadians(Object, FMath::DegreesToRadians(Angle), FMath::DegreesToRadians(RelativeAngle), Tag);
 }
 
-FLearningAgentsObservationObjectElement ULearningAgentsObservationObject::MakeAngleObservationRadians(const float Angle, const float RelativeAngle, const FName Name)
+FLearningAgentsObservationObjectElement ULearningAgentsObservations::MakeAngleObservationRadians(ULearningAgentsObservationObject* Object, const float Angle, const float RelativeAngle, const FName Tag)
 {
 	const float LocalAngle = FMath::FindDeltaAngleRadians(RelativeAngle, Angle);
 
-	return MakeContinuousObservationFromArrayView({
+	return MakeContinuousObservationFromArrayView(Object, {
 		(float)FMath::Sin(Angle),
 		(float)FMath::Cos(Angle),
-		}, Name);
+		}, Tag);
 }
 
-FLearningAgentsObservationObjectElement ULearningAgentsObservationObject::MakeVelocityObservation(const FVector Velocity, const FTransform RelativeTransform, const float VelocityScale, const FName Name)
+FLearningAgentsObservationObjectElement ULearningAgentsObservations::MakeVelocityObservation(ULearningAgentsObservationObject* Object, const FVector Velocity, const FTransform RelativeTransform, const float VelocityScale, const FName Tag)
 {
 	const FVector LocalVelocity = RelativeTransform.InverseTransformVectorNoScale(Velocity);
 
-	return MakeContinuousObservationFromArrayView({
+	return MakeContinuousObservationFromArrayView(Object, {
 		(float)LocalVelocity.X / FMath::Max(VelocityScale, UE_SMALL_NUMBER),
 		(float)LocalVelocity.Y / FMath::Max(VelocityScale, UE_SMALL_NUMBER),
 		(float)LocalVelocity.Z / FMath::Max(VelocityScale, UE_SMALL_NUMBER),
-		}, Name);
+		}, Tag);
 }
 
-FLearningAgentsObservationObjectElement ULearningAgentsObservationObject::MakeDirectionObservation(
+FLearningAgentsObservationObjectElement ULearningAgentsObservations::MakeDirectionObservation(
+	ULearningAgentsObservationObject* Object,
 	const FVector Direction,
 	const FTransform RelativeTransform,
-	const FName Name,
+	const FName Tag,
 	const bool bVisualLoggerEnabled,
 	ULearningAgentsManagerListener* VisualLoggerListener,
 	const int32 VisualLoggerAgentId,
@@ -1863,7 +1980,7 @@ FLearningAgentsObservationObjectElement ULearningAgentsObservationObject::MakeDi
 #if UE_LEARNING_AGENTS_ENABLE_VISUAL_LOG
 	if (bVisualLoggerEnabled && VisualLoggerListener)
 	{
-		const ULearningAgentsVisualLoggerObject* VisualLoggerObject = VisualLoggerListener->GetOrAddVisualLoggerObject(Name);
+		const ULearningAgentsVisualLoggerObject* VisualLoggerObject = VisualLoggerListener->GetOrAddVisualLoggerObject(Tag);
 
 		UE_LEARNING_AGENTS_VLOG_ARROW(VisualLoggerObject, LogLearning, Display,
 			VisualLoggerDirectionLocation,
@@ -1879,9 +1996,9 @@ FLearningAgentsObservationObjectElement ULearningAgentsObservationObject::MakeDi
 
 		UE_LEARNING_AGENTS_VLOG_STRING(VisualLoggerObject, LogLearning, Display, VisualLoggerLocation,
 			VisualLoggerColor.ToFColor(true),
-			TEXT("Listener: %s\nName: %s\nAgent Id: % 3i\nDirection: [% 6.1f % 6.1f % 6.1f]\nLocal Direction: [% 6.1f % 6.1f % 6.1f]\nEncoded: [% 6.2f % 6.2f % 6.2f]"),
+			TEXT("Listener: %s\nTag: %s\nAgent Id: % 3i\nDirection: [% 6.1f % 6.1f % 6.1f]\nLocal Direction: [% 6.1f % 6.1f % 6.1f]\nEncoded: [% 6.2f % 6.2f % 6.2f]"),
 			*VisualLoggerListener->GetName(),
-			*Name.ToString(),
+			*Tag.ToString(),
 			VisualLoggerAgentId,
 			Direction.X, Direction.Y, Direction.Z,
 			LocalDirection.X, LocalDirection.Y, LocalDirection.Z,
@@ -1889,19 +2006,20 @@ FLearningAgentsObservationObjectElement ULearningAgentsObservationObject::MakeDi
 	}
 #endif
 
-	return MakeContinuousObservationFromArrayView({
+	return MakeContinuousObservationFromArrayView(Object, {
 		(float)LocalDirection.X,
 		(float)LocalDirection.Y,
 		(float)LocalDirection.Z,
-		}, Name);
+		}, Tag);
 }
 
-FLearningAgentsObservationObjectElement ULearningAgentsObservationObject::MakeLocationAlongSplineObservation(
+FLearningAgentsObservationObjectElement ULearningAgentsObservations::MakeLocationAlongSplineObservation(
+	ULearningAgentsObservationObject* Object,
 	const USplineComponent* SplineComponent,
 	const float DistanceAlongSpline,
 	const FTransform RelativeTransform,
 	const float LocationScale,
-	const FName Name,
+	const FName Tag,
 	const bool bVisualLoggerEnabled,
 	ULearningAgentsManagerListener* VisualLoggerListener,
 	const int32 VisualLoggerAgentId,
@@ -1914,10 +2032,10 @@ FLearningAgentsObservationObjectElement ULearningAgentsObservationObject::MakeLo
 		return FLearningAgentsObservationObjectElement();
 	}
 
-	return MakeLocationObservation(SplineComponent->GetLocationAtDistanceAlongSpline(DistanceAlongSpline, ESplineCoordinateSpace::World), RelativeTransform, LocationScale, Name, bVisualLoggerEnabled, VisualLoggerListener, VisualLoggerAgentId, VisualLoggerLocation, VisualLoggerColor);
+	return MakeLocationObservation(Object, SplineComponent->GetLocationAtDistanceAlongSpline(DistanceAlongSpline, ESplineCoordinateSpace::World), RelativeTransform, LocationScale, Tag, bVisualLoggerEnabled, VisualLoggerListener, VisualLoggerAgentId, VisualLoggerLocation, VisualLoggerColor);
 }
 
-FLearningAgentsObservationObjectElement ULearningAgentsObservationObject::MakeProportionAlongSplineObservation(const USplineComponent* SplineComponent, const float DistanceAlongSpline, const FName Name)
+FLearningAgentsObservationObjectElement ULearningAgentsObservations::MakeProportionAlongSplineObservation(ULearningAgentsObservationObject* Object, const USplineComponent* SplineComponent, const float DistanceAlongSpline, const FName Tag)
 {
 	if (!SplineComponent)
 	{
@@ -1931,20 +2049,21 @@ FLearningAgentsObservationObjectElement ULearningAgentsObservationObject::MakePr
 		const float WrapDistance = FMath::Wrap(DistanceAlongSpline, 0.0f, TotalDistance);
 		const float Proportion = WrapDistance / FMath::Max(TotalDistance, UE_SMALL_NUMBER);
 		const float Angle = FMath::Wrap(UE_TWO_PI * Proportion, -UE_PI, UE_PI);
-		return MakeExclusiveUnionObservation(TEXT("Angle"), MakeAngleObservation(Angle), Name);
+		return MakeExclusiveUnionObservation(Object, TEXT("Angle"), MakeAngleObservation(Object, Angle), Tag);
 	}
 	else
 	{
 		const float Proportion = FMath::Clamp(DistanceAlongSpline / FMath::Max(SplineComponent->GetSplineLength(), UE_SMALL_NUMBER), 0.0f, 1.0f);
-		return MakeExclusiveUnionObservation(TEXT("Proportion"), MakeFloatObservation(Proportion), Name);
+		return MakeExclusiveUnionObservation(Object, TEXT("Proportion"), MakeFloatObservation(Object, Proportion), Tag);
 	}
 }
 
-FLearningAgentsObservationObjectElement ULearningAgentsObservationObject::MakeDirectionAlongSplineObservation(
+FLearningAgentsObservationObjectElement ULearningAgentsObservations::MakeDirectionAlongSplineObservation(
+	ULearningAgentsObservationObject* Object,
 	const USplineComponent* SplineComponent,
 	const float DistanceAlongSpline,
 	const FTransform RelativeTransform,
-	const FName Name,
+	const FName Tag,
 	const bool bVisualLoggerEnabled,
 	ULearningAgentsManagerListener* VisualLoggerListener,
 	const int32 VisualLoggerAgentId,
@@ -1967,10 +2086,10 @@ FLearningAgentsObservationObjectElement ULearningAgentsObservationObject::MakeDi
 	const FVector VisualLoggerDirectionLocation = FVector::ZeroVector;
 #endif
 
-	return MakeDirectionObservation(SplineComponent->GetDirectionAtDistanceAlongSpline(DistanceAlongSpline, ESplineCoordinateSpace::World), RelativeTransform, Name, bVisualLoggerEnabled, VisualLoggerListener, VisualLoggerAgentId, VisualLoggerDirectionLocation, VisualLoggerLocation, VisualLoggerArrowLength, VisualLoggerColor);
+	return MakeDirectionObservation(Object, SplineComponent->GetDirectionAtDistanceAlongSpline(DistanceAlongSpline, ESplineCoordinateSpace::World), RelativeTransform, Tag, bVisualLoggerEnabled, VisualLoggerListener, VisualLoggerAgentId, VisualLoggerDirectionLocation, VisualLoggerLocation, VisualLoggerArrowLength, VisualLoggerColor);
 }
 
-FLearningAgentsObservationObjectElement ULearningAgentsObservationObject::MakePropertiesAlongSplineObservation(const USplineComponent* SplineComponent, const float DistanceAlongSpline, const FTransform RelativeTransform, const float LocationScale, const FName Name)
+FLearningAgentsObservationObjectElement ULearningAgentsObservations::MakePropertiesAlongSplineObservation(ULearningAgentsObservationObject* Object, const USplineComponent* SplineComponent, const float DistanceAlongSpline, const FTransform RelativeTransform, const float LocationScale, const FName Tag)
 {
 	if (!SplineComponent)
 	{
@@ -1978,20 +2097,20 @@ FLearningAgentsObservationObjectElement ULearningAgentsObservationObject::MakePr
 		return FLearningAgentsObservationObjectElement();
 	}
 
-	return MakeStructObservationFromArrayViews(
+	return MakeStructObservationFromArrayViews(Object,
 		{
 			TEXT("Location"),
 			TEXT("Proportion"),
 			TEXT("Direction")
 		},
 		{
-			MakeLocationAlongSplineObservation(SplineComponent, DistanceAlongSpline, RelativeTransform, LocationScale),
-			MakeProportionAlongSplineObservation(SplineComponent, DistanceAlongSpline),
-			MakeDirectionAlongSplineObservation(SplineComponent, DistanceAlongSpline, RelativeTransform),
-		}, Name);
+			MakeLocationAlongSplineObservation(Object, SplineComponent, DistanceAlongSpline, RelativeTransform, LocationScale),
+			MakeProportionAlongSplineObservation(Object, SplineComponent, DistanceAlongSpline),
+			MakeDirectionAlongSplineObservation(Object, SplineComponent, DistanceAlongSpline, RelativeTransform),
+		}, Tag);
 }
 
-FLearningAgentsObservationObjectElement ULearningAgentsObservationObject::MakeProportionAlongRayObservation(const FVector RayStart, const FVector RayEnd, const FTransform RayTransform, const ECollisionChannel CollisionChannel, const FName Name)
+FLearningAgentsObservationObjectElement ULearningAgentsObservations::MakeProportionAlongRayObservation(ULearningAgentsObservationObject* Object, const FVector RayStart, const FVector RayEnd, const FTransform RayTransform, const ECollisionChannel CollisionChannel, const FName Tag)
 {
 	const FVector RayStartWorld = RayTransform.TransformPosition(RayStart);
 	const FVector RayEndWorld = RayTransform.TransformPosition(RayEnd);
@@ -2002,19 +2121,19 @@ FLearningAgentsObservationObjectElement ULearningAgentsObservationObject::MakePr
 	ObjectQueryParams.AddObjectTypesToQuery(CollisionChannel);
 
 	FHitResult TraceHit;
-	const bool bHit = GetWorld()->LineTraceSingleByObjectType(TraceHit, RayStartWorld, RayEndWorld, ObjectQueryParams);
+	const bool bHit = Object->GetWorld()->LineTraceSingleByObjectType(TraceHit, RayStartWorld, RayEndWorld, ObjectQueryParams);
 
 	if (bHit)
 	{
-		return MakeFloatObservation(1.0f - TraceHit.Time, 1.0f, Name);
+		return MakeFloatObservation(Object, 1.0f - TraceHit.Time, 1.0f, Tag);
 	}
 	else
 	{
-		return MakeFloatObservation(0.0f, 1.0f, Name);
+		return MakeFloatObservation(Object, 0.0f, 1.0f, Tag);
 	}
 }
 
-FLearningAgentsObservationObjectElement ULearningAgentsObservationObject::MakeProportionAlongRaysObservation(const TArray<FVector>& RayStarts, const TArray<FVector>& RayEnds, const FTransform RayTransform, const ECollisionChannel CollisionChannel, const FName Name)
+FLearningAgentsObservationObjectElement ULearningAgentsObservations::MakeProportionAlongRaysObservation(ULearningAgentsObservationObject* Object, const TArray<FVector>& RayStarts, const TArray<FVector>& RayEnds, const FTransform RayTransform, const ECollisionChannel CollisionChannel, const FName Tag)
 {
 	if (RayStarts.Num() != RayEnds.Num())
 	{
@@ -2029,31 +2148,37 @@ FLearningAgentsObservationObjectElement ULearningAgentsObservationObject::MakePr
 
 	for (int32 RayIdx = 0; RayIdx < RayNum; RayIdx++)
 	{
-		RayElements.Add(MakeProportionAlongRayObservation(RayStarts[RayIdx], RayEnds[RayIdx], RayTransform, CollisionChannel));
+		RayElements.Add(MakeProportionAlongRayObservation(Object, RayStarts[RayIdx], RayEnds[RayIdx], RayTransform, CollisionChannel));
 	}
 
-	return MakeStaticArrayObservationFromArrayView(RayElements, Name);
+	return MakeStaticArrayObservationFromArrayView(Object, RayElements, Tag);
 }
 
-bool ULearningAgentsObservationObject::GetNullObservation(const FLearningAgentsObservationObjectElement Element, const FName Name) const
+
+bool ULearningAgentsObservations::GetNullObservation(const ULearningAgentsObservationObject* Object, const FLearningAgentsObservationObjectElement Element, const FName Tag)
 {
-	if (!ObservationObject.IsValid(Element.ObjectElement))
+	if (!Object)
+	{
+		UE_LOG(LogLearning, Error, TEXT("GetNullObservation: Object is nullptr."));
+		return false;
+	}
+
+	if (!Object->ObservationObject.IsValid(Element.ObjectElement))
 	{
-		UE_LOG(LogLearning, Error, TEXT("%s: Invalid Observation Object."), *GetName());
+		UE_LOG(LogLearning, Error, TEXT("GetNullObservation: Invalid Observation Object."));
 		return false;
 	}
 
-	if (ObservationObject.GetName(Element.ObjectElement) != Name)
+	if (Object->ObservationObject.GetTag(Element.ObjectElement) != Tag)
 	{
-		UE_LOG(LogLearning, Warning, TEXT("%s: Observation name does not match. Observation is '%s' but asked for '%s'."), *GetName(), *ObservationObject.GetName(Element.ObjectElement).ToString(), *Name.ToString());
+		UE_LOG(LogLearning, Warning, TEXT("GetNullObservation: Observation tag does not match. Observation is '%s' but asked for '%s'."), *Object->ObservationObject.GetTag(Element.ObjectElement).ToString(), *Tag.ToString());
 	}
 
-	if (ObservationObject.GetType(Element.ObjectElement) != UE::Learning::Observation::EType::Null)
+	if (Object->ObservationObject.GetType(Element.ObjectElement) != UE::Learning::Observation::EType::Null)
 	{
-		UE_LOG(LogLearning, Error, TEXT("%s: Observation '%s' type does not match. Observation is '%s' but asked for '%s'."),
-			*GetName(),
-			*ObservationObject.GetName(Element.ObjectElement).ToString(),
-			UE::Learning::Agents::Observation::Private::GetObservationTypeString(ObservationObject.GetType(Element.ObjectElement)),
+		UE_LOG(LogLearning, Error, TEXT("GetNullObservation: Observation '%s' type does not match. Observation is '%s' but asked for '%s'."),
+			*Object->ObservationObject.GetTag(Element.ObjectElement).ToString(),
+			UE::Learning::Agents::Observation::Private::GetObservationTypeString(Object->ObservationObject.GetType(Element.ObjectElement)),
 			UE::Learning::Agents::Observation::Private::GetObservationTypeString(UE::Learning::Observation::EType::Null));
 		return false;
 	}
@@ -2061,39 +2186,45 @@ bool ULearningAgentsObservationObject::GetNullObservation(const FLearningAgentsO
 	return true;
 }
 
-bool ULearningAgentsObservationObject::GetContinuousObservationNum(int32& OutNum, const FLearningAgentsObservationObjectElement Element, const FName Name) const
+bool ULearningAgentsObservations::GetContinuousObservationNum(int32& OutNum, const ULearningAgentsObservationObject* Object, const FLearningAgentsObservationObjectElement Element, const FName Tag)
 {
-	if (!ObservationObject.IsValid(Element.ObjectElement))
+	if (!Object)
 	{
-		UE_LOG(LogLearning, Error, TEXT("%s: Invalid Observation Object."), *GetName());
+		UE_LOG(LogLearning, Error, TEXT("GetContinuousObservationNum: Object is nullptr."));
 		OutNum = 0;
 		return false;
 	}
 
-	if (ObservationObject.GetName(Element.ObjectElement) != Name)
+	if (!Object->ObservationObject.IsValid(Element.ObjectElement))
 	{
-		UE_LOG(LogLearning, Warning, TEXT("%s: Observation name does not match. Observation is '%s' but asked for '%s'."), *GetName(), *ObservationObject.GetName(Element.ObjectElement).ToString(), *Name.ToString());
+		UE_LOG(LogLearning, Error, TEXT("GetContinuousObservationNum: Invalid Observation Object."));
+		OutNum = 0;
+		return false;
+	}
+
+	if (Object->ObservationObject.GetTag(Element.ObjectElement) != Tag)
+	{
+		UE_LOG(LogLearning, Warning, TEXT("GetContinuousObservationNum: Observation tag does not match. Observation is '%s' but asked for '%s'."), *Object->ObservationObject.GetTag(Element.ObjectElement).ToString(), *Tag.ToString());
 	}
 
-	if (ObservationObject.GetType(Element.ObjectElement) != UE::Learning::Observation::EType::Continuous)
+	if (Object->ObservationObject.GetType(Element.ObjectElement) != UE::Learning::Observation::EType::Continuous)
 	{
-		UE_LOG(LogLearning, Error, TEXT("%s: Observation '%s' type does not match. Observation is '%s' but asked for '%s'."),
-			*GetName(),
-			*ObservationObject.GetName(Element.ObjectElement).ToString(),
-			UE::Learning::Agents::Observation::Private::GetObservationTypeString(ObservationObject.GetType(Element.ObjectElement)),
+		UE_LOG(LogLearning, Error, TEXT("GetContinuousObservationNum: Observation '%s' type does not match. Observation is '%s' but asked for '%s'."),
+			*Object->ObservationObject.GetTag(Element.ObjectElement).ToString(),
+			UE::Learning::Agents::Observation::Private::GetObservationTypeString(Object->ObservationObject.GetType(Element.ObjectElement)),
 			UE::Learning::Agents::Observation::Private::GetObservationTypeString(UE::Learning::Observation::EType::Continuous));
 		OutNum = 0;
 		return false;
 	}
 
-	OutNum = ObservationObject.GetContinuous(Element.ObjectElement).Values.Num();
+	OutNum = Object->ObservationObject.GetContinuous(Element.ObjectElement).Values.Num();
 	return true;
 }
 
-bool ULearningAgentsObservationObject::GetContinuousObservation(TArray<float>& OutValues, const FLearningAgentsObservationObjectElement Element, const FName Name) const
+bool ULearningAgentsObservations::GetContinuousObservation(TArray<float>& OutValues, const ULearningAgentsObservationObject* Object, const FLearningAgentsObservationObjectElement Element, const FName Tag)
 {
 	int32 OutValueNum = 0;
-	if (!GetContinuousObservationNum(OutValueNum, Element, Name))
+	if (!GetContinuousObservationNum(OutValueNum, Object, Element, Tag))
 	{
 		OutValues.Empty();
 		return false;
@@ -2101,7 +2232,7 @@ bool ULearningAgentsObservationObject::GetContinuousObservation(TArray<float>& O
 
 	OutValues.SetNumUninitialized(OutValueNum);
 
-	if (!GetContinuousObservationToArrayView(OutValues, Element, Name))
+	if (!GetContinuousObservationToArrayView(OutValues, Object, Element, Tag))
 	{
 		OutValues.Empty();
 		return false;
@@ -2110,37 +2241,43 @@ bool ULearningAgentsObservationObject::GetContinuousObservation(TArray<float>& O
 	return true;
 }
 
-bool ULearningAgentsObservationObject::GetContinuousObservationToArrayView(TArrayView<float> OutValues, const FLearningAgentsObservationObjectElement Element, const FName Name) const
+bool ULearningAgentsObservations::GetContinuousObservationToArrayView(TArrayView<float> OutValues, const ULearningAgentsObservationObject* Object, const FLearningAgentsObservationObjectElement Element, const FName Tag)
 {
-	if (!ObservationObject.IsValid(Element.ObjectElement))
+	if (!Object)
 	{
-		UE_LOG(LogLearning, Error, TEXT("%s: Invalid Observation Object."), *GetName());
+		UE_LOG(LogLearning, Error, TEXT("GetContinuousObservationToArrayView: Object is nullptr."));
 		UE::Learning::Array::Zero<1, float>(OutValues);
 		return false;
 	}
 
-	if (ObservationObject.GetName(Element.ObjectElement) != Name)
+	if (!Object->ObservationObject.IsValid(Element.ObjectElement))
 	{
-		UE_LOG(LogLearning, Warning, TEXT("%s: Observation name does not match. Observation is '%s' but asked for '%s'."), *GetName(), *ObservationObject.GetName(Element.ObjectElement).ToString(), *Name.ToString());
+		UE_LOG(LogLearning, Error, TEXT("GetContinuousObservationToArrayView: Invalid Observation Object."));
+		UE::Learning::Array::Zero<1, float>(OutValues);
+		return false;
 	}
 
-	if (ObservationObject.GetType(Element.ObjectElement) != UE::Learning::Observation::EType::Continuous)
+	if (Object->ObservationObject.GetTag(Element.ObjectElement) != Tag)
 	{
-		UE_LOG(LogLearning, Error, TEXT("%s: Observation '%s' type does not match. Observation is '%s' but asked for '%s'."),
-			*GetName(),
-			*ObservationObject.GetName(Element.ObjectElement).ToString(),
-			UE::Learning::Agents::Observation::Private::GetObservationTypeString(ObservationObject.GetType(Element.ObjectElement)),
+		UE_LOG(LogLearning, Warning, TEXT("GetContinuousObservationToArrayView: Observation tag does not match. Observation is '%s' but asked for '%s'."), *Object->ObservationObject.GetTag(Element.ObjectElement).ToString(), *Tag.ToString());
+	}
+
+	if (Object->ObservationObject.GetType(Element.ObjectElement) != UE::Learning::Observation::EType::Continuous)
+	{
+		UE_LOG(LogLearning, Error, TEXT("GetContinuousObservationToArrayView: Observation '%s' type does not match. Observation is '%s' but asked for '%s'."),
+			*Object->ObservationObject.GetTag(Element.ObjectElement).ToString(),
+			UE::Learning::Agents::Observation::Private::GetObservationTypeString(Object->ObservationObject.GetType(Element.ObjectElement)),
 			UE::Learning::Agents::Observation::Private::GetObservationTypeString(UE::Learning::Observation::EType::Continuous));
 		UE::Learning::Array::Zero<1, float>(OutValues);
 		return false;
 	}
 
-	const TArrayView<const float> Values = ObservationObject.GetContinuous(Element.ObjectElement).Values;
+	const TArrayView<const float> Values = Object->ObservationObject.GetContinuous(Element.ObjectElement).Values;
 
 	if (Values.Num() != OutValues.Num())
 	{
-		UE_LOG(LogLearning, Error, TEXT("%s: Observation '%s' size does not match. Observation is '%i' values but asked for '%i'."),
-			*GetName(), *ObservationObject.GetName(Element.ObjectElement).ToString(),
+		UE_LOG(LogLearning, Error, TEXT("GetContinuousObservationToArrayView: Observation '%s' size does not match. Observation is '%i' values but asked for '%i'."),
+			*Object->ObservationObject.GetTag(Element.ObjectElement).ToString(),
 			Values.Num(), OutValues.Num());
 		UE::Learning::Array::Zero<1, float>(OutValues);
 		return false;
@@ -2150,10 +2287,10 @@ bool ULearningAgentsObservationObject::GetContinuousObservationToArrayView(TArra
 	return true;
 }
 
-bool ULearningAgentsObservationObject::GetExclusiveDiscreteObservation(int32& OutIndex, const FLearningAgentsObservationObjectElement Element, const FName Name) const
+bool ULearningAgentsObservations::GetExclusiveDiscreteObservation(int32& OutIndex, const ULearningAgentsObservationObject* Object, const FLearningAgentsObservationObjectElement Element, const FName Tag)
 {
 	int32 DiscreteValueNum;
-	if (!GetContinuousObservationNum(DiscreteValueNum, Element, Name))
+	if (!GetContinuousObservationNum(DiscreteValueNum, Object, Element, Tag))
 	{
 		OutIndex = INDEX_NONE;
 		return false;
@@ -2161,7 +2298,7 @@ bool ULearningAgentsObservationObject::GetExclusiveDiscreteObservation(int32& Ou
 
 	TArray<float, TInlineAllocator<32>> OneHot;
 	OneHot.SetNumUninitialized(DiscreteValueNum);
-	if (!GetContinuousObservationToArrayView(OneHot, Element, Name))
+	if (!GetContinuousObservationToArrayView(OneHot, Object, Element, Tag))
 	{
 		OutIndex = INDEX_NONE;
 		return false;
@@ -2180,10 +2317,10 @@ bool ULearningAgentsObservationObject::GetExclusiveDiscreteObservation(int32& Ou
 	return OutIndex != INDEX_NONE;
 }
 
-bool ULearningAgentsObservationObject::GetInclusiveDiscreteObservationNum(int32& OutNum, const FLearningAgentsObservationObjectElement Element, const FName Name) const
+bool ULearningAgentsObservations::GetInclusiveDiscreteObservationNum(int32& OutNum, const ULearningAgentsObservationObject* Object, const FLearningAgentsObservationObjectElement Element, const FName Tag)
 {
 	int32 DiscreteValueNum;
-	if (!GetContinuousObservationNum(DiscreteValueNum, Element, Name))
+	if (!GetContinuousObservationNum(DiscreteValueNum, Object, Element, Tag))
 	{
 		OutNum = 0;
 		return false;
@@ -2191,7 +2328,7 @@ bool ULearningAgentsObservationObject::GetInclusiveDiscreteObservationNum(int32&
 
 	TArray<float, TInlineAllocator<32>> OneHot;
 	OneHot.SetNumUninitialized(DiscreteValueNum);
-	if (!GetContinuousObservationToArrayView(OneHot, Element, Name))
+	if (!GetContinuousObservationToArrayView(OneHot, Object, Element, Tag))
 	{
 		OutNum = 0;
 		return false;
@@ -2206,30 +2343,36 @@ bool ULearningAgentsObservationObject::GetInclusiveDiscreteObservationNum(int32&
 	return true;
 }
 
-bool ULearningAgentsObservationObject::GetInclusiveDiscreteObservation(TArray<int32>& OutIndices, const FLearningAgentsObservationObjectElement Element, const FName Name) const
+bool ULearningAgentsObservations::GetInclusiveDiscreteObservation(TArray<int32>& OutIndices, const ULearningAgentsObservationObject* Object, const FLearningAgentsObservationObjectElement Element, const FName Tag)
 {
 	int32 OutIndicesNum;
-	if (!GetInclusiveDiscreteObservationNum(OutIndicesNum, Element, Name))
+	if (!GetInclusiveDiscreteObservationNum(OutIndicesNum, Object, Element, Tag))
 	{
 		OutIndices.Empty();
 		return false;
 	}
 
 	OutIndices.SetNumUninitialized(OutIndicesNum);
-	return GetInclusiveDiscreteObservationToArrayView(OutIndices, Element, Name);
+	return GetInclusiveDiscreteObservationToArrayView(OutIndices, Object, Element, Tag);
 }
 
-bool ULearningAgentsObservationObject::GetInclusiveDiscreteObservationToArrayView(TArrayView<int32> OutIndices, const FLearningAgentsObservationObjectElement Element, const FName Name) const
+bool ULearningAgentsObservations::GetInclusiveDiscreteObservationToArrayView(TArrayView<int32> OutIndices, const ULearningAgentsObservationObject* Object, const FLearningAgentsObservationObjectElement Element, const FName Tag)
 {
+	if (!Object)
+	{
+		UE_LOG(LogLearning, Error, TEXT("GetInclusiveDiscreteObservationToArrayView: Object is nullptr."));
+		return false;
+	}
+
 	int32 DiscreteValueNum;
-	if (!GetContinuousObservationNum(DiscreteValueNum, Element, Name))
+	if (!GetContinuousObservationNum(DiscreteValueNum, Object, Element, Tag))
 	{
 		return false;
 	}
 
 	TArray<float, TInlineAllocator<32>> OneHot;
 	OneHot.SetNumUninitialized(DiscreteValueNum);
-	if (!GetContinuousObservationToArrayView(OneHot, Element, Name))
+	if (!GetContinuousObservationToArrayView(OneHot, Object, Element, Tag))
 	{
 		return false;
 	}
@@ -2242,8 +2385,8 @@ bool ULearningAgentsObservationObject::GetInclusiveDiscreteObservationToArrayVie
 
 	if (IndicesNum != OutIndices.Num())
 	{
-		UE_LOG(LogLearning, Error, TEXT("%s: Observation '%s' size does not match. Observation is '%i' indices but asked for '%i'."),
-			*GetName(), *ObservationObject.GetName(Element.ObjectElement).ToString(),
+		UE_LOG(LogLearning, Error, TEXT("GetInclusiveDiscreteObservationToArrayView: Observation '%s' size does not match. Observation is '%i' indices but asked for '%i'."),
+			*Object->ObservationObject.GetTag(Element.ObjectElement).ToString(),
 			IndicesNum, OutIndices.Num());
 		UE::Learning::Array::Zero<1, int32>(OutIndices);
 		return false;
@@ -2258,10 +2401,10 @@ bool ULearningAgentsObservationObject::GetInclusiveDiscreteObservationToArrayVie
 	return true;
 }
 
-bool ULearningAgentsObservationObject::GetIndexObservation(int32& OutIndex, const FLearningAgentsObservationObjectElement Element, const FName Name) const
+bool ULearningAgentsObservations::GetIndexObservation(int32& OutIndex, const ULearningAgentsObservationObject* Object, const FLearningAgentsObservationObjectElement Element, const FName Tag)
 {
 	int32 DiscreteValueNum;
-	if (!GetContinuousObservationNum(DiscreteValueNum, Element, Name))
+	if (!GetContinuousObservationNum(DiscreteValueNum, Object, Element, Tag))
 	{
 		OutIndex = INDEX_NONE;
 		return false;
@@ -2269,7 +2412,7 @@ bool ULearningAgentsObservationObject::GetIndexObservation(int32& OutIndex, cons
 
 	TArray<float, TInlineAllocator<32>> OneHot;
 	OneHot.SetNumUninitialized(DiscreteValueNum);
-	if (!GetContinuousObservationToArrayView(OneHot, Element, Name))
+	if (!GetContinuousObservationToArrayView(OneHot, Object, Element, Tag))
 	{
 		OutIndex = INDEX_NONE;
 		return false;
@@ -2284,10 +2427,10 @@ bool ULearningAgentsObservationObject::GetIndexObservation(int32& OutIndex, cons
 	return true;
 }
 
-bool ULearningAgentsObservationObject::GetCountObservation(int32& OutNum, const FLearningAgentsObservationObjectElement Element, const int32 MaxNum, const FName Name) const
+bool ULearningAgentsObservations::GetCountObservation(int32& OutNum, const ULearningAgentsObservationObject* Object, const FLearningAgentsObservationObjectElement Element, const int32 MaxNum, const FName Tag)
 {
 	float FloatNum = 0.0f;
-	if (!GetContinuousObservationToArrayView(MakeArrayView(&FloatNum, 1), Element, Name))
+	if (!GetContinuousObservationToArrayView(MakeArrayView(&FloatNum, 1), Object, Element, Tag))
 	{
 		OutNum = -1;
 		return false;
@@ -2297,41 +2440,47 @@ bool ULearningAgentsObservationObject::GetCountObservation(int32& OutNum, const
 	return true;
 }
 
-bool ULearningAgentsObservationObject::GetStructObservationNum(int32& OutNum, const FLearningAgentsObservationObjectElement Element, const FName Name) const
+bool ULearningAgentsObservations::GetStructObservationNum(int32& OutNum, const ULearningAgentsObservationObject* Object, const FLearningAgentsObservationObjectElement Element, const FName Tag)
 {
-	if (!ObservationObject.IsValid(Element.ObjectElement))
+	if (!Object)
 	{
-		UE_LOG(LogLearning, Error, TEXT("%s: Invalid Observation Object."), *GetName());
+		UE_LOG(LogLearning, Error, TEXT("GetStructObservationNum: Object is nullptr."));
 		OutNum = 0;
 		return false;
 	}
 
-	if (ObservationObject.GetName(Element.ObjectElement) != Name)
+	if (!Object->ObservationObject.IsValid(Element.ObjectElement))
 	{
-		UE_LOG(LogLearning, Warning, TEXT("%s: Observation name does not match. Observation is '%s' but asked for '%s'."), *GetName(), *ObservationObject.GetName(Element.ObjectElement).ToString(), *Name.ToString());
+		UE_LOG(LogLearning, Error, TEXT("GetStructObservationNum: Invalid Observation Object."));
+		OutNum = 0;
+		return false;
 	}
 
-	if (ObservationObject.GetType(Element.ObjectElement) != UE::Learning::Observation::EType::And)
+	if (Object->ObservationObject.GetTag(Element.ObjectElement) != Tag)
 	{
-		UE_LOG(LogLearning, Error, TEXT("%s: Observation '%s' type does not match. Observation is '%s' but asked for '%s'."),
-			*GetName(),
-			*ObservationObject.GetName(Element.ObjectElement).ToString(),
-			UE::Learning::Agents::Observation::Private::GetObservationTypeString(ObservationObject.GetType(Element.ObjectElement)),
+		UE_LOG(LogLearning, Warning, TEXT("GetStructObservationNum: Observation tag does not match. Observation is '%s' but asked for '%s'."), *Object->ObservationObject.GetTag(Element.ObjectElement).ToString(), *Tag.ToString());
+	}
+
+	if (Object->ObservationObject.GetType(Element.ObjectElement) != UE::Learning::Observation::EType::And)
+	{
+		UE_LOG(LogLearning, Error, TEXT("GetStructObservationNum: Observation '%s' type does not match. Observation is '%s' but asked for '%s'."),
+			*Object->ObservationObject.GetTag(Element.ObjectElement).ToString(),
+			UE::Learning::Agents::Observation::Private::GetObservationTypeString(Object->ObservationObject.GetType(Element.ObjectElement)),
 			UE::Learning::Agents::Observation::Private::GetObservationTypeString(UE::Learning::Observation::EType::And));
 		OutNum = 0;
 		return false;
 	}
 
-	const UE::Learning::Observation::FObjectAndParameters Parameters = ObservationObject.GetAnd(Element.ObjectElement);
+	const UE::Learning::Observation::FObjectAndParameters Parameters = Object->ObservationObject.GetAnd(Element.ObjectElement);
 
 	OutNum = Parameters.Elements.Num();
 	return true;
 }
 
-bool ULearningAgentsObservationObject::GetStructObservation(TMap<FName, FLearningAgentsObservationObjectElement>& OutElements, const FLearningAgentsObservationObjectElement Element, const FName Name) const
+bool ULearningAgentsObservations::GetStructObservation(TMap<FName, FLearningAgentsObservationObjectElement>& OutElements, const ULearningAgentsObservationObject* Object, const FLearningAgentsObservationObjectElement Element, const FName Tag)
 {
 	int32 OutElementNum = 0;
-	if (!GetStructObservationNum(OutElementNum, Element, Name))
+	if (!GetStructObservationNum(OutElementNum, Object, Element, Tag))
 	{
 		OutElements.Empty();
 		return false;
@@ -2342,7 +2491,7 @@ bool ULearningAgentsObservationObject::GetStructObservation(TMap<FName, FLearnin
 	SubElementNames.SetNumUninitialized(OutElementNum);
 	SubElements.SetNumUninitialized(OutElementNum);
 
-	if (!GetStructObservationToArrayViews(SubElementNames, SubElements, Element, Name))
+	if (!GetStructObservationToArrayViews(SubElementNames, SubElements, Object, Element, Tag))
 	{
 		OutElements.Empty();
 		return false;
@@ -2357,10 +2506,10 @@ bool ULearningAgentsObservationObject::GetStructObservation(TMap<FName, FLearnin
 	return true;
 }
 
-bool ULearningAgentsObservationObject::GetStructObservationToArrays(TArray<FName>& OutElementNames, TArray<FLearningAgentsObservationObjectElement>& OutElements, const FLearningAgentsObservationObjectElement Element, const FName Name) const
+bool ULearningAgentsObservations::GetStructObservationToArrays(TArray<FName>& OutElementNames, TArray<FLearningAgentsObservationObjectElement>& OutElements, const ULearningAgentsObservationObject* Object, const FLearningAgentsObservationObjectElement Element, const FName Tag)
 {
 	int32 OutElementNum = 0;
-	if (!GetStructObservationNum(OutElementNum, Element, Name))
+	if (!GetStructObservationNum(OutElementNum, Object, Element, Tag))
 	{
 		OutElementNames.Empty();
 		OutElements.Empty();
@@ -2370,7 +2519,7 @@ bool ULearningAgentsObservationObject::GetStructObservationToArrays(TArray<FName
 	OutElementNames.SetNumUninitialized(OutElementNum);
 	OutElements.SetNumUninitialized(OutElementNum);
 
-	if (!GetStructObservationToArrayViews(OutElementNames, OutElements, Element, Name))
+	if (!GetStructObservationToArrayViews(OutElementNames, OutElements, Object, Element, Tag))
 	{
 		OutElementNames.Empty();
 		OutElements.Empty();
@@ -2380,45 +2529,51 @@ bool ULearningAgentsObservationObject::GetStructObservationToArrays(TArray<FName
 	return true;
 }
 
-bool ULearningAgentsObservationObject::GetStructObservationToArrayViews(TArrayView<FName> OutElementNames, TArrayView<FLearningAgentsObservationObjectElement> OutElements, const FLearningAgentsObservationObjectElement Element, const FName Name) const
+bool ULearningAgentsObservations::GetStructObservationToArrayViews(TArrayView<FName> OutElementNames, TArrayView<FLearningAgentsObservationObjectElement> OutElements, const ULearningAgentsObservationObject* Object, const FLearningAgentsObservationObjectElement Element, const FName Tag)
 {
-	if (!ObservationObject.IsValid(Element.ObjectElement))
+	if (!Object)
+	{
+		UE_LOG(LogLearning, Error, TEXT("GetStructObservationToArrayViews: Object is nullptr."));
+		UE::Learning::Array::Set<1, FName>(OutElementNames, NAME_None);
+		UE::Learning::Array::Set<1, FLearningAgentsObservationObjectElement>(OutElements, FLearningAgentsObservationObjectElement());
+		return false;
+	}
+
+	if (!Object->ObservationObject.IsValid(Element.ObjectElement))
 	{
-		UE_LOG(LogLearning, Error, TEXT("%s: Invalid Observation Object."), *GetName());
+		UE_LOG(LogLearning, Error, TEXT("GetStructObservationToArrayViews: Invalid Observation Object."));
 		UE::Learning::Array::Set<1, FName>(OutElementNames, NAME_None);
 		UE::Learning::Array::Set<1, FLearningAgentsObservationObjectElement>(OutElements, FLearningAgentsObservationObjectElement());
 		return false;
 	}
 
-	if (ObservationObject.GetName(Element.ObjectElement) != Name)
+	if (Object->ObservationObject.GetTag(Element.ObjectElement) != Tag)
 	{
-		UE_LOG(LogLearning, Warning, TEXT("%s: Observation name does not match. Observation is '%s' but asked for '%s'."), *GetName(), *ObservationObject.GetName(Element.ObjectElement).ToString(), *Name.ToString());
+		UE_LOG(LogLearning, Warning, TEXT("GetStructObservationToArrayViews: Observation tag does not match. Observation is '%s' but asked for '%s'."), *Object->ObservationObject.GetTag(Element.ObjectElement).ToString(), *Tag.ToString());
 	}
 
-	if (ObservationObject.GetType(Element.ObjectElement) != UE::Learning::Observation::EType::And)
+	if (Object->ObservationObject.GetType(Element.ObjectElement) != UE::Learning::Observation::EType::And)
 	{
-		UE_LOG(LogLearning, Error, TEXT("%s: Observation '%s' type does not match. Observation is '%s' but asked for '%s'."),
-			*GetName(),
-			*ObservationObject.GetName(Element.ObjectElement).ToString(),
-			UE::Learning::Agents::Observation::Private::GetObservationTypeString(ObservationObject.GetType(Element.ObjectElement)),
+		UE_LOG(LogLearning, Error, TEXT("GetStructObservationToArrayViews: Observation '%s' type does not match. Observation is '%s' but asked for '%s'."),
+			*Object->ObservationObject.GetTag(Element.ObjectElement).ToString(),
+			UE::Learning::Agents::Observation::Private::GetObservationTypeString(Object->ObservationObject.GetType(Element.ObjectElement)),
 			UE::Learning::Agents::Observation::Private::GetObservationTypeString(UE::Learning::Observation::EType::And));
 		UE::Learning::Array::Set<1, FName>(OutElementNames, NAME_None);
 		UE::Learning::Array::Set<1, FLearningAgentsObservationObjectElement>(OutElements, FLearningAgentsObservationObjectElement());
 		return false;
 	}
 
-	const UE::Learning::Observation::FObjectAndParameters Parameters = ObservationObject.GetAnd(Element.ObjectElement);
+	const UE::Learning::Observation::FObjectAndParameters Parameters = Object->ObservationObject.GetAnd(Element.ObjectElement);
 
 	if (Parameters.Elements.Num() == 0)
 	{
-		UE_LOG(LogLearning, Warning, TEXT("%s: Getting zero-sized And Observation."), *GetName());
+		UE_LOG(LogLearning, Warning, TEXT("GetStructObservationToArrayViews: Getting zero-sized And Observation."));
 	}
 
 	if (Parameters.Elements.Num() != OutElements.Num())
 	{
-		UE_LOG(LogLearning, Error, TEXT("%s: Observation '%s' size does not match. Observation is '%i' elements but asked for '%i'."),
-			*GetName(),
-			*ObservationObject.GetName(Element.ObjectElement).ToString(),
+		UE_LOG(LogLearning, Error, TEXT("GetStructObservationToArrayViews: Observation '%s' size does not match. Observation is '%i' elements but asked for '%i'."),
+			*Object->ObservationObject.GetTag(Element.ObjectElement).ToString(),
 			Parameters.Elements.Num(), OutElements.Num());
 		UE::Learning::Array::Set<1, FName>(OutElementNames, NAME_None);
 		UE::Learning::Array::Set<1, FLearningAgentsObservationObjectElement>(OutElements, FLearningAgentsObservationObjectElement());
@@ -2427,9 +2582,9 @@ bool ULearningAgentsObservationObject::GetStructObservationToArrayViews(TArrayVi
 
 	for (int32 ElementIdx = 0; ElementIdx < Parameters.Elements.Num(); ElementIdx++)
 	{
-		if (!ObservationObject.IsValid(Parameters.Elements[ElementIdx]))
+		if (!Object->ObservationObject.IsValid(Parameters.Elements[ElementIdx]))
 		{
-			UE_LOG(LogLearning, Error, TEXT("%s: Invalid Observation Object."), *GetName());
+			UE_LOG(LogLearning, Error, TEXT("GetStructObservationToArrayViews: Invalid Observation Object."));
 			UE::Learning::Array::Set<1, FName>(OutElementNames, NAME_None);
 			UE::Learning::Array::Set<1, FLearningAgentsObservationObjectElement>(OutElements, FLearningAgentsObservationObjectElement());
 			return false;
@@ -2442,74 +2597,87 @@ bool ULearningAgentsObservationObject::GetStructObservationToArrayViews(TArrayVi
 	return true;
 }
 
-bool ULearningAgentsObservationObject::GetExclusiveUnionObservation(FName& OutElementName, FLearningAgentsObservationObjectElement& OutElement, const FLearningAgentsObservationObjectElement Element, const FName Name) const
+bool ULearningAgentsObservations::GetExclusiveUnionObservation(FName& OutElementName, FLearningAgentsObservationObjectElement& OutElement, const ULearningAgentsObservationObject* Object, const FLearningAgentsObservationObjectElement Element, const FName Tag)
 {
-	if (!ObservationObject.IsValid(Element.ObjectElement))
+	if (!Object)
 	{
-		UE_LOG(LogLearning, Error, TEXT("%s: Invalid Observation Object."), *GetName());
+		UE_LOG(LogLearning, Error, TEXT("GetStructObservationToArrayViews: Object is nullptr."));
 		OutElementName = NAME_None;
 		OutElement = FLearningAgentsObservationObjectElement();
 		return false;
 	}
 
-	if (ObservationObject.GetName(Element.ObjectElement) != Name)
+	if (!Object->ObservationObject.IsValid(Element.ObjectElement))
 	{
-		UE_LOG(LogLearning, Warning, TEXT("%s: Observation name does not match. Observation is '%s' but asked for '%s'."), *GetName(), *ObservationObject.GetName(Element.ObjectElement).ToString(), *Name.ToString());
+		UE_LOG(LogLearning, Error, TEXT("GetStructObservationToArrayViews: Invalid Observation Object."));
+		OutElementName = NAME_None;
+		OutElement = FLearningAgentsObservationObjectElement();
+		return false;
 	}
 
-	if (ObservationObject.GetType(Element.ObjectElement) != UE::Learning::Observation::EType::OrExclusive)
+	if (Object->ObservationObject.GetTag(Element.ObjectElement) != Tag)
 	{
-		UE_LOG(LogLearning, Error, TEXT("%s: Observation '%s' type does not match. Observation is '%s' but asked for '%s'."),
-			*GetName(),
-			*ObservationObject.GetName(Element.ObjectElement).ToString(),
-			UE::Learning::Agents::Observation::Private::GetObservationTypeString(ObservationObject.GetType(Element.ObjectElement)),
+		UE_LOG(LogLearning, Warning, TEXT("GetStructObservationToArrayViews: Observation tag does not match. Observation is '%s' but asked for '%s'."), *Object->ObservationObject.GetTag(Element.ObjectElement).ToString(), *Tag.ToString());
+	}
+
+	if (Object->ObservationObject.GetType(Element.ObjectElement) != UE::Learning::Observation::EType::OrExclusive)
+	{
+		UE_LOG(LogLearning, Error, TEXT("GetStructObservationToArrayViews: Observation '%s' type does not match. Observation is '%s' but asked for '%s'."),
+			*Object->ObservationObject.GetTag(Element.ObjectElement).ToString(),
+			UE::Learning::Agents::Observation::Private::GetObservationTypeString(Object->ObservationObject.GetType(Element.ObjectElement)),
 			UE::Learning::Agents::Observation::Private::GetObservationTypeString(UE::Learning::Observation::EType::OrExclusive));
 		OutElementName = NAME_None;
 		OutElement = FLearningAgentsObservationObjectElement();
 		return false;
 	}
 
-	const UE::Learning::Observation::FObjectOrExclusiveParameters Parameters = ObservationObject.GetOrExclusive(Element.ObjectElement);
+	const UE::Learning::Observation::FObjectOrExclusiveParameters Parameters = Object->ObservationObject.GetOrExclusive(Element.ObjectElement);
 	OutElementName = Parameters.ElementName;
 	OutElement = { Parameters.Element };
 	return true;
 }
 
-bool ULearningAgentsObservationObject::GetInclusiveUnionObservationNum(int32& OutNum, const FLearningAgentsObservationObjectElement Element, const FName Name) const
+bool ULearningAgentsObservations::GetInclusiveUnionObservationNum(int32& OutNum, const ULearningAgentsObservationObject* Object, const FLearningAgentsObservationObjectElement Element, const FName Tag)
 {
-	if (!ObservationObject.IsValid(Element.ObjectElement))
+	if (!Object)
 	{
-		UE_LOG(LogLearning, Error, TEXT("%s: Invalid Observation Object."), *GetName());
+		UE_LOG(LogLearning, Error, TEXT("GetInclusiveUnionObservationNum: Object is nullptr."));
 		OutNum = 0;
 		return false;
 	}
 
-	if (ObservationObject.GetName(Element.ObjectElement) != Name)
+	if (!Object->ObservationObject.IsValid(Element.ObjectElement))
 	{
-		UE_LOG(LogLearning, Warning, TEXT("%s: Observation name does not match. Observation is '%s' but asked for '%s'."), *GetName(), *ObservationObject.GetName(Element.ObjectElement).ToString(), *Name.ToString());
+		UE_LOG(LogLearning, Error, TEXT("GetInclusiveUnionObservationNum: Invalid Observation Object."));
+		OutNum = 0;
+		return false;
 	}
 
-	if (ObservationObject.GetType(Element.ObjectElement) != UE::Learning::Observation::EType::OrInclusive)
+	if (Object->ObservationObject.GetTag(Element.ObjectElement) != Tag)
 	{
-		UE_LOG(LogLearning, Error, TEXT("%s: Observation '%s' type does not match. Observation is '%s' but asked for '%s'."),
-			*GetName(),
-			*ObservationObject.GetName(Element.ObjectElement).ToString(),
-			UE::Learning::Agents::Observation::Private::GetObservationTypeString(ObservationObject.GetType(Element.ObjectElement)),
+		UE_LOG(LogLearning, Warning, TEXT("GetInclusiveUnionObservationNum: Observation tag does not match. Observation is '%s' but asked for '%s'."), *Object->ObservationObject.GetTag(Element.ObjectElement).ToString(), *Tag.ToString());
+	}
+
+	if (Object->ObservationObject.GetType(Element.ObjectElement) != UE::Learning::Observation::EType::OrInclusive)
+	{
+		UE_LOG(LogLearning, Error, TEXT("GetInclusiveUnionObservationNum: Observation '%s' type does not match. Observation is '%s' but asked for '%s'."),
+			*Object->ObservationObject.GetTag(Element.ObjectElement).ToString(),
+			UE::Learning::Agents::Observation::Private::GetObservationTypeString(Object->ObservationObject.GetType(Element.ObjectElement)),
 			UE::Learning::Agents::Observation::Private::GetObservationTypeString(UE::Learning::Observation::EType::OrInclusive));
 		OutNum = 0;
 		return false;
 	}
 
-	const UE::Learning::Observation::FObjectOrInclusiveParameters Parameters = ObservationObject.GetOrInclusive(Element.ObjectElement);
+	const UE::Learning::Observation::FObjectOrInclusiveParameters Parameters = Object->ObservationObject.GetOrInclusive(Element.ObjectElement);
 
 	OutNum = Parameters.Elements.Num();
 	return true;
 }
 
-bool ULearningAgentsObservationObject::GetInclusiveUnionObservation(TMap<FName, FLearningAgentsObservationObjectElement>& OutElements, const FLearningAgentsObservationObjectElement Element, const FName Name) const
+bool ULearningAgentsObservations::GetInclusiveUnionObservation(TMap<FName, FLearningAgentsObservationObjectElement>& OutElements, const ULearningAgentsObservationObject* Object, const FLearningAgentsObservationObjectElement Element, const FName Tag)
 {
 	int32 OutElementNum = 0;
-	if (!GetInclusiveUnionObservationNum(OutElementNum, Element, Name))
+	if (!GetInclusiveUnionObservationNum(OutElementNum, Object, Element, Tag))
 	{
 		OutElements.Empty();
 		return false;
@@ -2520,7 +2688,7 @@ bool ULearningAgentsObservationObject::GetInclusiveUnionObservation(TMap<FName,
 	SubElementNames.SetNumUninitialized(OutElementNum);
 	SubElements.SetNumUninitialized(OutElementNum);
 
-	if (!GetInclusiveUnionObservationToArrayViews(SubElementNames, SubElements, Element, Name))
+	if (!GetInclusiveUnionObservationToArrayViews(SubElementNames, SubElements, Object, Element, Tag))
 	{
 		OutElements.Empty();
 		return false;
@@ -2535,10 +2703,10 @@ bool ULearningAgentsObservationObject::GetInclusiveUnionObservation(TMap<FName,
 	return true;
 }
 
-bool ULearningAgentsObservationObject::GetInclusiveUnionObservationToArrays(TArray<FName>& OutElementNames, TArray<FLearningAgentsObservationObjectElement>& OutElements, const FLearningAgentsObservationObjectElement Element, const FName Name) const
+bool ULearningAgentsObservations::GetInclusiveUnionObservationToArrays(TArray<FName>& OutElementNames, TArray<FLearningAgentsObservationObjectElement>& OutElements, const ULearningAgentsObservationObject* Object, const FLearningAgentsObservationObjectElement Element, const FName Tag)
 {
 	int32 OutElementNum = 0;
-	if (!GetInclusiveUnionObservationNum(OutElementNum, Element, Name))
+	if (!GetInclusiveUnionObservationNum(OutElementNum, Object, Element, Tag))
 	{
 		OutElementNames.Empty();
 		OutElements.Empty();
@@ -2548,7 +2716,7 @@ bool ULearningAgentsObservationObject::GetInclusiveUnionObservationToArrays(TArr
 	OutElementNames.SetNumUninitialized(OutElementNum);
 	OutElements.SetNumUninitialized(OutElementNum);
 
-	if (!GetInclusiveUnionObservationToArrayViews(OutElementNames, OutElements, Element, Name))
+	if (!GetInclusiveUnionObservationToArrayViews(OutElementNames, OutElements, Object, Element, Tag))
 	{
 		OutElementNames.Empty();
 		OutElements.Empty();
@@ -2558,40 +2726,46 @@ bool ULearningAgentsObservationObject::GetInclusiveUnionObservationToArrays(TArr
 	return true;
 }
 
-bool ULearningAgentsObservationObject::GetInclusiveUnionObservationToArrayViews(TArrayView<FName> OutElementNames, TArrayView<FLearningAgentsObservationObjectElement> OutElements, const FLearningAgentsObservationObjectElement Element, const FName Name) const
+bool ULearningAgentsObservations::GetInclusiveUnionObservationToArrayViews(TArrayView<FName> OutElementNames, TArrayView<FLearningAgentsObservationObjectElement> OutElements, const ULearningAgentsObservationObject* Object, const FLearningAgentsObservationObjectElement Element, const FName Tag)
 {
-	if (!ObservationObject.IsValid(Element.ObjectElement))
+	if (!Object)
 	{
-		UE_LOG(LogLearning, Error, TEXT("%s: Invalid Observation Object."), *GetName());
+		UE_LOG(LogLearning, Error, TEXT("GetInclusiveUnionObservationToArrayViews: Object is nullptr."));
 		UE::Learning::Array::Set<1, FName>(OutElementNames, NAME_None);
 		UE::Learning::Array::Set<1, FLearningAgentsObservationObjectElement>(OutElements, FLearningAgentsObservationObjectElement());
 		return false;
 	}
 
-	if (ObservationObject.GetName(Element.ObjectElement) != Name)
+	if (!Object->ObservationObject.IsValid(Element.ObjectElement))
 	{
-		UE_LOG(LogLearning, Warning, TEXT("%s: Observation name does not match. Observation is '%s' but asked for '%s'."), *GetName(), *ObservationObject.GetName(Element.ObjectElement).ToString(), *Name.ToString());
+		UE_LOG(LogLearning, Error, TEXT("GetInclusiveUnionObservationToArrayViews: Invalid Observation Object."));
+		UE::Learning::Array::Set<1, FName>(OutElementNames, NAME_None);
+		UE::Learning::Array::Set<1, FLearningAgentsObservationObjectElement>(OutElements, FLearningAgentsObservationObjectElement());
+		return false;
 	}
 
-	if (ObservationObject.GetType(Element.ObjectElement) != UE::Learning::Observation::EType::OrInclusive)
+	if (Object->ObservationObject.GetTag(Element.ObjectElement) != Tag)
 	{
-		UE_LOG(LogLearning, Error, TEXT("%s: Observation '%s' type does not match. Observation is '%s' but asked for '%s'."),
-			*GetName(),
-			*ObservationObject.GetName(Element.ObjectElement).ToString(),
-			UE::Learning::Agents::Observation::Private::GetObservationTypeString(ObservationObject.GetType(Element.ObjectElement)),
+		UE_LOG(LogLearning, Warning, TEXT("GetInclusiveUnionObservationToArrayViews: Observation tag does not match. Observation is '%s' but asked for '%s'."), *Object->ObservationObject.GetTag(Element.ObjectElement).ToString(), *Tag.ToString());
+	}
+
+	if (Object->ObservationObject.GetType(Element.ObjectElement) != UE::Learning::Observation::EType::OrInclusive)
+	{
+		UE_LOG(LogLearning, Error, TEXT("GetInclusiveUnionObservationToArrayViews: Observation '%s' type does not match. Observation is '%s' but asked for '%s'."),
+			*Object->ObservationObject.GetTag(Element.ObjectElement).ToString(),
+			UE::Learning::Agents::Observation::Private::GetObservationTypeString(Object->ObservationObject.GetType(Element.ObjectElement)),
 			UE::Learning::Agents::Observation::Private::GetObservationTypeString(UE::Learning::Observation::EType::OrInclusive));
 		UE::Learning::Array::Set<1, FName>(OutElementNames, NAME_None);
 		UE::Learning::Array::Set<1, FLearningAgentsObservationObjectElement>(OutElements, FLearningAgentsObservationObjectElement());
 		return false;
 	}
 
-	const UE::Learning::Observation::FObjectOrInclusiveParameters Parameters = ObservationObject.GetOrInclusive(Element.ObjectElement);
+	const UE::Learning::Observation::FObjectOrInclusiveParameters Parameters = Object->ObservationObject.GetOrInclusive(Element.ObjectElement);
 
 	if (Parameters.Elements.Num() != OutElements.Num())
 	{
-		UE_LOG(LogLearning, Error, TEXT("%s: Observation '%s' size does not match. Observation is '%i' elements but asked for '%i'."),
-			*GetName(),
-			*ObservationObject.GetName(Element.ObjectElement).ToString(),
+		UE_LOG(LogLearning, Error, TEXT("GetInclusiveUnionObservationToArrayViews: Observation '%s' size does not match. Observation is '%i' elements but asked for '%i'."),
+			*Object->ObservationObject.GetTag(Element.ObjectElement).ToString(),
 			Parameters.Elements.Num(), OutElements.Num());
 		UE::Learning::Array::Set<1, FName>(OutElementNames, NAME_None);
 		UE::Learning::Array::Set<1, FLearningAgentsObservationObjectElement>(OutElements, FLearningAgentsObservationObjectElement());
@@ -2600,9 +2774,9 @@ bool ULearningAgentsObservationObject::GetInclusiveUnionObservationToArrayViews(
 
 	for (int32 ElementIdx = 0; ElementIdx < Parameters.Elements.Num(); ElementIdx++)
 	{
-		if (!ObservationObject.IsValid(Parameters.Elements[ElementIdx]))
+		if (!Object->ObservationObject.IsValid(Parameters.Elements[ElementIdx]))
 		{
-			UE_LOG(LogLearning, Error, TEXT("%s: Invalid Observation Object."), *GetName());
+			UE_LOG(LogLearning, Error, TEXT("GetInclusiveUnionObservationToArrayViews: Invalid Observation Object."));
 			UE::Learning::Array::Set<1, FLearningAgentsObservationObjectElement>(OutElements, FLearningAgentsObservationObjectElement());
 			UE::Learning::Array::Set<1, FName>(OutElementNames, NAME_None);
 			return false;
@@ -2615,39 +2789,45 @@ bool ULearningAgentsObservationObject::GetInclusiveUnionObservationToArrayViews(
 	return true;
 }
 
-bool ULearningAgentsObservationObject::GetStaticArrayObservationNum(int32& OutNum, const FLearningAgentsObservationObjectElement Element, const FName Name) const
+bool ULearningAgentsObservations::GetStaticArrayObservationNum(int32& OutNum, const ULearningAgentsObservationObject* Object, const FLearningAgentsObservationObjectElement Element, const FName Tag)
 {
-	if (!ObservationObject.IsValid(Element.ObjectElement))
+	if (!Object)
+	{
+		UE_LOG(LogLearning, Error, TEXT("GetStaticArrayObservationNum: Object is nullptr."));
+		OutNum = 0;
+		return false;
+	}
+
+	if (!Object->ObservationObject.IsValid(Element.ObjectElement))
 	{
-		UE_LOG(LogLearning, Error, TEXT("%s: Invalid Observation Object."), *GetName());
+		UE_LOG(LogLearning, Error, TEXT("GetStaticArrayObservationNum: Invalid Observation Object."));
 		OutNum = 0;
 		return false;
 	}
 
-	if (ObservationObject.GetName(Element.ObjectElement) != Name)
+	if (Object->ObservationObject.GetTag(Element.ObjectElement) != Tag)
 	{
-		UE_LOG(LogLearning, Warning, TEXT("%s: Observation name does not match. Observation is '%s' but asked for '%s'."), *GetName(), *ObservationObject.GetName(Element.ObjectElement).ToString(), *Name.ToString());
+		UE_LOG(LogLearning, Warning, TEXT("GetStaticArrayObservationNum: Observation tag does not match. Observation is '%s' but asked for '%s'."), *Object->ObservationObject.GetTag(Element.ObjectElement).ToString(), *Tag.ToString());
 	}
 
-	if (ObservationObject.GetType(Element.ObjectElement) != UE::Learning::Observation::EType::Array)
+	if (Object->ObservationObject.GetType(Element.ObjectElement) != UE::Learning::Observation::EType::Array)
 	{
-		UE_LOG(LogLearning, Error, TEXT("%s: Observation '%s' type does not match. Observation is '%s' but asked for '%s'."),
-			*GetName(),
-			*ObservationObject.GetName(Element.ObjectElement).ToString(),
-			UE::Learning::Agents::Observation::Private::GetObservationTypeString(ObservationObject.GetType(Element.ObjectElement)),
+		UE_LOG(LogLearning, Error, TEXT("GetStaticArrayObservationNum: Observation '%s' type does not match. Observation is '%s' but asked for '%s'."),
+			*Object->ObservationObject.GetTag(Element.ObjectElement).ToString(),
+			UE::Learning::Agents::Observation::Private::GetObservationTypeString(Object->ObservationObject.GetType(Element.ObjectElement)),
 			UE::Learning::Agents::Observation::Private::GetObservationTypeString(UE::Learning::Observation::EType::Array));
 		OutNum = 0;
 		return false;
 	}
 
-	OutNum = ObservationObject.GetArray(Element.ObjectElement).Elements.Num();
+	OutNum = Object->ObservationObject.GetArray(Element.ObjectElement).Elements.Num();
 	return true;
 }
 
-bool ULearningAgentsObservationObject::GetStaticArrayObservation(TArray<FLearningAgentsObservationObjectElement>& OutElements, const FLearningAgentsObservationObjectElement Element, const FName Name) const
+bool ULearningAgentsObservations::GetStaticArrayObservation(TArray<FLearningAgentsObservationObjectElement>& OutElements, const ULearningAgentsObservationObject* Object, const FLearningAgentsObservationObjectElement Element, const FName Tag)
 {
 	int32 OutElementNum = 0;
-	if (!GetStaticArrayObservationNum(OutElementNum, Element, Name))
+	if (!GetStaticArrayObservationNum(OutElementNum, Object, Element, Tag))
 	{
 		OutElements.Empty();
 		return false;
@@ -2655,7 +2835,7 @@ bool ULearningAgentsObservationObject::GetStaticArrayObservation(TArray<FLearnin
 
 	OutElements.SetNumUninitialized(OutElementNum);
 
-	if (!GetStaticArrayObservationToArrayView(OutElements, Element, Name))
+	if (!GetStaticArrayObservationToArrayView(OutElements, Object, Element, Tag))
 	{
 		OutElements.Empty();
 		return false;
@@ -2664,42 +2844,48 @@ bool ULearningAgentsObservationObject::GetStaticArrayObservation(TArray<FLearnin
 	return true;
 }
 
-bool ULearningAgentsObservationObject::GetStaticArrayObservationToArrayView(TArrayView<FLearningAgentsObservationObjectElement> OutElements, const FLearningAgentsObservationObjectElement Element, const FName Name) const
+bool ULearningAgentsObservations::GetStaticArrayObservationToArrayView(TArrayView<FLearningAgentsObservationObjectElement> OutElements, const ULearningAgentsObservationObject* Object, const FLearningAgentsObservationObjectElement Element, const FName Tag)
 {
-	if (!ObservationObject.IsValid(Element.ObjectElement))
+	if (!Object)
 	{
-		UE_LOG(LogLearning, Error, TEXT("%s: Invalid Observation Object."), *GetName());
+		UE_LOG(LogLearning, Error, TEXT("GetStaticArrayObservationNum: Object is nullptr."));
 		UE::Learning::Array::Set<1, FLearningAgentsObservationObjectElement>(OutElements, FLearningAgentsObservationObjectElement());
 		return false;
 	}
 
-	if (ObservationObject.GetName(Element.ObjectElement) != Name)
+	if (!Object->ObservationObject.IsValid(Element.ObjectElement))
 	{
-		UE_LOG(LogLearning, Warning, TEXT("%s: Observation name does not match. Observation is '%s' but asked for '%s'."), *GetName(), *ObservationObject.GetName(Element.ObjectElement).ToString(), *Name.ToString());
+		UE_LOG(LogLearning, Error, TEXT("GetStaticArrayObservationToArrayView: Invalid Observation Object."));
+		UE::Learning::Array::Set<1, FLearningAgentsObservationObjectElement>(OutElements, FLearningAgentsObservationObjectElement());
+		return false;
 	}
 
-	if (ObservationObject.GetType(Element.ObjectElement) != UE::Learning::Observation::EType::Array)
+	if (Object->ObservationObject.GetTag(Element.ObjectElement) != Tag)
 	{
-		UE_LOG(LogLearning, Error, TEXT("%s: Observation '%s' type does not match. Observation is '%s' but asked for '%s'."),
-			*GetName(),
-			*ObservationObject.GetName(Element.ObjectElement).ToString(),
-			UE::Learning::Agents::Observation::Private::GetObservationTypeString(ObservationObject.GetType(Element.ObjectElement)),
+		UE_LOG(LogLearning, Warning, TEXT("GetStaticArrayObservationToArrayView: Observation tag does not match. Observation is '%s' but asked for '%s'."), *Object->ObservationObject.GetTag(Element.ObjectElement).ToString(), *Tag.ToString());
+	}
+
+	if (Object->ObservationObject.GetType(Element.ObjectElement) != UE::Learning::Observation::EType::Array)
+	{
+		UE_LOG(LogLearning, Error, TEXT("GetStaticArrayObservationToArrayView: Observation '%s' type does not match. Observation is '%s' but asked for '%s'."),
+			*Object->ObservationObject.GetTag(Element.ObjectElement).ToString(),
+			UE::Learning::Agents::Observation::Private::GetObservationTypeString(Object->ObservationObject.GetType(Element.ObjectElement)),
 			UE::Learning::Agents::Observation::Private::GetObservationTypeString(UE::Learning::Observation::EType::Array));
 		UE::Learning::Array::Set<1, FLearningAgentsObservationObjectElement>(OutElements, FLearningAgentsObservationObjectElement());
 		return false;
 	}
 
-	const UE::Learning::Observation::FObjectArrayParameters Parameters = ObservationObject.GetArray(Element.ObjectElement);
+	const UE::Learning::Observation::FObjectArrayParameters Parameters = Object->ObservationObject.GetArray(Element.ObjectElement);
 
 	if (Parameters.Elements.Num() == 0)
 	{
-		UE_LOG(LogLearning, Warning, TEXT("%s: Getting zero-sized Static Array Observation."), *GetName());
+		UE_LOG(LogLearning, Warning, TEXT("GetStaticArrayObservationToArrayView: Getting zero-sized Static Array Observation."));
 	}
 
 	if (Parameters.Elements.Num() != OutElements.Num())
 	{
-		UE_LOG(LogLearning, Error, TEXT("%s: Observation '%s' size does not match. Observation is '%i' elements but asked for '%i'."),
-			*GetName(), *ObservationObject.GetName(Element.ObjectElement).ToString(),
+		UE_LOG(LogLearning, Error, TEXT("GetStaticArrayObservationToArrayView: Observation '%s' size does not match. Observation is '%i' elements but asked for '%i'."),
+			*Object->ObservationObject.GetTag(Element.ObjectElement).ToString(),
 			Parameters.Elements.Num(), OutElements.Num());
 		UE::Learning::Array::Set<1, FLearningAgentsObservationObjectElement>(OutElements, FLearningAgentsObservationObjectElement());
 		return false;
@@ -2707,9 +2893,9 @@ bool ULearningAgentsObservationObject::GetStaticArrayObservationToArrayView(TArr
 
 	for (int32 ElementIdx = 0; ElementIdx < Parameters.Elements.Num(); ElementIdx++)
 	{
-		if (!ObservationObject.IsValid(Parameters.Elements[ElementIdx]))
+		if (!Object->ObservationObject.IsValid(Parameters.Elements[ElementIdx]))
 		{
-			UE_LOG(LogLearning, Error, TEXT("%s: Invalid Observation Object."), *GetName());
+			UE_LOG(LogLearning, Error, TEXT("GetStaticArrayObservationToArrayView: Invalid Observation Object."));
 			UE::Learning::Array::Set<1, FLearningAgentsObservationObjectElement>(OutElements, FLearningAgentsObservationObjectElement());
 			return false;
 		}
@@ -2720,73 +2906,85 @@ bool ULearningAgentsObservationObject::GetStaticArrayObservationToArrayView(TArr
 	return true;
 }
 
-bool ULearningAgentsObservationObject::GetSetObservationNum(int32& OutNum, const FLearningAgentsObservationObjectElement Element, const FName Name) const
+bool ULearningAgentsObservations::GetSetObservationNum(int32& OutNum, const ULearningAgentsObservationObject* Object, const FLearningAgentsObservationObjectElement Element, const FName Tag)
 {
-	if (!ObservationObject.IsValid(Element.ObjectElement))
+	if (!Object)
+	{
+		UE_LOG(LogLearning, Error, TEXT("GetSetObservationNum: Object is nullptr."));
+		OutNum = 0;
+		return false;
+	}
+
+	if (!Object->ObservationObject.IsValid(Element.ObjectElement))
 	{
-		UE_LOG(LogLearning, Error, TEXT("%s: Invalid Observation Object."), *GetName());
+		UE_LOG(LogLearning, Error, TEXT("GetSetObservationNum: Invalid Observation Object."));
 		OutNum = 0;
 		return false;
 	}
 
-	if (ObservationObject.GetName(Element.ObjectElement) != Name)
+	if (Object->ObservationObject.GetTag(Element.ObjectElement) != Tag)
 	{
-		UE_LOG(LogLearning, Warning, TEXT("%s: Observation name does not match. Observation is '%s' but asked for '%s'."), *GetName(), *ObservationObject.GetName(Element.ObjectElement).ToString(), *Name.ToString());
+		UE_LOG(LogLearning, Warning, TEXT("GetSetObservationNum: Observation tag does not match. Observation is '%s' but asked for '%s'."), *Object->ObservationObject.GetTag(Element.ObjectElement).ToString(), *Tag.ToString());
 	}
 
-	if (ObservationObject.GetType(Element.ObjectElement) != UE::Learning::Observation::EType::Set)
+	if (Object->ObservationObject.GetType(Element.ObjectElement) != UE::Learning::Observation::EType::Set)
 	{
-		UE_LOG(LogLearning, Error, TEXT("%s: Observation '%s' type does not match. Observation is '%s' but asked for '%s'."),
-			*GetName(),
-			*ObservationObject.GetName(Element.ObjectElement).ToString(),
-			UE::Learning::Agents::Observation::Private::GetObservationTypeString(ObservationObject.GetType(Element.ObjectElement)),
+		UE_LOG(LogLearning, Error, TEXT("GetSetObservationNum: Observation '%s' type does not match. Observation is '%s' but asked for '%s'."),
+			*Object->ObservationObject.GetTag(Element.ObjectElement).ToString(),
+			UE::Learning::Agents::Observation::Private::GetObservationTypeString(Object->ObservationObject.GetType(Element.ObjectElement)),
 			UE::Learning::Agents::Observation::Private::GetObservationTypeString(UE::Learning::Observation::EType::Set));
 		OutNum = 0;
 		return false;
 	}
 
-	OutNum = ObservationObject.GetSet(Element.ObjectElement).Elements.Num();
+	OutNum = Object->ObservationObject.GetSet(Element.ObjectElement).Elements.Num();
 	return true;
 }
 
-bool ULearningAgentsObservationObject::GetSetObservation(TSet<FLearningAgentsObservationObjectElement>& OutElements, const FLearningAgentsObservationObjectElement Element, const FName Name) const
+bool ULearningAgentsObservations::GetSetObservation(TSet<FLearningAgentsObservationObjectElement>& OutElements, const ULearningAgentsObservationObject* Object, const FLearningAgentsObservationObjectElement Element, const FName Tag)
 {
 	int32 OutElementNum = 0;
-	if (!GetSetObservationNum(OutElementNum, Element, Name))
+	if (!GetSetObservationNum(OutElementNum, Object, Element, Tag))
 	{
 		OutElements.Empty();
 		return false;
 	}
 
-	if (!ObservationObject.IsValid(Element.ObjectElement))
+	if (!Object)
 	{
-		UE_LOG(LogLearning, Error, TEXT("%s: Invalid Observation Object."), *GetName());
+		UE_LOG(LogLearning, Error, TEXT("GetSetObservation: Object is nullptr."));
 		OutElements.Empty();
 		return false;
 	}
 
-	if (ObservationObject.GetName(Element.ObjectElement) != Name)
+	if (!Object->ObservationObject.IsValid(Element.ObjectElement))
 	{
-		UE_LOG(LogLearning, Warning, TEXT("%s: Observation name does not match. Observation is '%s' but asked for '%s'."), *GetName(), *ObservationObject.GetName(Element.ObjectElement).ToString(), *Name.ToString());
+		UE_LOG(LogLearning, Error, TEXT("GetSetObservation: Invalid Observation Object."));
+		OutElements.Empty();
+		return false;
 	}
 
-	if (ObservationObject.GetType(Element.ObjectElement) != UE::Learning::Observation::EType::Set)
+	if (Object->ObservationObject.GetTag(Element.ObjectElement) != Tag)
 	{
-		UE_LOG(LogLearning, Error, TEXT("%s: Observation '%s' type does not match. Observation is '%s' but asked for '%s'."),
-			*GetName(),
-			*ObservationObject.GetName(Element.ObjectElement).ToString(),
-			UE::Learning::Agents::Observation::Private::GetObservationTypeString(ObservationObject.GetType(Element.ObjectElement)),
+		UE_LOG(LogLearning, Warning, TEXT("GetSetObservation: Observation tag does not match. Observation is '%s' but asked for '%s'."), *Object->ObservationObject.GetTag(Element.ObjectElement).ToString(), *Tag.ToString());
+	}
+
+	if (Object->ObservationObject.GetType(Element.ObjectElement) != UE::Learning::Observation::EType::Set)
+	{
+		UE_LOG(LogLearning, Error, TEXT("GetSetObservation: Observation '%s' type does not match. Observation is '%s' but asked for '%s'."),
+			*Object->ObservationObject.GetTag(Element.ObjectElement).ToString(),
+			UE::Learning::Agents::Observation::Private::GetObservationTypeString(Object->ObservationObject.GetType(Element.ObjectElement)),
 			UE::Learning::Agents::Observation::Private::GetObservationTypeString(UE::Learning::Observation::EType::Set));
 		OutElements.Empty();
 		return false;
 	}
 
-	const UE::Learning::Observation::FObjectSetParameters Parameters = ObservationObject.GetSet(Element.ObjectElement);
+	const UE::Learning::Observation::FObjectSetParameters Parameters = Object->ObservationObject.GetSet(Element.ObjectElement);
 
 	if (Parameters.Elements.Num() != OutElements.Num())
 	{
-		UE_LOG(LogLearning, Error, TEXT("%s: Observation '%s' size does not match. Observation is '%i' elements but asked for '%i'."),
-			*GetName(), *ObservationObject.GetName(Element.ObjectElement).ToString(),
+		UE_LOG(LogLearning, Error, TEXT("GetSetObservation: Observation '%s' size does not match. Observation is '%i' elements but asked for '%i'."),
+			*Object->ObservationObject.GetTag(Element.ObjectElement).ToString(),
 			Parameters.Elements.Num(), OutElements.Num());
 		OutElements.Empty();
 		return false;
@@ -2795,9 +2993,9 @@ bool ULearningAgentsObservationObject::GetSetObservation(TSet<FLearningAgentsObs
 	OutElements.Empty(Parameters.Elements.Num());
 	for (int32 ElementIdx = 0; ElementIdx < Parameters.Elements.Num(); ElementIdx++)
 	{
-		if (!ObservationObject.IsValid(Parameters.Elements[ElementIdx]))
+		if (!Object->ObservationObject.IsValid(Parameters.Elements[ElementIdx]))
 		{
-			UE_LOG(LogLearning, Error, TEXT("%s: Invalid Observation Object."), *GetName());
+			UE_LOG(LogLearning, Error, TEXT("GetSetObservation: Invalid Observation Object."));
 			OutElements.Empty();
 			return false;
 		}
@@ -2808,10 +3006,10 @@ bool ULearningAgentsObservationObject::GetSetObservation(TSet<FLearningAgentsObs
 	return true;
 }
 
-bool ULearningAgentsObservationObject::GetSetObservationToArray(TArray<FLearningAgentsObservationObjectElement>& OutElements, const FLearningAgentsObservationObjectElement Element, const FName Name) const
+bool ULearningAgentsObservations::GetSetObservationToArray(TArray<FLearningAgentsObservationObjectElement>& OutElements, const ULearningAgentsObservationObject* Object, const FLearningAgentsObservationObjectElement Element, const FName Tag)
 {
 	int32 OutElementNum = 0;
-	if (!GetSetObservationNum(OutElementNum, Element, Name))
+	if (!GetSetObservationNum(OutElementNum, Object, Element, Tag))
 	{
 		OutElements.Empty();
 		return false;
@@ -2819,7 +3017,7 @@ bool ULearningAgentsObservationObject::GetSetObservationToArray(TArray<FLearning
 
 	OutElements.SetNumUninitialized(OutElementNum);
 
-	if (!GetSetObservationToArrayView(OutElements, Element, Name))
+	if (!GetSetObservationToArrayView(OutElements, Object, Element, Tag))
 	{
 		OutElements.Empty();
 		return false;
@@ -2828,37 +3026,43 @@ bool ULearningAgentsObservationObject::GetSetObservationToArray(TArray<FLearning
 	return true;
 }
 
-bool ULearningAgentsObservationObject::GetSetObservationToArrayView(TArrayView<FLearningAgentsObservationObjectElement> OutElements, const FLearningAgentsObservationObjectElement Element, const FName Name) const
+bool ULearningAgentsObservations::GetSetObservationToArrayView(TArrayView<FLearningAgentsObservationObjectElement> OutElements, const ULearningAgentsObservationObject* Object, const FLearningAgentsObservationObjectElement Element, const FName Tag)
 {
-	if (!ObservationObject.IsValid(Element.ObjectElement))
+	if (!Object)
+	{
+		UE_LOG(LogLearning, Error, TEXT("GetSetObservationToArrayView: Object is nullptr."));
+		UE::Learning::Array::Set<1, FLearningAgentsObservationObjectElement>(OutElements, FLearningAgentsObservationObjectElement());
+		return false;
+	}
+
+	if (!Object->ObservationObject.IsValid(Element.ObjectElement))
 	{
-		UE_LOG(LogLearning, Error, TEXT("%s: Invalid Observation Object."), *GetName());
+		UE_LOG(LogLearning, Error, TEXT("GetSetObservationToArrayView: Invalid Observation Object."));
 		UE::Learning::Array::Set<1, FLearningAgentsObservationObjectElement>(OutElements, FLearningAgentsObservationObjectElement());
 		return false;
 	}
 
-	if (ObservationObject.GetName(Element.ObjectElement) != Name)
+	if (Object->ObservationObject.GetTag(Element.ObjectElement) != Tag)
 	{
-		UE_LOG(LogLearning, Warning, TEXT("%s: Observation name does not match. Observation is '%s' but asked for '%s'."), *GetName(), *ObservationObject.GetName(Element.ObjectElement).ToString(), *Name.ToString());
+		UE_LOG(LogLearning, Warning, TEXT("GetSetObservationToArrayView: Observation tag does not match. Observation is '%s' but asked for '%s'."), *Object->ObservationObject.GetTag(Element.ObjectElement).ToString(), *Tag.ToString());
 	}
 
-	if (ObservationObject.GetType(Element.ObjectElement) != UE::Learning::Observation::EType::Set)
+	if (Object->ObservationObject.GetType(Element.ObjectElement) != UE::Learning::Observation::EType::Set)
 	{
-		UE_LOG(LogLearning, Error, TEXT("%s: Observation '%s' type does not match. Observation is '%s' but asked for '%s'."),
-			*GetName(),
-			*ObservationObject.GetName(Element.ObjectElement).ToString(),
-			UE::Learning::Agents::Observation::Private::GetObservationTypeString(ObservationObject.GetType(Element.ObjectElement)),
+		UE_LOG(LogLearning, Error, TEXT("GetSetObservationToArrayView: Observation '%s' type does not match. Observation is '%s' but asked for '%s'."),
+			*Object->ObservationObject.GetTag(Element.ObjectElement).ToString(),
+			UE::Learning::Agents::Observation::Private::GetObservationTypeString(Object->ObservationObject.GetType(Element.ObjectElement)),
 			UE::Learning::Agents::Observation::Private::GetObservationTypeString(UE::Learning::Observation::EType::Set));
 		UE::Learning::Array::Set<1, FLearningAgentsObservationObjectElement>(OutElements, FLearningAgentsObservationObjectElement());
 		return false;
 	}
 
-	const UE::Learning::Observation::FObjectSetParameters Parameters = ObservationObject.GetSet(Element.ObjectElement);
+	const UE::Learning::Observation::FObjectSetParameters Parameters = Object->ObservationObject.GetSet(Element.ObjectElement);
 
 	if (Parameters.Elements.Num() != OutElements.Num())
 	{
-		UE_LOG(LogLearning, Error, TEXT("%s: Observation '%s' size does not match. Observation is '%i' elements but asked for '%i'."),
-			*GetName(), *ObservationObject.GetName(Element.ObjectElement).ToString(),
+		UE_LOG(LogLearning, Error, TEXT("GetSetObservationToArrayView: Observation '%s' size does not match. Observation is '%i' elements but asked for '%i'."),
+			*Object->ObservationObject.GetTag(Element.ObjectElement).ToString(),
 			Parameters.Elements.Num(), OutElements.Num());
 		UE::Learning::Array::Set<1, FLearningAgentsObservationObjectElement>(OutElements, FLearningAgentsObservationObjectElement());
 		return false;
@@ -2866,9 +3070,9 @@ bool ULearningAgentsObservationObject::GetSetObservationToArrayView(TArrayView<F
 
 	for (int32 ElementIdx = 0; ElementIdx < Parameters.Elements.Num(); ElementIdx++)
 	{
-		if (!ObservationObject.IsValid(Parameters.Elements[ElementIdx]))
+		if (!Object->ObservationObject.IsValid(Parameters.Elements[ElementIdx]))
 		{
-			UE_LOG(LogLearning, Error, TEXT("%s: Invalid Observation Object."), *GetName());
+			UE_LOG(LogLearning, Error, TEXT("GetSetObservationToArrayView: Invalid Observation Object."));
 			UE::Learning::Array::Set<1, FLearningAgentsObservationObjectElement>(OutElements, FLearningAgentsObservationObjectElement());
 			return false;
 		}
@@ -2879,11 +3083,11 @@ bool ULearningAgentsObservationObject::GetSetObservationToArrayView(TArrayView<F
 	return true;
 }
 
-bool ULearningAgentsObservationObject::GetPairObservation(FLearningAgentsObservationObjectElement& OutKey, FLearningAgentsObservationObjectElement& OutValue, const FLearningAgentsObservationObjectElement Element, const FName Name) const
+bool ULearningAgentsObservations::GetPairObservation(FLearningAgentsObservationObjectElement& OutKey, FLearningAgentsObservationObjectElement& OutValue, const ULearningAgentsObservationObject* Object, const FLearningAgentsObservationObjectElement Element, const FName Tag)
 {
 	TStaticArray<FName, 2> OutElementNames;
 	TStaticArray<FLearningAgentsObservationObjectElement, 2> OutElements;
-	if (!GetStructObservationToArrayViews(OutElementNames, OutElements, Element, Name))
+	if (!GetStructObservationToArrayViews(OutElementNames, OutElements, Object, Element, Tag))
 	{
 		OutKey = FLearningAgentsObservationObjectElement();
 		OutValue = FLearningAgentsObservationObjectElement();
@@ -2895,15 +3099,15 @@ bool ULearningAgentsObservationObject::GetPairObservation(FLearningAgentsObserva
 	return true;
 }
 
-bool ULearningAgentsObservationObject::GetArrayObservationNum(int32& OutNum, const FLearningAgentsObservationObjectElement Element, const FName Name) const
+bool ULearningAgentsObservations::GetArrayObservationNum(int32& OutNum, const ULearningAgentsObservationObject* Object, const FLearningAgentsObservationObjectElement Element, const FName Tag)
 {
-	return GetSetObservationNum(OutNum, Element, Name);
+	return GetSetObservationNum(OutNum, Object, Element, Tag);
 }
 
-bool ULearningAgentsObservationObject::GetArrayObservation(TArray<FLearningAgentsObservationObjectElement>& OutElements, const FLearningAgentsObservationObjectElement Element, const FName Name) const
+bool ULearningAgentsObservations::GetArrayObservation(TArray<FLearningAgentsObservationObjectElement>& OutElements, const ULearningAgentsObservationObject* Object, const FLearningAgentsObservationObjectElement Element, const FName Tag)
 {
 	int32 OutElementNum = 0;
-	if (!GetArrayObservationNum(OutElementNum, Element, Name))
+	if (!GetArrayObservationNum(OutElementNum, Object, Element, Tag))
 	{
 		OutElements.Empty();
 		return false;
@@ -2911,7 +3115,7 @@ bool ULearningAgentsObservationObject::GetArrayObservation(TArray<FLearningAgent
 
 	OutElements.SetNumUninitialized(OutElementNum);
 
-	if (!GetArrayObservationToArrayView(OutElements, Element, Name))
+	if (!GetArrayObservationToArrayView(OutElements, Object, Element, Tag))
 	{
 		OutElements.Empty();
 		return false;
@@ -2920,11 +3124,11 @@ bool ULearningAgentsObservationObject::GetArrayObservation(TArray<FLearningAgent
 	return true;
 }
 
-bool ULearningAgentsObservationObject::GetArrayObservationToArrayView(TArrayView<FLearningAgentsObservationObjectElement> OutElements, const FLearningAgentsObservationObjectElement Element, const FName Name) const
+bool ULearningAgentsObservations::GetArrayObservationToArrayView(TArrayView<FLearningAgentsObservationObjectElement> OutElements, const ULearningAgentsObservationObject* Object, const FLearningAgentsObservationObjectElement Element, const FName Tag)
 {
 	TArray<FLearningAgentsObservationObjectElement, TInlineAllocator<16>> Pairs;
 	Pairs.SetNumUninitialized(OutElements.Num());
-	if (!GetSetObservationToArrayView(Pairs, Element, Name))
+	if (!GetSetObservationToArrayView(Pairs, Object, Element, Tag))
 	{
 		UE::Learning::Array::Set<1, FLearningAgentsObservationObjectElement>(OutElements, FLearningAgentsObservationObjectElement());
 		return false;
@@ -2933,14 +3137,14 @@ bool ULearningAgentsObservationObject::GetArrayObservationToArrayView(TArrayView
 	for (int32 PairIdx = 0; PairIdx < Pairs.Num(); PairIdx++)
 	{
 		FLearningAgentsObservationObjectElement Key, Value;
-		if (!GetPairObservation(Pairs[PairIdx], Key, Value))
+		if (!GetPairObservation(Key, Value, Object, Pairs[PairIdx]))
 		{
 			UE::Learning::Array::Set<1, FLearningAgentsObservationObjectElement>(OutElements, FLearningAgentsObservationObjectElement());
 			return false;
 		}
 
 		int32 Index = INDEX_NONE;
-		if (!GetIndexObservation(Index, Key))
+		if (!GetIndexObservation(Index, Object, Key))
 		{
 			UE::Learning::Array::Set<1, FLearningAgentsObservationObjectElement>(OutElements, FLearningAgentsObservationObjectElement());
 			return false;
@@ -2952,15 +3156,15 @@ bool ULearningAgentsObservationObject::GetArrayObservationToArrayView(TArrayView
 	return true;
 }
 
-bool ULearningAgentsObservationObject::GetMapObservationNum(int32& OutNum, const FLearningAgentsObservationObjectElement Element, const FName Name) const
+bool ULearningAgentsObservations::GetMapObservationNum(int32& OutNum, const ULearningAgentsObservationObject* Object, const FLearningAgentsObservationObjectElement Element, const FName Tag)
 {
-	return GetSetObservationNum(OutNum, Element, Name);
+	return GetSetObservationNum(OutNum, Object, Element, Tag);
 }
 
-bool ULearningAgentsObservationObject::GetMapObservation(TMap<FLearningAgentsObservationObjectElement, FLearningAgentsObservationObjectElement>& OutElements, const FLearningAgentsObservationObjectElement Element, const FName Name) const
+bool ULearningAgentsObservations::GetMapObservation(TMap<FLearningAgentsObservationObjectElement, FLearningAgentsObservationObjectElement>& OutElements, const ULearningAgentsObservationObject* Object, const FLearningAgentsObservationObjectElement Element, const FName Tag)
 {
 	int32 OutElementNum = 0;
-	if (!GetMapObservationNum(OutElementNum, Element, Name))
+	if (!GetMapObservationNum(OutElementNum, Object, Element, Tag))
 	{
 		OutElements.Empty();
 		return false;
@@ -2968,7 +3172,7 @@ bool ULearningAgentsObservationObject::GetMapObservation(TMap<FLearningAgentsObs
 
 	TArray<FLearningAgentsObservationObjectElement, TInlineAllocator<16>> Pairs;
 	Pairs.SetNumUninitialized(OutElementNum);
-	if (!GetSetObservationToArrayView(Pairs, Element, Name))
+	if (!GetSetObservationToArrayView(Pairs, Object, Element, Tag))
 	{
 		OutElements.Empty();
 		return false;
@@ -2978,7 +3182,7 @@ bool ULearningAgentsObservationObject::GetMapObservation(TMap<FLearningAgentsObs
 	for (int32 PairIdx = 0; PairIdx < OutElementNum; PairIdx++)
 	{
 		FLearningAgentsObservationObjectElement Key, Value;
-		if (!GetPairObservation(Pairs[PairIdx], Key, Value))
+		if (!GetPairObservation(Key, Value, Object, Pairs[PairIdx]))
 		{
 			OutElements.Empty();
 			return false;
@@ -2990,10 +3194,10 @@ bool ULearningAgentsObservationObject::GetMapObservation(TMap<FLearningAgentsObs
 	return true;
 }
 
-bool ULearningAgentsObservationObject::GetMapObservationToArrays(TArray<FLearningAgentsObservationObjectElement>& OutKeys, TArray<FLearningAgentsObservationObjectElement>& OutValues, const FLearningAgentsObservationObjectElement Element, const FName Name) const
+bool ULearningAgentsObservations::GetMapObservationToArrays(TArray<FLearningAgentsObservationObjectElement>& OutKeys, TArray<FLearningAgentsObservationObjectElement>& OutValues, const ULearningAgentsObservationObject* Object, const FLearningAgentsObservationObjectElement Element, const FName Tag)
 {
 	int32 OutElementNum = 0;
-	if (!GetMapObservationNum(OutElementNum, Element, Name))
+	if (!GetMapObservationNum(OutElementNum, Object, Element, Tag))
 	{
 		OutKeys.Empty();
 		OutValues.Empty();
@@ -3002,7 +3206,7 @@ bool ULearningAgentsObservationObject::GetMapObservationToArrays(TArray<FLearnin
 
 	OutKeys.SetNumUninitialized(OutElementNum);
 	OutValues.SetNumUninitialized(OutElementNum);
-	if (!GetMapObservationToArrayViews(OutKeys, OutValues, Element, Name))
+	if (!GetMapObservationToArrayViews(OutKeys, OutValues, Object, Element, Tag))
 	{
 		OutKeys.Empty();
 		OutValues.Empty();
@@ -3012,11 +3216,11 @@ bool ULearningAgentsObservationObject::GetMapObservationToArrays(TArray<FLearnin
 	return true;
 }
 
-bool ULearningAgentsObservationObject::GetMapObservationToArrayViews(TArrayView<FLearningAgentsObservationObjectElement> OutKeys, TArrayView<FLearningAgentsObservationObjectElement> OutValues, const FLearningAgentsObservationObjectElement Element, const FName Name) const
+bool ULearningAgentsObservations::GetMapObservationToArrayViews(TArrayView<FLearningAgentsObservationObjectElement> OutKeys, TArrayView<FLearningAgentsObservationObjectElement> OutValues, const ULearningAgentsObservationObject* Object, const FLearningAgentsObservationObjectElement Element, const FName Tag)
 {
 	TArray<FLearningAgentsObservationObjectElement, TInlineAllocator<16>> Pairs;
 	Pairs.SetNumUninitialized(OutKeys.Num());
-	if (!GetSetObservationToArrayView(Pairs, Element, Name))
+	if (!GetSetObservationToArrayView(Pairs, Object, Element, Tag))
 	{
 		UE::Learning::Array::Set<1, FLearningAgentsObservationObjectElement>(OutKeys, FLearningAgentsObservationObjectElement());
 		UE::Learning::Array::Set<1, FLearningAgentsObservationObjectElement>(OutValues, FLearningAgentsObservationObjectElement());
@@ -3026,7 +3230,7 @@ bool ULearningAgentsObservationObject::GetMapObservationToArrayViews(TArrayView<
 	for (int32 PairIdx = 0; PairIdx < Pairs.Num(); PairIdx++)
 	{
 		FLearningAgentsObservationObjectElement Key, Value;
-		if (!GetPairObservation(Pairs[PairIdx], Key, Value))
+		if (!GetPairObservation(Key, Value, Object, Pairs[PairIdx]))
 		{
 			UE::Learning::Array::Set<1, FLearningAgentsObservationObjectElement>(OutKeys, FLearningAgentsObservationObjectElement());
 			UE::Learning::Array::Set<1, FLearningAgentsObservationObjectElement>(OutValues, FLearningAgentsObservationObjectElement());
@@ -3040,17 +3244,17 @@ bool ULearningAgentsObservationObject::GetMapObservationToArrayViews(TArrayView<
 	return true;
 }
 
-bool ULearningAgentsObservationObject::GetEnumObservation(uint8& OutEnumValue, const UEnum* Enum, const FLearningAgentsObservationObjectElement Element, const FName Name) const
+bool ULearningAgentsObservations::GetEnumObservation(uint8& OutEnumValue, const ULearningAgentsObservationObject* Object, const FLearningAgentsObservationObjectElement Element, const UEnum* Enum, const FName Tag)
 {
 	if (!Enum)
 	{
-		UE_LOG(LogLearning, Error, TEXT("%s: Enum is nullptr."), *GetName());
+		UE_LOG(LogLearning, Error, TEXT("GetEnumObservation: Enum is nullptr."));
 		OutEnumValue = 0;
 		return false;
 	}
 
 	int32 EnumValueNum;
-	if (!GetContinuousObservationNum(EnumValueNum, Element, Name))
+	if (!GetContinuousObservationNum(EnumValueNum, Object, Element, Tag))
 	{
 		OutEnumValue = 0;
 		return false;
@@ -3058,14 +3262,14 @@ bool ULearningAgentsObservationObject::GetEnumObservation(uint8& OutEnumValue, c
 	
 	if (EnumValueNum != Enum->NumEnums() - 1)
 	{
-		UE_LOG(LogLearning, Error, TEXT("%s: Too many values for Enum '%s'. Expected %i, got %i."), *GetName(), *Enum->GetName(), Enum->NumEnums() - 1, EnumValueNum);
+		UE_LOG(LogLearning, Error, TEXT("GetEnumObservation: Too many values for Enum '%s'. Expected %i, got %i."), *Enum->GetName(), Enum->NumEnums() - 1, EnumValueNum);
 		OutEnumValue = 0;
 		return false;
 	}
 
 	TArray<float, TInlineAllocator<32>> OneHot;
 	OneHot.SetNumUninitialized(EnumValueNum);
-	if (!GetContinuousObservationToArrayView(OneHot, Element, Name))
+	if (!GetContinuousObservationToArrayView(OneHot, Object, Element, Tag))
 	{
 		OutEnumValue = 0;
 		return false;
@@ -3083,7 +3287,7 @@ bool ULearningAgentsObservationObject::GetEnumObservation(uint8& OutEnumValue, c
 
 	if (EnumValueIndex == INDEX_NONE)
 	{
-		UE_LOG(LogLearning, Error, TEXT("%s: Index not found."), *GetName());
+		UE_LOG(LogLearning, Error, TEXT("GetEnumObservation: Index not found."));
 		OutEnumValue = 0;
 		return false;
 	}
@@ -3092,7 +3296,7 @@ bool ULearningAgentsObservationObject::GetEnumObservation(uint8& OutEnumValue, c
 
 	if (EnumValue == INDEX_NONE)
 	{
-		UE_LOG(LogLearning, Error, TEXT("%s: Enum Value not found for index %i."), *GetName(), EnumValueIndex);
+		UE_LOG(LogLearning, Error, TEXT("GetEnumObservation: Enum Value not found for index %i."), EnumValueIndex);
 		OutEnumValue = 0;
 		return false;
 	}
@@ -3101,24 +3305,24 @@ bool ULearningAgentsObservationObject::GetEnumObservation(uint8& OutEnumValue, c
 	return true;
 }
 
-bool ULearningAgentsObservationObject::GetBitmaskObservation(int32& OutBitmaskValue, const UEnum* Enum, const FLearningAgentsObservationObjectElement Element, const FName Name) const
+bool ULearningAgentsObservations::GetBitmaskObservation(int32& OutBitmaskValue, const ULearningAgentsObservationObject* Object, const FLearningAgentsObservationObjectElement Element, const UEnum* Enum, const FName Tag)
 {
 	if (!Enum)
 	{
-		UE_LOG(LogLearning, Error, TEXT("%s: Enum is nullptr."), *GetName());
+		UE_LOG(LogLearning, Error, TEXT("GetBitmaskObservation: Enum is nullptr."));
 		OutBitmaskValue = 0;
 		return false;
 	}
 
 	if (Enum->NumEnums() - 1 > 32)
 	{
-		UE_LOG(LogLearning, Error, TEXT("%s: Too many values in Enum to use as Bitmask (%i)."), *GetName(), Enum->NumEnums() - 1);
+		UE_LOG(LogLearning, Error, TEXT("GetBitmaskObservation: Too many values in Enum to use as Bitmask (%i)."), Enum->NumEnums() - 1);
 		OutBitmaskValue = 0;
 		return false;
 	}
 
 	int32 EnumValueNum;
-	if (!GetContinuousObservationNum(EnumValueNum, Element, Name))
+	if (!GetContinuousObservationNum(EnumValueNum, Object, Element, Tag))
 	{
 		OutBitmaskValue = 0;
 		return false;
@@ -3126,14 +3330,14 @@ bool ULearningAgentsObservationObject::GetBitmaskObservation(int32& OutBitmaskVa
 
 	if (EnumValueNum != Enum->NumEnums() - 1)
 	{
-		UE_LOG(LogLearning, Error, TEXT("%s: Too many values for Enum '%s'. Expected %i, got %i."), *GetName(), *Enum->GetName(), Enum->NumEnums() - 1, EnumValueNum);
+		UE_LOG(LogLearning, Error, TEXT("GetBitmaskObservation: Too many values for Enum '%s'. Expected %i, got %i."), *Enum->GetName(), Enum->NumEnums() - 1, EnumValueNum);
 		OutBitmaskValue = 0;
 		return false;
 	}
 
 	TArray<float, TInlineAllocator<32>> OneHot;
 	OneHot.Init(0.0f, EnumValueNum);
-	if (!GetContinuousObservationToArrayView(OneHot, Element, Name))
+	if (!GetContinuousObservationToArrayView(OneHot, Object, Element, Tag))
 	{
 		OutBitmaskValue = 0;
 		return false;
@@ -3151,10 +3355,10 @@ bool ULearningAgentsObservationObject::GetBitmaskObservation(int32& OutBitmaskVa
 }
 
 
-bool ULearningAgentsObservationObject::GetOptionalObservation(ELearningAgentsOptionalObservation& OutOption, FLearningAgentsObservationObjectElement& OutElement, const FLearningAgentsObservationObjectElement Element, const FName Name) const
+bool ULearningAgentsObservations::GetOptionalObservation(ELearningAgentsOptionalObservation& OutOption, FLearningAgentsObservationObjectElement& OutElement, const ULearningAgentsObservationObject* Object, const FLearningAgentsObservationObjectElement Element, const FName Tag)
 {
 	FName OutName = NAME_None;
-	if (!GetExclusiveUnionObservation(OutName, OutElement, Element, Name))
+	if (!GetExclusiveUnionObservation(OutName, OutElement, Object, Element, Tag))
 	{
 		OutOption = ELearningAgentsOptionalObservation::Null;
 		return false;
@@ -3164,10 +3368,10 @@ bool ULearningAgentsObservationObject::GetOptionalObservation(ELearningAgentsOpt
 	return true;
 }
 
-bool ULearningAgentsObservationObject::GetEitherObservation(ELearningAgentsEitherObservation& OutEither, FLearningAgentsObservationObjectElement& OutElement, const FLearningAgentsObservationObjectElement Element, const FName Name) const
+bool ULearningAgentsObservations::GetEitherObservation(ELearningAgentsEitherObservation& OutEither, FLearningAgentsObservationObjectElement& OutElement, const ULearningAgentsObservationObject* Object, const FLearningAgentsObservationObjectElement Element, const FName Tag)
 {
 	FName OutName = NAME_None;
-	if (!GetExclusiveUnionObservation(OutName, OutElement, Element, Name))
+	if (!GetExclusiveUnionObservation(OutName, OutElement, Object, Element, Tag))
 	{
 		OutEither = ELearningAgentsEitherObservation::A;
 		return false;
@@ -3177,40 +3381,46 @@ bool ULearningAgentsObservationObject::GetEitherObservation(ELearningAgentsEithe
 	return true;
 }
 
-bool ULearningAgentsObservationObject::GetEncodingObservation(FLearningAgentsObservationObjectElement& OutElement, const FLearningAgentsObservationObjectElement Element, const FName Name) const
+bool ULearningAgentsObservations::GetEncodingObservation(FLearningAgentsObservationObjectElement& OutElement, const ULearningAgentsObservationObject* Object, const FLearningAgentsObservationObjectElement Element, const FName Tag)
 {
-	if (!ObservationObject.IsValid(Element.ObjectElement))
+	if (!Object)
+	{
+		UE_LOG(LogLearning, Error, TEXT("GetEncodingObservation: Object is nullptr."));
+		OutElement = FLearningAgentsObservationObjectElement();
+		return false;
+	}
+
+	if (!Object->ObservationObject.IsValid(Element.ObjectElement))
 	{
-		UE_LOG(LogLearning, Error, TEXT("%s: Invalid Observation Object."), *GetName());
+		UE_LOG(LogLearning, Error, TEXT("GetEncodingObservation: Invalid Observation Object."));
 		OutElement = FLearningAgentsObservationObjectElement();
 		return false;
 	}
 
-	if (ObservationObject.GetName(Element.ObjectElement) != Name)
+	if (Object->ObservationObject.GetTag(Element.ObjectElement) != Tag)
 	{
-		UE_LOG(LogLearning, Warning, TEXT("%s: Observation name does not match. Observation is '%s' but asked for '%s'."), *GetName(), *ObservationObject.GetName(Element.ObjectElement).ToString(), *Name.ToString());
+		UE_LOG(LogLearning, Warning, TEXT("GetEncodingObservation: Observation tag does not match. Observation is '%s' but asked for '%s'."), *Object->ObservationObject.GetTag(Element.ObjectElement).ToString(), *Tag.ToString());
 	}
 
-	if (ObservationObject.GetType(Element.ObjectElement) != UE::Learning::Observation::EType::Encoding)
+	if (Object->ObservationObject.GetType(Element.ObjectElement) != UE::Learning::Observation::EType::Encoding)
 	{
-		UE_LOG(LogLearning, Error, TEXT("%s: Observation '%s' type does not match. Observation is '%s' but asked for '%s'."),
-			*GetName(),
-			*ObservationObject.GetName(Element.ObjectElement).ToString(),
-			UE::Learning::Agents::Observation::Private::GetObservationTypeString(ObservationObject.GetType(Element.ObjectElement)),
+		UE_LOG(LogLearning, Error, TEXT("GetEncodingObservation: Observation '%s' type does not match. Observation is '%s' but asked for '%s'."),
+			*Object->ObservationObject.GetTag(Element.ObjectElement).ToString(),
+			UE::Learning::Agents::Observation::Private::GetObservationTypeString(Object->ObservationObject.GetType(Element.ObjectElement)),
 			UE::Learning::Agents::Observation::Private::GetObservationTypeString(UE::Learning::Observation::EType::Encoding));
 		OutElement = FLearningAgentsObservationObjectElement();
 		return false;
 	}
 
-	OutElement = { ObservationObject.GetEncoding(Element.ObjectElement).Element };
+	OutElement = { Object->ObservationObject.GetEncoding(Element.ObjectElement).Element };
 
 	return true;
 }
 
-bool ULearningAgentsObservationObject::GetBoolObservation(bool& bOutValue, const FLearningAgentsObservationObjectElement Element, const FName Name) const
+bool ULearningAgentsObservations::GetBoolObservation(bool& bOutValue, const ULearningAgentsObservationObject* Object, const FLearningAgentsObservationObjectElement Element, const FName Tag)
 {
 	float OutValue = 0.0f;
-	if (!GetFloatObservation(OutValue, Element, 1.0f, Name))
+	if (!GetFloatObservation(OutValue, Object, Element, 1.0f, Tag))
 	{
 		bOutValue = false;
 		return false;
@@ -3220,10 +3430,10 @@ bool ULearningAgentsObservationObject::GetBoolObservation(bool& bOutValue, const
 	return true;
 }
 
-bool ULearningAgentsObservationObject::GetFloatObservation(float& OutValue, const FLearningAgentsObservationObjectElement Element, const float FloatScale, const FName Name) const
+bool ULearningAgentsObservations::GetFloatObservation(float& OutValue, const ULearningAgentsObservationObject* Object, const FLearningAgentsObservationObjectElement Element, const float FloatScale, const FName Tag)
 {
 	float OutValuesData;
-	if (!GetContinuousObservationToArrayView(MakeArrayView(&OutValuesData, 1), Element, Name))
+	if (!GetContinuousObservationToArrayView(MakeArrayView(&OutValuesData, 1), Object, Element, Tag))
 	{
 		OutValue = 0.0f;
 		return false;
@@ -3233,10 +3443,10 @@ bool ULearningAgentsObservationObject::GetFloatObservation(float& OutValue, cons
 	return true;
 }
 
-bool ULearningAgentsObservationObject::GetLocationObservation(FVector& OutLocation, const FLearningAgentsObservationObjectElement Element, const FTransform RelativeTransform, const float LocationScale, const FName Name) const
+bool ULearningAgentsObservations::GetLocationObservation(FVector& OutLocation, const ULearningAgentsObservationObject* Object, const FLearningAgentsObservationObjectElement Element, const FTransform RelativeTransform, const float LocationScale, const FName Tag)
 {
 	TStaticArray<float, 3> OutValues;
-	if (!GetContinuousObservationToArrayView(OutValues, Element, Name))
+	if (!GetContinuousObservationToArrayView(OutValues, Object, Element, Tag))
 	{
 		OutLocation = FVector::ZeroVector;
 		return false;
@@ -3246,10 +3456,10 @@ bool ULearningAgentsObservationObject::GetLocationObservation(FVector& OutLocati
 	return true;
 }
 
-bool ULearningAgentsObservationObject::GetRotationObservation(FRotator& OutRotation, const FLearningAgentsObservationObjectElement Element, const FRotator RelativeRotation, const FName Name) const
+bool ULearningAgentsObservations::GetRotationObservation(FRotator& OutRotation, const ULearningAgentsObservationObject* Object, const FLearningAgentsObservationObjectElement Element, const FRotator RelativeRotation, const FName Tag)
 {
 	FQuat OutRotationQuat;
-	if (!GetRotationObservationAsQuat(OutRotationQuat, Element, FQuat::MakeFromRotator(RelativeRotation), Name))
+	if (!GetRotationObservationAsQuat(OutRotationQuat, Object, Element, FQuat::MakeFromRotator(RelativeRotation), Tag))
 	{
 		OutRotation = FRotator::ZeroRotator;
 		return false;
@@ -3259,10 +3469,10 @@ bool ULearningAgentsObservationObject::GetRotationObservation(FRotator& OutRotat
 	return true;
 }
 
-bool ULearningAgentsObservationObject::GetRotationObservationAsQuat(FQuat& OutRotation, const FLearningAgentsObservationObjectElement Element, const FQuat RelativeRotation, const FName Name) const
+bool ULearningAgentsObservations::GetRotationObservationAsQuat(FQuat& OutRotation, const ULearningAgentsObservationObject* Object, const FLearningAgentsObservationObjectElement Element, const FQuat RelativeRotation, const FName Tag)
 {
 	TStaticArray<float, 6> OutValues;
-	if (!GetContinuousObservationToArrayView(OutValues, Element, Name))
+	if (!GetContinuousObservationToArrayView(OutValues, Object, Element, Tag))
 	{
 		OutRotation = FQuat::Identity;
 		return false;
@@ -3283,10 +3493,10 @@ bool ULearningAgentsObservationObject::GetRotationObservationAsQuat(FQuat& OutRo
 	return true;
 }
 
-bool ULearningAgentsObservationObject::GetScaleObservation(FVector& OutScale, const FLearningAgentsObservationObjectElement Element, const FVector RelativeScale, const FName Name) const
+bool ULearningAgentsObservations::GetScaleObservation(FVector& OutScale, const ULearningAgentsObservationObject* Object, const FLearningAgentsObservationObjectElement Element, const FVector RelativeScale, const FName Tag)
 {
 	TStaticArray<float, 3> OutValues;
-	if (!GetContinuousObservationToArrayView(OutValues, Element, Name))
+	if (!GetContinuousObservationToArrayView(OutValues, Object, Element, Tag))
 	{
 		OutScale = FVector::OneVector;
 		return false;
@@ -3296,11 +3506,11 @@ bool ULearningAgentsObservationObject::GetScaleObservation(FVector& OutScale, co
 	return true;
 }
 
-bool ULearningAgentsObservationObject::GetTransformObservation(FTransform& OutTransform, const FLearningAgentsObservationObjectElement Element, const FTransform RelativeTransform, const float LocationScale, const FName Name) const
+bool ULearningAgentsObservations::GetTransformObservation(FTransform& OutTransform, const ULearningAgentsObservationObject* Object, const FLearningAgentsObservationObjectElement Element, const FTransform RelativeTransform, const float LocationScale, const FName Tag)
 {
 	TStaticArray<FName, 3> OutElementNames;
 	TStaticArray<FLearningAgentsObservationObjectElement, 3> OutElements;
-	if (!GetStructObservationToArrayViews(OutElementNames, OutElements, Element, Name))
+	if (!GetStructObservationToArrayViews(OutElementNames, OutElements, Object, Element, Tag))
 	{
 		OutTransform = FTransform::Identity;
 		return false;
@@ -3308,7 +3518,7 @@ bool ULearningAgentsObservationObject::GetTransformObservation(FTransform& OutTr
 
 	const int32 LocationElement = MakeArrayView(OutElementNames).Find(TEXT("Location"));
 	FVector OutLocation;
-	if (LocationElement == INDEX_NONE || !GetLocationObservation(OutLocation, OutElements[LocationElement], RelativeTransform, LocationScale))
+	if (LocationElement == INDEX_NONE || !GetLocationObservation(OutLocation, Object, OutElements[LocationElement], RelativeTransform, LocationScale))
 	{
 		OutTransform = FTransform::Identity;
 		return false;
@@ -3316,7 +3526,7 @@ bool ULearningAgentsObservationObject::GetTransformObservation(FTransform& OutTr
 
 	const int32 RotationElement = MakeArrayView(OutElementNames).Find(TEXT("Rotation"));
 	FQuat OutRotation;
-	if (RotationElement == INDEX_NONE || !GetRotationObservationAsQuat(OutRotation, OutElements[RotationElement], RelativeTransform.GetRotation()))
+	if (RotationElement == INDEX_NONE || !GetRotationObservationAsQuat(OutRotation, Object, OutElements[RotationElement], RelativeTransform.GetRotation()))
 	{
 		OutTransform = FTransform::Identity;
 		return false;
@@ -3324,7 +3534,7 @@ bool ULearningAgentsObservationObject::GetTransformObservation(FTransform& OutTr
 
 	const int32 ScaleElement = MakeArrayView(OutElementNames).Find(TEXT("Scale"));
 	FVector OutScale;
-	if (ScaleElement == INDEX_NONE || !GetScaleObservation(OutScale, OutElements[ScaleElement], RelativeTransform.GetScale3D()))
+	if (ScaleElement == INDEX_NONE || !GetScaleObservation(OutScale, Object, OutElements[ScaleElement], RelativeTransform.GetScale3D()))
 	{
 		OutTransform = FTransform::Identity;
 		return false;
@@ -3334,10 +3544,10 @@ bool ULearningAgentsObservationObject::GetTransformObservation(FTransform& OutTr
 	return true;
 }
 
-bool ULearningAgentsObservationObject::GetAngleObservationRadians(float& OutAngle, const FLearningAgentsObservationObjectElement Element, const float RelativeAngle, const FName Name) const
+bool ULearningAgentsObservations::GetAngleObservationRadians(float& OutAngle, const ULearningAgentsObservationObject* Object, const FLearningAgentsObservationObjectElement Element, const float RelativeAngle, const FName Tag)
 {
 	TStaticArray<float, 2> OutValues;
-	if (!GetContinuousObservationToArrayView(OutValues, Element, Name))
+	if (!GetContinuousObservationToArrayView(OutValues, Object, Element, Tag))
 	{
 		OutAngle = 0.0f;
 		return false;
@@ -3348,9 +3558,9 @@ bool ULearningAgentsObservationObject::GetAngleObservationRadians(float& OutAngl
 }
 
 
-bool ULearningAgentsObservationObject::GetAngleObservation(float& OutAngle, const FLearningAgentsObservationObjectElement Element, const float RelativeAngle, const FName Name) const
+bool ULearningAgentsObservations::GetAngleObservation(float& OutAngle, const ULearningAgentsObservationObject* Object, const FLearningAgentsObservationObjectElement Element, const float RelativeAngle, const FName Tag)
 {
-	if (!GetAngleObservationRadians(OutAngle, Element, FMath::DegreesToRadians(RelativeAngle), Name))
+	if (!GetAngleObservationRadians(OutAngle, Object, Element, FMath::DegreesToRadians(RelativeAngle), Tag))
 	{
 		return false;
 	}
@@ -3359,10 +3569,10 @@ bool ULearningAgentsObservationObject::GetAngleObservation(float& OutAngle, cons
 	return true;
 }
 
-bool ULearningAgentsObservationObject::GetVelocityObservation(FVector& OutVelocity, const FLearningAgentsObservationObjectElement Element, const FTransform RelativeTransform, const float VelocityScale, const FName Name) const
+bool ULearningAgentsObservations::GetVelocityObservation(FVector& OutVelocity, const ULearningAgentsObservationObject* Object, const FLearningAgentsObservationObjectElement Element, const FTransform RelativeTransform, const float VelocityScale, const FName Tag)
 {
 	TStaticArray<float, 3> OutValues;
-	if (!GetContinuousObservationToArrayView(OutValues, Element, Name))
+	if (!GetContinuousObservationToArrayView(OutValues, Object, Element, Tag))
 	{
 		OutVelocity = FVector::ZeroVector;
 		return false;
@@ -3372,10 +3582,10 @@ bool ULearningAgentsObservationObject::GetVelocityObservation(FVector& OutVeloci
 	return true;
 }
 
-bool ULearningAgentsObservationObject::GetDirectionObservation(FVector& OutDirection, const FLearningAgentsObservationObjectElement Element, const FTransform RelativeTransform, const FName Name) const
+bool ULearningAgentsObservations::GetDirectionObservation(FVector& OutDirection, const ULearningAgentsObservationObject* Object, const FLearningAgentsObservationObjectElement Element, const FTransform RelativeTransform, const FName Tag)
 {
 	TStaticArray<float, 3> OutValues;
-	if (!GetContinuousObservationToArrayView(OutValues, Element, Name))
+	if (!GetContinuousObservationToArrayView(OutValues, Object, Element, Tag))
 	{
 		OutDirection = FVector::ForwardVector;
 		return false;
@@ -3385,16 +3595,16 @@ bool ULearningAgentsObservationObject::GetDirectionObservation(FVector& OutDirec
 	return true;
 }
 
-bool ULearningAgentsObservationObject::GetLocationAlongSplineObservation(FVector& OutLocation, const FLearningAgentsObservationObjectElement Element, const FTransform RelativeTransform, const float LocationScale, const FName Name)
+bool ULearningAgentsObservations::GetLocationAlongSplineObservation(FVector& OutLocation, const ULearningAgentsObservationObject* Object, const FLearningAgentsObservationObjectElement Element, const FTransform RelativeTransform, const float LocationScale, const FName Tag)
 {
-	return GetLocationObservation(OutLocation, Element, RelativeTransform, LocationScale, Name);
+	return GetLocationObservation(OutLocation, Object, Element, RelativeTransform, LocationScale, Tag);
 }
 
-bool ULearningAgentsObservationObject::GetProportionAlongSplineObservation(bool& bOutIsClosedLoop, float& OutAngle, float& OutPropotion, const FLearningAgentsObservationObjectElement Element, const FName Name)
+bool ULearningAgentsObservations::GetProportionAlongSplineObservation(bool& bOutIsClosedLoop, float& OutAngle, float& OutPropotion, const ULearningAgentsObservationObject* Object, const FLearningAgentsObservationObjectElement Element, const FName Tag)
 {
 	FName SubName;
 	FLearningAgentsObservationObjectElement SubElement;
-	if (!GetExclusiveUnionObservation(SubName, SubElement, Element, Name))
+	if (!GetExclusiveUnionObservation(SubName, SubElement, Object, Element, Tag))
 	{
 		bOutIsClosedLoop = false;
 		OutAngle = 0.0f;
@@ -3406,26 +3616,26 @@ bool ULearningAgentsObservationObject::GetProportionAlongSplineObservation(bool&
 	{
 		bOutIsClosedLoop = true;
 		OutPropotion = 0.0f;
-		return GetAngleObservation(OutAngle, SubElement);
+		return GetAngleObservation(OutAngle, Object, SubElement);
 	}
 	else
 	{
 		bOutIsClosedLoop = false;
 		OutAngle = 0.0f;
-		return GetFloatObservation(OutPropotion, SubElement);
+		return GetFloatObservation(OutPropotion, Object, SubElement);
 	}
 }
 
-bool ULearningAgentsObservationObject::GetDirectionAlongSplineObservation(FVector& OutDirection, const FLearningAgentsObservationObjectElement Element, const FTransform RelativeTransform, const FName Name)
+bool ULearningAgentsObservations::GetDirectionAlongSplineObservation(FVector& OutDirection, const ULearningAgentsObservationObject* Object, const FLearningAgentsObservationObjectElement Element, const FTransform RelativeTransform, const FName Tag)
 {
-	return GetDirectionObservation(OutDirection, Element, RelativeTransform, Name);
+	return GetDirectionObservation(OutDirection, Object, Element, RelativeTransform, Tag);
 }
 
-bool ULearningAgentsObservationObject::GetPropertiesAlongSplineObservation(FVector& OutLocation, bool& bOutIsClosedLoop, float& OutAngle, float& OutPropotion, FVector& OutDirection, const FLearningAgentsObservationObjectElement Element, const FTransform RelativeTransform, const float LocationScale, const FName Name)
+bool ULearningAgentsObservations::GetPropertiesAlongSplineObservation(FVector& OutLocation, bool& bOutIsClosedLoop, float& OutAngle, float& OutPropotion, FVector& OutDirection, const ULearningAgentsObservationObject* Object, const FLearningAgentsObservationObjectElement Element, const FTransform RelativeTransform, const float LocationScale, const FName Tag)
 {
 	TStaticArray<FName, 3> OutElementNames;
 	TStaticArray<FLearningAgentsObservationObjectElement, 3> OutElements;
-	if (!GetStructObservationToArrayViews(OutElementNames, OutElements, Element, Name))
+	if (!GetStructObservationToArrayViews(OutElementNames, OutElements, Object, Element, Tag))
 	{
 		OutLocation = FVector::ZeroVector;
 		bOutIsClosedLoop = false;
@@ -3436,7 +3646,7 @@ bool ULearningAgentsObservationObject::GetPropertiesAlongSplineObservation(FVect
 	}
 
 	const int32 LocationElement = MakeArrayView(OutElementNames).Find(TEXT("Location"));
-	if (LocationElement == INDEX_NONE || !GetLocationAlongSplineObservation(OutLocation, OutElements[LocationElement], RelativeTransform, LocationScale))
+	if (LocationElement == INDEX_NONE || !GetLocationAlongSplineObservation(OutLocation, Object, OutElements[LocationElement], RelativeTransform, LocationScale))
 	{
 		OutLocation = FVector::ZeroVector;
 		bOutIsClosedLoop = false;
@@ -3447,7 +3657,7 @@ bool ULearningAgentsObservationObject::GetPropertiesAlongSplineObservation(FVect
 	}
 
 	const int32 ProportionElement = MakeArrayView(OutElementNames).Find(TEXT("Proportion"));
-	if (ProportionElement == INDEX_NONE || !GetProportionAlongSplineObservation(bOutIsClosedLoop, OutAngle, OutPropotion, OutElements[ProportionElement]))
+	if (ProportionElement == INDEX_NONE || !GetProportionAlongSplineObservation(bOutIsClosedLoop, OutAngle, OutPropotion, Object, OutElements[ProportionElement]))
 	{
 		OutLocation = FVector::ZeroVector;
 		bOutIsClosedLoop = false;
@@ -3458,7 +3668,7 @@ bool ULearningAgentsObservationObject::GetPropertiesAlongSplineObservation(FVect
 	}
 
 	const int32 DirectionElement = MakeArrayView(OutElementNames).Find(TEXT("Direction"));
-	if (DirectionElement == INDEX_NONE || !GetDirectionAlongSplineObservation(OutDirection, OutElements[ProportionElement], RelativeTransform))
+	if (DirectionElement == INDEX_NONE || !GetDirectionAlongSplineObservation(OutDirection, Object, OutElements[ProportionElement], RelativeTransform))
 	{
 		OutLocation = FVector::ZeroVector;
 		bOutIsClosedLoop = false;
@@ -3471,9 +3681,9 @@ bool ULearningAgentsObservationObject::GetPropertiesAlongSplineObservation(FVect
 	return true;
 }
 
-bool ULearningAgentsObservationObject::GetProportionAlongRayObservation(float& OutProportion, const FLearningAgentsObservationObjectElement Element, const FName Name)
+bool ULearningAgentsObservations::GetProportionAlongRayObservation(float& OutProportion, const ULearningAgentsObservationObject* Object, const FLearningAgentsObservationObjectElement Element, const FName Tag)
 {
-	if (!GetFloatObservation(OutProportion, Element, 1.0f, Name))
+	if (!GetFloatObservation(OutProportion, Object, Element, 1.0f, Tag))
 	{
 		OutProportion = 0.0f;
 		return false;
@@ -3483,22 +3693,22 @@ bool ULearningAgentsObservationObject::GetProportionAlongRayObservation(float& O
 	return true;
 }
 
-bool ULearningAgentsObservationObject::GetProportionAlongRaysObservationNum(int32& OutProportionNum, const FLearningAgentsObservationObjectElement Element, const FName Name)
+bool ULearningAgentsObservations::GetProportionAlongRaysObservationNum(int32& OutProportionNum, const ULearningAgentsObservationObject* Object, const FLearningAgentsObservationObjectElement Element, const FName Tag)
 {
-	return GetStaticArrayObservationNum(OutProportionNum, Element, Name);
+	return GetStaticArrayObservationNum(OutProportionNum, Object, Element, Tag);
 }
 
-bool ULearningAgentsObservationObject::GetProportionAlongRaysObservation(TArray<float>& OutProportions, const FLearningAgentsObservationObjectElement Element, const FName Name)
+bool ULearningAgentsObservations::GetProportionAlongRaysObservation(TArray<float>& OutProportions, const ULearningAgentsObservationObject* Object, const FLearningAgentsObservationObjectElement Element, const FName Tag)
 {
 	int32 ProportionNum;
-	if (!GetProportionAlongRaysObservationNum(ProportionNum, Element, Name))
+	if (!GetProportionAlongRaysObservationNum(ProportionNum, Object, Element, Tag))
 	{
 		OutProportions.Empty();
 		return false;
 	}
 
 	OutProportions.SetNumUninitialized(ProportionNum);
-	if (!GetProportionAlongRaysObservationToArrayView(OutProportions, Element, Name))
+	if (!GetProportionAlongRaysObservationToArrayView(OutProportions, Object, Element, Tag))
 	{
 		OutProportions.Empty();
 		return false;
@@ -3507,19 +3717,26 @@ bool ULearningAgentsObservationObject::GetProportionAlongRaysObservation(TArray<
 	return true;
 }
 
-bool ULearningAgentsObservationObject::GetProportionAlongRaysObservationToArrayView(TArrayView<float> OutProportions, const FLearningAgentsObservationObjectElement Element, const FName Name)
+bool ULearningAgentsObservations::GetProportionAlongRaysObservationToArrayView(TArrayView<float> OutProportions, const ULearningAgentsObservationObject* Object, const FLearningAgentsObservationObjectElement Element, const FName Tag)
 {
 	int32 ProportionNum;
-	if (!GetStaticArrayObservationNum(ProportionNum, Element, Name))
+	if (!GetStaticArrayObservationNum(ProportionNum, Object, Element, Tag))
+	{
+		UE::Learning::Array::Zero<1, float>(OutProportions);
+		return false;
+	}
+
+	if (!Object)
 	{
+		UE_LOG(LogLearning, Error, TEXT("GetProportionAlongRaysObservationToArrayView: Object is nullptr."));
 		UE::Learning::Array::Zero<1, float>(OutProportions);
 		return false;
 	}
 
 	if (ProportionNum != OutProportions.Num())
 	{
-		UE_LOG(LogLearning, Error, TEXT("%s: Observation '%s' size does not match. Observation is '%i' elements but asked for '%i'."),
-			*GetName(), *ObservationObject.GetName(Element.ObjectElement).ToString(),
+		UE_LOG(LogLearning, Error, TEXT("GetProportionAlongRaysObservationToArrayView: Observation '%s' size does not match. Observation is '%i' elements but asked for '%i'."),
+			*Object->ObservationObject.GetTag(Element.ObjectElement).ToString(),
 			ProportionNum, OutProportions.Num());
 		UE::Learning::Array::Zero<1, float>(OutProportions);
 		return false;
@@ -3527,7 +3744,7 @@ bool ULearningAgentsObservationObject::GetProportionAlongRaysObservationToArrayV
 
 	TArray<FLearningAgentsObservationObjectElement, TInlineAllocator<32>> SubElements;
 	SubElements.SetNumUninitialized(ProportionNum);
-	if (!GetStaticArrayObservationToArrayView(SubElements, Element, Name))
+	if (!GetStaticArrayObservationToArrayView(SubElements, Object, Element, Tag))
 	{
 		UE::Learning::Array::Zero<1, float>(OutProportions);
 		return false;
@@ -3535,7 +3752,7 @@ bool ULearningAgentsObservationObject::GetProportionAlongRaysObservationToArrayV
 
 	for (int32 SubElementIdx = 0; SubElementIdx < ProportionNum; SubElementIdx++)
 	{
-		if (!GetProportionAlongRayObservation(OutProportions[SubElementIdx], SubElements[SubElementIdx]))
+		if (!GetProportionAlongRayObservation(OutProportions[SubElementIdx], Object, SubElements[SubElementIdx]))
 		{
 			UE::Learning::Array::Zero<1, float>(OutProportions);
 			return false;
