LearningAgents: Refactored to make Observation and Action functions static and use "AdvancedDisplay" markup. Added Visual Logging and comments for Completions.

diff --git a/Engine/Plugins/Experimental/LearningAgents/Source/LearningAgents/Private/LearningAgentsPolicy.cpp b/Engine/Plugins/Experimental/LearningAgents/Source/LearningAgents/Private/LearningAgentsPolicy.cpp
--- a/Engine/Plugins/Experimental/LearningAgents/Source/LearningAgents/Private/LearningAgentsPolicy.cpp
+++ b/Engine/Plugins/Experimental/LearningAgents/Source/LearningAgents/Private/LearningAgentsPolicy.cpp
@@ -631,7 +631,7 @@ void ULearningAgentsPolicy::DecodeAndSampleActions(const float ActionNoiseScale)
 			Seeds[AgentId],
 			Interactor->ActionVectors[AgentId],
 			ActionDistributionVectors[AgentId],
-			Interactor->ActionSchema->GetActionSchema(),
+			Interactor->ActionSchema->ActionSchema,
 			Interactor->ActionSchemaElement.SchemaElement,
 			ActionNoiseScale);
 	}
diff --git a/Engine/Plugins/Experimental/LearningAgents/Source/LearningAgents/Public/LearningAgentsActions.h b/Engine/Plugins/Experimental/LearningAgents/Source/LearningAgents/Public/LearningAgentsActions.h
--- a/Engine/Plugins/Experimental/LearningAgents/Source/LearningAgents/Public/LearningAgentsActions.h
+++ b/Engine/Plugins/Experimental/LearningAgents/Source/LearningAgents/Public/LearningAgentsActions.h
@@ -7,6 +7,9 @@
 
 #include "LearningAgentsNeuralNetwork.h" // Included for ELearningAgentsActivationFunction
 
+#include "Engine/EngineTypes.h"
+#include "Kismet/BlueprintFunctionLibrary.h"
+
 #include "LearningAgentsActions.generated.h"
 
 class ULearningAgentsActionSchema;
@@ -38,6 +41,36 @@ bool operator==(const FLearningAgentsActionObjectElement& Lhs, const FLearningAg
 /** Hashing operator for Action Object Elements */
 uint32 GetTypeHash(const FLearningAgentsActionObjectElement& Element);
 
+/**
+ * Action Schema
+ *
+ * This object is used to construct a schema describing some structure of actions.
+ */
+UCLASS(BlueprintType)
+class LEARNINGAGENTS_API ULearningAgentsActionSchema : public UObject
+{
+	GENERATED_BODY()
+
+public:
+
+	UE::Learning::Action::FSchema ActionSchema;
+};
+
+/**
+ * Action Object
+ *
+ * This object is used to construct or get the values of actions.
+ */
+UCLASS(BlueprintType)
+class LEARNINGAGENTS_API ULearningAgentsActionObject : public UObject
+{
+	GENERATED_BODY()
+
+public:
+
+	UE::Learning::Action::FObject ActionObject;
+};
+
 /** Enum Type representing either action A or action B */
 UENUM(BlueprintType)
 enum class ELearningAgentsEitherAction : uint8
@@ -54,365 +87,340 @@ enum class ELearningAgentsOptionalAction : uint8
 	Valid,
 };
 
-/**
- * Action Schema
- * 
- * This object is used to construct a schema describing some structure of actions.
- */
-UCLASS(BlueprintType)
-class LEARNINGAGENTS_API ULearningAgentsActionSchema : public UObject
+UCLASS()
+class LEARNINGAGENTS_API ULearningAgentsActions : public UBlueprintFunctionLibrary
 {
 	GENERATED_BODY()
 
 public:
 
-	/** Gets the internal action schema object */
-	const UE::Learning::Action::FSchema& GetActionSchema() const;
-
 	/**
 	 * Validates that the given object matches the schema. Will log errors on objects that don't match.
-	 * 
-	 * @param SchemaElement			Schema Element
+	 *
+	 * @param Schema				Action Schema
+	 * @param SchemaElement			Action Schema Element
 	 * @param Object				Action Object
 	 * @param ObjectElement			Action Object Element
 	 * @returns						true if the object matches the schema
 	 */
-	UFUNCTION(BlueprintPure, Category = "LearningAgents")
-	bool ValidateObjectMatchesSchema(
+	UFUNCTION(BlueprintPure = false, Category = "LearningAgents")
+	static bool ValidateObjectMatchesSchema(
+		const ULearningAgentsActionSchema* Schema,
 		const FLearningAgentsActionSchemaElement SchemaElement,
 		const ULearningAgentsActionObject* Object,
-		const FLearningAgentsActionObjectElement ObjectElement) const;
+		const FLearningAgentsActionObjectElement ObjectElement);
+
+	UFUNCTION(BlueprintPure = false, Category = "LearningAgents")
+	static void LogAction(
+		const ULearningAgentsActionObject* Object, 
+		const FLearningAgentsActionObjectElement Element);
 
 public:
 
-	UFUNCTION(BlueprintPure, Category = "LearningAgents")
-	FLearningAgentsActionSchemaElement SpecifyNullAction(const FName Name = TEXT("Null"));
+	UFUNCTION(BlueprintPure, Category = "LearningAgents", meta = (AdvancedDisplay = 1))
+	static FLearningAgentsActionSchemaElement SpecifyNullAction(ULearningAgentsActionSchema* Schema, const FName Tag = TEXT("Null"));
 
-	UFUNCTION(BlueprintPure, Category = "LearningAgents")
-	FLearningAgentsActionSchemaElement SpecifyContinuousAction(const int32 Size, const FName Name = TEXT("Continuous"));
+	UFUNCTION(BlueprintPure, Category = "LearningAgents", meta = (AdvancedDisplay = 2))
+	static FLearningAgentsActionSchemaElement SpecifyContinuousAction(ULearningAgentsActionSchema* Schema, const int32 Size, const FName Tag = TEXT("Continuous"));
 
-	UFUNCTION(BlueprintPure, Category = "LearningAgents", meta=(AutoCreateRefTerm="PriorProbabilities"))
-	FLearningAgentsActionSchemaElement SpecifyExclusiveDiscreteAction(const int32 Size, const TArray<float>& PriorProbabilities, const FName Name = TEXT("DiscreteExclusive"));
-	FLearningAgentsActionSchemaElement SpecifyExclusiveDiscreteActionFromArrayView(const int32 Size, const TArrayView<const float> PriorProbabilities = {}, const FName Name = TEXT("DiscreteExclusive"));
+	UFUNCTION(BlueprintPure, Category = "LearningAgents", meta = (AdvancedDisplay = 2, AutoCreateRefTerm = "PriorProbabilities"))
+	static FLearningAgentsActionSchemaElement SpecifyExclusiveDiscreteAction(ULearningAgentsActionSchema* Schema, const int32 Size, const TArray<float>& PriorProbabilities, const FName Tag = TEXT("DiscreteExclusive"));
+	static FLearningAgentsActionSchemaElement SpecifyExclusiveDiscreteActionFromArrayView(ULearningAgentsActionSchema* Schema, const int32 Size, const TArrayView<const float> PriorProbabilities = {}, const FName Tag = TEXT("DiscreteExclusive"));
 
-	UFUNCTION(BlueprintPure, Category = "LearningAgents", meta = (AutoCreateRefTerm = "PriorProbabilities"))
-	FLearningAgentsActionSchemaElement SpecifyInclusiveDiscreteAction(const int32 Size, const TArray<float>& PriorProbabilities, const FName Name = TEXT("DiscreteInclusive"));
-	FLearningAgentsActionSchemaElement SpecifyInclusiveDiscreteActionFromArrayView(const int32 Size, const TArrayView<const float> PriorProbabilities = {}, const FName Name = TEXT("DiscreteInclusive"));
+	UFUNCTION(BlueprintPure, Category = "LearningAgents", meta = (AdvancedDisplay = 2, AutoCreateRefTerm = "PriorProbabilities"))
+	static FLearningAgentsActionSchemaElement SpecifyInclusiveDiscreteAction(ULearningAgentsActionSchema* Schema, const int32 Size, const TArray<float>& PriorProbabilities, const FName Tag = TEXT("DiscreteInclusive"));
+	static FLearningAgentsActionSchemaElement SpecifyInclusiveDiscreteActionFromArrayView(ULearningAgentsActionSchema* Schema, const int32 Size, const TArrayView<const float> PriorProbabilities = {}, const FName Tag = TEXT("DiscreteInclusive"));
 
-	UFUNCTION(BlueprintPure, Category = "LearningAgents")
-	FLearningAgentsActionSchemaElement SpecifyStructAction(const TMap<FName, FLearningAgentsActionSchemaElement>& Elements, const FName Name = TEXT("Struct"));
+	UFUNCTION(BlueprintPure, Category = "LearningAgents", meta = (AdvancedDisplay = 2))
+	static FLearningAgentsActionSchemaElement SpecifyStructAction(ULearningAgentsActionSchema* Schema, const TMap<FName, FLearningAgentsActionSchemaElement>& Elements, const FName Tag = TEXT("Struct"));
 
-	UFUNCTION(BlueprintPure, Category = "LearningAgents")
-	FLearningAgentsActionSchemaElement SpecifyStructActionFromArrays(const TArray<FName>& ElementNames, const TArray<FLearningAgentsActionSchemaElement>& Elements, const FName Name = TEXT("Struct"));
-	FLearningAgentsActionSchemaElement SpecifyStructActionFromArrayViews(const TArrayView<const FName> ElementNames, const TArrayView<const FLearningAgentsActionSchemaElement> Elements, const FName Name = TEXT("Struct"));
+	UFUNCTION(BlueprintPure, Category = "LearningAgents", meta = (AdvancedDisplay = 3))
+	static FLearningAgentsActionSchemaElement SpecifyStructActionFromArrays(ULearningAgentsActionSchema* Schema, const TArray<FName>& ElementNames, const TArray<FLearningAgentsActionSchemaElement>& Elements, const FName Tag = TEXT("Struct"));
+	static FLearningAgentsActionSchemaElement SpecifyStructActionFromArrayViews(ULearningAgentsActionSchema* Schema, const TArrayView<const FName> ElementNames, const TArrayView<const FLearningAgentsActionSchemaElement> Elements, const FName Tag = TEXT("Struct"));
 
-	UFUNCTION(BlueprintPure, Category = "LearningAgents", meta = (AutoCreateRefTerm = "PriorProbabilities"))
-	FLearningAgentsActionSchemaElement SpecifyExclusiveUnionAction(const TMap<FName, FLearningAgentsActionSchemaElement>& Elements, const TMap<FName, float>& PriorProbabilities, const FName Name = TEXT("ExclusiveUnion"));
-	
-	UFUNCTION(BlueprintPure, Category = "LearningAgents", meta = (AutoCreateRefTerm = "PriorProbabilities"))
-	FLearningAgentsActionSchemaElement SpecifyExclusiveUnionActionFromArrays(const TArray<FName>& ElementNames, const TArray<FLearningAgentsActionSchemaElement>& Elements, const TArray<float>& PriorProbabilities, const FName Name = TEXT("ExclusiveUnion"));
-	FLearningAgentsActionSchemaElement SpecifyExclusiveUnionActionFromArrayViews(const TArrayView<const FName> ElementNames, const TArrayView<const FLearningAgentsActionSchemaElement> Elements, const TArrayView<const float> PriorProbabilities = {}, const FName Name = TEXT("ExclusiveUnion"));
+	UFUNCTION(BlueprintPure, Category = "LearningAgents", meta = (AdvancedDisplay = 2, AutoCreateRefTerm = "PriorProbabilities"))
+	static FLearningAgentsActionSchemaElement SpecifyExclusiveUnionAction(ULearningAgentsActionSchema* Schema, const TMap<FName, FLearningAgentsActionSchemaElement>& Elements, const TMap<FName, float>& PriorProbabilities, const FName Tag = TEXT("ExclusiveUnion"));
 
-	UFUNCTION(BlueprintPure, Category = "LearningAgents", meta = (AutoCreateRefTerm = "PriorProbabilities"))
-	FLearningAgentsActionSchemaElement SpecifyInclusiveUnionAction(const TMap<FName, FLearningAgentsActionSchemaElement>& Elements, const TMap<FName, float>& PriorProbabilities, const FName Name = TEXT("InclusiveUnion"));
-	
-	UFUNCTION(BlueprintPure, Category = "LearningAgents", meta = (AutoCreateRefTerm = "PriorProbabilities"))
-	FLearningAgentsActionSchemaElement SpecifyInclusiveUnionActionFromArrays(const TArray<FName> ElementNames, const TArray<FLearningAgentsActionSchemaElement>& Elements, const TArray<float>& PriorProbabilities, const FName Name = TEXT("InclusiveUnion"));
-	FLearningAgentsActionSchemaElement SpecifyInclusiveUnionActionFromArrayViews(const TArrayView<const FName> ElementNames, const TArrayView<const FLearningAgentsActionSchemaElement> Elements, const TArrayView<const float> PriorProbabilities = {}, const FName Name = TEXT("InclusiveUnion"));
+	UFUNCTION(BlueprintPure, Category = "LearningAgents", meta = (AdvancedDisplay = 3, AutoCreateRefTerm = "PriorProbabilities"))
+	static FLearningAgentsActionSchemaElement SpecifyExclusiveUnionActionFromArrays(ULearningAgentsActionSchema* Schema, const TArray<FName>& ElementNames, const TArray<FLearningAgentsActionSchemaElement>& Elements, const TArray<float>& PriorProbabilities, const FName Tag = TEXT("ExclusiveUnion"));
+	static FLearningAgentsActionSchemaElement SpecifyExclusiveUnionActionFromArrayViews(ULearningAgentsActionSchema* Schema, const TArrayView<const FName> ElementNames, const TArrayView<const FLearningAgentsActionSchemaElement> Elements, const TArrayView<const float> PriorProbabilities = {}, const FName Tag = TEXT("ExclusiveUnion"));
 
-	UFUNCTION(BlueprintPure, Category = "LearningAgents")
-	FLearningAgentsActionSchemaElement SpecifyStaticArrayAction(const FLearningAgentsActionSchemaElement Element, const int32 Num, const FName Name = TEXT("StaticArray"));
+	UFUNCTION(BlueprintPure, Category = "LearningAgents", meta = (AdvancedDisplay = 2, AutoCreateRefTerm = "PriorProbabilities"))
+	static FLearningAgentsActionSchemaElement SpecifyInclusiveUnionAction(ULearningAgentsActionSchema* Schema, const TMap<FName, FLearningAgentsActionSchemaElement>& Elements, const TMap<FName, float>& PriorProbabilities, const FName Tag = TEXT("InclusiveUnion"));
 
-	UFUNCTION(BlueprintPure, Category = "LearningAgents")
-	FLearningAgentsActionSchemaElement SpecifyPairAction(const FLearningAgentsActionSchemaElement Key, const FLearningAgentsActionSchemaElement Value, const FName Name = TEXT("Pair"));
+	UFUNCTION(BlueprintPure, Category = "LearningAgents", meta = (AdvancedDisplay = 3, AutoCreateRefTerm = "PriorProbabilities"))
+	static FLearningAgentsActionSchemaElement SpecifyInclusiveUnionActionFromArrays(ULearningAgentsActionSchema* Schema, const TArray<FName> ElementNames, const TArray<FLearningAgentsActionSchemaElement>& Elements, const TArray<float>& PriorProbabilities, const FName Tag = TEXT("InclusiveUnion"));
+	static FLearningAgentsActionSchemaElement SpecifyInclusiveUnionActionFromArrayViews(ULearningAgentsActionSchema* Schema, const TArrayView<const FName> ElementNames, const TArrayView<const FLearningAgentsActionSchemaElement> Elements, const TArrayView<const float> PriorProbabilities = {}, const FName Tag = TEXT("InclusiveUnion"));
 
-	UFUNCTION(BlueprintPure, Category = "LearningAgents", meta = (AutoCreateRefTerm = "PriorProbabilities"))
-	FLearningAgentsActionSchemaElement SpecifyEnumAction(const UEnum* Enum, const TMap<uint8, float>& PriorProbabilities, const FName Name = TEXT("Enum"));
+	UFUNCTION(BlueprintPure, Category = "LearningAgents", meta = (AdvancedDisplay = 3))
+	static FLearningAgentsActionSchemaElement SpecifyStaticArrayAction(ULearningAgentsActionSchema* Schema, const FLearningAgentsActionSchemaElement Element, const int32 Num, const FName Tag = TEXT("StaticArray"));
 
-	UFUNCTION(BlueprintPure, Category = "LearningAgents", meta = (AutoCreateRefTerm = "PriorProbabilities"))
-	FLearningAgentsActionSchemaElement SpecifyEnumActionFromArray(const UEnum* Enum, const TArray<float>& PriorProbabilities, const FName Name = TEXT("Enum"));
-	FLearningAgentsActionSchemaElement SpecifyEnumActionFromArrayView(const UEnum* Enum, const TArrayView<const float> PriorProbabilities = {}, const FName Name = TEXT("Enum"));
+	UFUNCTION(BlueprintPure, Category = "LearningAgents", meta = (AdvancedDisplay = 3))
+	static FLearningAgentsActionSchemaElement SpecifyPairAction(ULearningAgentsActionSchema* Schema, const FLearningAgentsActionSchemaElement Key, const FLearningAgentsActionSchemaElement Value, const FName Tag = TEXT("Pair"));
 
-	UFUNCTION(BlueprintPure, Category = "LearningAgents", meta = (AutoCreateRefTerm = "PriorProbabilities"))
-	FLearningAgentsActionSchemaElement SpecifyBitmaskAction(const UEnum* Enum, const TMap<uint8, float>& PriorProbabilities, const FName Name = TEXT("Bitmask"));
+	UFUNCTION(BlueprintPure, Category = "LearningAgents", meta = (AdvancedDisplay = 2, AutoCreateRefTerm = "PriorProbabilities"))
+	static FLearningAgentsActionSchemaElement SpecifyEnumAction(ULearningAgentsActionSchema* Schema, const UEnum* Enum, const TMap<uint8, float>& PriorProbabilities, const FName Tag = TEXT("Enum"));
 
-	UFUNCTION(BlueprintPure, Category = "LearningAgents", meta = (AutoCreateRefTerm = "PriorProbabilities"))
-	FLearningAgentsActionSchemaElement SpecifyBitmaskActionFromArray(const UEnum* Enum, const TArray<float>& PriorProbabilities, const FName Name = TEXT("Bitmask"));
-	FLearningAgentsActionSchemaElement SpecifyBitmaskActionFromArrayView(const UEnum* Enum, const TArrayView<const float> PriorProbabilities = {}, const FName Name = TEXT("Bitmask"));
+	UFUNCTION(BlueprintPure, Category = "LearningAgents", meta = (AdvancedDisplay = 2, AutoCreateRefTerm = "PriorProbabilities"))
+	static FLearningAgentsActionSchemaElement SpecifyEnumActionFromArray(ULearningAgentsActionSchema* Schema, const UEnum* Enum, const TArray<float>& PriorProbabilities, const FName Tag = TEXT("Enum"));
+	static FLearningAgentsActionSchemaElement SpecifyEnumActionFromArrayView(ULearningAgentsActionSchema* Schema, const UEnum* Enum, const TArrayView<const float> PriorProbabilities = {}, const FName Tag = TEXT("Enum"));
 
-	UFUNCTION(BlueprintPure, Category = "LearningAgents")
-	FLearningAgentsActionSchemaElement SpecifyOptionalAction(const FLearningAgentsActionSchemaElement Element, const float PriorProbability = 0.5f, const FName Name = TEXT("Optional"));
+	UFUNCTION(BlueprintPure, Category = "LearningAgents", meta = (AdvancedDisplay = 2, AutoCreateRefTerm = "PriorProbabilities"))
+	static FLearningAgentsActionSchemaElement SpecifyBitmaskAction(ULearningAgentsActionSchema* Schema, const UEnum* Enum, const TMap<uint8, float>& PriorProbabilities, const FName Tag = TEXT("Bitmask"));
 
-	UFUNCTION(BlueprintPure, Category = "LearningAgents")
-	FLearningAgentsActionSchemaElement SpecifyEitherAction(const FLearningAgentsActionSchemaElement A, const FLearningAgentsActionSchemaElement B, const float PriorProbabilityOfA = 0.5f, const FName Name = TEXT("Either"));
+	UFUNCTION(BlueprintPure, Category = "LearningAgents", meta = (AdvancedDisplay = 2, AutoCreateRefTerm = "PriorProbabilities"))
+	static FLearningAgentsActionSchemaElement SpecifyBitmaskActionFromArray(ULearningAgentsActionSchema* Schema, const UEnum* Enum, const TArray<float>& PriorProbabilities, const FName Tag = TEXT("Bitmask"));
+	static FLearningAgentsActionSchemaElement SpecifyBitmaskActionFromArrayView(ULearningAgentsActionSchema* Schema, const UEnum* Enum, const TArrayView<const float> PriorProbabilities = {}, const FName Tag = TEXT("Bitmask"));
 
-	UFUNCTION(BlueprintPure, Category = "LearningAgents")
-	FLearningAgentsActionSchemaElement SpecifyEncodingAction(const FLearningAgentsActionSchemaElement Element, const int32 EncodingSize = 128, const int32 LayerNum = 1, const ELearningAgentsActivationFunction ActivationFunction = ELearningAgentsActivationFunction::ELU, const FName Name = TEXT("Encoding"));
+	UFUNCTION(BlueprintPure, Category = "LearningAgents", meta = (AdvancedDisplay = 2))
+	static FLearningAgentsActionSchemaElement SpecifyOptionalAction(ULearningAgentsActionSchema* Schema, const FLearningAgentsActionSchemaElement Element, const float PriorProbability = 0.5f, const FName Tag = TEXT("Optional"));
 
-	UFUNCTION(BlueprintPure, Category = "LearningAgents")
-	FLearningAgentsActionSchemaElement SpecifyBoolAction(const float PriorProbability = 0.5f, const FName Name = TEXT("Bool"));
+	UFUNCTION(BlueprintPure, Category = "LearningAgents", meta = (AdvancedDisplay = 3))
+	static FLearningAgentsActionSchemaElement SpecifyEitherAction(ULearningAgentsActionSchema* Schema, const FLearningAgentsActionSchemaElement A, const FLearningAgentsActionSchemaElement B, const float PriorProbabilityOfA = 0.5f, const FName Tag = TEXT("Either"));
 
-	UFUNCTION(BlueprintPure, Category = "LearningAgents")
-	FLearningAgentsActionSchemaElement SpecifyFloatAction(const FName Name = TEXT("Float"));
+	UFUNCTION(BlueprintPure, Category = "LearningAgents", meta = (AdvancedDisplay = 2))
+	static FLearningAgentsActionSchemaElement SpecifyEncodingAction(ULearningAgentsActionSchema* Schema, const FLearningAgentsActionSchemaElement Element, const int32 EncodingSize = 128, const int32 LayerNum = 1, const ELearningAgentsActivationFunction ActivationFunction = ELearningAgentsActivationFunction::ELU, const FName Tag = TEXT("Encoding"));
 
-	UFUNCTION(BlueprintPure, Category = "LearningAgents")
-	FLearningAgentsActionSchemaElement SpecifyLocationAction(const FName Name = TEXT("Location"));
+	UFUNCTION(BlueprintPure, Category = "LearningAgents", meta = (AdvancedDisplay = 1))
+	static FLearningAgentsActionSchemaElement SpecifyBoolAction(ULearningAgentsActionSchema* Schema, const float PriorProbability = 0.5f, const FName Tag = TEXT("Bool"));
 
-	UFUNCTION(BlueprintPure, Category = "LearningAgents")
-	FLearningAgentsActionSchemaElement SpecifyRotationAction(const FName Name = TEXT("Rotation"));
+	UFUNCTION(BlueprintPure, Category = "LearningAgents", meta = (AdvancedDisplay = 1))
+	static FLearningAgentsActionSchemaElement SpecifyFloatAction(ULearningAgentsActionSchema* Schema, const FName Tag = TEXT("Float"));
 
-	UFUNCTION(BlueprintPure, Category = "LearningAgents")
-	FLearningAgentsActionSchemaElement SpecifyScaleAction(const FName Name = TEXT("Scale"));
+	UFUNCTION(BlueprintPure, Category = "LearningAgents", meta = (AdvancedDisplay = 1))
+	static FLearningAgentsActionSchemaElement SpecifyLocationAction(ULearningAgentsActionSchema* Schema, const FName Tag = TEXT("Location"));
 
-	UFUNCTION(BlueprintPure, Category = "LearningAgents")
-	FLearningAgentsActionSchemaElement SpecifyTransformAction(const FName Name = TEXT("Transform"));
+	UFUNCTION(BlueprintPure, Category = "LearningAgents", meta = (AdvancedDisplay = 1))
+	static FLearningAgentsActionSchemaElement SpecifyRotationAction(ULearningAgentsActionSchema* Schema, const FName Tag = TEXT("Rotation"));
 
-	UFUNCTION(BlueprintPure, Category = "LearningAgents")
-	FLearningAgentsActionSchemaElement SpecifyAngleAction(const FName Name = TEXT("Angle"));
+	UFUNCTION(BlueprintPure, Category = "LearningAgents", meta = (AdvancedDisplay = 1))
+	static FLearningAgentsActionSchemaElement SpecifyScaleAction(ULearningAgentsActionSchema* Schema, const FName Tag = TEXT("Scale"));
 
-	UFUNCTION(BlueprintPure, Category = "LearningAgents")
-	FLearningAgentsActionSchemaElement SpecifyVelocityAction(const FName Name = TEXT("Velocity"));
+	UFUNCTION(BlueprintPure, Category = "LearningAgents", meta = (AdvancedDisplay = 1))
+	static FLearningAgentsActionSchemaElement SpecifyTransformAction(ULearningAgentsActionSchema* Schema, const FName Tag = TEXT("Transform"));
 
-	UFUNCTION(BlueprintPure, Category = "LearningAgents")
-	FLearningAgentsActionSchemaElement SpecifyDirectionAction(const FName Name = TEXT("Direction"));
+	UFUNCTION(BlueprintPure, Category = "LearningAgents", meta = (AdvancedDisplay = 1))
+	static FLearningAgentsActionSchemaElement SpecifyAngleAction(ULearningAgentsActionSchema* Schema, const FName Tag = TEXT("Angle"));
 
-	UFUNCTION(BlueprintPure, Category = "LearningAgents")
-	FLearningAgentsActionSchemaElement SpecifySpeedAction(const FName Name = TEXT("Speed"));
+	UFUNCTION(BlueprintPure, Category = "LearningAgents", meta = (AdvancedDisplay = 1))
+	static FLearningAgentsActionSchemaElement SpecifyVelocityAction(ULearningAgentsActionSchema* Schema, const FName Tag = TEXT("Velocity"));
 
-private:
+	UFUNCTION(BlueprintPure, Category = "LearningAgents", meta = (AdvancedDisplay = 1))
+	static FLearningAgentsActionSchemaElement SpecifyDirectionAction(ULearningAgentsActionSchema* Schema, const FName Tag = TEXT("Direction"));
 
-	UE::Learning::Action::FSchema ActionSchema;
-};
-
-/**
- * Action Object
- *
- * This object is used to construct or get the values of actions.
- */
-UCLASS(BlueprintType)
-class LEARNINGAGENTS_API ULearningAgentsActionObject : public UObject
-{
-	GENERATED_BODY()
+	UFUNCTION(BlueprintPure, Category = "LearningAgents", meta = (AdvancedDisplay = 1))
+	static FLearningAgentsActionSchemaElement SpecifySpeedAction(ULearningAgentsActionSchema* Schema, const FName Tag = TEXT("Speed"));
 
 public:
 
-	/** Gets the internal action object */
-	const UE::Learning::Action::FObject& GetActionObject() const;
+	UFUNCTION(BlueprintPure, Category = "LearningAgents", meta = (AdvancedDisplay = 1))
+	static FLearningAgentsActionObjectElement MakeNullAction(ULearningAgentsActionObject* Object, const FName Tag = TEXT("Null"));
 
-	/** Gets the internal action object */
-	UE::Learning::Action::FObject& GetActionObject();
+	UFUNCTION(BlueprintPure, Category = "LearningAgents", meta = (AdvancedDisplay = 2))
+	static FLearningAgentsActionObjectElement MakeContinuousAction(ULearningAgentsActionObject* Object, const TArray<float>& Values, const FName Tag = TEXT("Continuous"));
+	static FLearningAgentsActionObjectElement MakeContinuousActionFromArrayView(ULearningAgentsActionObject* Object, const TArrayView<const float> Values, const FName Tag = TEXT("Continuous"));
 
-public:
+	UFUNCTION(BlueprintPure, Category = "LearningAgents", meta = (AdvancedDisplay = 2))
+	static FLearningAgentsActionObjectElement MakeExclusiveDiscreteAction(ULearningAgentsActionObject* Object, const int32 Index, const FName Tag = TEXT("DiscreteExclusive"));
 
-	UFUNCTION(BlueprintPure = false, Category = "LearningAgents")
-	void LogAction(const FLearningAgentsActionObjectElement Element);
+	UFUNCTION(BlueprintPure, Category = "LearningAgents", meta = (AdvancedDisplay = 2))
+	static FLearningAgentsActionObjectElement MakeInclusiveDiscreteAction(ULearningAgentsActionObject* Object, const TArray<int32>& Indices, const FName Tag = TEXT("DiscreteInclusive"));
+	static FLearningAgentsActionObjectElement MakeInclusiveDiscreteActionFromArrayView(ULearningAgentsActionObject* Object, const TArrayView<const int32> Indices, const FName Tag = TEXT("DiscreteInclusive"));
 
-public:
+	UFUNCTION(BlueprintPure, Category = "LearningAgents", meta = (AdvancedDisplay = 2))
+	static FLearningAgentsActionObjectElement MakeStructAction(ULearningAgentsActionObject* Object, const TMap<FName, FLearningAgentsActionObjectElement>& Elements, const FName Tag = TEXT("Struct"));
 
-	UFUNCTION(BlueprintPure, Category = "LearningAgents")
-	FLearningAgentsActionObjectElement MakeNullAction(const FName Name = TEXT("Null"));
+	UFUNCTION(BlueprintPure, Category = "LearningAgents", meta = (AdvancedDisplay = 3))
+	static FLearningAgentsActionObjectElement MakeStructActionFromArrays(ULearningAgentsActionObject* Object, const TArray<FName>& ElementNames, const TArray<FLearningAgentsActionObjectElement>& Elements, const FName Tag = TEXT("Struct"));
+	static FLearningAgentsActionObjectElement MakeStructActionFromArrayViews(ULearningAgentsActionObject* Object, const TArrayView<const FName> ElementNames, const TArrayView<const FLearningAgentsActionObjectElement> Elements, const FName Tag = TEXT("Struct"));
 
-	UFUNCTION(BlueprintPure, Category = "LearningAgents")
-	FLearningAgentsActionObjectElement MakeContinuousAction(const TArray<float>& Values, const FName Name = TEXT("Continuous"));
-	FLearningAgentsActionObjectElement MakeContinuousActionFromArrayView(const TArrayView<const float> Values, const FName Name = TEXT("Continuous"));
+	UFUNCTION(BlueprintPure, Category = "LearningAgents", meta = (AdvancedDisplay = 3))
+	static FLearningAgentsActionObjectElement MakeExclusiveUnionAction(ULearningAgentsActionObject* Object, const FName ElementName, const FLearningAgentsActionObjectElement Element, const FName Tag = TEXT("ExclusiveUnion"));
 
-	UFUNCTION(BlueprintPure, Category = "LearningAgents")
-	FLearningAgentsActionObjectElement MakeExclusiveDiscreteAction(const int32 Index, const FName Name = TEXT("DiscreteExclusive"));
+	UFUNCTION(BlueprintPure, Category = "LearningAgents", meta = (AdvancedDisplay = 2))
+	static FLearningAgentsActionObjectElement MakeInclusiveUnionAction(ULearningAgentsActionObject* Object, const TMap<FName, FLearningAgentsActionObjectElement>& Elements, const FName Tag = TEXT("InclusiveUnion"));
 
-	UFUNCTION(BlueprintPure, Category = "LearningAgents")
-	FLearningAgentsActionObjectElement MakeInclusiveDiscreteAction(const TArray<int32>& Indices, const FName Name = TEXT("DiscreteInclusive"));
-	FLearningAgentsActionObjectElement MakeInclusiveDiscreteActionFromArrayView(const TArrayView<const int32> Indices, const FName Name = TEXT("DiscreteInclusive"));
+	UFUNCTION(BlueprintPure, Category = "LearningAgents", meta = (AdvancedDisplay = 3))
+	static FLearningAgentsActionObjectElement MakeInclusiveUnionActionFromArrays(ULearningAgentsActionObject* Object, const TArray<FName>& ElementNames, const TArray<FLearningAgentsActionObjectElement>& Elements, const FName Tag = TEXT("InclusiveUnion"));
+	static FLearningAgentsActionObjectElement MakeInclusiveUnionActionFromArrayViews(ULearningAgentsActionObject* Object, const TArrayView<const FName> ElementNames, const TArrayView<const FLearningAgentsActionObjectElement> Elements, const FName Tag = TEXT("InclusiveUnion"));
 
-	UFUNCTION(BlueprintPure, Category = "LearningAgents")
-	FLearningAgentsActionObjectElement MakeStructAction(const TMap<FName, FLearningAgentsActionObjectElement>& Elements, const FName Name = TEXT("Struct"));
-	
-	UFUNCTION(BlueprintPure, Category = "LearningAgents")
-	FLearningAgentsActionObjectElement MakeStructActionFromArrays(const TArray<FName>& ElementNames, const TArray<FLearningAgentsActionObjectElement>& Elements, const FName Name = TEXT("Struct"));
-	FLearningAgentsActionObjectElement MakeStructActionFromArrayViews(const TArrayView<const FName> ElementNames, const TArrayView<const FLearningAgentsActionObjectElement> Elements, const FName Name = TEXT("Struct"));
+	UFUNCTION(BlueprintPure, Category = "LearningAgents", meta = (AdvancedDisplay = 2))
+	static FLearningAgentsActionObjectElement MakeStaticArrayAction(ULearningAgentsActionObject* Object, const TArray<FLearningAgentsActionObjectElement>& Elements, const FName Tag = TEXT("StaticArray"));
+	static FLearningAgentsActionObjectElement MakeStaticArrayActionFromArrayView(ULearningAgentsActionObject* Object, const TArrayView<const FLearningAgentsActionObjectElement> Elements, const FName Tag = TEXT("StaticArray"));
 
-	UFUNCTION(BlueprintPure, Category = "LearningAgents")
-	FLearningAgentsActionObjectElement MakeExclusiveUnionAction(const FName ElementName, const FLearningAgentsActionObjectElement Element, const FName Name = TEXT("ExclusiveUnion"));
+	UFUNCTION(BlueprintPure, Category = "LearningAgents", meta = (AdvancedDisplay = 3))
+	static FLearningAgentsActionObjectElement MakePairAction(ULearningAgentsActionObject* Object, const FLearningAgentsActionObjectElement Key, const FLearningAgentsActionObjectElement Value, const FName Tag = TEXT("Pair"));
 
-	UFUNCTION(BlueprintPure, Category = "LearningAgents")
-	FLearningAgentsActionObjectElement MakeInclusiveUnionAction(const TMap<FName, FLearningAgentsActionObjectElement>& Elements, const FName Name = TEXT("InclusiveUnion"));
+	UFUNCTION(BlueprintPure, Category = "LearningAgents", meta = (AdvancedDisplay = 2))
+	static FLearningAgentsActionObjectElement MakeEnumAction(ULearningAgentsActionObject* Object, const UEnum* Enum, const uint8 EnumValue, const FName Tag = TEXT("Enum"));
 
-	UFUNCTION(BlueprintPure, Category = "LearningAgents")
-	FLearningAgentsActionObjectElement MakeInclusiveUnionActionFromArrays(const TArray<FName>& ElementNames, const TArray<FLearningAgentsActionObjectElement>& Elements, const FName Name = TEXT("InclusiveUnion"));
-	FLearningAgentsActionObjectElement MakeInclusiveUnionActionFromArrayViews(const TArrayView<const FName> ElementNames, const TArrayView<const FLearningAgentsActionObjectElement> Elements, const FName Name = TEXT("InclusiveUnion"));
+	UFUNCTION(BlueprintPure, Category = "LearningAgents", meta = (AdvancedDisplay = 2))
+	static FLearningAgentsActionObjectElement MakeBitmaskAction(ULearningAgentsActionObject* Object, const UEnum* Enum, const int32 BitmaskValue, const FName Tag = TEXT("Bitmask"));
 
-	UFUNCTION(BlueprintPure, Category = "LearningAgents")
-	FLearningAgentsActionObjectElement MakeStaticArrayAction(const TArray<FLearningAgentsActionObjectElement>& Elements, const FName Name = TEXT("StaticArray"));
-	FLearningAgentsActionObjectElement MakeStaticArrayActionFromArrayView(const TArrayView<const FLearningAgentsActionObjectElement> Elements, const FName Name = TEXT("StaticArray"));
+	UFUNCTION(BlueprintPure, Category = "LearningAgents", meta = (AdvancedDisplay = 3))
+	static FLearningAgentsActionObjectElement MakeOptionalAction(ULearningAgentsActionObject* Object, const FLearningAgentsActionObjectElement Element, const ELearningAgentsOptionalAction Option, const FName Tag = TEXT("Optional"));
 
-	UFUNCTION(BlueprintPure, Category = "LearningAgents")
-	FLearningAgentsActionObjectElement MakePairAction(const FLearningAgentsActionObjectElement Key, const FLearningAgentsActionObjectElement Value, const FName Name = TEXT("Pair"));
+	UFUNCTION(BlueprintPure, Category = "LearningAgents", meta = (AdvancedDisplay = 1))
+	static FLearningAgentsActionObjectElement MakeOptionalNullAction(ULearningAgentsActionObject* Object, const FName Tag = TEXT("Optional"));
 
-	UFUNCTION(BlueprintPure, Category = "LearningAgents")
-	FLearningAgentsActionObjectElement MakeEnumAction(const UEnum* Enum, const uint8 EnumValue, const FName Name = TEXT("Enum"));
+	UFUNCTION(BlueprintPure, Category = "LearningAgents", meta = (AdvancedDisplay = 2))
+	static FLearningAgentsActionObjectElement MakeOptionalValidAction(ULearningAgentsActionObject* Object, const FLearningAgentsActionObjectElement Element, const FName Tag = TEXT("Optional"));
 
-	UFUNCTION(BlueprintPure, Category = "LearningAgents")
-	FLearningAgentsActionObjectElement MakeBitmaskAction(const UEnum* Enum, const int32 BitmaskValue, const FName Name = TEXT("Bitmask"));
+	UFUNCTION(BlueprintPure, Category = "LearningAgents", meta = (AdvancedDisplay = 3))
+	static FLearningAgentsActionObjectElement MakeEitherAction(ULearningAgentsActionObject* Object, const FLearningAgentsActionObjectElement Element, const ELearningAgentsEitherAction Either, const FName Tag = TEXT("Either"));
 
-	UFUNCTION(BlueprintPure, Category = "LearningAgents")
-	FLearningAgentsActionObjectElement MakeOptionalAction(const FLearningAgentsActionObjectElement Element, const ELearningAgentsOptionalAction Option, const FName Name = TEXT("Optional"));
+	UFUNCTION(BlueprintPure, Category = "LearningAgents", meta = (AdvancedDisplay = 2, DisplayName = "Make Either A Action"))
+	static FLearningAgentsActionObjectElement MakeEitherAAction(ULearningAgentsActionObject* Object, const FLearningAgentsActionObjectElement A, const FName Tag = TEXT("Either"));
 
-	UFUNCTION(BlueprintPure, Category = "LearningAgents")
-	FLearningAgentsActionObjectElement MakeOptionalNullAction(const FName Name = TEXT("Optional"));
+	UFUNCTION(BlueprintPure, Category = "LearningAgents", meta = (AdvancedDisplay = 2, DisplayName = "Make Either B Action"))
+	static FLearningAgentsActionObjectElement MakeEitherBAction(ULearningAgentsActionObject* Object, const FLearningAgentsActionObjectElement B, const FName Tag = TEXT("Either"));
 
-	UFUNCTION(BlueprintPure, Category = "LearningAgents")
-	FLearningAgentsActionObjectElement MakeOptionalValidAction(const FLearningAgentsActionObjectElement Element, const FName Name = TEXT("Optional"));
+	UFUNCTION(BlueprintPure, Category = "LearningAgents", meta = (AdvancedDisplay = 2))
+	static FLearningAgentsActionObjectElement MakeEncodingAction(ULearningAgentsActionObject* Object, const FLearningAgentsActionObjectElement Element, const FName Tag = TEXT("Encoding"));
 
-	UFUNCTION(BlueprintPure, Category = "LearningAgents")
-	FLearningAgentsActionObjectElement MakeEitherAction(const FLearningAgentsActionObjectElement Element, const ELearningAgentsEitherAction Either, const FName Name = TEXT("Either"));
+	UFUNCTION(BlueprintPure, Category = "LearningAgents", meta = (AdvancedDisplay = 2))
+	static FLearningAgentsActionObjectElement MakeBoolAction(ULearningAgentsActionObject* Object, const bool bValue, const FName Tag = TEXT("Bool"));
 
-	UFUNCTION(BlueprintPure, Category = "LearningAgents", meta = (DisplayName = "Make Either A Action"))
-	FLearningAgentsActionObjectElement MakeEitherAAction(const FLearningAgentsActionObjectElement A, const FName Name = TEXT("Either"));
+	UFUNCTION(BlueprintPure, Category = "LearningAgents", meta = (AdvancedDisplay = 3))
+	static FLearningAgentsActionObjectElement MakeFloatAction(ULearningAgentsActionObject* Object, const float Value, const float FloatScale = 1.0f, const FName Tag = TEXT("Float"));
 
-	UFUNCTION(BlueprintPure, Category = "LearningAgents", meta = (DisplayName = "Make Either B Action"))
-	FLearningAgentsActionObjectElement MakeEitherBAction(const FLearningAgentsActionObjectElement B, const FName Name = TEXT("Either"));
+	UFUNCTION(BlueprintPure, Category = "LearningAgents", meta = (AdvancedDisplay = 4))
+	static FLearningAgentsActionObjectElement MakeLocationAction(ULearningAgentsActionObject* Object, const FVector Location, const FTransform RelativeTransform = FTransform(), const float LocationScale = 100.0f, const FName Tag = TEXT("Location"));
 
-	UFUNCTION(BlueprintPure, Category = "LearningAgents")
-	FLearningAgentsActionObjectElement MakeEncodingAction(const FLearningAgentsActionObjectElement Element, const FName Name = TEXT("Encoding"));
+	UFUNCTION(BlueprintPure, Category = "LearningAgents", meta = (AdvancedDisplay = 4))
+	static FLearningAgentsActionObjectElement MakeRotationAction(ULearningAgentsActionObject* Object, const FRotator Rotation, const FRotator RelativeRotation = FRotator::ZeroRotator, const float RotationScale = 90.0f, const FName Tag = TEXT("Rotation"));
 
-	UFUNCTION(BlueprintPure, Category = "LearningAgents")
-	FLearningAgentsActionObjectElement MakeBoolAction(const bool bValue, const FName Name = TEXT("Bool"));
+	UFUNCTION(BlueprintPure, Category = "LearningAgents", meta = (AdvancedDisplay = 4))
+	static FLearningAgentsActionObjectElement MakeRotationActionFromQuat(ULearningAgentsActionObject* Object, const FQuat Rotation, const FQuat RelativeRotation, const float RotationScale = 90.0f, const FName Tag = TEXT("Rotation"));
 
-	UFUNCTION(BlueprintPure, Category = "LearningAgents")
-	FLearningAgentsActionObjectElement MakeFloatAction(const float Value, const float FloatScale = 1.0f, const FName Name = TEXT("Float"));
+	UFUNCTION(BlueprintPure, Category = "LearningAgents", meta = (AdvancedDisplay = 3))
+	static FLearningAgentsActionObjectElement MakeScaleAction(ULearningAgentsActionObject* Object, const FVector Scale, const FVector RelativeScale = FVector(1, 1, 1), const FName Tag = TEXT("Scale"));
 
-	UFUNCTION(BlueprintPure, Category = "LearningAgents")
-	FLearningAgentsActionObjectElement MakeLocationAction(const FVector Location, const FTransform RelativeTransform = FTransform(), const float LocationScale = 100.0f, const FName Name = TEXT("Location"));
+	UFUNCTION(BlueprintPure, Category = "LearningAgents", meta = (AdvancedDisplay = 4))
+	static FLearningAgentsActionObjectElement MakeTransformAction(ULearningAgentsActionObject* Object, const FTransform Transform, const FTransform RelativeTransform = FTransform(), const float LocationScale = 100.0f, const FName Tag = TEXT("Transform"));
 
-	UFUNCTION(BlueprintPure, Category = "LearningAgents")
-	FLearningAgentsActionObjectElement MakeRotationAction(const FRotator Rotation, const FRotator RelativeRotation = FRotator::ZeroRotator, const float RotationScale = 90.0f, const FName Name = TEXT("Rotation"));
-	FLearningAgentsActionObjectElement MakeRotationActionFromQuat(const FQuat Rotation, const FQuat RelativeRotation = FQuat::Identity, const float RotationScale = 90.0f, const FName Name = TEXT("Rotation"));
+	UFUNCTION(BlueprintPure, Category = "LearningAgents", meta = (AdvancedDisplay = 4))
+	static FLearningAgentsActionObjectElement MakeAngleAction(ULearningAgentsActionObject* Object, const float Angle, const float RelativeAngle = 0.0f, const float AngleScale = 90.0f, const FName Tag = TEXT("Angle"));
 
-	UFUNCTION(BlueprintPure, Category = "LearningAgents")
-	FLearningAgentsActionObjectElement MakeScaleAction(const FVector Scale, const FVector RelativeScale = FVector(1,1,1), const FName Name = TEXT("Scale"));
+	UFUNCTION(BlueprintPure, Category = "LearningAgents", meta = (AdvancedDisplay = 4))
+	static FLearningAgentsActionObjectElement MakeAngleActionRadians(ULearningAgentsActionObject* Object, const float Angle, const float RelativeAngle = 0.0f, const float AngleScale = 1.57079632679f, const FName Tag = TEXT("Angle"));
 
-	UFUNCTION(BlueprintPure, Category = "LearningAgents")
-	FLearningAgentsActionObjectElement MakeTransformAction(const FTransform Transform, const FTransform RelativeTransform = FTransform(), const float LocationScale = 100.0f, const FName Name = TEXT("Transform"));
+	UFUNCTION(BlueprintPure, Category = "LearningAgents", meta = (AdvancedDisplay = 4))
+	static FLearningAgentsActionObjectElement MakeVelocityAction(ULearningAgentsActionObject* Object, const FVector Velocity, const FTransform RelativeTransform = FTransform(), const float VelocityScale = 200.0f, const FName Tag = TEXT("Velocity"));
 
-	UFUNCTION(BlueprintPure, Category = "LearningAgents")
-	FLearningAgentsActionObjectElement MakeAngleAction(const float Angle, const float RelativeAngle = 0.0f, const float AngleScale = 90.0f, const FName Name = TEXT("Angle"));
+	UFUNCTION(BlueprintPure, Category = "LearningAgents", meta = (AdvancedDisplay = 3))
+	static FLearningAgentsActionObjectElement MakeDirectionAction(ULearningAgentsActionObject* Object, const FVector Direction, const FTransform RelativeTransform = FTransform(), const FName Tag = TEXT("Direction"));
 
-	UFUNCTION(BlueprintPure, Category = "LearningAgents")
-	FLearningAgentsActionObjectElement MakeAngleActionRadians(const float Angle, const float RelativeAngle = 0.0f, const float AngleScale = 1.57079632679f, const FName Name = TEXT("Angle"));
+	UFUNCTION(BlueprintPure, Category = "LearningAgents", meta = (AdvancedDisplay = 3))
+	static FLearningAgentsActionObjectElement MakeSpeedAction(ULearningAgentsActionObject* Object, const float Speed, const float SpeedScale = 200.0f, const FName Tag = TEXT("Speed"));
 
-	UFUNCTION(BlueprintPure, Category = "LearningAgents")
-	FLearningAgentsActionObjectElement MakeVelocityAction(const FVector Velocity, const FTransform RelativeTransform = FTransform(), const float VelocityScale = 200.0f, const FName Name = TEXT("Velocity"));
+public:
 
-	UFUNCTION(BlueprintPure, Category = "LearningAgents")
-	FLearningAgentsActionObjectElement MakeDirectionAction(const FVector Direction, const FTransform RelativeTransform = FTransform(), const FName Name = TEXT("Direction"));
+	UFUNCTION(BlueprintPure = false, Category = "LearningAgents", meta = (AdvancedDisplay = 2, ReturnDisplayName = "Success"))
+	static bool GetNullAction(const ULearningAgentsActionObject* Object, const FLearningAgentsActionObjectElement Element, const FName Tag = TEXT("Null"));
 
-	UFUNCTION(BlueprintPure, Category = "LearningAgents")
-	FLearningAgentsActionObjectElement MakeSpeedAction(const float Speed, const float SpeedScale = 200.0f, const FName Name = TEXT("Speed"));
+	UFUNCTION(BlueprintPure = false, Category = "LearningAgents", meta = (AdvancedDisplay = 3, ReturnDisplayName = "Success"))
+	static bool GetContinuousActionNum(int32& OutNum, const ULearningAgentsActionObject* Object, const FLearningAgentsActionObjectElement Element, const FName Tag = TEXT("Continuous"));
 
-public:
+	UFUNCTION(BlueprintPure = false, Category = "LearningAgents", meta = (AdvancedDisplay = 3, ReturnDisplayName = "Success"))
+	static bool GetContinuousAction(TArray<float>& OutValues, const ULearningAgentsActionObject* Object, const FLearningAgentsActionObjectElement Element, const FName Tag = TEXT("Continuous"));
+	static bool GetContinuousActionToArrayView(TArrayView<float> OutValues, const ULearningAgentsActionObject* Object, const FLearningAgentsActionObjectElement Element, const FName Tag = TEXT("Continuous"));
 
-	UFUNCTION(BlueprintPure = false, Category = "LearningAgents", meta = (ReturnDisplayName = "Success"))
-	bool GetNullAction(const FLearningAgentsActionObjectElement Element, const FName Name = TEXT("Null")) const;
+	UFUNCTION(BlueprintPure = false, Category = "LearningAgents", meta = (AdvancedDisplay = 3, ReturnDisplayName = "Success"))
+	static bool GetExclusiveDiscreteAction(int32& OutIndex, const ULearningAgentsActionObject* Object, const FLearningAgentsActionObjectElement Element, const FName Tag = TEXT("DiscreteExclusive"));
 
-	UFUNCTION(BlueprintPure = false, Category = "LearningAgents", meta = (ReturnDisplayName = "Success"))
-	bool GetContinuousActionNum(int32& OutNum, const FLearningAgentsActionObjectElement Element, const FName Name = TEXT("Continuous")) const;
+	UFUNCTION(BlueprintPure = false, Category = "LearningAgents", meta = (AdvancedDisplay = 3, ReturnDisplayName = "Success"))
+	static bool GetInclusiveDiscreteActionNum(int32& OutNum, const ULearningAgentsActionObject* Object, const FLearningAgentsActionObjectElement Element, const FName Tag = TEXT("DiscreteInclusive"));
 
-	UFUNCTION(BlueprintPure = false, Category = "LearningAgents", meta = (ReturnDisplayName = "Success"))
-	bool GetContinuousAction(TArray<float>& OutValues, const FLearningAgentsActionObjectElement Element, const FName Name = TEXT("Continuous")) const;
-	bool GetContinuousActionToArrayView(TArrayView<float> OutValues, const FLearningAgentsActionObjectElement Element, const FName Name = TEXT("Continuous")) const;
-	
-	UFUNCTION(BlueprintPure = false, Category = "LearningAgents", meta = (ReturnDisplayName = "Success"))
-	bool GetExclusiveDiscreteAction(int32& OutIndex, const FLearningAgentsActionObjectElement Element, const FName Name = TEXT("DiscreteExclusive")) const;
+	UFUNCTION(BlueprintPure = false, Category = "LearningAgents", meta = (AdvancedDisplay = 3, ReturnDisplayName = "Success"))
+	static bool GetInclusiveDiscreteAction(TArray<int32>& OutIndices, const ULearningAgentsActionObject* Object, const FLearningAgentsActionObjectElement Element, const FName Tag = TEXT("DiscreteInclusive"));
+	static bool GetInclusiveDiscreteActionToArrayView(TArrayView<int32> OutIndices, const ULearningAgentsActionObject* Object, const FLearningAgentsActionObjectElement Element, const FName Tag = TEXT("DiscreteInclusive"));
 
-	UFUNCTION(BlueprintPure = false, Category = "LearningAgents", meta = (ReturnDisplayName = "Success"))
-	bool GetInclusiveDiscreteActionNum(int32& OutNum, const FLearningAgentsActionObjectElement Element, const FName Name = TEXT("DiscreteInclusive")) const;
+	UFUNCTION(BlueprintPure = false, Category = "LearningAgents", meta = (AdvancedDisplay = 3, ReturnDisplayName = "Success"))
+	static bool GetStructActionNum(int32& OutNum, const ULearningAgentsActionObject* Object, const FLearningAgentsActionObjectElement Element, const FName Tag = TEXT("Struct"));
 
-	UFUNCTION(BlueprintPure = false, Category = "LearningAgents", meta = (ReturnDisplayName = "Success"))
-	bool GetInclusiveDiscreteAction(TArray<int32>& OutIndices, const FLearningAgentsActionObjectElement Element, const FName Name = TEXT("DiscreteInclusive")) const;
-	bool GetInclusiveDiscreteActionToArrayView(TArrayView<int32> OutIndices, const FLearningAgentsActionObjectElement Element, const FName Name = TEXT("DiscreteInclusive")) const;
+	UFUNCTION(BlueprintPure = false, Category = "LearningAgents", meta = (AdvancedDisplay = 3, ReturnDisplayName = "Success"))
+	static bool GetStructAction(TMap<FName, FLearningAgentsActionObjectElement>& OutElements, const ULearningAgentsActionObject* Object, const FLearningAgentsActionObjectElement Element, const FName Tag = TEXT("Struct"));
 
-	UFUNCTION(BlueprintPure = false, Category = "LearningAgents", meta = (ReturnDisplayName = "Success"))
-	bool GetStructActionNum(int32& OutNum, const FLearningAgentsActionObjectElement Element, const FName Name = TEXT("Struct")) const;
+	UFUNCTION(BlueprintPure = false, Category = "LearningAgents", meta = (AdvancedDisplay = 4, ReturnDisplayName = "Success"))
+	static bool GetStructActionToArrays(TArray<FName>& OutElementNames, TArray<FLearningAgentsActionObjectElement>& OutElements, const ULearningAgentsActionObject* Object, const FLearningAgentsActionObjectElement Element, const FName Tag = TEXT("Struct"));
+	static bool GetStructActionToArrayViews(TArrayView<FName> OutElementNames, TArrayView<FLearningAgentsActionObjectElement> OutElements, const ULearningAgentsActionObject* Object, const FLearningAgentsActionObjectElement Element, const FName Tag = TEXT("Struct"));
 
-	UFUNCTION(BlueprintPure = false, Category = "LearningAgents", meta = (ReturnDisplayName = "Success"))
-	bool GetStructAction(TMap<FName, FLearningAgentsActionObjectElement>& OutElements, const FLearningAgentsActionObjectElement Element, const FName Name = TEXT("Struct")) const;
-	
-	UFUNCTION(BlueprintPure = false, Category = "LearningAgents", meta = (ReturnDisplayName = "Success"))
-	bool GetStructActionToArrays(TArray<FName>& OutElementNames, TArray<FLearningAgentsActionObjectElement>& OutElements, const FLearningAgentsActionObjectElement Element, const FName Name = TEXT("Struct")) const;
-	bool GetStructActionToArrayViews(TArrayView<FName> OutElementNames, TArrayView<FLearningAgentsActionObjectElement> OutElements, const FLearningAgentsActionObjectElement Element, const FName Name = TEXT("Struct")) const;
+	UFUNCTION(BlueprintPure = false, Category = "LearningAgents", meta = (AdvancedDisplay = 4, ReturnDisplayName = "Success"))
+	static bool GetExclusiveUnionAction(FName& OutElementName, FLearningAgentsActionObjectElement& OutElement, const ULearningAgentsActionObject* Object, const FLearningAgentsActionObjectElement Element, const FName Tag = TEXT("ExclusiveUnion"));
 
-	UFUNCTION(BlueprintPure = false, Category = "LearningAgents", meta = (ReturnDisplayName = "Success"))
-	bool GetExclusiveUnionAction(FName& OutElementName, FLearningAgentsActionObjectElement& OutElement, const FLearningAgentsActionObjectElement Element, const FName Name = TEXT("ExclusiveUnion")) const;
+	UFUNCTION(BlueprintPure = false, Category = "LearningAgents", meta = (AdvancedDisplay = 3, ReturnDisplayName = "Success"))
+	static bool GetInclusiveUnionActionNum(int32& OutNum, const ULearningAgentsActionObject* Object, const FLearningAgentsActionObjectElement Element, const FName Tag = TEXT("InclusiveUnion"));
 
-	UFUNCTION(BlueprintPure = false, Category = "LearningAgents", meta = (ReturnDisplayName = "Success"))
-	bool GetInclusiveUnionActionNum(int32& OutNum, const FLearningAgentsActionObjectElement Element, const FName Name = TEXT("InclusiveUnion")) const;
+	UFUNCTION(BlueprintPure = false, Category = "LearningAgents", meta = (AdvancedDisplay = 3, ReturnDisplayName = "Success"))
+	static bool GetInclusiveUnionAction(TMap<FName, FLearningAgentsActionObjectElement>& OutElements, const ULearningAgentsActionObject* Object, const FLearningAgentsActionObjectElement Element, const FName Tag = TEXT("InclusiveUnion"));
 
-	UFUNCTION(BlueprintPure = false, Category = "LearningAgents", meta = (ReturnDisplayName = "Success"))
-	bool GetInclusiveUnionAction(TMap<FName, FLearningAgentsActionObjectElement>& OutElements, const FLearningAgentsActionObjectElement Element, const FName Name = TEXT("InclusiveUnion")) const;
-	
-	UFUNCTION(BlueprintPure = false, Category = "LearningAgents", meta = (ReturnDisplayName = "Success"))
-	bool GetInclusiveUnionActionToArrays(TArray<FName>& OutElementNames, TArray<FLearningAgentsActionObjectElement>& OutElements, const FLearningAgentsActionObjectElement Element, const FName Name = TEXT("InclusiveUnion")) const;
-	bool GetInclusiveUnionActionToArrayViews(TArrayView<FName> OutElementNames, TArrayView<FLearningAgentsActionObjectElement> OutElements, const FLearningAgentsActionObjectElement Element, const FName Name = TEXT("InclusiveUnion")) const;
+	UFUNCTION(BlueprintPure = false, Category = "LearningAgents", meta = (AdvancedDisplay = 4, ReturnDisplayName = "Success"))
+	static bool GetInclusiveUnionActionToArrays(TArray<FName>& OutElementNames, TArray<FLearningAgentsActionObjectElement>& OutElements, const ULearningAgentsActionObject* Object, const FLearningAgentsActionObjectElement Element, const FName Tag = TEXT("InclusiveUnion"));
+	static bool GetInclusiveUnionActionToArrayViews(TArrayView<FName> OutElementNames, TArrayView<FLearningAgentsActionObjectElement> OutElements, const ULearningAgentsActionObject* Object, const FLearningAgentsActionObjectElement Element, const FName Tag = TEXT("InclusiveUnion"));
 
-	UFUNCTION(BlueprintPure = false, Category = "LearningAgents", meta = (ReturnDisplayName = "Success"))
-	bool GetStaticArrayActionNum(int32& OutNum, const FLearningAgentsActionObjectElement Element, const FName Name = TEXT("StaticArray")) const;
+	UFUNCTION(BlueprintPure = false, Category = "LearningAgents", meta = (AdvancedDisplay = 3, ReturnDisplayName = "Success"))
+	static bool GetStaticArrayActionNum(int32& OutNum, const ULearningAgentsActionObject* Object, const FLearningAgentsActionObjectElement Element, const FName Tag = TEXT("StaticArray"));
 
-	UFUNCTION(BlueprintPure = false, Category = "LearningAgents", meta = (ReturnDisplayName = "Success"))
-	bool GetStaticArrayAction(TArray<FLearningAgentsActionObjectElement>& OutElements, const FLearningAgentsActionObjectElement Element, const FName Name = TEXT("StaticArray")) const;
-	bool GetStaticArrayActionToArrayView(TArrayView<FLearningAgentsActionObjectElement> OutElements, const FLearningAgentsActionObjectElement Element, const FName Name = TEXT("StaticArray")) const;
+	UFUNCTION(BlueprintPure = false, Category = "LearningAgents", meta = (AdvancedDisplay = 3, ReturnDisplayName = "Success"))
+	static bool GetStaticArrayAction(TArray<FLearningAgentsActionObjectElement>& OutElements, const ULearningAgentsActionObject* Object, const FLearningAgentsActionObjectElement Element, const FName Tag = TEXT("StaticArray"));
+	static bool GetStaticArrayActionToArrayView(TArrayView<FLearningAgentsActionObjectElement> OutElements, const ULearningAgentsActionObject* Object, const FLearningAgentsActionObjectElement Element, const FName Tag = TEXT("StaticArray"));
 
-	UFUNCTION(BlueprintPure = false, Category = "LearningAgents", meta = (ReturnDisplayName = "Success"))
-	bool GetPairAction(FLearningAgentsActionObjectElement& OutKey, FLearningAgentsActionObjectElement& OutValue, const FLearningAgentsActionObjectElement Element, const FName Name = TEXT("Pair")) const;
+	UFUNCTION(BlueprintPure = false, Category = "LearningAgents", meta = (AdvancedDisplay = 4, ReturnDisplayName = "Success"))
+	static bool GetPairAction(FLearningAgentsActionObjectElement& OutKey, FLearningAgentsActionObjectElement& OutValue, const ULearningAgentsActionObject* Object, const FLearningAgentsActionObjectElement Element, const FName Tag = TEXT("Pair"));
 
-	UFUNCTION(BlueprintPure = false, Category = "LearningAgents", meta = (ReturnDisplayName = "Success"))
-	bool GetEnumAction(uint8& OutEnumValue, const UEnum* Enum, const FLearningAgentsActionObjectElement Element, const FName Name = TEXT("Enum")) const;
+	UFUNCTION(BlueprintPure = false, Category = "LearningAgents", meta = (AdvancedDisplay = 4, ReturnDisplayName = "Success"))
+	static bool GetEnumAction(uint8& OutEnumValue, const ULearningAgentsActionObject* Object, const FLearningAgentsActionObjectElement Element, const UEnum* Enum, const FName Tag = TEXT("Enum"));
 
-	UFUNCTION(BlueprintPure = false, Category = "LearningAgents", meta = (ReturnDisplayName = "Success"))
-	bool GetBitmaskAction(int32& OutBitmaskValue, const UEnum* Enum, const FLearningAgentsActionObjectElement Element, const FName Name = TEXT("Bitmask")) const;
+	UFUNCTION(BlueprintPure = false, Category = "LearningAgents", meta = (AdvancedDisplay = 4, ReturnDisplayName = "Success"))
+	static bool GetBitmaskAction(int32& OutBitmaskValue, const ULearningAgentsActionObject* Object, const FLearningAgentsActionObjectElement Element, const UEnum* Enum, const FName Tag = TEXT("Bitmask"));
 
-	UFUNCTION(BlueprintPure = false, Category = "LearningAgents", meta = (ExpandEnumAsExecs = "OutOption", ReturnDisplayName = "Success"))
-	bool GetOptionalAction(ELearningAgentsOptionalAction& OutOption, FLearningAgentsActionObjectElement& OutElement, const FLearningAgentsActionObjectElement Element, const FName Name = TEXT("Optional")) const;
-	
-	UFUNCTION(BlueprintPure = false, Category = "LearningAgents", meta = (ExpandEnumAsExecs = "OutEither", ReturnDisplayName = "Success"))
-	bool GetEitherAction(ELearningAgentsEitherAction& OutEither, FLearningAgentsActionObjectElement& OutElement, const FLearningAgentsActionObjectElement Element, const FName Name = TEXT("Either")) const;
+	UFUNCTION(BlueprintPure = false, Category = "LearningAgents", meta = (AdvancedDisplay = 4, ExpandEnumAsExecs = "OutOption", ReturnDisplayName = "Success"))
+	static bool GetOptionalAction(ELearningAgentsOptionalAction& OutOption, FLearningAgentsActionObjectElement& OutElement, const ULearningAgentsActionObject* Object, const FLearningAgentsActionObjectElement Element, const FName Tag = TEXT("Optional"));
 
-	UFUNCTION(BlueprintPure = false, Category = "LearningAgents", meta = (ReturnDisplayName = "Success"))
-	bool GetEncodingAction(FLearningAgentsActionObjectElement& OutElement, const FLearningAgentsActionObjectElement Element, const FName Name = TEXT("Encoding")) const;
+	UFUNCTION(BlueprintPure = false, Category = "LearningAgents", meta = (AdvancedDisplay = 4, ExpandEnumAsExecs = "OutEither", ReturnDisplayName = "Success"))
+	static bool GetEitherAction(ELearningAgentsEitherAction& OutEither, FLearningAgentsActionObjectElement& OutElement, const ULearningAgentsActionObject* Object, const FLearningAgentsActionObjectElement Element, const FName Tag = TEXT("Either"));
 
-	UFUNCTION(BlueprintPure = false, Category = "LearningAgents", meta = (ReturnDisplayName = "Success"))
-	bool GetBoolAction(bool& bOutValue, const FLearningAgentsActionObjectElement Element, const FName Name = TEXT("Bool")) const;
+	UFUNCTION(BlueprintPure = false, Category = "LearningAgents", meta = (AdvancedDisplay = 3, ReturnDisplayName = "Success"))
+	static bool GetEncodingAction(FLearningAgentsActionObjectElement& OutElement, const ULearningAgentsActionObject* Object, const FLearningAgentsActionObjectElement Element, const FName Tag = TEXT("Encoding"));
 
-	UFUNCTION(BlueprintPure = false, Category = "LearningAgents", meta = (ReturnDisplayName = "Success"))
-	bool GetFloatAction(float& OutValue, const FLearningAgentsActionObjectElement Element, const float FloatScale = 1.0f, const FName Name = TEXT("Float")) const;
+	UFUNCTION(BlueprintPure = false, Category = "LearningAgents", meta = (AdvancedDisplay = 3, ReturnDisplayName = "Success"))
+	static bool GetBoolAction(bool& bOutValue, const ULearningAgentsActionObject* Object, const FLearningAgentsActionObjectElement Element, const FName Tag = TEXT("Bool"));
 
-	UFUNCTION(BlueprintPure = false, Category = "LearningAgents", meta = (ReturnDisplayName = "Success"))
-	bool GetLocationAction(FVector& OutLocation, const FLearningAgentsActionObjectElement Element, const FTransform RelativeTransform = FTransform(), const float LocationScale = 100.0f, const FName Name = TEXT("Location")) const;
+	UFUNCTION(BlueprintPure = false, Category = "LearningAgents", meta = (AdvancedDisplay = 4, ReturnDisplayName = "Success"))
+	static bool GetFloatAction(float& OutValue, const ULearningAgentsActionObject* Object, const FLearningAgentsActionObjectElement Element, const float FloatScale = 1.0f, const FName Tag = TEXT("Float"));
 
-	UFUNCTION(BlueprintPure = false, Category = "LearningAgents", meta = (ReturnDisplayName = "Success"))
-	bool GetRotationAction(FRotator& OutRotation, const FLearningAgentsActionObjectElement Element, const FRotator RelativeRotation = FRotator::ZeroRotator, const float RotationScale = 90.0f, const FName Name = TEXT("Rotation")) const;
-	bool GetRotationActionAsQuat(FQuat& OutRotation, const FLearningAgentsActionObjectElement Element, const FQuat RelativeRotation = FQuat::Identity, const float RotationScale = 90.0f, const FName Name = TEXT("Rotation")) const;
+	UFUNCTION(BlueprintPure = false, Category = "LearningAgents", meta = (AdvancedDisplay = 5, ReturnDisplayName = "Success"))
+	static bool GetLocationAction(FVector& OutLocation, const ULearningAgentsActionObject* Object, const FLearningAgentsActionObjectElement Element, const FTransform RelativeTransform = FTransform(), const float LocationScale = 100.0f, const FName Tag = TEXT("Location"));
 
-	UFUNCTION(BlueprintPure = false, Category = "LearningAgents", meta = (ReturnDisplayName = "Success"))
-	bool GetScaleAction(FVector& OutScale, const FLearningAgentsActionObjectElement Element, const FVector RelativeScale = FVector(1,1,1), const float Scale = 1.0f, const FName Name = TEXT("Scale")) const;
+	UFUNCTION(BlueprintPure = false, Category = "LearningAgents", meta = (AdvancedDisplay = 5, ReturnDisplayName = "Success"))
+	static bool GetRotationAction(FRotator& OutRotation, const ULearningAgentsActionObject* Object, const FLearningAgentsActionObjectElement Element, const FRotator RelativeRotation = FRotator::ZeroRotator, const float RotationScale = 90.0f, const FName Tag = TEXT("Rotation"));
 
-	UFUNCTION(BlueprintPure = false, Category = "LearningAgents", meta = (ReturnDisplayName = "Success"))
-	bool GetTransformAction(FTransform& OutTransform, const FLearningAgentsActionObjectElement Element, const FTransform RelativeTransform = FTransform(), const float LocationScale = 100.0f, const float RotationScale = 1.0f, const float ScaleScale = 1.0f, const FName Name = TEXT("Transform")) const;
+	UFUNCTION(BlueprintPure = false, Category = "LearningAgents", meta = (AdvancedDisplay = 5, ReturnDisplayName = "Success"))
+	static bool GetRotationActionAsQuat(FQuat& OutRotation, const ULearningAgentsActionObject* Object, const FLearningAgentsActionObjectElement Element, const FQuat RelativeRotation, const float RotationScale = 90.0f, const FName Tag = TEXT("Rotation"));
 
-	UFUNCTION(BlueprintPure = false, Category = "LearningAgents", meta = (ReturnDisplayName = "Success"))
-	bool GetAngleAction(float& OutAngle, const FLearningAgentsActionObjectElement Element, const float RelativeAngle = 0.0f, const float AngleScale = 90.0f, const FName Name = TEXT("Angle")) const;
+	UFUNCTION(BlueprintPure = false, Category = "LearningAgents", meta = (AdvancedDisplay = 5, ReturnDisplayName = "Success"))
+	static bool GetScaleAction(FVector& OutScale, const ULearningAgentsActionObject* Object, const FLearningAgentsActionObjectElement Element, const FVector RelativeScale = FVector(1, 1, 1), const float Scale = 1.0f, const FName Tag = TEXT("Scale"));
 
-	UFUNCTION(BlueprintPure = false, Category = "LearningAgents", meta = (ReturnDisplayName = "Success"))
-	bool GetAngleActionRadians(float& OutAngle, const FLearningAgentsActionObjectElement Element, const float RelativeAngle = 0.0f, const float AngleScale = 1.57079632679f, const FName Name = TEXT("Angle")) const;
+	UFUNCTION(BlueprintPure = false, Category = "LearningAgents", meta = (AdvancedDisplay = 7, ReturnDisplayName = "Success"))
+	static bool GetTransformAction(FTransform& OutTransform, const ULearningAgentsActionObject* Object, const FLearningAgentsActionObjectElement Element, const FTransform RelativeTransform = FTransform(), const float LocationScale = 100.0f, const float RotationScale = 1.0f, const float ScaleScale = 1.0f, const FName Tag = TEXT("Transform"));
 
-	UFUNCTION(BlueprintPure = false, Category = "LearningAgents", meta = (ReturnDisplayName = "Success"))
-	bool GetVelocityAction(FVector& OutVelocity, const FLearningAgentsActionObjectElement Element, const FTransform RelativeTransform = FTransform(), const float VelocityScale = 200.0f, const FName Name = TEXT("Velocity")) const;
+	UFUNCTION(BlueprintPure = false, Category = "LearningAgents", meta = (AdvancedDisplay = 5, ReturnDisplayName = "Success"))
+	static bool GetAngleAction(float& OutAngle, const ULearningAgentsActionObject* Object, const FLearningAgentsActionObjectElement Element, const float RelativeAngle = 0.0f, const float AngleScale = 90.0f, const FName Tag = TEXT("Angle"));
 
-	UFUNCTION(BlueprintPure = false, Category = "LearningAgents", meta = (ReturnDisplayName = "Success"))
-	bool GetDirectionAction(FVector& OutDirection, const FLearningAgentsActionObjectElement Element, const FTransform RelativeTransform = FTransform(), const FName Name = TEXT("Direction")) const;
+	UFUNCTION(BlueprintPure = false, Category = "LearningAgents", meta = (AdvancedDisplay = 5, ReturnDisplayName = "Success"))
+	static bool GetAngleActionRadians(float& OutAngle, const ULearningAgentsActionObject* Object, const FLearningAgentsActionObjectElement Element, const float RelativeAngle = 0.0f, const float AngleScale = 1.57079632679f, const FName Tag = TEXT("Angle"));
 
-	UFUNCTION(BlueprintPure = false, Category = "LearningAgents", meta = (ReturnDisplayName = "Success"))
-	bool GetSpeedAction(float& OutSpeed, const FLearningAgentsActionObjectElement Element, const float SpeedScale = 200.0f, const FName Name = TEXT("Speed")) const;
+	UFUNCTION(BlueprintPure = false, Category = "LearningAgents", meta = (AdvancedDisplay = 5, ReturnDisplayName = "Success"))
+	static bool GetVelocityAction(FVector& OutVelocity, const ULearningAgentsActionObject* Object, const FLearningAgentsActionObjectElement Element, const FTransform RelativeTransform = FTransform(), const float VelocityScale = 200.0f, const FName Tag = TEXT("Velocity"));
 
-private:
+	UFUNCTION(BlueprintPure = false, Category = "LearningAgents", meta = (AdvancedDisplay = 4, ReturnDisplayName = "Success"))
+	static bool GetDirectionAction(FVector& OutDirection, const ULearningAgentsActionObject* Object, const FLearningAgentsActionObjectElement Element, const FTransform RelativeTransform = FTransform(), const FName Tag = TEXT("Direction"));
 
-	UE::Learning::Action::FObject ActionObject;
+	UFUNCTION(BlueprintPure = false, Category = "LearningAgents", meta = (AdvancedDisplay = 4, ReturnDisplayName = "Success"))
+	static bool GetSpeedAction(float& OutSpeed, const ULearningAgentsActionObject* Object, const FLearningAgentsActionObjectElement Element, const float SpeedScale = 200.0f, const FName Tag = TEXT("Speed"));
 };
+
+
+
+
