LearningAgents: Refactored to make Observation and Action functions static and use "AdvancedDisplay" markup. Added Visual Logging and comments for Completions.

diff --git a/Engine/Plugins/Experimental/LearningAgents/Source/LearningAgents/Public/LearningAgentsObservations.h b/Engine/Plugins/Experimental/LearningAgents/Source/LearningAgents/Public/LearningAgentsObservations.h
--- a/Engine/Plugins/Experimental/LearningAgents/Source/LearningAgents/Public/LearningAgentsObservations.h
+++ b/Engine/Plugins/Experimental/LearningAgents/Source/LearningAgents/Public/LearningAgentsObservations.h
@@ -44,6 +44,36 @@ bool operator==(const FLearningAgentsObservationObjectElement& Lhs, const FLearn
 /** Hashing operator for Observation Object Elements */
 uint32 GetTypeHash(const FLearningAgentsObservationObjectElement& Element);
 
+/**
+ * Observation Schema
+ *
+ * This object is used to construct a schema describing some structure of observations.
+ */
+UCLASS(BlueprintType)
+class LEARNINGAGENTS_API ULearningAgentsObservationSchema : public UObject
+{
+	GENERATED_BODY()
+
+public:
+
+	UE::Learning::Observation::FSchema ObservationSchema;
+};
+
+/**
+ * Observation Object
+ *
+ * This object is used to construct or get the values of observations.
+ */
+UCLASS(BlueprintType)
+class LEARNINGAGENTS_API ULearningAgentsObservationObject : public UObject
+{
+	GENERATED_BODY()
+
+public:
+
+	UE::Learning::Observation::FObject ObservationObject;
+};
+
 /** Enum Type representing either observation A or observation B */
 UENUM(BlueprintType)
 enum class ELearningAgentsEitherObservation : uint8
@@ -60,16 +90,13 @@ enum class ELearningAgentsOptionalObservation : uint8
 	Valid,
 };
 
-/**
- * Observation Functions
- * 
- * Convenience functions useful for constructing observations.
- */
 UCLASS()
-class ULearningAgentsObservationFunctions : public UBlueprintFunctionLibrary
+class LEARNINGAGENTS_API ULearningAgentsObservations : public UBlueprintFunctionLibrary
 {
 	GENERATED_BODY()
 
+public:
+
 	/** Project a transform onto the ground plane, leaving just rotation around the vertical axis */
 	UFUNCTION(BlueprintPure, Category = "LearningAgents")
 	static FTransform ProjectTransformOntoGroundPlane(const FTransform Transform, const FVector LocalForwardVector = FVector::ForwardVector, const float GroundPlaneHeight = 0.0f);
@@ -78,304 +105,265 @@ class ULearningAgentsObservationFunctions : public UBlueprintFunctionLibrary
 	/** Find an Enum type by Name. This can be used to find Enum types defined in C++. This call can be expensive so the result should be cached. */
 	UFUNCTION(BlueprintCallable, Category = "LearningAgents")
 	static UEnum* FindEnumByName(const FString& Name);
-};
-
-/**
- * Observation Schema
- *
- * This object is used to construct a schema describing some structure of observations.
- */
-UCLASS(BlueprintType)
-class LEARNINGAGENTS_API ULearningAgentsObservationSchema : public UObject
-{
-	GENERATED_BODY()
-
-public:
-
-	/** Gets the internal observation schema object */
-	const UE::Learning::Observation::FSchema& GetObservationSchema() const;
 
 	/**
 	 * Validates that the given object matches the schema. Will log errors on objects that don't match.
 	 *
-	 * @param SchemaElement			Schema Element
+	 * @param Schema				Observation Schema
+	 * @param SchemaElement			Observation Schema Element
 	 * @param Object				Observation Object
 	 * @param ObjectElement			Observation Object Element
 	 * @returns						true if the object matches the schema
 	 */
-	UFUNCTION(BlueprintPure, Category = "LearningAgents")
-	bool ValidateObjectMatchesSchema(
+	UFUNCTION(BlueprintPure = false, Category = "LearningAgents")
+	static bool ValidateObjectMatchesSchema(
+		const ULearningAgentsObservationSchema* Schema,
 		const FLearningAgentsObservationSchemaElement SchemaElement,
 		const ULearningAgentsObservationObject* Object,
-		const FLearningAgentsObservationObjectElement ObjectElement) const;
+		const FLearningAgentsObservationObjectElement ObjectElement);
+
+	UFUNCTION(BlueprintPure = false, Category = "LearningAgents")
+	static void LogObservation(const ULearningAgentsObservationObject* Object, const FLearningAgentsObservationObjectElement Element);
 
 public:
 
-	// Basic Observations
+	// Specify Basic Observations
 
-	UFUNCTION(BlueprintPure, Category = "LearningAgents")
-	FLearningAgentsObservationSchemaElement SpecifyNullObservation(const FName Name = TEXT("Null"));
+	UFUNCTION(BlueprintPure, Category = "LearningAgents", meta = (AdvancedDisplay = 1))
+	static FLearningAgentsObservationSchemaElement SpecifyNullObservation(ULearningAgentsObservationSchema* Schema, const FName Tag = TEXT("Null"));
 
-	UFUNCTION(BlueprintPure, Category = "LearningAgents")
-	FLearningAgentsObservationSchemaElement SpecifyContinuousObservation(const int32 Size, const FName Name = TEXT("Continuous"));
+	UFUNCTION(BlueprintPure, Category = "LearningAgents", meta = (AdvancedDisplay = 2))
+	static FLearningAgentsObservationSchemaElement SpecifyContinuousObservation(ULearningAgentsObservationSchema* Schema, const int32 Size, const FName Tag = TEXT("Continuous"));
 
-	UFUNCTION(BlueprintPure, Category = "LearningAgents")
-	FLearningAgentsObservationSchemaElement SpecifyExclusiveDiscreteObservation(const int32 Size, const FName Name = TEXT("ExclusiveDiscrete"));
+	UFUNCTION(BlueprintPure, Category = "LearningAgents", meta = (AdvancedDisplay = 2))
+	static FLearningAgentsObservationSchemaElement SpecifyExclusiveDiscreteObservation(ULearningAgentsObservationSchema* Schema, const int32 Size, const FName Tag = TEXT("ExclusiveDiscrete"));
 
-	UFUNCTION(BlueprintPure, Category = "LearningAgents")
-	FLearningAgentsObservationSchemaElement SpecifyInclusiveDiscreteObservation(const int32 Size, const FName Name = TEXT("InclusiveDiscrete"));
+	UFUNCTION(BlueprintPure, Category = "LearningAgents", meta = (AdvancedDisplay = 2))
+	static FLearningAgentsObservationSchemaElement SpecifyInclusiveDiscreteObservation(ULearningAgentsObservationSchema* Schema, const int32 Size, const FName Tag = TEXT("InclusiveDiscrete"));
 
-	UFUNCTION(BlueprintPure, Category = "LearningAgents")
-	FLearningAgentsObservationSchemaElement SpecifyIndexObservation(const int32 Size, const FName Name = TEXT("Index"));
+	UFUNCTION(BlueprintPure, Category = "LearningAgents", meta = (AdvancedDisplay = 2))
+	static FLearningAgentsObservationSchemaElement SpecifyIndexObservation(ULearningAgentsObservationSchema* Schema, const int32 Size, const FName Tag = TEXT("Index"));
 
-	UFUNCTION(BlueprintPure, Category = "LearningAgents")
-	FLearningAgentsObservationSchemaElement SpecifyCountObservation(const FName Name = TEXT("Count"));
+	UFUNCTION(BlueprintPure, Category = "LearningAgents", meta = (AdvancedDisplay = 1))
+	static FLearningAgentsObservationSchemaElement SpecifyCountObservation(ULearningAgentsObservationSchema* Schema, const FName Tag = TEXT("Count"));
 
-	UFUNCTION(BlueprintPure, Category = "LearningAgents")
-	FLearningAgentsObservationSchemaElement SpecifyStructObservation(const TMap<FName, FLearningAgentsObservationSchemaElement>& Elements, const FName Name = TEXT("Struct"));
-	
-	UFUNCTION(BlueprintPure, Category = "LearningAgents")
-	FLearningAgentsObservationSchemaElement SpecifyStructObservationFromArrays(const TArray<FName>& ElementNames, const TArray<FLearningAgentsObservationSchemaElement>& Elements, const FName Name = TEXT("Struct"));
-	FLearningAgentsObservationSchemaElement SpecifyStructObservationFromArrayViews(const TArrayView<const FName> ElementNames, const TArrayView<const FLearningAgentsObservationSchemaElement> Elements, const FName Name = TEXT("Struct"));
+	UFUNCTION(BlueprintPure, Category = "LearningAgents", meta = (AdvancedDisplay = 2))
+	static FLearningAgentsObservationSchemaElement SpecifyStructObservation(ULearningAgentsObservationSchema* Schema, const TMap<FName, FLearningAgentsObservationSchemaElement>& Elements, const FName Tag = TEXT("Struct"));
 
-	UFUNCTION(BlueprintPure, Category = "LearningAgents")
-	FLearningAgentsObservationSchemaElement SpecifyExclusiveUnionObservation(const TMap<FName, FLearningAgentsObservationSchemaElement>& Elements, const int32 EncodingSize = 128, const FName Name = TEXT("ExclusiveUnion"));
-	
-	UFUNCTION(BlueprintPure, Category = "LearningAgents")
-	FLearningAgentsObservationSchemaElement SpecifyExclusiveUnionObservationFromArrays(const TArray<FName>& ElementNames, const TArray<FLearningAgentsObservationSchemaElement>& Elements, const int32 EncodingSize = 128, const FName Name = TEXT("ExclusiveUnion"));
-	FLearningAgentsObservationSchemaElement SpecifyExclusiveUnionObservationFromArrayViews(const TArrayView<const FName> ElementNames, const TArrayView<const FLearningAgentsObservationSchemaElement> Elements, const int32 EncodingSize = 128, const FName Name = TEXT("ExclusiveUnion"));
+	UFUNCTION(BlueprintPure, Category = "LearningAgents", meta = (AdvancedDisplay = 3))
+	static FLearningAgentsObservationSchemaElement SpecifyStructObservationFromArrays(ULearningAgentsObservationSchema* Schema, const TArray<FName>& ElementNames, const TArray<FLearningAgentsObservationSchemaElement>& Elements, const FName Tag = TEXT("Struct"));
+	static FLearningAgentsObservationSchemaElement SpecifyStructObservationFromArrayViews(ULearningAgentsObservationSchema* Schema, const TArrayView<const FName> ElementNames, const TArrayView<const FLearningAgentsObservationSchemaElement> Elements, const FName Tag = TEXT("Struct"));
 
-	UFUNCTION(BlueprintPure, Category = "LearningAgents")
-	FLearningAgentsObservationSchemaElement SpecifyInclusiveUnionObservation(const TMap<FName, FLearningAgentsObservationSchemaElement>& Elements, const int32 AttentionEncodingSize = 32, const int32 AttentionHeadNum = 4, const int32 ValueEncodingSize = 32, const FName Name = TEXT("InclusiveUnion"));
-	
-	UFUNCTION(BlueprintPure, Category = "LearningAgents")
-	FLearningAgentsObservationSchemaElement SpecifyInclusiveUnionObservationFromArrays(const TArray<FName>& ElementNames, const TArray<FLearningAgentsObservationSchemaElement>& Elements, const int32 AttentionEncodingSize = 32, const int32 AttentionHeadNum = 4, const int32 ValueEncodingSize = 32, const FName Name = TEXT("InclusiveUnion"));
-	FLearningAgentsObservationSchemaElement SpecifyInclusiveUnionObservationFromArrayViews(const TArrayView<const FName> ElementNames, const TArrayView<const FLearningAgentsObservationSchemaElement> Elements, const int32 AttentionEncodingSize = 32, const int32 AttentionHeadNum = 4, const int32 ValueEncodingSize = 32, const FName Name = TEXT("InclusiveUnion"));
+	UFUNCTION(BlueprintPure, Category = "LearningAgents", meta = (AdvancedDisplay = 2))
+	static FLearningAgentsObservationSchemaElement SpecifyExclusiveUnionObservation(ULearningAgentsObservationSchema* Schema, const TMap<FName, FLearningAgentsObservationSchemaElement>& Elements, const int32 EncodingSize = 128, const FName Tag = TEXT("ExclusiveUnion"));
 
-	UFUNCTION(BlueprintPure, Category = "LearningAgents")
-	FLearningAgentsObservationSchemaElement SpecifyStaticArrayObservation(const FLearningAgentsObservationSchemaElement Element, const int32 Num, const FName Name = TEXT("StaticArray"));
+	UFUNCTION(BlueprintPure, Category = "LearningAgents", meta = (AdvancedDisplay = 3))
+	static FLearningAgentsObservationSchemaElement SpecifyExclusiveUnionObservationFromArrays(ULearningAgentsObservationSchema* Schema, const TArray<FName>& ElementNames, const TArray<FLearningAgentsObservationSchemaElement>& Elements, const int32 EncodingSize = 128, const FName Tag = TEXT("ExclusiveUnion"));
+	static FLearningAgentsObservationSchemaElement SpecifyExclusiveUnionObservationFromArrayViews(ULearningAgentsObservationSchema* Schema, const TArrayView<const FName> ElementNames, const TArrayView<const FLearningAgentsObservationSchemaElement> Elements, const int32 EncodingSize = 128, const FName Tag = TEXT("ExclusiveUnion"));
 
-	UFUNCTION(BlueprintPure, Category = "LearningAgents")
-	FLearningAgentsObservationSchemaElement SpecifySetObservation(const FLearningAgentsObservationSchemaElement Element, const int32 MaxNum, const int32 AttentionEncodingSize = 32, const int32 AttentionHeadNum = 4, const int32 ValueEncodingSize = 32, const FName Name = TEXT("Set"));
+	UFUNCTION(BlueprintPure, Category = "LearningAgents", meta = (AdvancedDisplay = 2))
+	static FLearningAgentsObservationSchemaElement SpecifyInclusiveUnionObservation(ULearningAgentsObservationSchema* Schema, const TMap<FName, FLearningAgentsObservationSchemaElement>& Elements, const int32 AttentionEncodingSize = 32, const int32 AttentionHeadNum = 4, const int32 ValueEncodingSize = 32, const FName Tag = TEXT("InclusiveUnion"));
 
-	UFUNCTION(BlueprintPure, Category = "LearningAgents")
-	FLearningAgentsObservationSchemaElement SpecifyPairObservation(const FLearningAgentsObservationSchemaElement Key, const FLearningAgentsObservationSchemaElement Value, const FName Name = TEXT("Pair"));
+	UFUNCTION(BlueprintPure, Category = "LearningAgents", meta = (AdvancedDisplay = 3))
+	static FLearningAgentsObservationSchemaElement SpecifyInclusiveUnionObservationFromArrays(ULearningAgentsObservationSchema* Schema, const TArray<FName>& ElementNames, const TArray<FLearningAgentsObservationSchemaElement>& Elements, const int32 AttentionEncodingSize = 32, const int32 AttentionHeadNum = 4, const int32 ValueEncodingSize = 32, const FName Tag = TEXT("InclusiveUnion"));
+	static FLearningAgentsObservationSchemaElement SpecifyInclusiveUnionObservationFromArrayViews(ULearningAgentsObservationSchema* Schema, const TArrayView<const FName> ElementNames, const TArrayView<const FLearningAgentsObservationSchemaElement> Elements, const int32 AttentionEncodingSize = 32, const int32 AttentionHeadNum = 4, const int32 ValueEncodingSize = 32, const FName Tag = TEXT("InclusiveUnion"));
 
-	UFUNCTION(BlueprintPure, Category = "LearningAgents")
-	FLearningAgentsObservationSchemaElement SpecifyArrayObservation(const FLearningAgentsObservationSchemaElement Element, const int32 MaxNum, const int32 AttentionEncodingSize = 32, const int32 AttentionHeadNum = 4, const int32 ValueEncodingSize = 32, const FName Name = TEXT("Array"));
+	UFUNCTION(BlueprintPure, Category = "LearningAgents", meta = (AdvancedDisplay = 3))
+	static FLearningAgentsObservationSchemaElement SpecifyStaticArrayObservation(ULearningAgentsObservationSchema* Schema, const FLearningAgentsObservationSchemaElement Element, const int32 Num, const FName Tag = TEXT("StaticArray"));
 
-	UFUNCTION(BlueprintPure, Category = "LearningAgents")
-	FLearningAgentsObservationSchemaElement SpecifyMapObservation(const FLearningAgentsObservationSchemaElement KeyElement, const FLearningAgentsObservationSchemaElement ValueElement, const int32 MaxNum, const int32 AttentionEncodingSize = 32, const int32 AttentionHeadNum = 4, const int32 ValueEncodingSize = 32, const FName Name = TEXT("Map"));
+	UFUNCTION(BlueprintPure, Category = "LearningAgents", meta = (AdvancedDisplay = 3))
+	static FLearningAgentsObservationSchemaElement SpecifySetObservation(ULearningAgentsObservationSchema* Schema, const FLearningAgentsObservationSchemaElement Element, const int32 MaxNum, const int32 AttentionEncodingSize = 32, const int32 AttentionHeadNum = 4, const int32 ValueEncodingSize = 32, const FName Tag = TEXT("Set"));
 
-	UFUNCTION(BlueprintPure, Category = "LearningAgents")
-	FLearningAgentsObservationSchemaElement SpecifyEnumObservation(const UEnum* Enum, const FName Name = TEXT("Enum"));
+	UFUNCTION(BlueprintPure, Category = "LearningAgents", meta = (AdvancedDisplay = 3))
+	static FLearningAgentsObservationSchemaElement SpecifyPairObservation(ULearningAgentsObservationSchema* Schema, const FLearningAgentsObservationSchemaElement Key, const FLearningAgentsObservationSchemaElement Value, const FName Tag = TEXT("Pair"));
 
-	UFUNCTION(BlueprintPure, Category = "LearningAgents")
-	FLearningAgentsObservationSchemaElement SpecifyBitmaskObservation(const UEnum* Enum, const FName Name = TEXT("Bitmask"));
+	UFUNCTION(BlueprintPure, Category = "LearningAgents", meta = (AdvancedDisplay = 3))
+	static FLearningAgentsObservationSchemaElement SpecifyArrayObservation(ULearningAgentsObservationSchema* Schema, const FLearningAgentsObservationSchemaElement Element, const int32 MaxNum, const int32 AttentionEncodingSize = 32, const int32 AttentionHeadNum = 4, const int32 ValueEncodingSize = 32, const FName Tag = TEXT("Array"));
 
-	UFUNCTION(BlueprintPure, Category = "LearningAgents")
-	FLearningAgentsObservationSchemaElement SpecifyOptionalObservation(const FLearningAgentsObservationSchemaElement Element, const int32 EncodingSize = 128, const FName Name = TEXT("Optional"));
+	UFUNCTION(BlueprintPure, Category = "LearningAgents", meta = (AdvancedDisplay = 4))
+	static FLearningAgentsObservationSchemaElement SpecifyMapObservation(ULearningAgentsObservationSchema* Schema, const FLearningAgentsObservationSchemaElement KeyElement, const FLearningAgentsObservationSchemaElement ValueElement, const int32 MaxNum, const int32 AttentionEncodingSize = 32, const int32 AttentionHeadNum = 4, const int32 ValueEncodingSize = 32, const FName Tag = TEXT("Map"));
 
-	UFUNCTION(BlueprintPure, Category = "LearningAgents")
-	FLearningAgentsObservationSchemaElement SpecifyEitherObservation(const FLearningAgentsObservationSchemaElement A, const FLearningAgentsObservationSchemaElement B, const int32 EncodingSize = 128, const FName Name = TEXT("Either"));
+	UFUNCTION(BlueprintPure, Category = "LearningAgents", meta = (AdvancedDisplay = 2))
+	static FLearningAgentsObservationSchemaElement SpecifyEnumObservation(ULearningAgentsObservationSchema* Schema, const UEnum* Enum, const FName Tag = TEXT("Enum"));
 
-	UFUNCTION(BlueprintPure, Category = "LearningAgents")
-	FLearningAgentsObservationSchemaElement SpecifyEncodingObservation(const FLearningAgentsObservationSchemaElement Element, const int32 EncodingSize = 128, const int32 LayerNum = 1, const ELearningAgentsActivationFunction ActivationFunction = ELearningAgentsActivationFunction::ELU, const FName Name = TEXT("Encoding"));
+	UFUNCTION(BlueprintPure, Category = "LearningAgents", meta = (AdvancedDisplay = 2))
+	static FLearningAgentsObservationSchemaElement SpecifyBitmaskObservation(ULearningAgentsObservationSchema* Schema, const UEnum* Enum, const FName Tag = TEXT("Bitmask"));
+
+	UFUNCTION(BlueprintPure, Category = "LearningAgents", meta = (AdvancedDisplay = 2))
+	static FLearningAgentsObservationSchemaElement SpecifyOptionalObservation(ULearningAgentsObservationSchema* Schema, const FLearningAgentsObservationSchemaElement Element, const int32 EncodingSize = 128, const FName Tag = TEXT("Optional"));
+
+	UFUNCTION(BlueprintPure, Category = "LearningAgents", meta = (AdvancedDisplay = 3))
+	static FLearningAgentsObservationSchemaElement SpecifyEitherObservation(ULearningAgentsObservationSchema* Schema, const FLearningAgentsObservationSchemaElement A, const FLearningAgentsObservationSchemaElement B, const int32 EncodingSize = 128, const FName Tag = TEXT("Either"));
+
+	UFUNCTION(BlueprintPure, Category = "LearningAgents", meta = (AdvancedDisplay = 2))
+	static FLearningAgentsObservationSchemaElement SpecifyEncodingObservation(ULearningAgentsObservationSchema* Schema, const FLearningAgentsObservationSchemaElement Element, const int32 EncodingSize = 128, const int32 LayerNum = 1, const ELearningAgentsActivationFunction ActivationFunction = ELearningAgentsActivationFunction::ELU, const FName Tag = TEXT("Encoding"));
 
 	/**
 	 * Specifies a new bool observation. A true or false observation.
 	 *
-	 * @param Name The name of this new observation. Used during observation object validation and debugging.
+	 * @param Tag The tag of this new observation. Used during observation object validation and debugging.
 	 * @return The newly created observation schema element.
 	 */
-	UFUNCTION(BlueprintPure, Category = "LearningAgents")
-	FLearningAgentsObservationSchemaElement SpecifyBoolObservation(const FName Name = TEXT("Bool"));
+	UFUNCTION(BlueprintPure, Category = "LearningAgents", meta = (AdvancedDisplay = 1))
+	static FLearningAgentsObservationSchemaElement SpecifyBoolObservation(ULearningAgentsObservationSchema* Schema, const FName Tag = TEXT("Bool"));
 
 	/**
-	 * Specifies a new float observation. A simple observation which can be used as a catch-all for situations where a 
+	 * Specifies a new float observation. A simple observation which can be used as a catch-all for situations where a
 	 * type-specific observation does not exist.
-	 * 
-	 * @param Name The name of this new observation. Used during observation object validation and debugging.
+	 *
+	 * @param Tag The tag of this new observation. Used during observation object validation and debugging.
 	 * @return The newly created observation schema element.
 	 */
-	UFUNCTION(BlueprintPure, Category = "LearningAgents")
-	FLearningAgentsObservationSchemaElement SpecifyFloatObservation(const FName Name = TEXT("Float"));
+	UFUNCTION(BlueprintPure, Category = "LearningAgents", meta = (AdvancedDisplay = 1))
+	static FLearningAgentsObservationSchemaElement SpecifyFloatObservation(ULearningAgentsObservationSchema* Schema, const FName Tag = TEXT("Float"));
 
 	/**
 	 * Specifies a new location observation. Allows an agent to observe the location of some entity.
 	 *
-	 * @param Name The name of this new observation. Used during observation object validation and debugging.
+	 * @param Tag The tag of this new observation. Used during observation object validation and debugging.
 	 * @return The newly created observation schema element.
 	 */
-	UFUNCTION(BlueprintPure, Category = "LearningAgents")
-	FLearningAgentsObservationSchemaElement SpecifyLocationObservation(const FName Name = TEXT("Location"));
+	UFUNCTION(BlueprintPure, Category = "LearningAgents", meta = (AdvancedDisplay = 1))
+	static FLearningAgentsObservationSchemaElement SpecifyLocationObservation(ULearningAgentsObservationSchema* Schema, const FName Tag = TEXT("Location"));
 
-	UFUNCTION(BlueprintPure, Category = "LearningAgents")
-	FLearningAgentsObservationSchemaElement SpecifyRotationObservation(const FName Name = TEXT("Rotation"));
+	UFUNCTION(BlueprintPure, Category = "LearningAgents", meta = (AdvancedDisplay = 1))
+	static FLearningAgentsObservationSchemaElement SpecifyRotationObservation(ULearningAgentsObservationSchema* Schema, const FName Tag = TEXT("Rotation"));
 
-	UFUNCTION(BlueprintPure, Category = "LearningAgents")
-	FLearningAgentsObservationSchemaElement SpecifyScaleObservation(const FName Name = TEXT("Scale"));
+	UFUNCTION(BlueprintPure, Category = "LearningAgents", meta = (AdvancedDisplay = 1))
+	static FLearningAgentsObservationSchemaElement SpecifyScaleObservation(ULearningAgentsObservationSchema* Schema, const FName Tag = TEXT("Scale"));
 
-	UFUNCTION(BlueprintPure, Category = "LearningAgents")
-	FLearningAgentsObservationSchemaElement SpecifyTransformObservation(const FName Name = TEXT("Transform"));
+	UFUNCTION(BlueprintPure, Category = "LearningAgents", meta = (AdvancedDisplay = 1))
+	static FLearningAgentsObservationSchemaElement SpecifyTransformObservation(ULearningAgentsObservationSchema* Schema, const FName Tag = TEXT("Transform"));
 
-	UFUNCTION(BlueprintPure, Category = "LearningAgents")
-	FLearningAgentsObservationSchemaElement SpecifyAngleObservation(const FName Name = TEXT("Angle"));
+	UFUNCTION(BlueprintPure, Category = "LearningAgents", meta = (AdvancedDisplay = 1))
+	static FLearningAgentsObservationSchemaElement SpecifyAngleObservation(ULearningAgentsObservationSchema* Schema, const FName Tag = TEXT("Angle"));
 
-	UFUNCTION(BlueprintPure, Category = "LearningAgents")
-	FLearningAgentsObservationSchemaElement SpecifyVelocityObservation(const FName Name = TEXT("Velocity"));
+	UFUNCTION(BlueprintPure, Category = "LearningAgents", meta = (AdvancedDisplay = 1))
+	static FLearningAgentsObservationSchemaElement SpecifyVelocityObservation(ULearningAgentsObservationSchema* Schema, const FName Tag = TEXT("Velocity"));
 
-	UFUNCTION(BlueprintPure, Category = "LearningAgents")
-	FLearningAgentsObservationSchemaElement SpecifyDirectionObservation(const FName Name = TEXT("Direction"));
+	UFUNCTION(BlueprintPure, Category = "LearningAgents", meta = (AdvancedDisplay = 1))
+	static FLearningAgentsObservationSchemaElement SpecifyDirectionObservation(ULearningAgentsObservationSchema* Schema, const FName Tag = TEXT("Direction"));
 
-	// Spline Observations
+	// Specify Spline Observations
 
-	UFUNCTION(BlueprintPure, Category = "LearningAgents")
-	FLearningAgentsObservationSchemaElement SpecifyLocationAlongSplineObservation(const FName Name = TEXT("LocationAlongSpline"));
+	UFUNCTION(BlueprintPure, Category = "LearningAgents", meta = (AdvancedDisplay = 1))
+	static FLearningAgentsObservationSchemaElement SpecifyLocationAlongSplineObservation(ULearningAgentsObservationSchema* Schema, const FName Tag = TEXT("LocationAlongSpline"));
 
-	UFUNCTION(BlueprintPure, Category = "LearningAgents")
-	FLearningAgentsObservationSchemaElement SpecifyProportionAlongSplineObservation(const FName Name = TEXT("ProportionAlongSpline"));
+	UFUNCTION(BlueprintPure, Category = "LearningAgents", meta = (AdvancedDisplay = 1))
+	static FLearningAgentsObservationSchemaElement SpecifyProportionAlongSplineObservation(ULearningAgentsObservationSchema* Schema, const FName Tag = TEXT("ProportionAlongSpline"));
 
-	UFUNCTION(BlueprintPure, Category = "LearningAgents")
-	FLearningAgentsObservationSchemaElement SpecifyDirectionAlongSplineObservation(const FName Name = TEXT("DirectionAlongSpline"));
+	UFUNCTION(BlueprintPure, Category = "LearningAgents", meta = (AdvancedDisplay = 1))
+	static FLearningAgentsObservationSchemaElement SpecifyDirectionAlongSplineObservation(ULearningAgentsObservationSchema* Schema, const FName Tag = TEXT("DirectionAlongSpline"));
 
-	UFUNCTION(BlueprintPure, Category = "LearningAgents")
-	FLearningAgentsObservationSchemaElement SpecifyPropertiesAlongSplineObservation(const FName Name = TEXT("PropertiesAlongSpline"));
-
-	// Ray Cast Observations
+	UFUNCTION(BlueprintPure, Category = "LearningAgents", meta = (AdvancedDisplay = 1))
+	static FLearningAgentsObservationSchemaElement SpecifyPropertiesAlongSplineObservation(ULearningAgentsObservationSchema* Schema, const FName Tag = TEXT("PropertiesAlongSpline"));
 
-	UFUNCTION(BlueprintPure, Category = "LearningAgents")
-	FLearningAgentsObservationSchemaElement SpecifyProportionAlongRayObservation(const FName Name = TEXT("ProportionAlongRay"));
+	// Specify Ray Cast Observations
 
-	UFUNCTION(BlueprintPure, Category = "LearningAgents")
-	FLearningAgentsObservationSchemaElement SpecifyProportionAlongRaysObservation(const int32 Num, const FName Name = TEXT("ProportionAlongRays"));
+	UFUNCTION(BlueprintPure, Category = "LearningAgents", meta = (AdvancedDisplay = 1))
+	static FLearningAgentsObservationSchemaElement SpecifyProportionAlongRayObservation(ULearningAgentsObservationSchema* Schema, const FName Tag = TEXT("ProportionAlongRay"));
 
-private:
-
-	UE::Learning::Observation::FSchema ObservationSchema;
-};
-
-/**
- * Observation Object
- *
- * This object is used to construct or get the values of observations.
- */
-UCLASS(BlueprintType)
-class LEARNINGAGENTS_API ULearningAgentsObservationObject : public UObject
-{
-	GENERATED_BODY()
+	UFUNCTION(BlueprintPure, Category = "LearningAgents", meta = (AdvancedDisplay = 2))
+	static FLearningAgentsObservationSchemaElement SpecifyProportionAlongRaysObservation(ULearningAgentsObservationSchema* Schema, const int32 Num, const FName Tag = TEXT("ProportionAlongRays"));
 
 public:
 
-	/** Gets the internal observation object */
-	const UE::Learning::Observation::FObject& GetObservationObject() const;
+	// Make Basic Observations
 
-	/** Gets the internal observation object */
-	UE::Learning::Observation::FObject& GetObservationObject();
+	UFUNCTION(BlueprintPure, Category = "LearningAgents", meta = (AdvancedDisplay = 1))
+	static FLearningAgentsObservationObjectElement MakeNullObservation(ULearningAgentsObservationObject* Object, const FName Tag = TEXT("Null"));
 
-public:
+	UFUNCTION(BlueprintPure, Category = "LearningAgents", meta = (AdvancedDisplay = 2))
+	static FLearningAgentsObservationObjectElement MakeContinuousObservation(ULearningAgentsObservationObject* Object, const TArray<float>& Values, const FName Tag = TEXT("Continuous"));
+	static FLearningAgentsObservationObjectElement MakeContinuousObservationFromArrayView(ULearningAgentsObservationObject* Object, const TArrayView<const float> Values, const FName Tag = TEXT("Continuous"));
 
-	UFUNCTION(BlueprintPure = false, Category = "LearningAgents")
-	void LogObservation(const FLearningAgentsObservationObjectElement Element);
-
-public:
+	UFUNCTION(BlueprintPure, Category = "LearningAgents", meta = (AdvancedDisplay = 3))
+	static FLearningAgentsObservationObjectElement MakeExclusiveDiscreteObservation(ULearningAgentsObservationObject* Object, const int32 DiscreteIndex, const int32 Size, const FName Tag = TEXT("ExclusiveDiscrete"));
 
-	// Basic Observations
+	UFUNCTION(BlueprintPure, Category = "LearningAgents", meta = (AdvancedDisplay = 3))
+	static FLearningAgentsObservationObjectElement MakeInclusiveDiscreteObservation(ULearningAgentsObservationObject* Object, const TArray<int32>& DiscreteIndices, const int32 Size, const FName Tag = TEXT("InclusiveDiscrete"));
+	static FLearningAgentsObservationObjectElement MakeInclusiveDiscreteObservationFromArrayView(ULearningAgentsObservationObject* Object, const TArrayView<const int32> DiscreteIndices, const int32 Size, const FName Tag = TEXT("InclusiveDiscrete"));
 
-	UFUNCTION(BlueprintPure, Category = "LearningAgents")
-	FLearningAgentsObservationObjectElement MakeNullObservation(const FName Name = TEXT("Null"));
+	UFUNCTION(BlueprintPure, Category = "LearningAgents", meta = (AdvancedDisplay = 3))
+	static FLearningAgentsObservationObjectElement MakeIndexObservation(ULearningAgentsObservationObject* Object, const int32 Index, const int32 Size, const FName Tag = TEXT("Index"));
 
-	UFUNCTION(BlueprintPure, Category = "LearningAgents")
-	FLearningAgentsObservationObjectElement MakeContinuousObservation(const TArray<float>& Values, const FName Name = TEXT("Continuous"));
-	FLearningAgentsObservationObjectElement MakeContinuousObservationFromArrayView(const TArrayView<const float> Values, const FName Name = TEXT("Continuous"));
+	UFUNCTION(BlueprintPure, Category = "LearningAgents", meta = (AdvancedDisplay = 3))
+	static FLearningAgentsObservationObjectElement MakeCountObservation(ULearningAgentsObservationObject* Object, const int32 Num, const int32 MaxNum, const FName Tag = TEXT("Count"));
 
-	UFUNCTION(BlueprintPure, Category = "LearningAgents")
-	FLearningAgentsObservationObjectElement MakeExclusiveDiscreteObservation(const int32 DiscreteIndex, const int32 Size, const FName Name = TEXT("ExclusiveDiscrete"));
+	UFUNCTION(BlueprintPure, Category = "LearningAgents", meta = (AdvancedDisplay = 2))
+	static FLearningAgentsObservationObjectElement MakeStructObservation(ULearningAgentsObservationObject* Object, const TMap<FName, FLearningAgentsObservationObjectElement>& Elements, const FName Tag = TEXT("Struct"));
 
-	UFUNCTION(BlueprintPure, Category = "LearningAgents")
-	FLearningAgentsObservationObjectElement MakeInclusiveDiscreteObservation(const TArray<int32>& DiscreteIndices, const int32 Size, const FName Name = TEXT("InclusiveDiscrete"));
-	FLearningAgentsObservationObjectElement MakeInclusiveDiscreteObservationFromArrayView(const TArrayView<const int32> DiscreteIndices, const int32 Size, const FName Name = TEXT("InclusiveDiscrete"));
+	UFUNCTION(BlueprintPure, Category = "LearningAgents", meta = (AdvancedDisplay = 3))
+	static FLearningAgentsObservationObjectElement MakeStructObservationFromArrays(ULearningAgentsObservationObject* Object, const TArray<FName>& ElementNames, const TArray<FLearningAgentsObservationObjectElement>& Elements, const FName Tag = TEXT("Struct"));
+	static FLearningAgentsObservationObjectElement MakeStructObservationFromArrayViews(ULearningAgentsObservationObject* Object, const TArrayView<const FName> ElementNames, const TArrayView<const FLearningAgentsObservationObjectElement> Elements, const FName Tag = TEXT("Struct"));
 
-	UFUNCTION(BlueprintPure, Category = "LearningAgents")
-	FLearningAgentsObservationObjectElement MakeIndexObservation(const int32 Index, const int32 Size, const FName Name = TEXT("Index"));
+	UFUNCTION(BlueprintPure, Category = "LearningAgents", meta = (AdvancedDisplay = 3))
+	static FLearningAgentsObservationObjectElement MakeExclusiveUnionObservation(ULearningAgentsObservationObject* Object, const FName ElementName, const FLearningAgentsObservationObjectElement Element, const FName Tag = TEXT("ExclusiveUnion"));
 
-	UFUNCTION(BlueprintPure, Category = "LearningAgents")
-	FLearningAgentsObservationObjectElement MakeCountObservation(const int32 Num, const int32 MaxNum, const FName Name = TEXT("Count"));
+	UFUNCTION(BlueprintPure, Category = "LearningAgents", meta = (AdvancedDisplay = 2))
+	static FLearningAgentsObservationObjectElement MakeInclusiveUnionObservation(ULearningAgentsObservationObject* Object, const TMap<FName, FLearningAgentsObservationObjectElement>& Elements, const FName Tag = TEXT("InclusiveUnion"));
 
-	UFUNCTION(BlueprintPure, Category = "LearningAgents")
-	FLearningAgentsObservationObjectElement MakeStructObservation(const TMap<FName, FLearningAgentsObservationObjectElement>& Elements, const FName Name = TEXT("Struct"));
-	
-	UFUNCTION(BlueprintPure, Category = "LearningAgents")
-	FLearningAgentsObservationObjectElement MakeStructObservationFromArrays(const TArray<FName>& ElementNames, const TArray<FLearningAgentsObservationObjectElement>& Elements, const FName Name = TEXT("Struct"));
-	FLearningAgentsObservationObjectElement MakeStructObservationFromArrayViews(const TArrayView<const FName> ElementNames, const TArrayView<const FLearningAgentsObservationObjectElement> Elements, const FName Name = TEXT("Struct"));
+	UFUNCTION(BlueprintPure, Category = "LearningAgents", meta = (AdvancedDisplay = 3))
+	static FLearningAgentsObservationObjectElement MakeInclusiveUnionObservationFromArrays(ULearningAgentsObservationObject* Object, const TArray<FName>& ElementNames, const TArray<FLearningAgentsObservationObjectElement>& Elements, const FName Tag = TEXT("InclusiveUnion"));
+	static FLearningAgentsObservationObjectElement MakeInclusiveUnionObservationFromArrayViews(ULearningAgentsObservationObject* Object, const TArrayView<const FName> ElementNames, const TArrayView<const FLearningAgentsObservationObjectElement> Elements, const FName Tag = TEXT("InclusiveUnion"));
 
-	UFUNCTION(BlueprintPure, Category = "LearningAgents")
-	FLearningAgentsObservationObjectElement MakeExclusiveUnionObservation(const FName ElementName, const FLearningAgentsObservationObjectElement Element, const FName Name = TEXT("ExclusiveUnion"));
+	UFUNCTION(BlueprintPure, Category = "LearningAgents", meta = (AdvancedDisplay = 2))
+	static FLearningAgentsObservationObjectElement MakeStaticArrayObservation(ULearningAgentsObservationObject* Object, const TArray<FLearningAgentsObservationObjectElement>& Elements, const FName Tag = TEXT("StaticArray"));
+	static FLearningAgentsObservationObjectElement MakeStaticArrayObservationFromArrayView(ULearningAgentsObservationObject* Object, const TArrayView<const FLearningAgentsObservationObjectElement> Elements, const FName Tag = TEXT("StaticArray"));
 
-	UFUNCTION(BlueprintPure, Category = "LearningAgents")
-	FLearningAgentsObservationObjectElement MakeInclusiveUnionObservation(const TMap<FName, FLearningAgentsObservationObjectElement>& Elements, const FName Name = TEXT("InclusiveUnion"));
-	
-	UFUNCTION(BlueprintPure, Category = "LearningAgents")
-	FLearningAgentsObservationObjectElement MakeInclusiveUnionObservationFromArrays(const TArray<FName>& ElementNames, const TArray<FLearningAgentsObservationObjectElement>& Elements, const FName Name = TEXT("InclusiveUnion"));
-	FLearningAgentsObservationObjectElement MakeInclusiveUnionObservationFromArrayViews(const TArrayView<const FName> ElementNames, const TArrayView<const FLearningAgentsObservationObjectElement> Elements, const FName Name = TEXT("InclusiveUnion"));
+	UFUNCTION(BlueprintPure, Category = "LearningAgents", meta = (AdvancedDisplay = 2))
+	static FLearningAgentsObservationObjectElement MakeSetObservation(ULearningAgentsObservationObject* Object, const TSet<FLearningAgentsObservationObjectElement>& Elements, const FName Tag = TEXT("Set"));
 
-	UFUNCTION(BlueprintPure, Category = "LearningAgents")
-	FLearningAgentsObservationObjectElement MakeStaticArrayObservation(const TArray<FLearningAgentsObservationObjectElement>& Elements, const FName Name = TEXT("StaticArray"));
-	FLearningAgentsObservationObjectElement MakeStaticArrayObservationFromArrayView(const TArrayView<const FLearningAgentsObservationObjectElement> Elements, const FName Name = TEXT("StaticArray"));
+	UFUNCTION(BlueprintPure, Category = "LearningAgents", meta = (AdvancedDisplay = 2))
+	static FLearningAgentsObservationObjectElement MakeSetObservationFromArray(ULearningAgentsObservationObject* Object, const TArray<FLearningAgentsObservationObjectElement>& Elements, const FName Tag = TEXT("Set"));
+	static FLearningAgentsObservationObjectElement MakeSetObservationFromArrayView(ULearningAgentsObservationObject* Object, const TArrayView<const FLearningAgentsObservationObjectElement> Elements, const FName Tag = TEXT("Set"));
 
-	UFUNCTION(BlueprintPure, Category = "LearningAgents")
-	FLearningAgentsObservationObjectElement MakeSetObservation(const TSet<FLearningAgentsObservationObjectElement>& Elements, const FName Name = TEXT("Set"));
-	
-	UFUNCTION(BlueprintPure, Category = "LearningAgents")
-	FLearningAgentsObservationObjectElement MakeSetObservationFromArray(const TArray<FLearningAgentsObservationObjectElement>& Elements, const FName Name = TEXT("Set"));
-	FLearningAgentsObservationObjectElement MakeSetObservationFromArrayView(const TArrayView<const FLearningAgentsObservationObjectElement> Elements, const FName Name = TEXT("Set"));
+	UFUNCTION(BlueprintPure, Category = "LearningAgents", meta = (AdvancedDisplay = 3))
+	static FLearningAgentsObservationObjectElement MakePairObservation(ULearningAgentsObservationObject* Object, const FLearningAgentsObservationObjectElement Key, const FLearningAgentsObservationObjectElement Value, const FName Tag = TEXT("Pair"));
 
-	UFUNCTION(BlueprintPure, Category = "LearningAgents")
-	FLearningAgentsObservationObjectElement MakePairObservation(const FLearningAgentsObservationObjectElement Key, const FLearningAgentsObservationObjectElement Value, const FName Name = TEXT("Pair"));
+	UFUNCTION(BlueprintPure, Category = "LearningAgents", meta = (AdvancedDisplay = 2))
+	static FLearningAgentsObservationObjectElement MakeArrayObservation(ULearningAgentsObservationObject* Object, const TArray<FLearningAgentsObservationObjectElement>& Elements, const FName Tag = TEXT("Array"));
+	static FLearningAgentsObservationObjectElement MakeArrayObservationFromArrayView(ULearningAgentsObservationObject* Object, const TArrayView<const FLearningAgentsObservationObjectElement> Elements, const FName Tag = TEXT("Array"));
 
-	UFUNCTION(BlueprintPure, Category = "LearningAgents")
-	FLearningAgentsObservationObjectElement MakeArrayObservation(const TArray<FLearningAgentsObservationObjectElement>& Elements, const FName Name = TEXT("Array"));
-	FLearningAgentsObservationObjectElement MakeArrayObservationFromArrayView(const TArrayView<const FLearningAgentsObservationObjectElement> Elements, const FName Name = TEXT("Array"));
+	UFUNCTION(BlueprintPure, Category = "LearningAgents", meta = (AdvancedDisplay = 2))
+	static FLearningAgentsObservationObjectElement MakeMapObservation(ULearningAgentsObservationObject* Object, const TMap<FLearningAgentsObservationObjectElement, FLearningAgentsObservationObjectElement>& Map, const FName Tag = TEXT("Map"));
 
-	UFUNCTION(BlueprintPure, Category = "LearningAgents")
-	FLearningAgentsObservationObjectElement MakeMapObservation(const TMap<FLearningAgentsObservationObjectElement, FLearningAgentsObservationObjectElement>& Map, const FName Name = TEXT("Map"));
-	
-	UFUNCTION(BlueprintPure, Category = "LearningAgents")
-	FLearningAgentsObservationObjectElement MakeMapObservationFromArrays(const TArray<FLearningAgentsObservationObjectElement>& Keys, const TArray<FLearningAgentsObservationObjectElement>& Values, const FName Name = TEXT("Map"));
-	FLearningAgentsObservationObjectElement MakeMapObservationFromArrayViews(const TArrayView<const FLearningAgentsObservationObjectElement> Keys, const TArrayView<const FLearningAgentsObservationObjectElement> Values, const FName Name = TEXT("Map"));
+	UFUNCTION(BlueprintPure, Category = "LearningAgents", meta = (AdvancedDisplay = 3))
+	static FLearningAgentsObservationObjectElement MakeMapObservationFromArrays(ULearningAgentsObservationObject* Object, const TArray<FLearningAgentsObservationObjectElement>& Keys, const TArray<FLearningAgentsObservationObjectElement>& Values, const FName Tag = TEXT("Map"));
+	static FLearningAgentsObservationObjectElement MakeMapObservationFromArrayViews(ULearningAgentsObservationObject* Object, const TArrayView<const FLearningAgentsObservationObjectElement> Keys, const TArrayView<const FLearningAgentsObservationObjectElement> Values, const FName Tag = TEXT("Map"));
 
-	UFUNCTION(BlueprintPure, Category = "LearningAgents")
-	FLearningAgentsObservationObjectElement MakeEnumObservation(const UEnum* Enum, const uint8 EnumValue, const FName Name = TEXT("Enum"));
+	UFUNCTION(BlueprintPure, Category = "LearningAgents", meta = (AdvancedDisplay = 3))
+	static FLearningAgentsObservationObjectElement MakeEnumObservation(ULearningAgentsObservationObject* Object, const UEnum* Enum, const uint8 EnumValue, const FName Tag = TEXT("Enum"));
 
-	UFUNCTION(BlueprintPure, Category = "LearningAgents")
-	FLearningAgentsObservationObjectElement MakeBitmaskObservation(const UEnum* Enum, const int32 BitmaskValue, const FName Name = TEXT("Bitmask"));
+	UFUNCTION(BlueprintPure, Category = "LearningAgents", meta = (AdvancedDisplay = 3))
+	static FLearningAgentsObservationObjectElement MakeBitmaskObservation(ULearningAgentsObservationObject* Object, const UEnum* Enum, const int32 BitmaskValue, const FName Tag = TEXT("Bitmask"));
 
-	UFUNCTION(BlueprintPure, Category = "LearningAgents")
-	FLearningAgentsObservationObjectElement MakeOptionalObservation(const FLearningAgentsObservationObjectElement Element, const ELearningAgentsOptionalObservation Option, const FName Name = TEXT("Optional"));
+	UFUNCTION(BlueprintPure, Category = "LearningAgents", meta = (AdvancedDisplay = 3))
+	static FLearningAgentsObservationObjectElement MakeOptionalObservation(ULearningAgentsObservationObject* Object, const FLearningAgentsObservationObjectElement Element, const ELearningAgentsOptionalObservation Option, const FName Tag = TEXT("Optional"));
 
-	UFUNCTION(BlueprintPure, Category = "LearningAgents")
-	FLearningAgentsObservationObjectElement MakeOptionalNullObservation(const FName Name = TEXT("Optional"));
+	UFUNCTION(BlueprintPure, Category = "LearningAgents", meta = (AdvancedDisplay = 1))
+	static FLearningAgentsObservationObjectElement MakeOptionalNullObservation(ULearningAgentsObservationObject* Object, const FName Tag = TEXT("Optional"));
 
-	UFUNCTION(BlueprintPure, Category = "LearningAgents")
-	FLearningAgentsObservationObjectElement MakeOptionalValidObservation(const FLearningAgentsObservationObjectElement Element, const FName Name = TEXT("Optional"));
+	UFUNCTION(BlueprintPure, Category = "LearningAgents", meta = (AdvancedDisplay = 2))
+	static FLearningAgentsObservationObjectElement MakeOptionalValidObservation(ULearningAgentsObservationObject* Object, const FLearningAgentsObservationObjectElement Element, const FName Tag = TEXT("Optional"));
 
-	UFUNCTION(BlueprintPure, Category = "LearningAgents")
-	FLearningAgentsObservationObjectElement MakeEitherObservation(const FLearningAgentsObservationObjectElement Element, const ELearningAgentsEitherObservation Either, const FName Name = TEXT("Either"));
+	UFUNCTION(BlueprintPure, Category = "LearningAgents", meta = (AdvancedDisplay = 3))
+	static FLearningAgentsObservationObjectElement MakeEitherObservation(ULearningAgentsObservationObject* Object, const FLearningAgentsObservationObjectElement Element, const ELearningAgentsEitherObservation Either, const FName Tag = TEXT("Either"));
 
-	UFUNCTION(BlueprintPure, Category = "LearningAgents", meta=(DisplayName="Make Either A Observation"))
-	FLearningAgentsObservationObjectElement MakeEitherAObservation(const FLearningAgentsObservationObjectElement A, const FName Name = TEXT("Either"));
+	UFUNCTION(BlueprintPure, Category = "LearningAgents", meta = (AdvancedDisplay = 2, DisplayName = "Make Either A Observation"))
+	static FLearningAgentsObservationObjectElement MakeEitherAObservation(ULearningAgentsObservationObject* Object, const FLearningAgentsObservationObjectElement A, const FName Tag = TEXT("Either"));
 
-	UFUNCTION(BlueprintPure, Category = "LearningAgents", meta = (DisplayName = "Make Either B Observation"))
-	FLearningAgentsObservationObjectElement MakeEitherBObservation(const FLearningAgentsObservationObjectElement B, const FName Name = TEXT("Either"));
+	UFUNCTION(BlueprintPure, Category = "LearningAgents", meta = (AdvancedDisplay = 2, DisplayName = "Make Either B Observation"))
+	static FLearningAgentsObservationObjectElement MakeEitherBObservation(ULearningAgentsObservationObject* Object, const FLearningAgentsObservationObjectElement B, const FName Tag = TEXT("Either"));
 
-	UFUNCTION(BlueprintPure, Category = "LearningAgents")
-	FLearningAgentsObservationObjectElement MakeEncodingObservation(const FLearningAgentsObservationObjectElement Element, const FName Name = TEXT("Encoding"));
+	UFUNCTION(BlueprintPure, Category = "LearningAgents", meta = (AdvancedDisplay = 2))
+	static FLearningAgentsObservationObjectElement MakeEncodingObservation(ULearningAgentsObservationObject* Object, const FLearningAgentsObservationObjectElement Element, const FName Tag = TEXT("Encoding"));
 
 	/**
 	 * Make a new bool observation.
 	 *
 	 * @param Value The new value of this observation.
-	 * @param Name The name of the corresponding observation. Must match the name given during Specify.
+	 * @param Tag The tag of the corresponding observation. Must match the tag given during Specify.
 	 * @param bVisualLoggerEnabled When true, debug data will be sent to the visual logger.
 	 * @param VisualLoggerListener The listener object which is making this observation. This must be set to use logging.
 	 * @param VisualLoggerAgentId The agent id associated with this observation.
@@ -384,9 +372,10 @@ public:
 	 * @return The newly created observation object element.
 	 */
 	UFUNCTION(BlueprintPure, Category = "LearningAgents", meta = (AdvancedDisplay = 2, DefaultToSelf = "VisualLoggerListener"))
-	FLearningAgentsObservationObjectElement MakeBoolObservation(
+	static FLearningAgentsObservationObjectElement MakeBoolObservation(
+		ULearningAgentsObservationObject* Object,
 		const bool bValue,
-		const FName Name = TEXT("Bool"),
+		const FName Tag = TEXT("Bool"),
 		const bool bVisualLoggerEnabled = false,
 		ULearningAgentsManagerListener* VisualLoggerListener = nullptr,
 		const int32 VisualLoggerAgentId = -1,
@@ -398,7 +387,7 @@ public:
 	 *
 	 * @param Value The new value of this observation.
 	 * @param FloatScale Used to normalize the data for this observation.
-	 * @param Name The name of the corresponding observation. Must match the name given during Specify.
+	 * @param Tag The tag of the corresponding observation. Must match the tag given during Specify.
 	 * @param bVisualLoggerEnabled When true, debug data will be sent to the visual logger.
 	 * @param VisualLoggerListener The listener object which is making this observation. This must be set to use logging.
 	 * @param VisualLoggerAgentId The agent id associated with this observation.
@@ -407,10 +396,11 @@ public:
 	 * @return The newly created observation object element.
 	 */
 	UFUNCTION(BlueprintPure, Category = "LearningAgents", meta = (AdvancedDisplay = 3, DefaultToSelf = "VisualLoggerListener"))
-	FLearningAgentsObservationObjectElement MakeFloatObservation(
+	static FLearningAgentsObservationObjectElement MakeFloatObservation(
+		ULearningAgentsObservationObject* Object,
 		const float Value,
 		const float FloatScale = 1.0f,
-		const FName Name = TEXT("Float"),
+		const FName Tag = TEXT("Float"),
 		const bool bVisualLoggerEnabled = false,
 		ULearningAgentsManagerListener* VisualLoggerListener = nullptr,
 		const int32 VisualLoggerAgentId = -1,
@@ -423,7 +413,7 @@ public:
 	 * @param Location The location of interest to the agent.
 	 * @param RelativeTransform The transform the provided location should be encoded relative to.
 	 * @param LocationScale Used to normalize the data for this observation.
-	 * @param Name The name of the corresponding observation. Must match the name given during Specify.
+	 * @param Tag The tag of the corresponding observation. Must match the tag given during Specify.
 	 * @param bVisualLoggerEnabled When true, debug data will be sent to the visual logger.
 	 * @param VisualLoggerListener The listener object which is making this observation. This must be set to use logging.
 	 * @param VisualLoggerAgentId The agent id associated with this observation.
@@ -432,11 +422,12 @@ public:
 	 * @return The newly created observation object element.
 	 */
 	UFUNCTION(BlueprintPure, Category = "LearningAgents", meta = (AdvancedDisplay = 4, DefaultToSelf = "VisualLoggerListener"))
-	FLearningAgentsObservationObjectElement MakeLocationObservation(
+	static FLearningAgentsObservationObjectElement MakeLocationObservation(
+		ULearningAgentsObservationObject* Object,
 		const FVector Location,
 		const FTransform RelativeTransform = FTransform(),
 		const float LocationScale = 100.0f,
-		const FName Name = TEXT("Location"),
+		const FName Tag = TEXT("Location"),
 		const bool bVisualLoggerEnabled = false,
 		ULearningAgentsManagerListener* VisualLoggerListener = nullptr,
 		const int32 VisualLoggerAgentId = -1,
@@ -448,7 +439,7 @@ public:
 	 *
 	 * @param Rotation The rotation of interest to the agent.
 	 * @param RelativeRotation The rotation the provided rotation should be encoded relative to.
-	 * @param Name The name of the corresponding observation. Must match the name given during Specify.
+	 * @param Tag The tag of the corresponding observation. Must match the tag given during Specify.
 	 * @param bVisualLoggerEnabled When true, debug data will be sent to the visual logger.
 	 * @param VisualLoggerListener The listener object which is making this observation. This must be set to use logging.
 	 * @param VisualLoggerAgentId The agent id associated with this observation.
@@ -458,10 +449,11 @@ public:
 	 * @return The newly created observation object element.
 	 */
 	UFUNCTION(BlueprintPure, Category = "LearningAgents", meta = (AdvancedDisplay = 3, DefaultToSelf = "VisualLoggerListener"))
-	FLearningAgentsObservationObjectElement MakeRotationObservation(
+	static FLearningAgentsObservationObjectElement MakeRotationObservation(
+		ULearningAgentsObservationObject* Object,
 		const FRotator Rotation,
 		const FRotator RelativeRotation = FRotator::ZeroRotator,
-		const FName Name = TEXT("Rotation"),
+		const FName Tag = TEXT("Rotation"),
 		const bool bVisualLoggerEnabled = false,
 		ULearningAgentsManagerListener* VisualLoggerListener = nullptr,
 		const int32 VisualLoggerAgentId = -1,
@@ -474,7 +466,7 @@ public:
 	 *
 	 * @param Rotation The rotation of interest to the agent.
 	 * @param RelativeRotation The rotation the provided rotation should be encoded relative to.
-	 * @param Name The name of the corresponding observation. Must match the name given during Specify.
+	 * @param Tag The tag of the corresponding observation. Must match the tag given during Specify.
 	 * @param bVisualLoggerEnabled When true, debug data will be sent to the visual logger.
 	 * @param VisualLoggerListener The listener object which is making this observation. This must be set to use logging.
 	 * @param VisualLoggerAgentId The agent id associated with this observation.
@@ -484,10 +476,11 @@ public:
 	 * @return The newly created observation object element.
 	 */
 	UFUNCTION(BlueprintPure, Category = "LearningAgents", meta = (AdvancedDisplay = 3, DefaultToSelf = "VisualLoggerListener"))
-	FLearningAgentsObservationObjectElement MakeRotationObservationFromQuat(
+	static FLearningAgentsObservationObjectElement MakeRotationObservationFromQuat(
+		ULearningAgentsObservationObject* Object,
 		const FQuat Rotation,
 		const FQuat RelativeRotation,
-		const FName Name = TEXT("Rotation"),
+		const FName Tag = TEXT("Rotation"),
 		const bool bVisualLoggerEnabled = false,
 		ULearningAgentsManagerListener* VisualLoggerListener = nullptr,
 		const int32 VisualLoggerAgentId = -1,
@@ -500,7 +493,7 @@ public:
 	 *
 	 * @param Scale The scale of interest to the agent.
 	 * @param RelativeScale The scale the provided scale should be encoded relative to.
-	 * @param Name The name of the corresponding observation. Must match the name given during Specify.
+	 * @param Tag The tag of the corresponding observation. Must match the tag given during Specify.
 	 * @param bVisualLoggerEnabled When true, debug data will be sent to the visual logger.
 	 * @param VisualLoggerListener The listener object which is making this observation. This must be set to use logging.
 	 * @param VisualLoggerAgentId The agent id associated with this observation.
@@ -510,10 +503,11 @@ public:
 	 * @return The newly created observation object element.
 	 */
 	UFUNCTION(BlueprintPure, Category = "LearningAgents", meta = (AdvancedDisplay = 3, DefaultToSelf = "VisualLoggerListener"))
-	FLearningAgentsObservationObjectElement MakeScaleObservation(
+	static FLearningAgentsObservationObjectElement MakeScaleObservation(
+		ULearningAgentsObservationObject* Object,
 		const FVector Scale,
-		const FVector RelativeScale = FVector(1,1,1),
-		const FName Name = TEXT("Scale"),
+		const FVector RelativeScale = FVector(1, 1, 1),
+		const FName Tag = TEXT("Scale"),
 		const bool bVisualLoggerEnabled = false,
 		ULearningAgentsManagerListener* VisualLoggerListener = nullptr,
 		const int32 VisualLoggerAgentId = -1,
@@ -527,7 +521,7 @@ public:
 	 * @param Transform The transform of interest to the agent.
 	 * @param RelativeTransform The transform the provided transform should be encoded relative to.
 	 * @param LocationScale Used to normalize the transform's location for this observation.
-	 * @param Name The name of the corresponding observation. Must match the name given during Specify.
+	 * @param Tag The tag of the corresponding observation. Must match the tag given during Specify.
 	 * @param bVisualLoggerEnabled When true, debug data will be sent to the visual logger.
 	 * @param VisualLoggerListener The listener object which is making this observation. This must be set to use logging.
 	 * @param VisualLoggerAgentId The agent id associated with this observation.
@@ -536,31 +530,33 @@ public:
 	 * @return The newly created observation object element.
 	 */
 	UFUNCTION(BlueprintPure, Category = "LearningAgents", meta = (AdvancedDisplay = 4, DefaultToSelf = "VisualLoggerListener"))
-	FLearningAgentsObservationObjectElement MakeTransformObservation(
+	static FLearningAgentsObservationObjectElement MakeTransformObservation(
+		ULearningAgentsObservationObject* Object,
 		const FTransform Transform,
 		const FTransform RelativeTransform = FTransform(),
 		const float LocationScale = 100.0f,
-		const FName Name = TEXT("Transform"),
+		const FName Tag = TEXT("Transform"),
 		const bool bVisualLoggerEnabled = false,
 		ULearningAgentsManagerListener* VisualLoggerListener = nullptr,
 		const int32 VisualLoggerAgentId = -1,
 		const FVector VisualLoggerLocation = FVector::ZeroVector,
 		const FLinearColor VisualLoggerColor = FLinearColor::Red);
 
-	UFUNCTION(BlueprintPure, Category = "LearningAgents")
-	FLearningAgentsObservationObjectElement MakeAngleObservation(const float Angle, const float RelativeAngle = 0.0f, const FName Name = TEXT("Angle"));
+	UFUNCTION(BlueprintPure, Category = "LearningAgents", meta = (AdvancedDisplay = 3))
+	static FLearningAgentsObservationObjectElement MakeAngleObservation(ULearningAgentsObservationObject* Object, const float Angle, const float RelativeAngle = 0.0f, const FName Tag = TEXT("Angle"));
 
-	UFUNCTION(BlueprintPure, Category = "LearningAgents")
-	FLearningAgentsObservationObjectElement MakeAngleObservationRadians(const float Angle, const float RelativeAngle = 0.0f, const FName Name = TEXT("Angle"));
+	UFUNCTION(BlueprintPure, Category = "LearningAgents", meta = (AdvancedDisplay = 3))
+	static FLearningAgentsObservationObjectElement MakeAngleObservationRadians(ULearningAgentsObservationObject* Object, const float Angle, const float RelativeAngle = 0.0f, const FName Tag = TEXT("Angle"));
 
-	UFUNCTION(BlueprintPure, Category = "LearningAgents")
-	FLearningAgentsObservationObjectElement MakeVelocityObservation(const FVector Velocity, const FTransform RelativeTransform = FTransform(), const float VelocityScale = 200.0f, const FName Name = TEXT("Velocity"));
+	UFUNCTION(BlueprintPure, Category = "LearningAgents", meta = (AdvancedDisplay = 4))
+	static FLearningAgentsObservationObjectElement MakeVelocityObservation(ULearningAgentsObservationObject* Object, const FVector Velocity, const FTransform RelativeTransform = FTransform(), const float VelocityScale = 200.0f, const FName Tag = TEXT("Velocity"));
 
 	UFUNCTION(BlueprintPure, Category = "LearningAgents", meta = (AdvancedDisplay = 3, DefaultToSelf = "VisualLoggerListener"))
-	FLearningAgentsObservationObjectElement MakeDirectionObservation(
+	static FLearningAgentsObservationObjectElement MakeDirectionObservation(
+		ULearningAgentsObservationObject* Object,
 		const FVector Direction,
 		const FTransform RelativeTransform = FTransform(),
-		const FName Name = TEXT("Direction"),
+		const FName Tag = TEXT("Direction"),
 		const bool bVisualLoggerEnabled = false,
 		ULearningAgentsManagerListener* VisualLoggerListener = nullptr,
 		const int32 VisualLoggerAgentId = -1,
@@ -569,30 +565,32 @@ public:
 		const float VisualLoggerArrowLength = 100.0f,
 		const FLinearColor VisualLoggerColor = FLinearColor::Red);
 
-	// Spline Observations
+	// Make Spline Observations
 
 	UFUNCTION(BlueprintPure, Category = "LearningAgents", meta = (AdvancedDisplay = 5, DefaultToSelf = "VisualLoggerListener"))
-	FLearningAgentsObservationObjectElement MakeLocationAlongSplineObservation(
+	static FLearningAgentsObservationObjectElement MakeLocationAlongSplineObservation(
+		ULearningAgentsObservationObject* Object,
 		const USplineComponent* SplineComponent,
 		const float DistanceAlongSpline,
 		const FTransform RelativeTransform = FTransform(),
 		const float LocationScale = 100.0f,
-		const FName Name = TEXT("LocationAlongSpline"),
+		const FName Tag = TEXT("LocationAlongSpline"),
 		const bool bVisualLoggerEnabled = false,
 		ULearningAgentsManagerListener* VisualLoggerListener = nullptr,
 		const int32 VisualLoggerAgentId = -1,
 		const FVector VisualLoggerLocation = FVector::ZeroVector,
 		const FLinearColor VisualLoggerColor = FLinearColor::Red);
 
-	UFUNCTION(BlueprintPure, Category = "LearningAgents")
-	FLearningAgentsObservationObjectElement MakeProportionAlongSplineObservation(const USplineComponent* SplineComponent, const float DistanceAlongSpline, const FName Name = TEXT("ProportionAlongSpline"));
+	UFUNCTION(BlueprintPure, Category = "LearningAgents", meta = (AdvancedDisplay = 3))
+	static FLearningAgentsObservationObjectElement MakeProportionAlongSplineObservation(ULearningAgentsObservationObject* Object, const USplineComponent* SplineComponent, const float DistanceAlongSpline, const FName Tag = TEXT("ProportionAlongSpline"));
 
 	UFUNCTION(BlueprintPure, Category = "LearningAgents", meta = (AdvancedDisplay = 4, DefaultToSelf = "VisualLoggerListener"))
-	FLearningAgentsObservationObjectElement MakeDirectionAlongSplineObservation(
+	static FLearningAgentsObservationObjectElement MakeDirectionAlongSplineObservation(
+		ULearningAgentsObservationObject* Object,
 		const USplineComponent* SplineComponent,
 		const float DistanceAlongSpline,
 		const FTransform RelativeTransform = FTransform(),
-		const FName Name = TEXT("DirectionAlongSpline"),
+		const FName Tag = TEXT("DirectionAlongSpline"),
 		const bool bVisualLoggerEnabled = false,
 		ULearningAgentsManagerListener* VisualLoggerListener = nullptr,
 		const int32 VisualLoggerAgentId = -1,
@@ -600,181 +598,181 @@ public:
 		const float VisualLoggerArrowLength = 100.0f,
 		const FLinearColor VisualLoggerColor = FLinearColor::Red);
 
-	UFUNCTION(BlueprintPure, Category = "LearningAgents")
-	FLearningAgentsObservationObjectElement MakePropertiesAlongSplineObservation(const USplineComponent* SplineComponent, const float DistanceAlongSpline, const FTransform RelativeTransform = FTransform(), const float LocationScale = 100.0f, const FName Name = TEXT("PropertiesAlongSpline"));
+	UFUNCTION(BlueprintPure, Category = "LearningAgents", meta = (AdvancedDisplay = 5))
+	static FLearningAgentsObservationObjectElement MakePropertiesAlongSplineObservation(ULearningAgentsObservationObject* Object, const USplineComponent* SplineComponent, const float DistanceAlongSpline, const FTransform RelativeTransform = FTransform(), const float LocationScale = 100.0f, const FName Tag = TEXT("PropertiesAlongSpline"));
 
-	// Ray Cast Observations
+	// Make Ray Cast Observations
 
-	UFUNCTION(BlueprintPure, Category = "LearningAgents")
-	FLearningAgentsObservationObjectElement MakeProportionAlongRayObservation(const FVector RayStart, const FVector RayEnd, const FTransform RayTransform = FTransform(), const ECollisionChannel CollisionChannel = ECollisionChannel::ECC_WorldStatic, const FName Name = TEXT("ProportionAlongRay"));
+	UFUNCTION(BlueprintPure, Category = "LearningAgents", meta = (AdvancedDisplay = 5))
+	static FLearningAgentsObservationObjectElement MakeProportionAlongRayObservation(ULearningAgentsObservationObject* Object, const FVector RayStart, const FVector RayEnd, const FTransform RayTransform = FTransform(), const ECollisionChannel CollisionChannel = ECollisionChannel::ECC_WorldStatic, const FName Tag = TEXT("ProportionAlongRay"));
 
-	UFUNCTION(BlueprintPure, Category = "LearningAgents")
-	FLearningAgentsObservationObjectElement MakeProportionAlongRaysObservation(const TArray<FVector>& RayStarts, const TArray<FVector>& RayEnds, const FTransform RayTransform, const ECollisionChannel CollisionChannel = ECollisionChannel::ECC_WorldStatic, const FName Name = TEXT("ProportionAlongRays"));
+	UFUNCTION(BlueprintPure, Category = "LearningAgents", meta = (AdvancedDisplay = 5))
+	static FLearningAgentsObservationObjectElement MakeProportionAlongRaysObservation(ULearningAgentsObservationObject* Object, const TArray<FVector>& RayStarts, const TArray<FVector>& RayEnds, const FTransform RayTransform, const ECollisionChannel CollisionChannel = ECollisionChannel::ECC_WorldStatic, const FName Tag = TEXT("ProportionAlongRays"));
 
 public:
 
-	// Basic Observations
+	// Get Basic Observations
 
-	UFUNCTION(BlueprintPure = false, Category = "LearningAgents", meta=(ReturnDisplayName = "Success"))
-	bool GetNullObservation(const FLearningAgentsObservationObjectElement Element, const FName Name = TEXT("Null")) const;
+	UFUNCTION(BlueprintPure = false, Category = "LearningAgents", meta = (AdvancedDisplay = 2, ReturnDisplayName = "Success"))
+	static bool GetNullObservation(const ULearningAgentsObservationObject* Object, const FLearningAgentsObservationObjectElement Element, const FName Tag = TEXT("Null"));
 
-	UFUNCTION(BlueprintPure = false, Category = "LearningAgents", meta = (ReturnDisplayName = "Success"))
-	bool GetContinuousObservationNum(int32& OutNum, const FLearningAgentsObservationObjectElement Element, const FName Name = TEXT("Continuous")) const;
+	UFUNCTION(BlueprintPure = false, Category = "LearningAgents", meta = (AdvancedDisplay = 3, ReturnDisplayName = "Success"))
+	static bool GetContinuousObservationNum(int32& OutNum, const ULearningAgentsObservationObject* Object, const FLearningAgentsObservationObjectElement Element, const FName Tag = TEXT("Continuous"));
 
-	UFUNCTION(BlueprintPure = false, Category = "LearningAgents", meta = (ReturnDisplayName = "Success"))
-	bool GetContinuousObservation(TArray<float>& OutValues, const FLearningAgentsObservationObjectElement Element, const FName Name = TEXT("Continuous")) const;
-	bool GetContinuousObservationToArrayView(TArrayView<float> OutValues, const FLearningAgentsObservationObjectElement Element, const FName Name = TEXT("Continuous")) const;
+	UFUNCTION(BlueprintPure = false, Category = "LearningAgents", meta = (AdvancedDisplay = 3, ReturnDisplayName = "Success"))
+	static bool GetContinuousObservation(TArray<float>& OutValues, const ULearningAgentsObservationObject* Object, const FLearningAgentsObservationObjectElement Element, const FName Tag = TEXT("Continuous"));
+	static bool GetContinuousObservationToArrayView(TArrayView<float> OutValues, const ULearningAgentsObservationObject* Object, const FLearningAgentsObservationObjectElement Element, const FName Tag = TEXT("Continuous"));
 
-	UFUNCTION(BlueprintPure = false, Category = "LearningAgents", meta = (ReturnDisplayName = "Success"))
-	bool GetExclusiveDiscreteObservation(int32& OutIndex, const FLearningAgentsObservationObjectElement Element, const FName Name = TEXT("ExclusiveDiscrete")) const;
+	UFUNCTION(BlueprintPure = false, Category = "LearningAgents", meta = (AdvancedDisplay = 3, ReturnDisplayName = "Success"))
+	static bool GetExclusiveDiscreteObservation(int32& OutIndex, const ULearningAgentsObservationObject* Object, const FLearningAgentsObservationObjectElement Element, const FName Tag = TEXT("ExclusiveDiscrete"));
 
-	UFUNCTION(BlueprintPure = false, Category = "LearningAgents", meta = (ReturnDisplayName = "Success"))
-	bool GetInclusiveDiscreteObservationNum(int32& OutNum, const FLearningAgentsObservationObjectElement Element, const FName Name = TEXT("InclusiveDiscrete")) const;
+	UFUNCTION(BlueprintPure = false, Category = "LearningAgents", meta = (AdvancedDisplay = 3, ReturnDisplayName = "Success"))
+	static bool GetInclusiveDiscreteObservationNum(int32& OutNum, const ULearningAgentsObservationObject* Object, const FLearningAgentsObservationObjectElement Element, const FName Tag = TEXT("InclusiveDiscrete"));
 
-	UFUNCTION(BlueprintPure = false, Category = "LearningAgents", meta = (ReturnDisplayName = "Success"))
-	bool GetInclusiveDiscreteObservation(TArray<int32>& OutIndices, const FLearningAgentsObservationObjectElement Element, const FName Name = TEXT("InclusiveDiscrete")) const;
-	bool GetInclusiveDiscreteObservationToArrayView(TArrayView<int32> OutIndices, const FLearningAgentsObservationObjectElement Element, const FName Name = TEXT("InclusiveDiscrete")) const;
+	UFUNCTION(BlueprintPure = false, Category = "LearningAgents", meta = (AdvancedDisplay = 3, ReturnDisplayName = "Success"))
+	static bool GetInclusiveDiscreteObservation(TArray<int32>& OutIndices, const ULearningAgentsObservationObject* Object, const FLearningAgentsObservationObjectElement Element, const FName Tag = TEXT("InclusiveDiscrete"));
+	static bool GetInclusiveDiscreteObservationToArrayView(TArrayView<int32> OutIndices, const ULearningAgentsObservationObject* Object, const FLearningAgentsObservationObjectElement Element, const FName Tag = TEXT("InclusiveDiscrete"));
 
-	UFUNCTION(BlueprintPure = false, Category = "LearningAgents", meta = (ReturnDisplayName = "Success"))
-	bool GetIndexObservation(int32& OutIndex, const FLearningAgentsObservationObjectElement Element, const FName Name = TEXT("Index")) const;
+	UFUNCTION(BlueprintPure = false, Category = "LearningAgents", meta = (AdvancedDisplay = 3, ReturnDisplayName = "Success"))
+	static bool GetIndexObservation(int32& OutIndex, const ULearningAgentsObservationObject* Object, const FLearningAgentsObservationObjectElement Element, const FName Tag = TEXT("Index"));
 
-	UFUNCTION(BlueprintPure = false, Category = "LearningAgents", meta = (ReturnDisplayName = "Success"))
-	bool GetCountObservation(int32& OutNum, const FLearningAgentsObservationObjectElement Element, const int32 MaxNum, const FName Name = TEXT("Count")) const;
+	UFUNCTION(BlueprintPure = false, Category = "LearningAgents", meta = (AdvancedDisplay = 3, ReturnDisplayName = "Success"))
+	static bool GetCountObservation(int32& OutNum, const ULearningAgentsObservationObject* Object, const FLearningAgentsObservationObjectElement Element, const int32 MaxNum, const FName Tag = TEXT("Count"));
 
-	UFUNCTION(BlueprintPure = false, Category = "LearningAgents", meta = (ReturnDisplayName = "Success"))
-	bool GetStructObservationNum(int32& OutNum, const FLearningAgentsObservationObjectElement Element, const FName Name = TEXT("Struct")) const;
+	UFUNCTION(BlueprintPure = false, Category = "LearningAgents", meta = (AdvancedDisplay = 3, ReturnDisplayName = "Success"))
+	static bool GetStructObservationNum(int32& OutNum, const ULearningAgentsObservationObject* Object, const FLearningAgentsObservationObjectElement Element, const FName Tag = TEXT("Struct"));
 
-	UFUNCTION(BlueprintPure = false, Category = "LearningAgents", meta = (ReturnDisplayName = "Success"))
-	bool GetStructObservation(TMap<FName, FLearningAgentsObservationObjectElement>& OutElements, const FLearningAgentsObservationObjectElement Element, const FName Name = TEXT("Struct")) const;
-	
-	UFUNCTION(BlueprintPure = false, Category = "LearningAgents", meta = (ReturnDisplayName = "Success"))
-	bool GetStructObservationToArrays(TArray<FName>& OutElementNames, TArray<FLearningAgentsObservationObjectElement>& OutElements, const FLearningAgentsObservationObjectElement Element, const FName Name = TEXT("Struct")) const;
-	bool GetStructObservationToArrayViews(TArrayView<FName> OutElementNames, TArrayView<FLearningAgentsObservationObjectElement> OutElements, const FLearningAgentsObservationObjectElement Element, const FName Name = TEXT("Struct")) const;
+	UFUNCTION(BlueprintPure = false, Category = "LearningAgents", meta = (AdvancedDisplay = 3, ReturnDisplayName = "Success"))
+	static bool GetStructObservation(TMap<FName, FLearningAgentsObservationObjectElement>& OutElements, const ULearningAgentsObservationObject* Object, const FLearningAgentsObservationObjectElement Element, const FName Tag = TEXT("Struct"));
 
-	UFUNCTION(BlueprintPure = false, Category = "LearningAgents", meta = (ReturnDisplayName = "Success"))
-	bool GetExclusiveUnionObservation(FName& OutElementName, FLearningAgentsObservationObjectElement& OutElement, const FLearningAgentsObservationObjectElement Element, const FName Name = TEXT("ExclusiveUnion")) const;
+	UFUNCTION(BlueprintPure = false, Category = "LearningAgents", meta = (AdvancedDisplay = 4, ReturnDisplayName = "Success"))
+	static bool GetStructObservationToArrays(TArray<FName>& OutElementNames, TArray<FLearningAgentsObservationObjectElement>& OutElements, const ULearningAgentsObservationObject* Object, const FLearningAgentsObservationObjectElement Element, const FName Tag = TEXT("Struct"));
+	static bool GetStructObservationToArrayViews(TArrayView<FName> OutElementNames, TArrayView<FLearningAgentsObservationObjectElement> OutElements, const ULearningAgentsObservationObject* Object, const FLearningAgentsObservationObjectElement Element, const FName Tag = TEXT("Struct"));
 
-	UFUNCTION(BlueprintPure = false, Category = "LearningAgents", meta = (ReturnDisplayName = "Success"))
-	bool GetInclusiveUnionObservationNum(int32& OutNum, const FLearningAgentsObservationObjectElement Element, const FName Name = TEXT("InclusiveUnion")) const;
+	UFUNCTION(BlueprintPure = false, Category = "LearningAgents", meta = (AdvancedDisplay = 4, ReturnDisplayName = "Success"))
+	static bool GetExclusiveUnionObservation(FName& OutElementName, FLearningAgentsObservationObjectElement& OutElement, const ULearningAgentsObservationObject* Object, const FLearningAgentsObservationObjectElement Element, const FName Tag = TEXT("ExclusiveUnion"));
 
-	UFUNCTION(BlueprintPure = false, Category = "LearningAgents", meta = (ReturnDisplayName = "Success"))
-	bool GetInclusiveUnionObservation(TMap<FName, FLearningAgentsObservationObjectElement>& OutElements, const FLearningAgentsObservationObjectElement Element, const FName Name = TEXT("InclusiveUnion")) const;
-	
-	UFUNCTION(BlueprintPure = false, Category = "LearningAgents", meta = (ReturnDisplayName = "Success"))
-	bool GetInclusiveUnionObservationToArrays(TArray<FName>& OutElementNames, TArray<FLearningAgentsObservationObjectElement>& OutElements, const FLearningAgentsObservationObjectElement Element, const FName Name = TEXT("InclusiveUnion")) const;
-	bool GetInclusiveUnionObservationToArrayViews(TArrayView<FName> OutElementNames, TArrayView<FLearningAgentsObservationObjectElement> OutElements, const FLearningAgentsObservationObjectElement Element, const FName Name = TEXT("InclusiveUnion")) const;
+	UFUNCTION(BlueprintPure = false, Category = "LearningAgents", meta = (AdvancedDisplay = 3, ReturnDisplayName = "Success"))
+	static bool GetInclusiveUnionObservationNum(int32& OutNum, const ULearningAgentsObservationObject* Object, const FLearningAgentsObservationObjectElement Element, const FName Tag = TEXT("InclusiveUnion"));
 
-	UFUNCTION(BlueprintPure = false, Category = "LearningAgents", meta = (ReturnDisplayName = "Success"))
-	bool GetStaticArrayObservationNum(int32& OutNum, const FLearningAgentsObservationObjectElement Element, const FName Name = TEXT("StaticArray")) const;
+	UFUNCTION(BlueprintPure = false, Category = "LearningAgents", meta = (AdvancedDisplay = 3, ReturnDisplayName = "Success"))
+	static bool GetInclusiveUnionObservation(TMap<FName, FLearningAgentsObservationObjectElement>& OutElements, const ULearningAgentsObservationObject* Object, const FLearningAgentsObservationObjectElement Element, const FName Tag = TEXT("InclusiveUnion"));
 
-	UFUNCTION(BlueprintPure = false, Category = "LearningAgents", meta = (ReturnDisplayName = "Success"))
-	bool GetStaticArrayObservation(TArray<FLearningAgentsObservationObjectElement>& OutElements, const FLearningAgentsObservationObjectElement Element, const FName Name = TEXT("StaticArray")) const;
-	bool GetStaticArrayObservationToArrayView(TArrayView<FLearningAgentsObservationObjectElement> OutElements, const FLearningAgentsObservationObjectElement Element, const FName Name = TEXT("StaticArray")) const;
+	UFUNCTION(BlueprintPure = false, Category = "LearningAgents", meta = (AdvancedDisplay = 4, ReturnDisplayName = "Success"))
+	static bool GetInclusiveUnionObservationToArrays(TArray<FName>& OutElementNames, TArray<FLearningAgentsObservationObjectElement>& OutElements, const ULearningAgentsObservationObject* Object, const FLearningAgentsObservationObjectElement Element, const FName Tag = TEXT("InclusiveUnion"));
+	static bool GetInclusiveUnionObservationToArrayViews(TArrayView<FName> OutElementNames, TArrayView<FLearningAgentsObservationObjectElement> OutElements, const ULearningAgentsObservationObject* Object, const FLearningAgentsObservationObjectElement Element, const FName Tag = TEXT("InclusiveUnion"));
 
-	UFUNCTION(BlueprintPure = false, Category = "LearningAgents", meta = (ReturnDisplayName = "Success"))
-	bool GetSetObservationNum(int32& OutNum, const FLearningAgentsObservationObjectElement Element, const FName Name = TEXT("Set")) const;
+	UFUNCTION(BlueprintPure = false, Category = "LearningAgents", meta = (AdvancedDisplay = 3, ReturnDisplayName = "Success"))
+	static bool GetStaticArrayObservationNum(int32& OutNum, const ULearningAgentsObservationObject* Object, const FLearningAgentsObservationObjectElement Element, const FName Tag = TEXT("StaticArray"));
 
-	UFUNCTION(BlueprintPure = false, Category = "LearningAgents", meta = (ReturnDisplayName = "Success"))
-	bool GetSetObservation(TSet<FLearningAgentsObservationObjectElement>& OutElements, const FLearningAgentsObservationObjectElement Element, const FName Name = TEXT("Set")) const;
-	
-	UFUNCTION(BlueprintPure = false, Category = "LearningAgents", meta = (ReturnDisplayName = "Success"))
-	bool GetSetObservationToArray(TArray<FLearningAgentsObservationObjectElement>& OutElements, const FLearningAgentsObservationObjectElement Element, const FName Name = TEXT("Set")) const;
-	bool GetSetObservationToArrayView(TArrayView<FLearningAgentsObservationObjectElement> OutElements, const FLearningAgentsObservationObjectElement Element, const FName Name = TEXT("Set")) const;
+	UFUNCTION(BlueprintPure = false, Category = "LearningAgents", meta = (AdvancedDisplay = 3, ReturnDisplayName = "Success"))
+	static bool GetStaticArrayObservation(TArray<FLearningAgentsObservationObjectElement>& OutElements, const ULearningAgentsObservationObject* Object, const FLearningAgentsObservationObjectElement Element, const FName Tag = TEXT("StaticArray"));
+	static bool GetStaticArrayObservationToArrayView(TArrayView<FLearningAgentsObservationObjectElement> OutElements, const ULearningAgentsObservationObject* Object, const FLearningAgentsObservationObjectElement Element, const FName Tag = TEXT("StaticArray"));
 
-	UFUNCTION(BlueprintPure = false, Category = "LearningAgents", meta = (ReturnDisplayName = "Success"))
-	bool GetPairObservation(FLearningAgentsObservationObjectElement& OutKey, FLearningAgentsObservationObjectElement& OutValue, const FLearningAgentsObservationObjectElement Element, const FName Name = TEXT("Pair")) const;
+	UFUNCTION(BlueprintPure = false, Category = "LearningAgents", meta = (AdvancedDisplay = 3, ReturnDisplayName = "Success"))
+	static bool GetSetObservationNum(int32& OutNum, const ULearningAgentsObservationObject* Object, const FLearningAgentsObservationObjectElement Element, const FName Tag = TEXT("Set"));
 
-	UFUNCTION(BlueprintPure = false, Category = "LearningAgents", meta = (ReturnDisplayName = "Success"))
-	bool GetArrayObservationNum(int32& OutNum, const FLearningAgentsObservationObjectElement Element, const FName Name = TEXT("Array")) const;
+	UFUNCTION(BlueprintPure = false, Category = "LearningAgents", meta = (AdvancedDisplay = 3, ReturnDisplayName = "Success"))
+	static bool GetSetObservation(TSet<FLearningAgentsObservationObjectElement>& OutElements, const ULearningAgentsObservationObject* Object, const FLearningAgentsObservationObjectElement Element, const FName Tag = TEXT("Set"));
 
-	UFUNCTION(BlueprintPure = false, Category = "LearningAgents", meta = (ReturnDisplayName = "Success"))
-	bool GetArrayObservation(TArray<FLearningAgentsObservationObjectElement>& OutElements, const FLearningAgentsObservationObjectElement Element, const FName Name = TEXT("Array")) const;
-	bool GetArrayObservationToArrayView(TArrayView<FLearningAgentsObservationObjectElement> OutElements, const FLearningAgentsObservationObjectElement Element, const FName Name = TEXT("Array")) const;
+	UFUNCTION(BlueprintPure = false, Category = "LearningAgents", meta = (AdvancedDisplay = 3, ReturnDisplayName = "Success"))
+	static bool GetSetObservationToArray(TArray<FLearningAgentsObservationObjectElement>& OutElements, const ULearningAgentsObservationObject* Object, const FLearningAgentsObservationObjectElement Element, const FName Tag = TEXT("Set"));
+	static bool GetSetObservationToArrayView(TArrayView<FLearningAgentsObservationObjectElement> OutElements, const ULearningAgentsObservationObject* Object, const FLearningAgentsObservationObjectElement Element, const FName Tag = TEXT("Set"));
 
-	UFUNCTION(BlueprintPure = false, Category = "LearningAgents", meta = (ReturnDisplayName = "Success"))
-	bool GetMapObservationNum(int32& OutNum, const FLearningAgentsObservationObjectElement Element, const FName Name = TEXT("Map")) const;
+	UFUNCTION(BlueprintPure = false, Category = "LearningAgents", meta = (AdvancedDisplay = 4, ReturnDisplayName = "Success"))
+	static bool GetPairObservation(FLearningAgentsObservationObjectElement& OutKey, FLearningAgentsObservationObjectElement& OutValue, const ULearningAgentsObservationObject* Object, const FLearningAgentsObservationObjectElement Element, const FName Tag = TEXT("Pair"));
 
-	UFUNCTION(BlueprintPure = false, Category = "LearningAgents", meta = (ReturnDisplayName = "Success"))
-	bool GetMapObservation(TMap<FLearningAgentsObservationObjectElement,FLearningAgentsObservationObjectElement>& OutElements, const FLearningAgentsObservationObjectElement Element, const FName Name = TEXT("Map")) const;
-	
-	UFUNCTION(BlueprintPure = false, Category = "LearningAgents", meta = (ReturnDisplayName = "Success"))
-	bool GetMapObservationToArrays(TArray<FLearningAgentsObservationObjectElement>& OutKeys, TArray<FLearningAgentsObservationObjectElement>& OutValues, const FLearningAgentsObservationObjectElement Element, const FName Name = TEXT("Map")) const;
-	bool GetMapObservationToArrayViews(TArrayView<FLearningAgentsObservationObjectElement> OutKeys, TArrayView<FLearningAgentsObservationObjectElement> OutValues, const FLearningAgentsObservationObjectElement Element, const FName Name = TEXT("Map")) const;
+	UFUNCTION(BlueprintPure = false, Category = "LearningAgents", meta = (AdvancedDisplay = 3, ReturnDisplayName = "Success"))
+	static bool GetArrayObservationNum(int32& OutNum, const ULearningAgentsObservationObject* Object, const FLearningAgentsObservationObjectElement Element, const FName Tag = TEXT("Array"));
 
-	UFUNCTION(BlueprintPure = false, Category = "LearningAgents", meta = (ReturnDisplayName = "Success"))
-	bool GetEnumObservation(uint8& OutEnumValue, const UEnum* Enum, const FLearningAgentsObservationObjectElement Element, const FName Name = TEXT("Enum")) const;
+	UFUNCTION(BlueprintPure = false, Category = "LearningAgents", meta = (AdvancedDisplay = 3, ReturnDisplayName = "Success"))
+	static bool GetArrayObservation(TArray<FLearningAgentsObservationObjectElement>& OutElements, const ULearningAgentsObservationObject* Object, const FLearningAgentsObservationObjectElement Element, const FName Tag = TEXT("Array"));
+	static bool GetArrayObservationToArrayView(TArrayView<FLearningAgentsObservationObjectElement> OutElements, const ULearningAgentsObservationObject* Object, const FLearningAgentsObservationObjectElement Element, const FName Tag = TEXT("Array"));
 
-	UFUNCTION(BlueprintPure = false, Category = "LearningAgents", meta = (ReturnDisplayName = "Success"))
-	bool GetBitmaskObservation(int32& OutBitmaskValue, const UEnum* Enum, const FLearningAgentsObservationObjectElement Element, const FName Name = TEXT("Bitmask")) const;
+	UFUNCTION(BlueprintPure = false, Category = "LearningAgents", meta = (AdvancedDisplay = 3, ReturnDisplayName = "Success"))
+	static bool GetMapObservationNum(int32& OutNum, const ULearningAgentsObservationObject* Object, const FLearningAgentsObservationObjectElement Element, const FName Tag = TEXT("Map"));
 
-	UFUNCTION(BlueprintPure = false, Category = "LearningAgents", meta = (ExpandEnumAsExecs = "OutOption", ReturnDisplayName = "Success"))
-	bool GetOptionalObservation(ELearningAgentsOptionalObservation& OutOption, FLearningAgentsObservationObjectElement& OutElement, const FLearningAgentsObservationObjectElement Element, const FName Name = TEXT("Optional")) const;
+	UFUNCTION(BlueprintPure = false, Category = "LearningAgents", meta = (AdvancedDisplay = 3, ReturnDisplayName = "Success"))
+	static bool GetMapObservation(TMap<FLearningAgentsObservationObjectElement, FLearningAgentsObservationObjectElement>& OutElements, const ULearningAgentsObservationObject* Object, const FLearningAgentsObservationObjectElement Element, const FName Tag = TEXT("Map"));
 
-	UFUNCTION(BlueprintPure = false, Category = "LearningAgents", meta = (ExpandEnumAsExecs = "OutEither", ReturnDisplayName = "Success"))
-	bool GetEitherObservation(ELearningAgentsEitherObservation& OutEither, FLearningAgentsObservationObjectElement& OutElement, const FLearningAgentsObservationObjectElement Element, const FName Name = TEXT("Either")) const;
+	UFUNCTION(BlueprintPure = false, Category = "LearningAgents", meta = (AdvancedDisplay = 4, ReturnDisplayName = "Success"))
+	static bool GetMapObservationToArrays(TArray<FLearningAgentsObservationObjectElement>& OutKeys, TArray<FLearningAgentsObservationObjectElement>& OutValues, const ULearningAgentsObservationObject* Object, const FLearningAgentsObservationObjectElement Element, const FName Tag = TEXT("Map"));
+	static bool GetMapObservationToArrayViews(TArrayView<FLearningAgentsObservationObjectElement> OutKeys, TArrayView<FLearningAgentsObservationObjectElement> OutValues, const ULearningAgentsObservationObject* Object, const FLearningAgentsObservationObjectElement Element, const FName Tag = TEXT("Map"));
 
-	UFUNCTION(BlueprintPure = false, Category = "LearningAgents", meta = (ReturnDisplayName = "Success"))
-	bool GetEncodingObservation(FLearningAgentsObservationObjectElement& OutElement, const FLearningAgentsObservationObjectElement Element, const FName Name = TEXT("Encoding")) const;
+	UFUNCTION(BlueprintPure = false, Category = "LearningAgents", meta = (AdvancedDisplay = 4, ReturnDisplayName = "Success"))
+	static bool GetEnumObservation(uint8& OutEnumValue, const ULearningAgentsObservationObject* Object, const FLearningAgentsObservationObjectElement Element, const UEnum* Enum, const FName Tag = TEXT("Enum"));
 
-	UFUNCTION(BlueprintPure = false, Category = "LearningAgents", meta = (ReturnDisplayName = "Success"))
-	bool GetBoolObservation(bool& bOutValue, const FLearningAgentsObservationObjectElement Element, const FName Name = TEXT("Bool")) const;
+	UFUNCTION(BlueprintPure = false, Category = "LearningAgents", meta = (AdvancedDisplay = 4, ReturnDisplayName = "Success"))
+	static bool GetBitmaskObservation(int32& OutBitmaskValue, const ULearningAgentsObservationObject* Object, const FLearningAgentsObservationObjectElement Element, const UEnum* Enum, const FName Tag = TEXT("Bitmask"));
 
-	UFUNCTION(BlueprintPure = false, Category = "LearningAgents", meta = (ReturnDisplayName = "Success"))
-	bool GetFloatObservation(float& OutValue, const FLearningAgentsObservationObjectElement Element, const float FloatScale = 1.0f, const FName Name = TEXT("Float")) const;
+	UFUNCTION(BlueprintPure = false, Category = "LearningAgents", meta = (AdvancedDisplay = 3, ExpandEnumAsExecs = "OutOption", ReturnDisplayName = "Success"))
+	static bool GetOptionalObservation(ELearningAgentsOptionalObservation& OutOption, FLearningAgentsObservationObjectElement& OutElement, const ULearningAgentsObservationObject* Object, const FLearningAgentsObservationObjectElement Element, const FName Tag = TEXT("Optional"));
 
-	UFUNCTION(BlueprintPure = false, Category = "LearningAgents", meta = (ReturnDisplayName = "Success"))
-	bool GetLocationObservation(FVector& OutLocation, const FLearningAgentsObservationObjectElement Element, const FTransform RelativeTransform = FTransform(), const float LocationScale = 100.0f, const FName Name = TEXT("Location")) const;
+	UFUNCTION(BlueprintPure = false, Category = "LearningAgents", meta = (AdvancedDisplay = 3, ExpandEnumAsExecs = "OutEither", ReturnDisplayName = "Success"))
+	static bool GetEitherObservation(ELearningAgentsEitherObservation& OutEither, FLearningAgentsObservationObjectElement& OutElement, const ULearningAgentsObservationObject* Object, const FLearningAgentsObservationObjectElement Element, const FName Tag = TEXT("Either"));
 
-	UFUNCTION(BlueprintPure = false, Category = "LearningAgents", meta = (ReturnDisplayName = "Success"))
-	bool GetRotationObservation(FRotator& OutRotation, const FLearningAgentsObservationObjectElement Element, const FRotator RelativeRotation = FRotator::ZeroRotator, const FName Name = TEXT("Rotation")) const;
-	bool GetRotationObservationAsQuat(FQuat& OutRotation, const FLearningAgentsObservationObjectElement Element, const FQuat RelativeRotation = FQuat::Identity, const FName Name = TEXT("Rotation")) const;
+	UFUNCTION(BlueprintPure = false, Category = "LearningAgents", meta = (AdvancedDisplay = 3, ReturnDisplayName = "Success"))
+	static bool GetEncodingObservation(FLearningAgentsObservationObjectElement& OutElement, const ULearningAgentsObservationObject* Object, const FLearningAgentsObservationObjectElement Element, const FName Tag = TEXT("Encoding"));
 
-	UFUNCTION(BlueprintPure = false, Category = "LearningAgents", meta = (ReturnDisplayName = "Success"))
-	bool GetScaleObservation(FVector& OutScale, const FLearningAgentsObservationObjectElement Element, const FVector RelativeScale = FVector(1, 1, 1), const FName Name = TEXT("Scale")) const;
+	UFUNCTION(BlueprintPure = false, Category = "LearningAgents", meta = (AdvancedDisplay = 3, ReturnDisplayName = "Success"))
+	static bool GetBoolObservation(bool& bOutValue, const ULearningAgentsObservationObject* Object, const FLearningAgentsObservationObjectElement Element, const FName Tag = TEXT("Bool"));
 
-	UFUNCTION(BlueprintPure = false, Category = "LearningAgents", meta = (ReturnDisplayName = "Success"))
-	bool GetTransformObservation(FTransform& OutTransform, const FLearningAgentsObservationObjectElement Element, const FTransform RelativeTransform = FTransform(), const float LocationScale = 100.0f, const FName Name = TEXT("Transform")) const;
+	UFUNCTION(BlueprintPure = false, Category = "LearningAgents", meta = (AdvancedDisplay = 4, ReturnDisplayName = "Success"))
+	static bool GetFloatObservation(float& OutValue, const ULearningAgentsObservationObject* Object, const FLearningAgentsObservationObjectElement Element, const float FloatScale = 1.0f, const FName Tag = TEXT("Float"));
 
-	UFUNCTION(BlueprintPure = false, Category = "LearningAgents", meta = (ReturnDisplayName = "Success"))
-	bool GetAngleObservation(float& OutAngle, const FLearningAgentsObservationObjectElement Element, const float RelativeAngle = 0.0f, const FName Name = TEXT("Angle")) const;
+	UFUNCTION(BlueprintPure = false, Category = "LearningAgents", meta = (AdvancedDisplay = 5, ReturnDisplayName = "Success"))
+	static bool GetLocationObservation(FVector& OutLocation, const ULearningAgentsObservationObject* Object, const FLearningAgentsObservationObjectElement Element, const FTransform RelativeTransform = FTransform(), const float LocationScale = 100.0f, const FName Tag = TEXT("Location"));
 
-	UFUNCTION(BlueprintPure = false, Category = "LearningAgents", meta = (ReturnDisplayName = "Success"))
-	bool GetAngleObservationRadians(float& OutAngle, const FLearningAgentsObservationObjectElement Element, const float RelativeAngle = 0.0f, const FName Name = TEXT("Angle")) const;
+	UFUNCTION(BlueprintPure = false, Category = "LearningAgents", meta = (AdvancedDisplay = 4, ReturnDisplayName = "Success"))
+	static bool GetRotationObservation(FRotator& OutRotation, const ULearningAgentsObservationObject* Object, const FLearningAgentsObservationObjectElement Element, const FRotator RelativeRotation = FRotator::ZeroRotator, const FName Tag = TEXT("Rotation"));
 
-	UFUNCTION(BlueprintPure = false, Category = "LearningAgents", meta = (ReturnDisplayName = "Success"))
-	bool GetVelocityObservation(FVector& OutVelocity, const FLearningAgentsObservationObjectElement Element, const FTransform RelativeTransform = FTransform(), const float VelocityScale = 200.0f, const FName Name = TEXT("Velocity")) const;
+	UFUNCTION(BlueprintPure = false, Category = "LearningAgents", meta = (AdvancedDisplay = 4, ReturnDisplayName = "Success"))
+	static bool GetRotationObservationAsQuat(FQuat& OutRotation, const ULearningAgentsObservationObject* Object, const FLearningAgentsObservationObjectElement Element, const FQuat RelativeRotation, const FName Tag = TEXT("Rotation"));
 
-	UFUNCTION(BlueprintPure = false, Category = "LearningAgents", meta = (ReturnDisplayName = "Success"))
-	bool GetDirectionObservation(FVector& OutDirection, const FLearningAgentsObservationObjectElement Element, const FTransform RelativeTransform = FTransform(), const FName Name = TEXT("Direction")) const;
+	UFUNCTION(BlueprintPure = false, Category = "LearningAgents", meta = (AdvancedDisplay = 4, ReturnDisplayName = "Success"))
+	static bool GetScaleObservation(FVector& OutScale, const ULearningAgentsObservationObject* Object, const FLearningAgentsObservationObjectElement Element, const FVector RelativeScale = FVector(1, 1, 1), const FName Tag = TEXT("Scale"));
 
-	// Spline Observations
+	UFUNCTION(BlueprintPure = false, Category = "LearningAgents", meta = (AdvancedDisplay = 5, ReturnDisplayName = "Success"))
+	static bool GetTransformObservation(FTransform& OutTransform, const ULearningAgentsObservationObject* Object, const FLearningAgentsObservationObjectElement Element, const FTransform RelativeTransform = FTransform(), const float LocationScale = 100.0f, const FName Tag = TEXT("Transform"));
 
-	UFUNCTION(BlueprintPure = false, Category = "LearningAgents", meta = (ReturnDisplayName = "Success"))
-	bool GetLocationAlongSplineObservation(FVector& OutLocation, const FLearningAgentsObservationObjectElement Element, const FTransform RelativeTransform = FTransform(), const float LocationScale = 100.0f, const FName Name = TEXT("LocationAlongSpline"));
+	UFUNCTION(BlueprintPure = false, Category = "LearningAgents", meta = (AdvancedDisplay = 4, ReturnDisplayName = "Success"))
+	static bool GetAngleObservation(float& OutAngle, const ULearningAgentsObservationObject* Object, const FLearningAgentsObservationObjectElement Element, const float RelativeAngle = 0.0f, const FName Tag = TEXT("Angle"));
 
-	UFUNCTION(BlueprintPure = false, Category = "LearningAgents", meta = (ReturnDisplayName = "Success"))
-	bool GetProportionAlongSplineObservation(bool& bOutIsClosedLoop, float& OutAngle, float& OutPropotion, const FLearningAgentsObservationObjectElement Element, const FName Name = TEXT("ProportionAlongSpline"));
+	UFUNCTION(BlueprintPure = false, Category = "LearningAgents", meta = (AdvancedDisplay = 4, ReturnDisplayName = "Success"))
+	static bool GetAngleObservationRadians(float& OutAngle, const ULearningAgentsObservationObject* Object, const FLearningAgentsObservationObjectElement Element, const float RelativeAngle = 0.0f, const FName Tag = TEXT("Angle"));
 
-	UFUNCTION(BlueprintPure = false, Category = "LearningAgents", meta = (ReturnDisplayName = "Success"))
-	bool GetDirectionAlongSplineObservation(FVector& OutDirection, const FLearningAgentsObservationObjectElement Element, const FTransform RelativeTransform = FTransform(), const FName Name = TEXT("DirectionAlongSpline"));
+	UFUNCTION(BlueprintPure = false, Category = "LearningAgents", meta = (AdvancedDisplay = 5, ReturnDisplayName = "Success"))
+	static bool GetVelocityObservation(FVector& OutVelocity, const ULearningAgentsObservationObject* Object, const FLearningAgentsObservationObjectElement Element, const FTransform RelativeTransform = FTransform(), const float VelocityScale = 200.0f, const FName Tag = TEXT("Velocity"));
 
-	UFUNCTION(BlueprintPure = false, Category = "LearningAgents", meta = (ReturnDisplayName = "Success"))
-	bool GetPropertiesAlongSplineObservation(FVector& OutLocation, bool& bOutIsClosedLoop, float& OutAngle, float& OutPropotion, FVector& OutDirection, const FLearningAgentsObservationObjectElement Element, const FTransform RelativeTransform = FTransform(), const float LocationScale = 100.0f, const FName Name = TEXT("PropertiesAlongSpline"));
+	UFUNCTION(BlueprintPure = false, Category = "LearningAgents", meta = (AdvancedDisplay = 4, ReturnDisplayName = "Success"))
+	static bool GetDirectionObservation(FVector& OutDirection, const ULearningAgentsObservationObject* Object, const FLearningAgentsObservationObjectElement Element, const FTransform RelativeTransform = FTransform(), const FName Tag = TEXT("Direction"));
 
-	// Ray Cast Observations
+	// Get Spline Observations
 
-	UFUNCTION(BlueprintPure, Category = "LearningAgents", meta = (ReturnDisplayName = "Success"))
-	bool GetProportionAlongRayObservation(float& OutProportion, const FLearningAgentsObservationObjectElement Element, const FName Name = TEXT("ProportionAlongRay"));
+	UFUNCTION(BlueprintPure = false, Category = "LearningAgents", meta = (AdvancedDisplay = 5, ReturnDisplayName = "Success"))
+	static bool GetLocationAlongSplineObservation(FVector& OutLocation, const ULearningAgentsObservationObject* Object, const FLearningAgentsObservationObjectElement Element, const FTransform RelativeTransform = FTransform(), const float LocationScale = 100.0f, const FName Tag = TEXT("LocationAlongSpline"));
 
-	UFUNCTION(BlueprintPure, Category = "LearningAgents", meta = (ReturnDisplayName = "Success"))
-	bool GetProportionAlongRaysObservationNum(int32& OutProportionNum, const FLearningAgentsObservationObjectElement Element, const FName Name = TEXT("ProportionAlongRays"));
+	UFUNCTION(BlueprintPure = false, Category = "LearningAgents", meta = (AdvancedDisplay = 5, ReturnDisplayName = "Success"))
+	static bool GetProportionAlongSplineObservation(bool& bOutIsClosedLoop, float& OutAngle, float& OutPropotion, const ULearningAgentsObservationObject* Object, const FLearningAgentsObservationObjectElement Element, const FName Tag = TEXT("ProportionAlongSpline"));
 
-	UFUNCTION(BlueprintPure, Category = "LearningAgents", meta = (ReturnDisplayName = "Success"))
-	bool GetProportionAlongRaysObservation(TArray<float>& OutProportions, const FLearningAgentsObservationObjectElement Element, const FName Name = TEXT("ProportionAlongRays"));
-	bool GetProportionAlongRaysObservationToArrayView(TArrayView<float> OutProportions, const FLearningAgentsObservationObjectElement Element, const FName Name = TEXT("ProportionAlongRays"));
+	UFUNCTION(BlueprintPure = false, Category = "LearningAgents", meta = (AdvancedDisplay = 4, ReturnDisplayName = "Success"))
+	static bool GetDirectionAlongSplineObservation(FVector& OutDirection, const ULearningAgentsObservationObject* Object, const FLearningAgentsObservationObjectElement Element, const FTransform RelativeTransform = FTransform(), const FName Tag = TEXT("DirectionAlongSpline"));
 
+	UFUNCTION(BlueprintPure = false, Category = "LearningAgents", meta = (AdvancedDisplay = 9, ReturnDisplayName = "Success"))
+	static bool GetPropertiesAlongSplineObservation(FVector& OutLocation, bool& bOutIsClosedLoop, float& OutAngle, float& OutPropotion, FVector& OutDirection, const ULearningAgentsObservationObject* Object, const FLearningAgentsObservationObjectElement Element, const FTransform RelativeTransform = FTransform(), const float LocationScale = 100.0f, const FName Tag = TEXT("PropertiesAlongSpline"));
 
-private:
+	// Get Ray Cast Observations
+
+	UFUNCTION(BlueprintPure, Category = "LearningAgents", meta = (AdvancedDisplay = 3, ReturnDisplayName = "Success"))
+	static bool GetProportionAlongRayObservation(float& OutProportion, const ULearningAgentsObservationObject* Object, const FLearningAgentsObservationObjectElement Element, const FName Tag = TEXT("ProportionAlongRay"));
+
+	UFUNCTION(BlueprintPure, Category = "LearningAgents", meta = (AdvancedDisplay = 3, ReturnDisplayName = "Success"))
+	static bool GetProportionAlongRaysObservationNum(int32& OutProportionNum, const ULearningAgentsObservationObject* Object, const FLearningAgentsObservationObjectElement Element, const FName Tag = TEXT("ProportionAlongRays"));
+
+	UFUNCTION(BlueprintPure, Category = "LearningAgents", meta = (AdvancedDisplay = 3, ReturnDisplayName = "Success"))
+	static bool GetProportionAlongRaysObservation(TArray<float>& OutProportions, const ULearningAgentsObservationObject* Object, const FLearningAgentsObservationObjectElement Element, const FName Tag = TEXT("ProportionAlongRays"));
+	static bool GetProportionAlongRaysObservationToArrayView(TArrayView<float> OutProportions, const ULearningAgentsObservationObject* Object, const FLearningAgentsObservationObjectElement Element, const FName Tag = TEXT("ProportionAlongRays"));
 
-	UE::Learning::Observation::FObject ObservationObject;
 };
+
+
