LearningAgents: Refactored to make Observation and Action functions static and use "AdvancedDisplay" markup. Added Visual Logging and comments for Completions.

diff --git a/Engine/Plugins/Experimental/LearningAgents/Source/LearningAgentsTraining/Private/LearningAgentsCompletions.cpp b/Engine/Plugins/Experimental/LearningAgents/Source/LearningAgentsTraining/Private/LearningAgentsCompletions.cpp
--- a/Engine/Plugins/Experimental/LearningAgents/Source/LearningAgentsTraining/Private/LearningAgentsCompletions.cpp
+++ b/Engine/Plugins/Experimental/LearningAgents/Source/LearningAgentsTraining/Private/LearningAgentsCompletions.cpp
@@ -2,6 +2,9 @@
 
 #include "LearningAgentsCompletions.h"
 
+#include "LearningAgentsManagerListener.h"
+#include "LearningAgentsDebug.h"
+
 #include "LearningCompletion.h"
 #include "LearningLog.h"
 
@@ -14,7 +17,7 @@ namespace UE::Learning::Agents
 		case ECompletionMode::Running: return ELearningAgentsCompletion::Running;
 		case ECompletionMode::Terminated: return ELearningAgentsCompletion::Termination;
 		case ECompletionMode::Truncated: return ELearningAgentsCompletion::Truncation;
-		default:UE_LOG(LogLearning, Error, TEXT("Unknown Completion Mode.")); return ELearningAgentsCompletion::Running;
+		default: UE_LOG(LogLearning, Error, TEXT("Unknown Completion Mode.")); return ELearningAgentsCompletion::Running;
 		}
 	}
 
@@ -25,31 +28,47 @@ namespace UE::Learning::Agents
 		case ELearningAgentsCompletion::Running: return ECompletionMode::Running;
 		case ELearningAgentsCompletion::Termination: return ECompletionMode::Terminated;
 		case ELearningAgentsCompletion::Truncation: return ECompletionMode::Truncated;
-		default:UE_LOG(LogLearning, Error, TEXT("Unknown Completion.")); return ECompletionMode::Running;
+		default: UE_LOG(LogLearning, Error, TEXT("Unknown Completion.")); return ECompletionMode::Running;
+		}
+	}
+
+	namespace Private
+	{
+		const TCHAR* GetCompletionName(const ELearningAgentsCompletion Completion)
+		{
+			switch (Completion)
+			{
+			case ELearningAgentsCompletion::Running: return TEXT("Running");
+			case ELearningAgentsCompletion::Termination: return TEXT("Termination");
+			case ELearningAgentsCompletion::Truncation: return TEXT("Truncation");
+			default: return TEXT("Error");
+			}
 		}
 	}
 }
 
-bool ULearningAgentsCompletions::CompletionIsRunning(const ELearningAgentsCompletion Completion)
+
+bool ULearningAgentsCompletions::IsCompletionRunning(const ELearningAgentsCompletion Completion)
 {
 	return Completion == ELearningAgentsCompletion::Running;
 }
 
-bool ULearningAgentsCompletions::CompletionIsCompleted(const ELearningAgentsCompletion Completion)
+bool ULearningAgentsCompletions::IsCompletionCompleted(const ELearningAgentsCompletion Completion)
 {
 	return Completion == ELearningAgentsCompletion::Truncation || Completion == ELearningAgentsCompletion::Termination;
 }
 
-bool ULearningAgentsCompletions::CompletionIsTruncation(const ELearningAgentsCompletion Completion)
+bool ULearningAgentsCompletions::IsCompletionTruncation(const ELearningAgentsCompletion Completion)
 {
 	return Completion == ELearningAgentsCompletion::Truncation;
 }
 
-bool ULearningAgentsCompletions::CompletionIsTermination(const ELearningAgentsCompletion Completion)
+bool ULearningAgentsCompletions::IsCompletionTermination(const ELearningAgentsCompletion Completion)
 {
 	return Completion == ELearningAgentsCompletion::Termination;
 }
 
+
 ELearningAgentsCompletion ULearningAgentsCompletions::CompletionOr(ELearningAgentsCompletion A, ELearningAgentsCompletion B)
 {
 	if ((A == ELearningAgentsCompletion::Running && B != ELearningAgentsCompletion::Running) ||
@@ -76,49 +95,309 @@ ELearningAgentsCompletion ULearningAgentsCompletions::CompletionAnd(ELearningAge
 	}
 }
 
-ELearningAgentsCompletion ULearningAgentsCompletions::CompletionNot(const ELearningAgentsCompletion A, const ELearningAgentsCompletion NotRunningType)
+ELearningAgentsCompletion ULearningAgentsCompletions::CompletionNot(ELearningAgentsCompletion A, ELearningAgentsCompletion NotRunningType)
 {
 	return A == ELearningAgentsCompletion::Running ? NotRunningType : ELearningAgentsCompletion::Running;
 }
 
-ELearningAgentsCompletion ULearningAgentsCompletions::CompletionOnCondition(const bool bCondition, const ELearningAgentsCompletion CompletionType)
+
+ELearningAgentsCompletion ULearningAgentsCompletions::MakeCompletion(
+	const ELearningAgentsCompletion CompletionType,
+	const FName Tag,
+	const bool bVisualLoggerEnabled,
+	ULearningAgentsManagerListener* VisualLoggerListener,
+	const int32 VisualLoggerAgentId,
+	const FVector VisualLoggerLocation,
+	const FLinearColor VisualLoggerColor)
+{
+
+#if UE_LEARNING_AGENTS_ENABLE_VISUAL_LOG
+	if (bVisualLoggerEnabled && VisualLoggerListener)
+	{
+		const ULearningAgentsVisualLoggerObject* VisualLoggerObject = VisualLoggerListener->GetOrAddVisualLoggerObject(Tag);
+
+		UE_LEARNING_AGENTS_VLOG_STRING(VisualLoggerObject, LogLearning, Display, VisualLoggerLocation,
+			VisualLoggerColor.ToFColor(true),
+			TEXT("Listener: %s\nTag: %s\nAgent Id: % 3i\nCompletion: [%s]"),
+			*VisualLoggerListener->GetName(),
+			*Tag.ToString(),
+			VisualLoggerAgentId,
+			UE::Learning::Agents::Private::GetCompletionName(CompletionType));
+	}
+#endif
+
+	return CompletionType;
+}
+
+ELearningAgentsCompletion ULearningAgentsCompletions::MakeCompletionOnCondition(
+	const bool bCondition, 
+	const ELearningAgentsCompletion CompletionType,
+	const FName Tag,
+	const bool bVisualLoggerEnabled,
+	ULearningAgentsManagerListener* VisualLoggerListener,
+	const int32 VisualLoggerAgentId,
+	const FVector VisualLoggerLocation,
+	const FLinearColor VisualLoggerColor)
 {
-	return bCondition ? CompletionType : ELearningAgentsCompletion::Running;
+	const ELearningAgentsCompletion Completion = bCondition ? CompletionType : ELearningAgentsCompletion::Running;
+
+#if UE_LEARNING_AGENTS_ENABLE_VISUAL_LOG
+	if (bVisualLoggerEnabled && VisualLoggerListener)
+	{
+		const ULearningAgentsVisualLoggerObject* VisualLoggerObject = VisualLoggerListener->GetOrAddVisualLoggerObject(Tag);
+
+		UE_LEARNING_AGENTS_VLOG_STRING(VisualLoggerObject, LogLearning, Display, VisualLoggerLocation,
+			VisualLoggerColor.ToFColor(true),
+			TEXT("Listener: %s\nTag: %s\nAgent Id: % 3i\nCondition: [%s]\nCompletion Type: [%s]\nCompletion: [%s]"),
+			*VisualLoggerListener->GetName(),
+			*Tag.ToString(),
+			VisualLoggerAgentId,
+			bCondition ? TEXT("true") : TEXT("false"),
+			UE::Learning::Agents::Private::GetCompletionName(CompletionType),
+			UE::Learning::Agents::Private::GetCompletionName(Completion));
+	}
+#endif
+
+	return Completion;
 }
 
-ELearningAgentsCompletion ULearningAgentsCompletions::CompletionOnTimeElapsed(const float Time, const float TimeThreshold, const ELearningAgentsCompletion CompletionType)
+ELearningAgentsCompletion ULearningAgentsCompletions::MakeCompletionOnTimeElapsed(
+	const float Time, 
+	const float TimeThreshold, 
+	const ELearningAgentsCompletion CompletionType,
+	const FName Tag,
+	const bool bVisualLoggerEnabled,
+	ULearningAgentsManagerListener* VisualLoggerListener,
+	const int32 VisualLoggerAgentId,
+	const FVector VisualLoggerLocation,
+	const FLinearColor VisualLoggerColor)
 {
-	return CompletionOnCondition(Time > TimeThreshold, CompletionType);
+	const bool bCondition = Time > TimeThreshold;
+	const ELearningAgentsCompletion Completion = MakeCompletionOnCondition(bCondition, CompletionType);
+
+#if UE_LEARNING_AGENTS_ENABLE_VISUAL_LOG
+	if (bVisualLoggerEnabled && VisualLoggerListener)
+	{
+		const ULearningAgentsVisualLoggerObject* VisualLoggerObject = VisualLoggerListener->GetOrAddVisualLoggerObject(Tag);
+
+		UE_LEARNING_AGENTS_VLOG_STRING(VisualLoggerObject, LogLearning, Display, VisualLoggerLocation,
+			VisualLoggerColor.ToFColor(true),
+			TEXT("Listener: %s\nTag: %s\nAgent Id: % 3i\nTime: [% 6.2f]\nThreshold: [% 6.2f]\nCondition: [%s]\nCompletion Type: [%s]\nCompletion: [%s]"),
+			*VisualLoggerListener->GetName(),
+			*Tag.ToString(),
+			VisualLoggerAgentId,
+			Time,
+			TimeThreshold,
+			bCondition ? TEXT("true") : TEXT("false"),
+			UE::Learning::Agents::Private::GetCompletionName(CompletionType),
+			UE::Learning::Agents::Private::GetCompletionName(Completion));
+	}
+#endif
+
+	return Completion;
 }
 
-ELearningAgentsCompletion ULearningAgentsCompletions::CompletionOnEpisodeStepsRecorded(const int32 EpisodeSteps, const int32 MaxEpisodeSteps, const ELearningAgentsCompletion CompletionType)
+ELearningAgentsCompletion ULearningAgentsCompletions::MakeCompletionOnEpisodeStepsRecorded(
+	const int32 EpisodeSteps, 
+	const int32 MaxEpisodeSteps, 
+	const ELearningAgentsCompletion CompletionType,
+	const FName Tag,
+	const bool bVisualLoggerEnabled,
+	ULearningAgentsManagerListener* VisualLoggerListener,
+	const int32 VisualLoggerAgentId,
+	const FVector VisualLoggerLocation,
+	const FLinearColor VisualLoggerColor)
 {
-	return CompletionOnCondition(EpisodeSteps >= MaxEpisodeSteps, CompletionType);
+	const bool bCondition = EpisodeSteps >= MaxEpisodeSteps;
+	const ELearningAgentsCompletion Completion = MakeCompletionOnCondition(bCondition, CompletionType);
+
+#if UE_LEARNING_AGENTS_ENABLE_VISUAL_LOG
+	if (bVisualLoggerEnabled && VisualLoggerListener)
+	{
+		const ULearningAgentsVisualLoggerObject* VisualLoggerObject = VisualLoggerListener->GetOrAddVisualLoggerObject(Tag);
+
+		UE_LEARNING_AGENTS_VLOG_STRING(VisualLoggerObject, LogLearning, Display, VisualLoggerLocation,
+			VisualLoggerColor.ToFColor(true),
+			TEXT("Listener: %s\nTag: %s\nAgent Id: % 3i\nEpisodeSteps: [% 4i]\nMaxEpisodeSteps: [% 4i]\nCondition: [%s]\nCompletion Type: [%s]\nCompletion: [%s]"),
+			*VisualLoggerListener->GetName(),
+			*Tag.ToString(),
+			VisualLoggerAgentId,
+			EpisodeSteps,
+			MaxEpisodeSteps,
+			bCondition ? TEXT("true") : TEXT("false"),
+			UE::Learning::Agents::Private::GetCompletionName(CompletionType),
+			UE::Learning::Agents::Private::GetCompletionName(Completion));
+	}
+#endif
+
+	return Completion;
 }
 
-ELearningAgentsCompletion ULearningAgentsCompletions::CompletionOnLocationDifferenceBelowThreshold(const FVector LocationA, const FVector LocationB, const float DistanceThreshold, const ELearningAgentsCompletion CompletionType)
+ELearningAgentsCompletion ULearningAgentsCompletions::MakeCompletionOnLocationDifferenceBelowThreshold(
+	const FVector LocationA, 
+	const FVector LocationB, 
+	const float DistanceThreshold, 
+	const ELearningAgentsCompletion CompletionType,
+	const FName Tag,
+	const bool bVisualLoggerEnabled,
+	ULearningAgentsManagerListener* VisualLoggerListener,
+	const int32 VisualLoggerAgentId,
+	const FVector VisualLoggerLocation,
+	const FLinearColor VisualLoggerColor)
 {
-	return CompletionOnCondition(FVector::Distance(LocationA, LocationB) < DistanceThreshold, CompletionType);
+	const float Distance = FVector::Distance(LocationA, LocationB);
+	const bool bCondition = Distance < DistanceThreshold;
+	const ELearningAgentsCompletion Completion = MakeCompletionOnCondition(bCondition, CompletionType);
+
+#if UE_LEARNING_AGENTS_ENABLE_VISUAL_LOG
+	if (bVisualLoggerEnabled && VisualLoggerListener)
+	{
+		const ULearningAgentsVisualLoggerObject* VisualLoggerObject = VisualLoggerListener->GetOrAddVisualLoggerObject(Tag);
+
+		UE_LEARNING_AGENTS_VLOG_LOCATION(VisualLoggerObject, LogLearning, Display,
+			LocationA,
+			10,
+			VisualLoggerColor.ToFColor(true),
+			TEXT(""));
+
+		UE_LEARNING_AGENTS_VLOG_SEGMENT(VisualLoggerObject, LogLearning, Display,
+			LocationA,
+			LocationB,
+			VisualLoggerColor.ToFColor(true),
+			TEXT(""));
+
+		UE_LEARNING_AGENTS_VLOG_LOCATION(VisualLoggerObject, LogLearning, Display,
+			LocationB,
+			10,
+			VisualLoggerColor.ToFColor(true),
+			TEXT(""));
+
+		UE_LEARNING_AGENTS_VLOG_STRING(VisualLoggerObject, LogLearning, Display, VisualLoggerLocation,
+			VisualLoggerColor.ToFColor(true),
+			TEXT("Listener: %s\nTag: %s\nAgent Id: % 3i\nLocationA: [% 6.1f % 6.1f % 6.1f]\nLocationB: [% 6.1f % 6.1f % 6.1f]\nDistance: [% 6.2f]\nThreshold: [% 6.2f]\nCondition: [%s]\nCompletion Type: [%s]\nCompletion: [%s]"),
+			*VisualLoggerListener->GetName(),
+			*Tag.ToString(),
+			VisualLoggerAgentId,
+			LocationA.X, LocationA.Y, LocationA.Z,
+			LocationB.X, LocationB.Y, LocationB.Z,
+			Distance,
+			DistanceThreshold,
+			bCondition ? TEXT("true") : TEXT("false"),
+			UE::Learning::Agents::Private::GetCompletionName(CompletionType),
+			UE::Learning::Agents::Private::GetCompletionName(Completion));
+	}
+#endif
+
+	return Completion;
 }
 
-ELearningAgentsCompletion ULearningAgentsCompletions::CompletionOnLocationDifferenceAboveThreshold(const FVector LocationA, const FVector LocationB, const float DistanceThreshold, const ELearningAgentsCompletion CompletionType)
+ELearningAgentsCompletion ULearningAgentsCompletions::MakeCompletionOnLocationDifferenceAboveThreshold(
+	const FVector LocationA, 
+	const FVector LocationB, 
+	const float DistanceThreshold, 
+	const ELearningAgentsCompletion CompletionType,
+	const FName Tag,
+	const bool bVisualLoggerEnabled,
+	ULearningAgentsManagerListener* VisualLoggerListener,
+	const int32 VisualLoggerAgentId,
+	const FVector VisualLoggerLocation,
+	const FLinearColor VisualLoggerColor)
 {
-	return CompletionOnCondition(FVector::Distance(LocationA, LocationB) > DistanceThreshold, CompletionType);
+	const float Distance = FVector::Distance(LocationA, LocationB);
+	const bool bCondition = Distance > DistanceThreshold;
+	const ELearningAgentsCompletion Completion = MakeCompletionOnCondition(bCondition, CompletionType);
+
+#if UE_LEARNING_AGENTS_ENABLE_VISUAL_LOG
+	if (bVisualLoggerEnabled && VisualLoggerListener)
+	{
+		const ULearningAgentsVisualLoggerObject* VisualLoggerObject = VisualLoggerListener->GetOrAddVisualLoggerObject(Tag);
+
+		UE_LEARNING_AGENTS_VLOG_LOCATION(VisualLoggerObject, LogLearning, Display,
+			LocationA,
+			10,
+			VisualLoggerColor.ToFColor(true),
+			TEXT(""));
+
+		UE_LEARNING_AGENTS_VLOG_SEGMENT(VisualLoggerObject, LogLearning, Display,
+			LocationA,
+			LocationB,
+			VisualLoggerColor.ToFColor(true),
+			TEXT(""));
+
+		UE_LEARNING_AGENTS_VLOG_LOCATION(VisualLoggerObject, LogLearning, Display,
+			LocationB,
+			10,
+			VisualLoggerColor.ToFColor(true),
+			TEXT(""));
+
+		UE_LEARNING_AGENTS_VLOG_STRING(VisualLoggerObject, LogLearning, Display, VisualLoggerLocation,
+			VisualLoggerColor.ToFColor(true),
+			TEXT("Listener: %s\nTag: %s\nAgent Id: % 3i\nLocationA: [% 6.1f % 6.1f % 6.1f]\nLocationB: [% 6.1f % 6.1f % 6.1f]\nDistance: [% 6.2f]\nThreshold: [% 6.2f]\nCondition: [%s]\nCompletion Type: [%s]\nCompletion: [%s]"),
+			*VisualLoggerListener->GetName(),
+			*Tag.ToString(),
+			VisualLoggerAgentId,
+			LocationA.X, LocationA.Y, LocationA.Z,
+			LocationB.X, LocationB.Y, LocationB.Z,
+			Distance,
+			DistanceThreshold,
+			bCondition ? TEXT("true") : TEXT("false"),
+			UE::Learning::Agents::Private::GetCompletionName(CompletionType),
+			UE::Learning::Agents::Private::GetCompletionName(Completion));
+	}
+#endif
+
+	return Completion;
 }
 
-ELearningAgentsCompletion ULearningAgentsCompletions::CompletionOnLocationOutsideBounds(
+ELearningAgentsCompletion ULearningAgentsCompletions::MakeCompletionOnLocationOutsideBounds(
 	const FVector Location,
 	const FTransform BoundsTransform,
 	const FVector BoundsMins,
 	const FVector BoundsMaxs,
-	const ELearningAgentsCompletion CompletionType)
+	const ELearningAgentsCompletion CompletionType,
+	const FName Tag,
+	const bool bVisualLoggerEnabled,
+	ULearningAgentsManagerListener* VisualLoggerListener,
+	const int32 VisualLoggerAgentId,
+	const FVector VisualLoggerLocation,
+	const FLinearColor VisualLoggerColor)
 {
 	const FVector LocalLocation = BoundsTransform.InverseTransformPosition(Location);
-
-	return CompletionOnCondition(
+	
+	const bool bCondition = (
 		LocalLocation.X < BoundsMins.X || LocalLocation.X > BoundsMaxs.X ||
 		LocalLocation.Y < BoundsMins.Y || LocalLocation.Y > BoundsMaxs.Y ||
-		LocalLocation.Z < BoundsMins.Z || LocalLocation.Z > BoundsMaxs.Z,
-		CompletionType);
+		LocalLocation.Z < BoundsMins.Z || LocalLocation.Z > BoundsMaxs.Z);
+
+	const ELearningAgentsCompletion Completion = MakeCompletionOnCondition(bCondition, CompletionType);
+
+#if UE_LEARNING_AGENTS_ENABLE_VISUAL_LOG
+	if (bVisualLoggerEnabled && VisualLoggerListener)
+	{
+		const ULearningAgentsVisualLoggerObject* VisualLoggerObject = VisualLoggerListener->GetOrAddVisualLoggerObject(Tag);
+
+		UE_LEARNING_AGENTS_VLOG_OBOX(VisualLoggerObject, LogLearning, Display,
+			FBox(BoundsMins, BoundsMaxs),
+			BoundsTransform.ToMatrixWithScale(),
+			VisualLoggerColor.ToFColor(true),
+			TEXT(""));
+
+		UE_LEARNING_AGENTS_VLOG_STRING(VisualLoggerObject, LogLearning, Display, VisualLoggerLocation,
+			VisualLoggerColor.ToFColor(true),
+			TEXT("Listener: %s\nTag: %s\nAgent Id: % 3i\nLocation: [% 6.1f % 6.1f % 6.1f]\nLocal Location: [% 6.1f % 6.1f % 6.1f]\nBounds Mins: [% 6.1f % 6.1f % 6.1f]\nBounds Maxs: [% 6.1f % 6.1f % 6.1f]\nCondition: [%s]\nCompletion Type: [%s]\nCompletion: [%s]"),
+			*VisualLoggerListener->GetName(),
+			*Tag.ToString(),
+			VisualLoggerAgentId,
+			Location.X, Location.Y, Location.Z,
+			LocalLocation.X, LocalLocation.Y, LocalLocation.Z,
+			BoundsMins.X, BoundsMins.Y, BoundsMins.Z,
+			BoundsMaxs.X, BoundsMaxs.Y, BoundsMaxs.Z,
+			bCondition ? TEXT("true") : TEXT("false"),
+			UE::Learning::Agents::Private::GetCompletionName(CompletionType),
+			UE::Learning::Agents::Private::GetCompletionName(Completion));
+	}
+#endif
+
+	return Completion;
 }
 
diff --git a/Engine/Plugins/Experimental/LearningAgents/Source/LearningAgentsTraining/Private/LearningAgentsRecorder.cpp b/Engine/Plugins/Experimental/LearningAgents/Source/LearningAgentsTraining/Private/LearningAgentsRecorder.cpp
--- a/Engine/Plugins/Experimental/LearningAgents/Source/LearningAgentsTraining/Private/LearningAgentsRecorder.cpp
+++ b/Engine/Plugins/Experimental/LearningAgents/Source/LearningAgentsTraining/Private/LearningAgentsRecorder.cpp
@@ -318,6 +318,30 @@ void ULearningAgentsRecorder::EndRecording()
 	bIsRecording = false;
 }
 
+void ULearningAgentsRecorder::EndRecordingAndDiscard()
+{
+	if (!IsSetup())
+	{
+		UE_LOG(LogLearning, Error, TEXT("%s: Setup not complete."), *GetName());
+		return;
+	}
+
+	if (!IsRecording())
+	{
+		UE_LOG(LogLearning, Error, TEXT("%s: Cannot end recording as we are not currently recording!"), *GetName());
+		return;
+	}
+
+	// Discard Records
+
+	for (const int32 AgentId : Manager->GetAllAgentSet())
+	{
+		RecordBuffers[AgentId].Empty();
+	}
+
+	bIsRecording = false;
+}
+
 void ULearningAgentsRecorder::BeginRecording()
 {
 	if (!IsSetup())
diff --git a/Engine/Plugins/Experimental/LearningAgents/Source/LearningAgentsTraining/Private/LearningAgentsRecording.cpp b/Engine/Plugins/Experimental/LearningAgents/Source/LearningAgentsTraining/Private/LearningAgentsRecording.cpp
--- a/Engine/Plugins/Experimental/LearningAgents/Source/LearningAgentsTraining/Private/LearningAgentsRecording.cpp
+++ b/Engine/Plugins/Experimental/LearningAgents/Source/LearningAgentsTraining/Private/LearningAgentsRecording.cpp
@@ -17,6 +17,12 @@ namespace UE::Learning::Agents::Recording::Private
 	static constexpr int32 VersionNumber = 1;
 }
 
+void ULearningAgentsRecording::ResetRecording()
+{
+	Records.Empty();
+	ForceMarkDirty();
+}
+
 void ULearningAgentsRecording::LoadRecordingFromFile(const FFilePath& File)
 {
 	TArray<uint8> RecordingData;
diff --git a/Engine/Plugins/Experimental/LearningAgents/Source/LearningAgentsTraining/Private/LearningAgentsRewards.cpp b/Engine/Plugins/Experimental/LearningAgents/Source/LearningAgentsTraining/Private/LearningAgentsRewards.cpp
--- a/Engine/Plugins/Experimental/LearningAgents/Source/LearningAgentsTraining/Private/LearningAgentsRewards.cpp
+++ b/Engine/Plugins/Experimental/LearningAgents/Source/LearningAgentsTraining/Private/LearningAgentsRewards.cpp
@@ -12,7 +12,7 @@
 float ULearningAgentsRewards::MakeReward(
 	const float RewardValue,
 	const float RewardScale,
-	const FName Name,
+	const FName Tag,
 	const bool bVisualLoggerEnabled,
 	ULearningAgentsManagerListener* VisualLoggerListener,
 	const int32 VisualLoggerAgentId, 
@@ -24,13 +24,13 @@ float ULearningAgentsRewards::MakeReward(
 #if UE_LEARNING_AGENTS_ENABLE_VISUAL_LOG
 	if (bVisualLoggerEnabled && VisualLoggerListener)
 	{
-		const ULearningAgentsVisualLoggerObject* VisualLoggerObject = VisualLoggerListener->GetOrAddVisualLoggerObject(Name);
+		const ULearningAgentsVisualLoggerObject* VisualLoggerObject = VisualLoggerListener->GetOrAddVisualLoggerObject(Tag);
 
 		UE_LEARNING_AGENTS_VLOG_STRING(VisualLoggerObject, LogLearning, Display, VisualLoggerLocation,
 			VisualLoggerColor.ToFColor(true),
-			TEXT("Listener: %s\nName: %s\nAgent Id: % 3i\nValue: [% 6.1f]\nScale: [% 6.2f]\nReward: [% 6.2f]"),
+			TEXT("Listener: %s\nTag: %s\nAgent Id: % 3i\nValue: [% 6.1f]\nScale: [% 6.2f]\nReward: [% 6.2f]"),
 			*VisualLoggerListener->GetName(),
-			*Name.ToString(),
+			*Tag.ToString(),
 			VisualLoggerAgentId,
 			RewardValue,
 			RewardScale,
@@ -44,7 +44,7 @@ float ULearningAgentsRewards::MakeReward(
 float ULearningAgentsRewards::MakeRewardOnCondition(
 	const bool bCondition, 
 	const float RewardScale,
-	const FName Name,
+	const FName Tag,
 	const bool bVisualLoggerEnabled,
 	ULearningAgentsManagerListener* VisualLoggerListener,
 	const int32 VisualLoggerAgentId,
@@ -56,13 +56,13 @@ float ULearningAgentsRewards::MakeRewardOnCondition(
 #if UE_LEARNING_AGENTS_ENABLE_VISUAL_LOG
 	if (bVisualLoggerEnabled && VisualLoggerListener)
 	{
-		const ULearningAgentsVisualLoggerObject* VisualLoggerObject = VisualLoggerListener->GetOrAddVisualLoggerObject(Name);
+		const ULearningAgentsVisualLoggerObject* VisualLoggerObject = VisualLoggerListener->GetOrAddVisualLoggerObject(Tag);
 
 		UE_LEARNING_AGENTS_VLOG_STRING(VisualLoggerObject, LogLearning, Display, VisualLoggerLocation,
 			VisualLoggerColor.ToFColor(true),
-			TEXT("Listener: %s\nName: %s\nAgent Id: % 3i\nCondition: [%s]\nScale: [% 6.2f]\nReward: [% 6.2f]"),
+			TEXT("Listener: %s\nTag: %s\nAgent Id: % 3i\nCondition: [%s]\nScale: [% 6.2f]\nReward: [% 6.2f]"),
 			*VisualLoggerListener->GetName(),
-			*Name.ToString(),
+			*Tag.ToString(),
 			VisualLoggerAgentId,
 			bCondition ? TEXT("true") : TEXT("false"),
 			RewardScale,
@@ -78,7 +78,7 @@ float ULearningAgentsRewards::MakeRewardOnLocationDifferenceBelowThreshold(
 	const FVector LocationB, 
 	const float DistanceThreshold, 
 	const float RewardScale,
-	const FName Name,
+	const FName Tag,
 	const bool bVisualLoggerEnabled,
 	ULearningAgentsManagerListener* VisualLoggerListener,
 	const int32 VisualLoggerAgentId,
@@ -92,7 +92,7 @@ float ULearningAgentsRewards::MakeRewardOnLocationDifferenceBelowThreshold(
 #if UE_LEARNING_AGENTS_ENABLE_VISUAL_LOG
 	if (bVisualLoggerEnabled && VisualLoggerListener)
 	{
-		const ULearningAgentsVisualLoggerObject* VisualLoggerObject = VisualLoggerListener->GetOrAddVisualLoggerObject(Name);
+		const ULearningAgentsVisualLoggerObject* VisualLoggerObject = VisualLoggerListener->GetOrAddVisualLoggerObject(Tag);
 
 		UE_LEARNING_AGENTS_VLOG_LOCATION(VisualLoggerObject, LogLearning, Display,
 			LocationA,
@@ -114,9 +114,9 @@ float ULearningAgentsRewards::MakeRewardOnLocationDifferenceBelowThreshold(
 
 		UE_LEARNING_AGENTS_VLOG_STRING(VisualLoggerObject, LogLearning, Display, VisualLoggerLocation,
 			VisualLoggerColor.ToFColor(true),
-			TEXT("Listener: %s\nName: %s\nAgent Id: % 3i\nLocationA: [% 6.1f % 6.1f % 6.1f]\nLocationB: [% 6.1f % 6.1f % 6.1f]\nDistance: [% 6.2f]\nThreshold: [% 6.2f]\nCondition: [%s]\nScale: [% 6.2f]\nReward: [% 6.2f]"),
+			TEXT("Listener: %s\nTag: %s\nAgent Id: % 3i\nLocationA: [% 6.1f % 6.1f % 6.1f]\nLocationB: [% 6.1f % 6.1f % 6.1f]\nDistance: [% 6.2f]\nThreshold: [% 6.2f]\nCondition: [%s]\nScale: [% 6.2f]\nReward: [% 6.2f]"),
 			*VisualLoggerListener->GetName(),
-			*Name.ToString(),
+			*Tag.ToString(),
 			VisualLoggerAgentId,
 			LocationA.X, LocationA.Y, LocationA.Z,
 			LocationB.X, LocationB.Y, LocationB.Z,
@@ -136,7 +136,7 @@ float ULearningAgentsRewards::MakeRewardOnLocationDifferenceAboveThreshold(
 	const FVector LocationB,
 	const float DistanceThreshold,
 	const float RewardScale,
-	const FName Name,
+	const FName Tag,
 	const bool bVisualLoggerEnabled,
 	ULearningAgentsManagerListener* VisualLoggerListener,
 	const int32 VisualLoggerAgentId,
@@ -150,7 +150,7 @@ float ULearningAgentsRewards::MakeRewardOnLocationDifferenceAboveThreshold(
 #if UE_LEARNING_AGENTS_ENABLE_VISUAL_LOG
 	if (bVisualLoggerEnabled && VisualLoggerListener)
 	{
-		const ULearningAgentsVisualLoggerObject* VisualLoggerObject = VisualLoggerListener->GetOrAddVisualLoggerObject(Name);
+		const ULearningAgentsVisualLoggerObject* VisualLoggerObject = VisualLoggerListener->GetOrAddVisualLoggerObject(Tag);
 
 		UE_LEARNING_AGENTS_VLOG_LOCATION(VisualLoggerObject, LogLearning, Display,
 			LocationA,
@@ -172,9 +172,9 @@ float ULearningAgentsRewards::MakeRewardOnLocationDifferenceAboveThreshold(
 
 		UE_LEARNING_AGENTS_VLOG_STRING(VisualLoggerObject, LogLearning, Display, VisualLoggerLocation,
 			VisualLoggerColor.ToFColor(true),
-			TEXT("Listener: %s\nName: %s\nAgent Id: % 3i\nLocationA: [% 6.1f % 6.1f % 6.1f]\nLocationB: [% 6.1f % 6.1f % 6.1f]\nDistance: [% 6.2f]\nThreshold: [% 6.2f]\nCondition: [%s]\nScale: [% 6.2f]\nReward: [% 6.2f]"),
+			TEXT("Listener: %s\nTag: %s\nAgent Id: % 3i\nLocationA: [% 6.1f % 6.1f % 6.1f]\nLocationB: [% 6.1f % 6.1f % 6.1f]\nDistance: [% 6.2f]\nThreshold: [% 6.2f]\nCondition: [%s]\nScale: [% 6.2f]\nReward: [% 6.2f]"),
 			*VisualLoggerListener->GetName(),
-			*Name.ToString(),
+			*Tag.ToString(),
 			VisualLoggerAgentId,
 			LocationA.X, LocationA.Y, LocationA.Z,
 			LocationB.X, LocationB.Y, LocationB.Z,
@@ -194,7 +194,7 @@ float ULearningAgentsRewards::MakeRewardFromLocationSimilarity(
 	const FVector LocationB, 
 	const float LocationScale, 
 	const float RewardScale,
-	const FName Name,
+	const FName Tag,
 	const bool bVisualLoggerEnabled,
 	ULearningAgentsManagerListener* VisualLoggerListener,
 	const int32 VisualLoggerAgentId,
@@ -208,7 +208,7 @@ float ULearningAgentsRewards::MakeRewardFromLocationSimilarity(
 #if UE_LEARNING_AGENTS_ENABLE_VISUAL_LOG
 	if (bVisualLoggerEnabled && VisualLoggerListener)
 	{
-		const ULearningAgentsVisualLoggerObject* VisualLoggerObject = VisualLoggerListener->GetOrAddVisualLoggerObject(Name);
+		const ULearningAgentsVisualLoggerObject* VisualLoggerObject = VisualLoggerListener->GetOrAddVisualLoggerObject(Tag);
 
 		UE_LEARNING_AGENTS_VLOG_LOCATION(VisualLoggerObject, LogLearning, Display,
 			LocationA,
@@ -230,9 +230,9 @@ float ULearningAgentsRewards::MakeRewardFromLocationSimilarity(
 
 		UE_LEARNING_AGENTS_VLOG_STRING(VisualLoggerObject, LogLearning, Display, VisualLoggerLocation,
 			VisualLoggerColor.ToFColor(true),
-			TEXT("Listener: %s\nName: %s\nAgent Id: % 3i\nLocationA: [% 6.1f % 6.1f % 6.1f]\nLocationB: [% 6.1f % 6.1f % 6.1f]\nDifference: [% 6.2f]\nLocationScale: [% 6.2f]\nSimilarity: [% 6.2f]\nScale: [% 6.2f]\nReward: [% 6.2f]"),
+			TEXT("Listener: %s\nTag: %s\nAgent Id: % 3i\nLocationA: [% 6.1f % 6.1f % 6.1f]\nLocationB: [% 6.1f % 6.1f % 6.1f]\nDifference: [% 6.2f]\nLocationScale: [% 6.2f]\nSimilarity: [% 6.2f]\nScale: [% 6.2f]\nReward: [% 6.2f]"),
 			*VisualLoggerListener->GetName(),
-			*Name.ToString(),
+			*Tag.ToString(),
 			VisualLoggerAgentId,
 			LocationA.X, LocationA.Y, LocationA.Z,
 			LocationB.X, LocationB.Y, LocationB.Z,
@@ -250,9 +250,8 @@ float ULearningAgentsRewards::MakeRewardFromLocationSimilarity(
 float ULearningAgentsRewards::MakeRewardFromAngleSimilarity(
 	const float AngleA, 
 	const float AngleB, 
-	const float AngleScale, 
 	const float RewardScale,
-	const FName Name,
+	const FName Tag,
 	const bool bVisualLoggerEnabled,
 	ULearningAgentsManagerListener* VisualLoggerListener,
 	const int32 VisualLoggerAgentId,
@@ -262,13 +261,13 @@ float ULearningAgentsRewards::MakeRewardFromAngleSimilarity(
 	const FLinearColor VisualLoggerColor)
 {
 	const float AngleDifference = FMath::FindDeltaAngleDegrees(AngleA, AngleB);
-	const float Similarity = FMath::InvExpApprox(FMath::Square(AngleDifference / FMath::Max(AngleScale, UE_SMALL_NUMBER)));
+	const float Similarity = 1.0f - FMath::Abs(AngleDifference) / 180.0f;
 	const float Reward = Similarity * RewardScale;
 
 #if UE_LEARNING_AGENTS_ENABLE_VISUAL_LOG
 	if (bVisualLoggerEnabled && VisualLoggerListener)
 	{
-		const ULearningAgentsVisualLoggerObject* VisualLoggerObject = VisualLoggerListener->GetOrAddVisualLoggerObject(Name);
+		const ULearningAgentsVisualLoggerObject* VisualLoggerObject = VisualLoggerListener->GetOrAddVisualLoggerObject(Tag);
 
 		UE_LEARNING_AGENTS_VLOG_ANGLE(
 			VisualLoggerObject,
@@ -294,14 +293,13 @@ float ULearningAgentsRewards::MakeRewardFromAngleSimilarity(
 
 		UE_LEARNING_AGENTS_VLOG_STRING(VisualLoggerObject, LogLearning, Display, VisualLoggerLocation,
 			VisualLoggerColor.ToFColor(true),
-			TEXT("Listener: %s\nName: %s\nAgent Id: % 3i\nAngleA [% 6.2f]\nAngleB [% 6.2f]\nDifference: [% 6.2f]\nAngleScale: [% 6.2f]\nSimilarity: [% 6.2f]\nScale: [% 6.2f]\nReward: [% 6.2f]"),
+			TEXT("Listener: %s\nTag: %s\nAgent Id: % 3i\nAngleA [% 6.2f]\nAngleB [% 6.2f]\nDifference: [% 6.2f]\nSimilarity: [% 6.2f]\nScale: [% 6.2f]\nReward: [% 6.2f]"),
 			*VisualLoggerListener->GetName(),
-			*Name.ToString(),
+			*Tag.ToString(),
 			VisualLoggerAgentId,
 			AngleA,
 			AngleB,
 			AngleDifference,
-			AngleScale,
 			Similarity,
 			RewardScale,
 			Reward);
@@ -314,9 +312,8 @@ float ULearningAgentsRewards::MakeRewardFromAngleSimilarity(
 float ULearningAgentsRewards::MakeRewardFromRotationSimilarityAsQuats(
 	const FQuat RotationA, 
 	const FQuat RotationB, 
-	const float AngleScale, 
 	const float RewardScale,
-	const FName Name,
+	const FName Tag,
 	const bool bVisualLoggerEnabled,
 	ULearningAgentsManagerListener* VisualLoggerListener,
 	const int32 VisualLoggerAgentId,
@@ -328,13 +325,13 @@ float ULearningAgentsRewards::MakeRewardFromRotationSimilarityAsQuats(
 	FQuat Difference = RotationA.Inverse() * RotationB;
 	Difference.EnforceShortestArcWith(FQuat::Identity);
 	const float AngleDifference = FMath::RadiansToDegrees(Difference.GetAngle());
-	const float Similarity = FMath::InvExpApprox(FMath::Square(AngleDifference / (FMath::Max(AngleScale, UE_SMALL_NUMBER))));
+	const float Similarity = 1.0f - FMath::Abs(AngleDifference) / 180.0f;
 	const float Reward = Similarity * RewardScale;
 
 #if UE_LEARNING_AGENTS_ENABLE_VISUAL_LOG
 	if (bVisualLoggerEnabled && VisualLoggerListener)
 	{
-		const ULearningAgentsVisualLoggerObject* VisualLoggerObject = VisualLoggerListener->GetOrAddVisualLoggerObject(Name);
+		const ULearningAgentsVisualLoggerObject* VisualLoggerObject = VisualLoggerListener->GetOrAddVisualLoggerObject(Tag);
 
 		UE_LEARNING_AGENTS_VLOG_TRANSFORM(
 			VisualLoggerObject,
@@ -356,14 +353,13 @@ float ULearningAgentsRewards::MakeRewardFromRotationSimilarityAsQuats(
 
 		UE_LEARNING_AGENTS_VLOG_STRING(VisualLoggerObject, LogLearning, Display, VisualLoggerLocation,
 			VisualLoggerColor.ToFColor(true),
-			TEXT("Listener: %s\nName: %s\nAgent Id: % 3i\nRotationA [% 6.2f % 6.2f % 6.2f % 6.2f]\nRotationB [% 6.2f % 6.2f % 6.2f % 6.2f]\nDifference: [% 6.2f]\nAngleScale: [% 6.2f]\nSimilarity: [% 6.2f]\nScale: [% 6.2f]\nReward: [% 6.2f]"),
+			TEXT("Listener: %s\nTag: %s\nAgent Id: % 3i\nRotationA [% 6.2f % 6.2f % 6.2f % 6.2f]\nRotationB [% 6.2f % 6.2f % 6.2f % 6.2f]\nDifference: [% 6.2f]\nSimilarity: [% 6.2f]\nScale: [% 6.2f]\nReward: [% 6.2f]"),
 			*VisualLoggerListener->GetName(),
-			*Name.ToString(),
+			*Tag.ToString(),
 			VisualLoggerAgentId,
 			RotationA.X, RotationA.Y, RotationA.Z, RotationA.W,
 			RotationB.X, RotationB.Y, RotationB.Z, RotationB.W,
 			AngleDifference,
-			AngleScale,
 			Similarity,
 			RewardScale,
 			Reward);
@@ -376,9 +372,8 @@ float ULearningAgentsRewards::MakeRewardFromRotationSimilarityAsQuats(
 float ULearningAgentsRewards::MakeRewardFromRotationSimilarity(
 	const FRotator RotationA, 
 	const FRotator RotationB, 
-	const float AngleScale, 
 	const float RewardScale,
-	const FName Name,
+	const FName Tag,
 	const bool bVisualLoggerEnabled,
 	ULearningAgentsManagerListener* VisualLoggerListener,
 	const int32 VisualLoggerAgentId,
@@ -390,9 +385,8 @@ float ULearningAgentsRewards::MakeRewardFromRotationSimilarity(
 	return MakeRewardFromRotationSimilarityAsQuats(
 		RotationA.Quaternion(), 
 		RotationB.Quaternion(), 
-		AngleScale, 
 		RewardScale, 
-		Name, 
+		Tag, 
 		bVisualLoggerEnabled, 
 		VisualLoggerListener, 
 		VisualLoggerAgentId,
@@ -405,9 +399,8 @@ float ULearningAgentsRewards::MakeRewardFromRotationSimilarity(
 float ULearningAgentsRewards::MakeRewardFromDirectionSimilarity(
 	const FVector DirectionA, 
 	const FVector DirectionB, 
-	const float AngleScale, 
 	const float RewardScale,
-	const FName Name,
+	const FName Tag,
 	const bool bVisualLoggerEnabled,
 	ULearningAgentsManagerListener* VisualLoggerListener,
 	const int32 VisualLoggerAgentId,
@@ -418,13 +411,13 @@ float ULearningAgentsRewards::MakeRewardFromDirectionSimilarity(
 	const FLinearColor VisualLoggerColor)
 {
 	const float AngleDifference = FMath::RadiansToDegrees(FMath::Acos(DirectionA.GetSafeNormal(UE_SMALL_NUMBER, FVector::ForwardVector).Dot(DirectionB.GetSafeNormal(UE_SMALL_NUMBER, FVector::ForwardVector))));
-	const float Similarity = FMath::InvExpApprox(FMath::Square(AngleDifference / FMath::Max(AngleScale, UE_SMALL_NUMBER)));
+	const float Similarity = 1.0f - FMath::Abs(AngleDifference) / 180.0f;
 	const float Reward = Similarity * RewardScale;
 
 #if UE_LEARNING_AGENTS_ENABLE_VISUAL_LOG
 	if (bVisualLoggerEnabled && VisualLoggerListener)
 	{
-		const ULearningAgentsVisualLoggerObject* VisualLoggerObject = VisualLoggerListener->GetOrAddVisualLoggerObject(Name);
+		const ULearningAgentsVisualLoggerObject* VisualLoggerObject = VisualLoggerListener->GetOrAddVisualLoggerObject(Tag);
 
 		UE_LEARNING_AGENTS_VLOG_ARROW(
 			VisualLoggerObject,
@@ -446,14 +439,13 @@ float ULearningAgentsRewards::MakeRewardFromDirectionSimilarity(
 
 		UE_LEARNING_AGENTS_VLOG_STRING(VisualLoggerObject, LogLearning, Display, VisualLoggerLocation,
 			VisualLoggerColor.ToFColor(true),
-			TEXT("Listener: %s\nName: %s\nAgent Id: % 3i\nDirectionA [% 6.2f % 6.2f % 6.2f]\nDirectionB [% 6.2f % 6.2f % 6.2f]\nDifference: [% 6.2f]\nAngleScale: [% 6.2f]\nSimilarity: [% 6.2f]\nScale: [% 6.2f]\nReward: [% 6.2f]"),
+			TEXT("Listener: %s\nTag: %s\nAgent Id: % 3i\nDirectionA [% 6.2f % 6.2f % 6.2f]\nDirectionB [% 6.2f % 6.2f % 6.2f]\nDifference: [% 6.2f]\nSimilarity: [% 6.2f]\nScale: [% 6.2f]\nReward: [% 6.2f]"),
 			*VisualLoggerListener->GetName(),
-			*Name.ToString(),
+			*Tag.ToString(),
 			VisualLoggerAgentId,
 			DirectionA.X, DirectionA.Y, DirectionA.Z,
 			DirectionB.X, DirectionB.Y, DirectionB.Z,
 			AngleDifference,
-			AngleScale,
 			Similarity,
 			RewardScale,
 			Reward);
@@ -472,7 +464,7 @@ float ULearningAgentsRewards::MakeRewardFromVelocityAlongSpline(
 	const float VelocityScale, 
 	const float RewardScale, 
 	const float FiniteDifferenceDelta,
-	const FName Name,
+	const FName Tag,
 	const bool bVisualLoggerEnabled,
 	ULearningAgentsManagerListener* VisualLoggerListener,
 	const int32 VisualLoggerAgentId,
@@ -522,7 +514,7 @@ float ULearningAgentsRewards::MakeRewardFromVelocityAlongSpline(
 		const FVector SplineLocation1 = SplineComponent->GetLocationAtDistanceAlongSpline(RawDistance1, ESplineCoordinateSpace::World);
 		const FVector SplineVelocityDirection = (SplineLocation1 - SplineLocation0).GetSafeNormal(UE_SMALL_NUMBER, FVector::ForwardVector);
 
-		const ULearningAgentsVisualLoggerObject* VisualLoggerObject = VisualLoggerListener->GetOrAddVisualLoggerObject(Name);
+		const ULearningAgentsVisualLoggerObject* VisualLoggerObject = VisualLoggerListener->GetOrAddVisualLoggerObject(Tag);
 
 		UE_LEARNING_AGENTS_VLOG_ARROW(
 			VisualLoggerObject,
@@ -535,9 +527,9 @@ float ULearningAgentsRewards::MakeRewardFromVelocityAlongSpline(
 
 		UE_LEARNING_AGENTS_VLOG_STRING(VisualLoggerObject, LogLearning, Display, VisualLoggerLocation,
 			VisualLoggerColor.ToFColor(true),
-			TEXT("Listener: %s\nName: %s\nAgent Id: % 3i\nLocation [% 6.2f % 6.2f % 6.2f]\nVelocity [% 6.2f % 6.2f % 6.2f]\nDistanceVelocity: [% 6.2f]\nScale: [% 6.2f]\nReward: [% 6.2f]"),
+			TEXT("Listener: %s\nTag: %s\nAgent Id: % 3i\nLocation [% 6.2f % 6.2f % 6.2f]\nVelocity [% 6.2f % 6.2f % 6.2f]\nDistanceVelocity: [% 6.2f]\nScale: [% 6.2f]\nReward: [% 6.2f]"),
 			*VisualLoggerListener->GetName(),
-			*Name.ToString(),
+			*Tag.ToString(),
 			VisualLoggerAgentId,
 			Location.X, Location.Y, Location.Z,
 			Velocity.X, Velocity.Y, Velocity.Z,
diff --git a/Engine/Plugins/Experimental/LearningAgents/Source/LearningAgentsTraining/Public/LearningAgentsCompletions.h b/Engine/Plugins/Experimental/LearningAgents/Source/LearningAgentsTraining/Public/LearningAgentsCompletions.h
--- a/Engine/Plugins/Experimental/LearningAgents/Source/LearningAgentsTraining/Public/LearningAgentsCompletions.h
+++ b/Engine/Plugins/Experimental/LearningAgents/Source/LearningAgentsTraining/Public/LearningAgentsCompletions.h
@@ -6,6 +6,8 @@
 
 #include "LearningAgentsCompletions.generated.h"
 
+class ULearningAgentsManagerListener;
+
 namespace UE::Learning
 {
 	enum class ECompletionMode : uint8;
@@ -41,47 +43,216 @@ class LEARNINGAGENTSTRAINING_API ULearningAgentsCompletions : public UBlueprintF
 
 public:
 
+	/** Returns true if a completion is running, otherwise false. */
 	UFUNCTION(BlueprintPure, Category = "LearningAgents")
-	static bool CompletionIsRunning(const ELearningAgentsCompletion Completion);
+	static bool IsCompletionRunning(const ELearningAgentsCompletion Completion);
 
+	/** Returns true if a completion is either truncated or terminated, otherwise false. */
 	UFUNCTION(BlueprintPure, Category = "LearningAgents")
-	static bool CompletionIsCompleted(const ELearningAgentsCompletion Completion);
+	static bool IsCompletionCompleted(const ELearningAgentsCompletion Completion);
 
+	/** Returns true if a completion is truncated, otherwise false. */
 	UFUNCTION(BlueprintPure, Category = "LearningAgents")
-	static bool CompletionIsTruncation(const ELearningAgentsCompletion Completion);
+	static bool IsCompletionTruncation(const ELearningAgentsCompletion Completion);
 
+	/** Returns true if a completion is terminated, otherwise false. */
 	UFUNCTION(BlueprintPure, Category = "LearningAgents")
-	static bool CompletionIsTermination(const ELearningAgentsCompletion Completion);
+	static bool IsCompletionTermination(const ELearningAgentsCompletion Completion);
+
 
-	UFUNCTION(BlueprintPure, Category = "LearningAgents", meta = (CommutativeAssociativeBinaryOperator))
+	/** Returns a termination if either input is a termination, otherwise a truncation if either input is a truncation, otherwise returns running. */
+	UFUNCTION(BlueprintPure, Category = "LearningAgents", meta = (CommutativeAssociativeBinaryOperator, DisplayName="Completion OR", CompactNodeTitle = "OR"))
 	static ELearningAgentsCompletion CompletionOr(ELearningAgentsCompletion A, ELearningAgentsCompletion B);
 
-	UFUNCTION(BlueprintPure, Category = "LearningAgents", meta = (CommutativeAssociativeBinaryOperator))
+	/** Returns a termination if both inputs are a termination, otherwise a truncation if both inputs are either a truncation or termination, otherwise returns running. */
+	UFUNCTION(BlueprintPure, Category = "LearningAgents", meta = (CommutativeAssociativeBinaryOperator, DisplayName = "Completion AND", CompactNodeTitle="AND"))
 	static ELearningAgentsCompletion CompletionAnd(ELearningAgentsCompletion A, ELearningAgentsCompletion B);
 
-	UFUNCTION(BlueprintPure, Category = "LearningAgents")
-	static ELearningAgentsCompletion CompletionNot(const ELearningAgentsCompletion A, const ELearningAgentsCompletion NotRunningType = ELearningAgentsCompletion::Termination);
+	/** Returns running if the input A is either a termination or truncation, otherwise returns the completion specified by NotRunningType */
+	UFUNCTION(BlueprintPure, Category = "LearningAgents", meta = (DisplayName = "Completion NOT", CompactNodeTitle = "NOT"))
+	static ELearningAgentsCompletion CompletionNot(ELearningAgentsCompletion A, ELearningAgentsCompletion NotRunningType = ELearningAgentsCompletion::Termination);
 
-	UFUNCTION(BlueprintPure, Category = "LearningAgents")
-	static ELearningAgentsCompletion CompletionOnCondition(const bool bCondition, const ELearningAgentsCompletion CompletionType = ELearningAgentsCompletion::Termination);
+	/**
+	 * Make a completion.
+	 *
+	 * @param CompletionType The type of completion to make.
+	 * @param Tag The tag for the completion. Used for debugging.
+	 * @param bVisualLoggerEnabled When true, debug data will be sent to the visual logger.
+	 * @param VisualLoggerListener The listener object which is making this completion. This must be set to use logging.
+	 * @param VisualLoggerAgentId The agent id associated with this completion.
+	 * @param VisualLoggerLocation A location for the visual logger information in the world.
+	 * @param VisualLoggerColor The color for the visual logger display.
+	 * @return The resulting completion.
+	 */
+	UFUNCTION(BlueprintPure, Category = "LearningAgents", meta = (AdvancedDisplay = 1, DefaultToSelf = "VisualLoggerListener"))
+	static ELearningAgentsCompletion MakeCompletion(
+		const ELearningAgentsCompletion CompletionType = ELearningAgentsCompletion::Termination,
+		const FName Tag = TEXT("Completion"),
+		const bool bVisualLoggerEnabled = false,
+		ULearningAgentsManagerListener* VisualLoggerListener = nullptr,
+		const int32 VisualLoggerAgentId = -1,
+		const FVector VisualLoggerLocation = FVector::ZeroVector,
+		const FLinearColor VisualLoggerColor = FLinearColor::Yellow);
 
-	UFUNCTION(BlueprintPure, Category = "LearningAgents")
-	static ELearningAgentsCompletion CompletionOnTimeElapsed(const float Time, const float TimeThreshold = 10.0f, const ELearningAgentsCompletion CompletionType = ELearningAgentsCompletion::Truncation);
+	/**
+	 * Make a completion based on some condition.
+	 *
+	 * @param bCondition When true, returns the given CompletionType, otherwise returns Running.
+	 * @param CompletionType The type of completion to make.
+	 * @param Tag The tag for the completion. Used for debugging.
+	 * @param bVisualLoggerEnabled When true, debug data will be sent to the visual logger.
+	 * @param VisualLoggerListener The listener object which is making this completion. This must be set to use logging.
+	 * @param VisualLoggerAgentId The agent id associated with this completion.
+	 * @param VisualLoggerLocation A location for the visual logger information in the world.
+	 * @param VisualLoggerColor The color for the visual logger display.
+	 * @return The resulting completion.
+	 */
+	UFUNCTION(BlueprintPure, Category = "LearningAgents", meta = (AdvancedDisplay = 2, DefaultToSelf = "VisualLoggerListener"))
+	static ELearningAgentsCompletion MakeCompletionOnCondition(
+		const bool bCondition, 
+		const ELearningAgentsCompletion CompletionType = ELearningAgentsCompletion::Termination,
+		const FName Tag = TEXT("CompletionOnCondition"),
+		const bool bVisualLoggerEnabled = false,
+		ULearningAgentsManagerListener* VisualLoggerListener = nullptr,
+		const int32 VisualLoggerAgentId = -1,
+		const FVector VisualLoggerLocation = FVector::ZeroVector,
+		const FLinearColor VisualLoggerColor = FLinearColor::Yellow);
 
-	UFUNCTION(BlueprintPure, Category = "LearningAgents")
-	static ELearningAgentsCompletion CompletionOnEpisodeStepsRecorded(const int32 EpisodeSteps, const int32 MaxEpisodeSteps = 64, const ELearningAgentsCompletion CompletionType = ELearningAgentsCompletion::Truncation);
+	/**
+	 * Make a completion when a time goes above a threshold.
+	 *
+	 * @param Time The current time.
+	 * @param TimeThreshold The time threshold above which to complete with the given CompletionType.
+	 * @param CompletionType The type of completion to make
+	 * @param Tag The tag for the completion. Used for debugging.
+	 * @param bVisualLoggerEnabled When true, debug data will be sent to the visual logger.
+	 * @param VisualLoggerListener The listener object which is making this completion. This must be set to use logging.
+	 * @param VisualLoggerAgentId The agent id associated with this completion.
+	 * @param VisualLoggerLocation A location for the visual logger information in the world.
+	 * @param VisualLoggerColor The color for the visual logger display.
+	 * @return The resulting completion.
+	 */
+	UFUNCTION(BlueprintPure, Category = "LearningAgents", meta = (AdvancedDisplay = 3, DefaultToSelf = "VisualLoggerListener"))
+	static ELearningAgentsCompletion MakeCompletionOnTimeElapsed(
+		const float Time, 
+		const float TimeThreshold = 10.0f, 
+		const ELearningAgentsCompletion CompletionType = ELearningAgentsCompletion::Truncation,
+		const FName Tag = TEXT("CompletionOnTimeElapsed"),
+		const bool bVisualLoggerEnabled = false,
+		ULearningAgentsManagerListener* VisualLoggerListener = nullptr,
+		const int32 VisualLoggerAgentId = -1,
+		const FVector VisualLoggerLocation = FVector::ZeroVector,
+		const FLinearColor VisualLoggerColor = FLinearColor::Yellow);
 
-	UFUNCTION(BlueprintPure, Category = "LearningAgents")
-	static ELearningAgentsCompletion CompletionOnLocationDifferenceBelowThreshold(const FVector LocationA, const FVector LocationB, const float DistanceThreshold = 100.0f, const ELearningAgentsCompletion CompletionType = ELearningAgentsCompletion::Termination);
+	/**
+	 * Make a completion when the number of episode steps recorded exceeds some threshold.
+	 *
+	 * @param EpisodeSteps The number of steps recorded.
+	 * @param MaxEpisodeSteps The step threshold above which to complete with the given CompletionType.
+	 * @param CompletionType The type of completion to make
+	 * @param Tag The tag for the completion. Used for debugging.
+	 * @param bVisualLoggerEnabled When true, debug data will be sent to the visual logger.
+	 * @param VisualLoggerListener The listener object which is making this completion. This must be set to use logging.
+	 * @param VisualLoggerAgentId The agent id associated with this completion.
+	 * @param VisualLoggerLocation A location for the visual logger information in the world.
+	 * @param VisualLoggerColor The color for the visual logger display.
+	 * @return The resulting completion.
+	 */
+	UFUNCTION(BlueprintPure, Category = "LearningAgents", meta = (AdvancedDisplay = 3, DefaultToSelf = "VisualLoggerListener"))
+	static ELearningAgentsCompletion MakeCompletionOnEpisodeStepsRecorded(
+		const int32 EpisodeSteps, 
+		const int32 MaxEpisodeSteps = 64, 
+		const ELearningAgentsCompletion CompletionType = ELearningAgentsCompletion::Truncation,
+		const FName Tag = TEXT("CompletionOnEpisodeStepsRecorded"),
+		const bool bVisualLoggerEnabled = false,
+		ULearningAgentsManagerListener* VisualLoggerListener = nullptr,
+		const int32 VisualLoggerAgentId = -1,
+		const FVector VisualLoggerLocation = FVector::ZeroVector,
+		const FLinearColor VisualLoggerColor = FLinearColor::Yellow);
 
-	UFUNCTION(BlueprintPure, Category = "LearningAgents")
-	static ELearningAgentsCompletion CompletionOnLocationDifferenceAboveThreshold(const FVector LocationA, const FVector LocationB, const float DistanceThreshold = 100.0f, const ELearningAgentsCompletion CompletionType = ELearningAgentsCompletion::Termination);
+	/**
+	 * Make a completion when the distance between two locations is below some threshold.
+	 *
+	 * @param LocationA The first location.
+	 * @param LocationB The second location.
+	 * @param DistanceThreshold The distance threshold.
+	 * @param CompletionType The type of completion to make
+	 * @param Tag The tag for the completion. Used for debugging.
+	 * @param bVisualLoggerEnabled When true, debug data will be sent to the visual logger.
+	 * @param VisualLoggerListener The listener object which is making this completion. This must be set to use logging.
+	 * @param VisualLoggerAgentId The agent id associated with this completion.
+	 * @param VisualLoggerLocation A location for the visual logger information in the world.
+	 * @param VisualLoggerColor The color for the visual logger display.
+	 * @return The resulting completion.
+	 */
+	UFUNCTION(BlueprintPure, Category = "LearningAgents", meta = (AdvancedDisplay = 4, DefaultToSelf = "VisualLoggerListener"))
+	static ELearningAgentsCompletion MakeCompletionOnLocationDifferenceBelowThreshold(
+		const FVector LocationA, 
+		const FVector LocationB, 
+		const float DistanceThreshold = 100.0f, 
+		const ELearningAgentsCompletion CompletionType = ELearningAgentsCompletion::Termination,
+		const FName Tag = TEXT("CompletionOnLocationDifferenceBelowThreshold"),
+		const bool bVisualLoggerEnabled = false,
+		ULearningAgentsManagerListener* VisualLoggerListener = nullptr,
+		const int32 VisualLoggerAgentId = -1,
+		const FVector VisualLoggerLocation = FVector::ZeroVector,
+		const FLinearColor VisualLoggerColor = FLinearColor::Yellow);
 
-	UFUNCTION(BlueprintPure, Category = "LearningAgents")
-	static ELearningAgentsCompletion CompletionOnLocationOutsideBounds(
+	/**
+	 * Make a completion when the distance between two locations is above some threshold.
+	 *
+	 * @param LocationA The first location.
+	 * @param LocationB The second location.
+	 * @param DistanceThreshold The distance threshold.
+	 * @param CompletionType The type of completion to make
+	 * @param Tag The tag for the completion. Used for debugging.
+	 * @param bVisualLoggerEnabled When true, debug data will be sent to the visual logger.
+	 * @param VisualLoggerListener The listener object which is making this completion. This must be set to use logging.
+	 * @param VisualLoggerAgentId The agent id associated with this completion.
+	 * @param VisualLoggerLocation A location for the visual logger information in the world.
+	 * @param VisualLoggerColor The color for the visual logger display.
+	 * @return The resulting completion.
+	 */
+	UFUNCTION(BlueprintPure, Category = "LearningAgents", meta = (AdvancedDisplay = 4, DefaultToSelf = "VisualLoggerListener"))
+	static ELearningAgentsCompletion MakeCompletionOnLocationDifferenceAboveThreshold(
+		const FVector LocationA, 
+		const FVector LocationB, 
+		const float DistanceThreshold = 100.0f, 
+		const ELearningAgentsCompletion CompletionType = ELearningAgentsCompletion::Termination,
+		const FName Tag = TEXT("CompletionOnLocationDifferenceAboveThreshold"),
+		const bool bVisualLoggerEnabled = false,
+		ULearningAgentsManagerListener* VisualLoggerListener = nullptr,
+		const int32 VisualLoggerAgentId = -1,
+		const FVector VisualLoggerLocation = FVector::ZeroVector,
+		const FLinearColor VisualLoggerColor = FLinearColor::Yellow);
+
+	/**
+	 * Make a completion when a location moves outside of sound bounds.
+	 *
+	 * @param Location The location.
+	 * @param BoundsTransform The transform of the bounds object.
+	 * @param BoundsMins The minimums of the bounds object.
+	 * @param BoundsMaxs The maximums of the bounds object.
+	 * @param CompletionType The type of completion to make
+	 * @param Tag The tag for the completion. Used for debugging.
+	 * @param bVisualLoggerEnabled When true, debug data will be sent to the visual logger.
+	 * @param VisualLoggerListener The listener object which is making this completion. This must be set to use logging.
+	 * @param VisualLoggerAgentId The agent id associated with this completion.
+	 * @param VisualLoggerLocation A location for the visual logger information in the world.
+	 * @param VisualLoggerColor The color for the visual logger display.
+	 * @return The resulting completion.
+	 */
+	UFUNCTION(BlueprintPure, Category = "LearningAgents", meta = (AdvancedDisplay = 5, DefaultToSelf = "VisualLoggerListener"))
+	static ELearningAgentsCompletion MakeCompletionOnLocationOutsideBounds(
 		const FVector Location, 
 		const FTransform BoundsTransform = FTransform(),
 		const FVector BoundsMins = FVector(-100.0f, -100.0f, -100.0f),
 		const FVector BoundsMaxs = FVector(+100.0f, +100.0f, +100.0f),
-		const ELearningAgentsCompletion CompletionType = ELearningAgentsCompletion::Termination);
+		const ELearningAgentsCompletion CompletionType = ELearningAgentsCompletion::Termination,
+		const FName Tag = TEXT("CompletionOnLocationOutsideBounds"),
+		const bool bVisualLoggerEnabled = false,
+		ULearningAgentsManagerListener* VisualLoggerListener = nullptr,
+		const int32 VisualLoggerAgentId = -1,
+		const FVector VisualLoggerLocation = FVector::ZeroVector,
+		const FLinearColor VisualLoggerColor = FLinearColor::Yellow);
 };
diff --git a/Engine/Plugins/Experimental/LearningAgents/Source/LearningAgentsTraining/Public/LearningAgentsRecorder.h b/Engine/Plugins/Experimental/LearningAgents/Source/LearningAgentsTraining/Public/LearningAgentsRecorder.h
--- a/Engine/Plugins/Experimental/LearningAgents/Source/LearningAgentsTraining/Public/LearningAgentsRecorder.h
+++ b/Engine/Plugins/Experimental/LearningAgents/Source/LearningAgentsTraining/Public/LearningAgentsRecorder.h
@@ -107,6 +107,10 @@ public:
 	UFUNCTION(BlueprintCallable, Category = "LearningAgents")
 	void EndRecording();
 
+	/** Ends the recording of the observations and actions of each agent and discards them. */
+	UFUNCTION(BlueprintCallable, Category = "LearningAgents")
+	void EndRecordingAndDiscard();
+
 	/** Returns true if recording is active; Otherwise, false. */
 	UFUNCTION(BlueprintPure, Category = "LearningAgents")
 	bool IsRecording() const;
diff --git a/Engine/Plugins/Experimental/LearningAgents/Source/LearningAgentsTraining/Public/LearningAgentsRecording.h b/Engine/Plugins/Experimental/LearningAgents/Source/LearningAgentsTraining/Public/LearningAgentsRecording.h
--- a/Engine/Plugins/Experimental/LearningAgents/Source/LearningAgentsTraining/Public/LearningAgentsRecording.h
+++ b/Engine/Plugins/Experimental/LearningAgents/Source/LearningAgentsTraining/Public/LearningAgentsRecording.h
@@ -62,6 +62,10 @@ public:
 
 public:
 
+	/** Resets this recording asset to be empty. */
+	UFUNCTION(CallInEditor, Category = "LearningAgents")
+	void ResetRecording();
+
 	/** Load this recording from a file. */
 	UFUNCTION(BlueprintCallable, Category = "LearningAgents", meta = (RelativePath))
 	void LoadRecordingFromFile(const FFilePath& File);
diff --git a/Engine/Plugins/Experimental/LearningAgents/Source/LearningAgentsTraining/Public/LearningAgentsRewards.h b/Engine/Plugins/Experimental/LearningAgents/Source/LearningAgentsTraining/Public/LearningAgentsRewards.h
--- a/Engine/Plugins/Experimental/LearningAgents/Source/LearningAgentsTraining/Public/LearningAgentsRewards.h
+++ b/Engine/Plugins/Experimental/LearningAgents/Source/LearningAgentsTraining/Public/LearningAgentsRewards.h
@@ -21,7 +21,7 @@ public:
 	 *
 	 * @param RewardValue The float value used to create the reward.
 	 * @param RewardScale The scale of the reward. Use a negative scale to create a penalty.
-	 * @param Name The name of the reward. Used for debugging.
+	 * @param Tag The tag for the reward. Used for debugging.
 	 * @param bVisualLoggerEnabled When true, debug data will be sent to the visual logger.
 	 * @param VisualLoggerListener The listener object which is making this reward. This must be set to use logging.
 	 * @param VisualLoggerAgentId The agent id associated with this reward.
@@ -29,11 +29,11 @@ public:
 	 * @param VisualLoggerColor The color for the visual logger display.
 	 * @return The resulting reward value.
 	 */
-	UFUNCTION(BlueprintPure, Category = "LearningAgents", meta = (AdvancedDisplay = 3, DefaultToSelf = "VisualLoggerListener"))
+	UFUNCTION(BlueprintPure, Category = "LearningAgents", meta = (AdvancedDisplay = 2, DefaultToSelf = "VisualLoggerListener"))
 	static float MakeReward(
 		const float RewardValue, 
 		const float RewardScale = 1.0f, 
-		const FName Name = TEXT("Reward"),
+		const FName Tag = TEXT("Reward"),
 		const bool bVisualLoggerEnabled = false,
 		ULearningAgentsManagerListener* VisualLoggerListener = nullptr,
 		const int32 VisualLoggerAgentId = -1,
@@ -45,7 +45,7 @@ public:
 	 *
 	 * @param bCondition The condition under which to create a reward.
 	 * @param RewardScale The scale of the reward. Use a negative scale to create a penalty.
-	 * @param Name The name of the reward. Used for debugging.
+	 * @param Tag The tag for the reward. Used for debugging.
 	 * @param bVisualLoggerEnabled When true, debug data will be sent to the visual logger.
 	 * @param VisualLoggerListener The listener object which is making this reward. This must be set to use logging.
 	 * @param VisualLoggerAgentId The agent id associated with this reward.
@@ -53,11 +53,11 @@ public:
 	 * @param VisualLoggerColor The color for the visual logger display.
 	 * @return The resulting reward value.
 	 */
-	UFUNCTION(BlueprintPure, Category = "LearningAgents", meta = (AdvancedDisplay = 3, DefaultToSelf = "VisualLoggerListener"))
+	UFUNCTION(BlueprintPure, Category = "LearningAgents", meta = (AdvancedDisplay = 2, DefaultToSelf = "VisualLoggerListener"))
 	static float MakeRewardOnCondition(
 		const bool bCondition, 
 		const float RewardScale = 1.0f,
-		const FName Name = TEXT("RewardOnCondition"),
+		const FName Tag = TEXT("RewardOnCondition"),
 		const bool bVisualLoggerEnabled = false,
 		ULearningAgentsManagerListener* VisualLoggerListener = nullptr,
 		const int32 VisualLoggerAgentId = -1,
@@ -71,7 +71,7 @@ public:
 	 * @param LocationB The second location.
 	 * @param DistanceThreshold The distance threshold.
 	 * @param RewardScale The scale of the reward. Use a negative scale to create a penalty.
-	 * @param Name The name of the reward. Used for debugging.
+	 * @param Tag The tag for the reward. Used for debugging.
 	 * @param bVisualLoggerEnabled When true, debug data will be sent to the visual logger.
 	 * @param VisualLoggerListener The listener object which is making this reward. This must be set to use logging.
 	 * @param VisualLoggerAgentId The agent id associated with this reward.
@@ -79,13 +79,13 @@ public:
 	 * @param VisualLoggerColor The color for the visual logger display.
 	 * @return The resulting reward value.
 	 */
-	UFUNCTION(BlueprintPure, Category = "LearningAgents", meta = (AdvancedDisplay = 5, DefaultToSelf = "VisualLoggerListener"))
+	UFUNCTION(BlueprintPure, Category = "LearningAgents", meta = (AdvancedDisplay = 4, DefaultToSelf = "VisualLoggerListener"))
 	static float MakeRewardOnLocationDifferenceBelowThreshold(
 		const FVector LocationA, 
 		const FVector LocationB, 
 		const float DistanceThreshold = 100.0f, 
 		const float RewardScale = 1.0f,
-		const FName Name = TEXT("RewardOnLocationDifferenceBelowThreshold"),
+		const FName Tag = TEXT("RewardOnLocationDifferenceBelowThreshold"),
 		const bool bVisualLoggerEnabled = false,
 		ULearningAgentsManagerListener* VisualLoggerListener = nullptr,
 		const int32 VisualLoggerAgentId = -1,
@@ -99,7 +99,7 @@ public:
 	 * @param LocationB The second location.
 	 * @param DistanceThreshold The distance threshold.
 	 * @param RewardScale The scale of the reward. Use a negative scale to create a penalty.
-	 * @param Name The name of the reward. Used for debugging.
+	 * @param Tag The tag for the reward. Used for debugging.
 	 * @param bVisualLoggerEnabled When true, debug data will be sent to the visual logger.
 	 * @param VisualLoggerListener The listener object which is making this reward. This must be set to use logging.
 	 * @param VisualLoggerAgentId The agent id associated with this reward.
@@ -107,13 +107,13 @@ public:
 	 * @param VisualLoggerColor The color for the visual logger display.
 	 * @return The resulting reward value.
 	 */
-	UFUNCTION(BlueprintPure, Category = "LearningAgents", meta = (AdvancedDisplay = 5, DefaultToSelf = "VisualLoggerListener"))
+	UFUNCTION(BlueprintPure, Category = "LearningAgents", meta = (AdvancedDisplay = 4, DefaultToSelf = "VisualLoggerListener"))
 	static float MakeRewardOnLocationDifferenceAboveThreshold(
 		const FVector LocationA,
 		const FVector LocationB,
 		const float DistanceThreshold = 100.0f,
 		const float RewardScale = 1.0f,
-		const FName Name = TEXT("RewardOnLocationDifferenceAboveThreshold"),
+		const FName Tag = TEXT("RewardOnLocationDifferenceAboveThreshold"),
 		const bool bVisualLoggerEnabled = false,
 		ULearningAgentsManagerListener* VisualLoggerListener = nullptr,
 		const int32 VisualLoggerAgentId = -1,
@@ -127,7 +127,7 @@ public:
 	 * @param LocationB The second location.
 	 * @param LocationScale The expected scale for the distance between locations.
 	 * @param RewardScale The scale of the reward. Use a negative scale to create a penalty.
-	 * @param Name The name of the reward. Used for debugging.
+	 * @param Tag The tag for the reward. Used for debugging.
 	 * @param bVisualLoggerEnabled When true, debug data will be sent to the visual logger.
 	 * @param VisualLoggerListener The listener object which is making this reward. This must be set to use logging.
 	 * @param VisualLoggerAgentId The agent id associated with this reward.
@@ -135,13 +135,13 @@ public:
 	 * @param VisualLoggerColor The color for the visual logger display.
 	 * @return The resulting reward value.
 	 */
-	UFUNCTION(BlueprintPure, Category = "LearningAgents", meta = (AdvancedDisplay = 5, DefaultToSelf = "VisualLoggerListener"))
+	UFUNCTION(BlueprintPure, Category = "LearningAgents", meta = (AdvancedDisplay = 4, DefaultToSelf = "VisualLoggerListener"))
 	static float MakeRewardFromLocationSimilarity(
 		const FVector LocationA, 
 		const FVector LocationB, 
 		const float LocationScale = 100.0f, 
 		const float RewardScale = 1.0f,
-		const FName Name = TEXT("RewardFromLocationSimilarity"),
+		const FName Tag = TEXT("RewardFromLocationSimilarity"),
 		const bool bVisualLoggerEnabled = false,
 		ULearningAgentsManagerListener* VisualLoggerListener = nullptr,
 		const int32 VisualLoggerAgentId = -1,
@@ -153,9 +153,8 @@ public:
 	 *
 	 * @param RotationA The first rotation.
 	 * @param RotationB The second rotation.
-	 * @param AngleScale The expected scale for the angle between rotations.
 	 * @param RewardScale The scale of the reward. Use a negative scale to create a penalty.
-	 * @param Name The name of the reward. Used for debugging.
+	 * @param Tag The tag for the reward. Used for debugging.
 	 * @param bVisualLoggerEnabled When true, debug data will be sent to the visual logger.
 	 * @param VisualLoggerListener The listener object which is making this reward. This must be set to use logging.
 	 * @param VisualLoggerAgentId The agent id associated with this reward.
@@ -165,13 +164,12 @@ public:
 	 * @param VisualLoggerColor The color for the visual logger display.
 	 * @return The resulting reward value.
 	 */
-	UFUNCTION(BlueprintPure, Category = "LearningAgents", meta = (AdvancedDisplay = 5, DefaultToSelf = "VisualLoggerListener"))
+	UFUNCTION(BlueprintPure, Category = "LearningAgents", meta = (AdvancedDisplay = 3, DefaultToSelf = "VisualLoggerListener"))
 	static float MakeRewardFromRotationSimilarity(
 		const FRotator RotationA, 
 		const FRotator RotationB, 
-		const float AngleScale = 90.0f, 
 		const float RewardScale = 1.0f,
-		const FName Name = TEXT("RewardFromRotationSimilarity"),
+		const FName Tag = TEXT("RewardFromRotationSimilarity"),
 		const bool bVisualLoggerEnabled = false,
 		ULearningAgentsManagerListener* VisualLoggerListener = nullptr,
 		const int32 VisualLoggerAgentId = -1,
@@ -185,9 +183,8 @@ public:
 	 *
 	 * @param RotationA The first rotation.
 	 * @param RotationB The second rotation.
-	 * @param AngleScale The expected scale for the angle between rotations.
 	 * @param RewardScale The scale of the reward. Use a negative scale to create a penalty.
-	 * @param Name The name of the reward. Used for debugging.
+	 * @param Tag The tag for the reward. Used for debugging.
 	 * @param bVisualLoggerEnabled When true, debug data will be sent to the visual logger.
 	 * @param VisualLoggerListener The listener object which is making this reward. This must be set to use logging.
 	 * @param VisualLoggerAgentId The agent id associated with this reward.
@@ -197,13 +194,12 @@ public:
 	 * @param VisualLoggerColor The color for the visual logger display.
 	 * @return The resulting reward value.
 	 */
-	UFUNCTION(BlueprintPure, Category = "LearningAgents", meta = (AdvancedDisplay = 5, DefaultToSelf = "VisualLoggerListener"))
+	UFUNCTION(BlueprintPure, Category = "LearningAgents", meta = (AdvancedDisplay = 3, DefaultToSelf = "VisualLoggerListener"))
 	static float MakeRewardFromRotationSimilarityAsQuats(
 		const FQuat RotationA, 
 		const FQuat RotationB, 
-		const float AngleScale = 90.0f,
 		const float RewardScale = 1.0f,
-		const FName Name = TEXT("RewardFromRotationSimilarity"),
+		const FName Tag = TEXT("RewardFromRotationSimilarity"),
 		const bool bVisualLoggerEnabled = false,
 		ULearningAgentsManagerListener* VisualLoggerListener = nullptr,
 		const int32 VisualLoggerAgentId = -1,
@@ -217,9 +213,8 @@ public:
 	 *
 	 * @param AngleA The first angle.
 	 * @param AngleB The second angle.
-	 * @param AngleScale The expected scale for the angle between angles.
 	 * @param RewardScale The scale of the reward. Use a negative scale to create a penalty.
-	 * @param Name The name of the reward. Used for debugging.
+	 * @param Tag The tag for the reward. Used for debugging.
 	 * @param bVisualLoggerEnabled When true, debug data will be sent to the visual logger.
 	 * @param VisualLoggerListener The listener object which is making this reward. This must be set to use logging.
 	 * @param VisualLoggerAgentId The agent id associated with this reward.
@@ -229,13 +224,12 @@ public:
 	 * @param VisualLoggerColor The color for the visual logger display.
 	 * @return The resulting reward value.
 	 */
-	UFUNCTION(BlueprintPure, Category = "LearningAgents", meta = (AdvancedDisplay = 5, DefaultToSelf = "VisualLoggerListener"))
+	UFUNCTION(BlueprintPure, Category = "LearningAgents", meta = (AdvancedDisplay = 3, DefaultToSelf = "VisualLoggerListener"))
 	static float MakeRewardFromAngleSimilarity(
 		const float AngleA, 
 		const float AngleB, 
-		const float AngleScale = 90.0f, 
 		const float RewardScale = 1.0f,
-		const FName Name = TEXT("RewardFromAngleSimilarity"),
+		const FName Tag = TEXT("RewardFromAngleSimilarity"),
 		const bool bVisualLoggerEnabled = false,
 		ULearningAgentsManagerListener* VisualLoggerListener = nullptr,
 		const int32 VisualLoggerAgentId = -1,
@@ -249,9 +243,8 @@ public:
 	 *
 	 * @param DirectionA The first angle.
 	 * @param DirectionB The second angle.
-	 * @param AngleScale The expected scale for the angle between the directions.
 	 * @param RewardScale The scale of the reward. Use a negative scale to create a penalty.
-	 * @param Name The name of the reward. Used for debugging.
+	 * @param Tag The tag for the reward. Used for debugging.
 	 * @param bVisualLoggerEnabled When true, debug data will be sent to the visual logger.
 	 * @param VisualLoggerListener The listener object which is making this reward. This must be set to use logging.
 	 * @param VisualLoggerAgentId The agent id associated with this reward.
@@ -262,13 +255,12 @@ public:
 	 * @param VisualLoggerColor The color for the visual logger display.
 	 * @return The resulting reward value.
 	 */
-	UFUNCTION(BlueprintPure, Category = "LearningAgents", meta = (AdvancedDisplay = 5, DefaultToSelf = "VisualLoggerListener", DirectionA="1.0,0.0,0.0", DirectionB = "1.0,0.0,0.0"))
+	UFUNCTION(BlueprintPure, Category = "LearningAgents", meta = (AdvancedDisplay = 3, DefaultToSelf = "VisualLoggerListener", DirectionA="1.0,0.0,0.0", DirectionB = "1.0,0.0,0.0"))
 	static float MakeRewardFromDirectionSimilarity(
 		const FVector DirectionA, 
 		const FVector DirectionB, 
-		const float AngleScale = 90.0f, 
 		const float RewardScale = 1.0f,
-		const FName Name = TEXT("RewardFromDirectionSimilarity"),
+		const FName Tag = TEXT("RewardFromDirectionSimilarity"),
 		const bool bVisualLoggerEnabled = false,
 		ULearningAgentsManagerListener* VisualLoggerListener = nullptr,
 		const int32 VisualLoggerAgentId = -1,
@@ -287,7 +279,7 @@ public:
 	 * @param VelocityScale The expected scale for the velocity.
 	 * @param RewardScale The scale of the reward. Use a negative scale to create a penalty.
 	 * @param FiniteDifferenceDelta The finite difference to use when computing the velocity along the spline.
-	 * @param Name The name of the reward. Used for debugging.
+	 * @param Tag The tag for the reward. Used for debugging.
 	 * @param bVisualLoggerEnabled When true, debug data will be sent to the visual logger.
 	 * @param VisualLoggerListener The listener object which is making this reward. This must be set to use logging.
 	 * @param VisualLoggerAgentId The agent id associated with this reward.
@@ -295,7 +287,7 @@ public:
 	 * @param VisualLoggerColor The color for the visual logger display.
 	 * @return The resulting reward value.
 	 */
-	UFUNCTION(BlueprintPure, Category = "LearningAgents", meta = (AdvancedDisplay = 7, DefaultToSelf = "VisualLoggerListener"))
+	UFUNCTION(BlueprintPure, Category = "LearningAgents", meta = (AdvancedDisplay = 6, DefaultToSelf = "VisualLoggerListener"))
 	static float MakeRewardFromVelocityAlongSpline(
 		const USplineComponent* SplineComponent, 
 		const FVector Location, 
@@ -303,7 +295,7 @@ public:
 		const float VelocityScale = 200.0f, 
 		const float RewardScale = 1.0f, 
 		const float FiniteDifferenceDelta = 10.0f,
-		const FName Name = TEXT("RewardFromVelocityAlongSpline"),
+		const FName Tag = TEXT("RewardFromVelocityAlongSpline"),
 		const bool bVisualLoggerEnabled = false,
 		ULearningAgentsManagerListener* VisualLoggerListener = nullptr,
 		const int32 VisualLoggerAgentId = -1,
