[Backout] - CL31291642

diff --git a/Engine/Plugins/Experimental/LearningAgents/Source/LearningAgents/Private/LearningAgentsActions.cpp b/Engine/Plugins/Experimental/LearningAgents/Source/LearningAgents/Private/LearningAgentsActions.cpp
--- a/Engine/Plugins/Experimental/LearningAgents/Source/LearningAgents/Private/LearningAgentsActions.cpp
+++ b/Engine/Plugins/Experimental/LearningAgents/Source/LearningAgents/Private/LearningAgentsActions.cpp
@@ -17,6 +17,11 @@ uint32 GetTypeHash(const FLearningAgentsActionObjectElement& Element)
 	return (uint32)Element.ObjectElement.Index;
 }
 
+const UE::Learning::Action::FSchema& ULearningAgentsActionSchema::GetActionSchema() const
+{
+	return ActionSchema;
+}
+
 namespace UE::Learning::Agents::Action::Private
 {
 	static inline bool ContainsDuplicates(const TArrayView<const int32> Indices)
@@ -56,31 +61,32 @@ namespace UE::Learning::Agents::Action::Private
 		const Learning::Action::FSchema& Schema,
 		const Learning::Action::FSchemaElement SchemaElement,
 		const Learning::Action::FObject& Object,
-		const Learning::Action::FObjectElement ObjectElement)
+		const Learning::Action::FObjectElement ObjectElement,
+		const FString& ObjectName)
 	{
 		// Check Elements are Valid
 
 		if (!Schema.IsValid(SchemaElement))
 		{
-			UE_LOG(LogLearning, Error, TEXT("ValidateObjectMatchesSchema: Invalid Action Schema Element."));
+			UE_LOG(LogLearning, Error, TEXT("%s: Invalid Action Schema Element."), *ObjectName);
 			return false;
 		}
 
 		if (!Object.IsValid(ObjectElement))
 		{
-			UE_LOG(LogLearning, Error, TEXT("ValidateObjectMatchesSchema: Invalid Action Object Element."));
+			UE_LOG(LogLearning, Error, TEXT("%s: Invalid Action Object Element."), *ObjectName);
 			return false;
 		}
 
 		// Check Names Match
 
-		const FName ActionSchemaElementTag = Schema.GetTag(SchemaElement);
-		const FName ActionObjectElementTag = Object.GetTag(ObjectElement);
+		const FName ActionSchemaElementName = Schema.GetName(SchemaElement);
+		const FName ActionObjectElementName = Object.GetName(ObjectElement);
 
-		if (ActionSchemaElementTag != ActionObjectElementTag)
+		if (ActionSchemaElementName != ActionObjectElementName)
 		{
-			UE_LOG(LogLearning, Warning, TEXT("ValidateObjectMatchesSchema: Action tag does not match Schema. Expected '%s', got '%s'."),
-				*ActionSchemaElementTag.ToString(), *ActionObjectElementTag.ToString());
+			UE_LOG(LogLearning, Warning, TEXT("%s: Action name does not match Schema. Expected '%s', got '%s'."),
+				*ObjectName, *ActionSchemaElementName.ToString(), *ActionObjectElementName.ToString());
 		}
 
 		// Check Types Match
@@ -90,8 +96,9 @@ namespace UE::Learning::Agents::Action::Private
 
 		if (ActionSchemaElementType != ActionObjectElementType)
 		{
-			UE_LOG(LogLearning, Error, TEXT("ValidateObjectMatchesSchema: Action '%s' type does not match Schema. Expected type '%s', got type '%s'."),
-				*ActionSchemaElementTag.ToString(),
+			UE_LOG(LogLearning, Error, TEXT("%s: Action '%s' type does not match Schema. Expected type '%s', got type '%s'."),
+				*ObjectName,
+				*ActionSchemaElementName.ToString(),
 				GetActionTypeString(ActionSchemaElementType),
 				GetActionTypeString(ActionObjectElementType));
 			return false;
@@ -110,8 +117,9 @@ namespace UE::Learning::Agents::Action::Private
 
 			if (SchemaElementSize != ObjectElementSize)
 			{
-				UE_LOG(LogLearning, Error, TEXT("ValidateObjectMatchesSchema: Action '%s' size does not match Schema. Expected '%i', got '%i'."),
-					*ActionSchemaElementTag.ToString(),
+				UE_LOG(LogLearning, Error, TEXT("%s: Action '%s' size does not match Schema. Expected '%i', got '%i'."),
+					*ObjectName,
+					*ActionSchemaElementName.ToString(),
 					SchemaElementSize,
 					ObjectElementSize);
 				return false;
@@ -127,8 +135,9 @@ namespace UE::Learning::Agents::Action::Private
 
 			if (ObjectElementIndex < 0 || ObjectElementIndex >= SchemaElementSize)
 			{
-				UE_LOG(LogLearning, Error, TEXT("ValidateObjectMatchesSchema: Action '%s' index out of range for Schema. Expected '<%i', got '%i'."),
-					*ActionSchemaElementTag.ToString(),
+				UE_LOG(LogLearning, Error, TEXT("%s: Action '%s' index out of range for Schema. Expected '<%i', got '%i'."),
+					*ObjectName,
+					*ActionSchemaElementName.ToString(),
 					SchemaElementSize,
 					ObjectElementIndex);
 				return false;
@@ -144,8 +153,9 @@ namespace UE::Learning::Agents::Action::Private
 
 			if (ObjectElementIndices.Num() > SchemaElementSize)
 			{
-				UE_LOG(LogLearning, Error, TEXT("ValidateObjectMatchesSchema: Action '%s' too many indices provided. Expected at most '%i', got '%i'."),
-					*ActionSchemaElementTag.ToString(),
+				UE_LOG(LogLearning, Error, TEXT("%s: Action '%s' too many indices provided. Expected at most '%i', got '%i'."),
+					*ObjectName,
+					*ActionSchemaElementName.ToString(),
 					SchemaElementSize,
 					ObjectElementIndices.Num());
 				return false;
@@ -155,8 +165,9 @@ namespace UE::Learning::Agents::Action::Private
 			{
 				if (ObjectElementIndices[SubElementIdx] < 0 || ObjectElementIndices[SubElementIdx] >= SchemaElementSize)
 				{
-					UE_LOG(LogLearning, Error, TEXT("ValidateObjectMatchesSchema: Action '%s' index out of range for Schema. Expected '<%i', got '%i'."),
-						*ActionSchemaElementTag.ToString(),
+					UE_LOG(LogLearning, Error, TEXT("%s: Action '%s' index out of range for Schema. Expected '<%i', got '%i'."),
+						*ObjectName,
+						*ActionSchemaElementName.ToString(),
 						SchemaElementSize,
 						ObjectElementIndices[SubElementIdx]);
 					return false;
@@ -175,8 +186,9 @@ namespace UE::Learning::Agents::Action::Private
 
 			if (SchemaParameters.Elements.Num() != ObjectParameters.Elements.Num())
 			{
-				UE_LOG(LogLearning, Error, TEXT("ValidateObjectMatchesSchema: Action '%s' number of sub-elements does not match Schema. Expected '%i', got '%i'."),
-					*ActionSchemaElementTag.ToString(),
+				UE_LOG(LogLearning, Error, TEXT("%s: Action '%s' number of sub-elements does not match Schema. Expected '%i', got '%i'."),
+					*ObjectName,
+					*ActionSchemaElementName.ToString(),
 					SchemaParameters.Elements.Num(),
 					ObjectParameters.Elements.Num());
 				return false;
@@ -188,8 +200,9 @@ namespace UE::Learning::Agents::Action::Private
 
 				if (ObjectElementIdx == INDEX_NONE)
 				{
-					UE_LOG(LogLearning, Error, TEXT("ValidateObjectMatchesSchema: Action '%s' does not include '%s' action required by Schema."),
-						*ActionSchemaElementTag.ToString(),
+					UE_LOG(LogLearning, Error, TEXT("%s: Action '%s' does not include '%s' action required by Schema."),
+						*ObjectName,
+						*ActionSchemaElementName.ToString(),
 						*SchemaParameters.ElementNames[SchemaElementIdx].ToString());
 					return false;
 				}
@@ -198,7 +211,8 @@ namespace UE::Learning::Agents::Action::Private
 					Schema,
 					SchemaParameters.Elements[SchemaElementIdx],
 					Object,
-					ObjectParameters.Elements[ObjectElementIdx]))
+					ObjectParameters.Elements[ObjectElementIdx],
+					ObjectName))
 				{
 					return false;
 				}
@@ -217,8 +231,9 @@ namespace UE::Learning::Agents::Action::Private
 
 			if (SchemaSubElementIdx == INDEX_NONE)
 			{
-				UE_LOG(LogLearning, Error, TEXT("ValidateObjectMatchesSchema: Action '%s' Schema does not include '%s' action."),
-					*ActionSchemaElementTag.ToString(),
+				UE_LOG(LogLearning, Error, TEXT("%s: Action '%s' Schema does not include '%s' action."),
+					*ObjectName,
+					*ActionSchemaElementName.ToString(),
 					*ObjectParameters.ElementName.ToString());
 				return false;
 			}
@@ -227,7 +242,8 @@ namespace UE::Learning::Agents::Action::Private
 				Schema,
 				SchemaParameters.Elements[SchemaSubElementIdx],
 				Object,
-				ObjectParameters.Element);
+				ObjectParameters.Element,
+				ObjectName);
 		}
 
 		case Learning::Action::EType::OrInclusive:
@@ -237,8 +253,9 @@ namespace UE::Learning::Agents::Action::Private
 
 			if (ObjectParameters.Elements.Num() > SchemaParameters.Elements.Num())
 			{
-				UE_LOG(LogLearning, Error, TEXT("ValidateObjectMatchesSchema: Action '%s' too many sub-actions provided. Expected at most '%i', got '%i'."),
-					*ActionSchemaElementTag.ToString(),
+				UE_LOG(LogLearning, Error, TEXT("%s: Action '%s' too many sub-actions provided. Expected at most '%i', got '%i'."),
+					*ObjectName,
+					*ActionSchemaElementName.ToString(),
 					SchemaParameters.Elements.Num(),
 					ObjectParameters.Elements.Num());
 				return false;
@@ -250,8 +267,9 @@ namespace UE::Learning::Agents::Action::Private
 
 				if (SchemaSubElementIdx == INDEX_NONE)
 				{
-					UE_LOG(LogLearning, Error, TEXT("ValidateObjectMatchesSchema: Action '%s' Schema does not include '%s' action."),
-						*ActionSchemaElementTag.ToString(),
+					UE_LOG(LogLearning, Error, TEXT("%s: Action '%s' Schema does not include '%s' action."),
+						*ObjectName,
+						*ActionSchemaElementName.ToString(),
 						*ObjectParameters.ElementNames[ObjectSubElementIdx].ToString());
 					return false;
 				}
@@ -260,7 +278,8 @@ namespace UE::Learning::Agents::Action::Private
 					Schema,
 					SchemaParameters.Elements[SchemaSubElementIdx],
 					Object,
-					ObjectParameters.Elements[ObjectSubElementIdx]))
+					ObjectParameters.Elements[ObjectSubElementIdx],
+					ObjectName))
 				{
 					return false;
 				}
@@ -276,8 +295,9 @@ namespace UE::Learning::Agents::Action::Private
 
 			if (ObjectParameters.Elements.Num() != SchemaParameters.Num)
 			{
-				UE_LOG(LogLearning, Error, TEXT("ValidateObjectMatchesSchema: Action '%s' array incorrect size. Expected '%i' elements, got '%i'."),
-					*ActionSchemaElementTag.ToString(),
+				UE_LOG(LogLearning, Error, TEXT("%s: Action '%s' array incorrect size. Expected '%i' elements, got '%i'."),
+					*ObjectName,
+					*ActionSchemaElementName.ToString(),
 					SchemaParameters.Num,
 					ObjectParameters.Elements.Num());
 				return false;
@@ -289,7 +309,8 @@ namespace UE::Learning::Agents::Action::Private
 					Schema,
 					SchemaParameters.Element,
 					Object,
-					ObjectParameters.Elements[ElementIdx]))
+					ObjectParameters.Elements[ElementIdx],
+					ObjectName))
 				{
 					return false;
 				}
@@ -307,7 +328,8 @@ namespace UE::Learning::Agents::Action::Private
 					Schema,
 					SchemaParameters.Element,
 					Object,
-					ObjectParameters.Element);
+					ObjectParameters.Element,
+					ObjectName);
 		}
 
 		default:
@@ -322,22 +344,23 @@ namespace UE::Learning::Agents::Action::Private
 		const UE::Learning::Action::FObject& Object,
 		const UE::Learning::Action::FObjectElement ObjectElement,
 		const FString& Indentation,
-		const FString& Prefix)
+		const FString& Prefix,
+		const FString& ObjectName)
 	{
 		if (!Object.IsValid(ObjectElement))
 		{
-			UE_LOG(LogLearning, Error, TEXT("LogAction: Invalid Action Object Element."));
+			UE_LOG(LogLearning, Error, TEXT("%s: Invalid Action Object Element."), *ObjectName);
 			return;
 		}
 
 		const UE::Learning::Action::EType Type = Object.GetType(ObjectElement);
-		const FName Tag = Object.GetTag(ObjectElement);
+		const FName Name = Object.GetName(ObjectElement);
 
 		switch (Type)
 		{
 		case UE::Learning::Action::EType::Null:
 		{
-			UE_LOG(LogLearning, Display, TEXT("%s%s \"%s\" (%s)"), *Indentation, *Prefix, *Tag.ToString(), GetActionTypeString(Type));
+			UE_LOG(LogLearning, Display, TEXT("%s%s \"%s\" (%s)"), *Indentation, *Prefix, *Name.ToString(), GetActionTypeString(Type));
 			return;
 		}
 
@@ -345,7 +368,7 @@ namespace UE::Learning::Agents::Action::Private
 		{
 			const UE::Learning::Action::FObjectContinuousParameters Parameters = Object.GetContinuous(ObjectElement);
 
-			UE_LOG(LogLearning, Display, TEXT("%s%s \"%s\" (%s) %s"), *Indentation, *Prefix, *Tag.ToString(), GetActionTypeString(Type), *UE::Learning::Array::FormatFloat(Parameters.Values));
+			UE_LOG(LogLearning, Display, TEXT("%s%s \"%s\" (%s) %s"), *Indentation, *Prefix, *Name.ToString(), GetActionTypeString(Type), *UE::Learning::Array::FormatFloat(Parameters.Values));
 			return;
 		}
 
@@ -353,7 +376,7 @@ namespace UE::Learning::Agents::Action::Private
 		{
 			const UE::Learning::Action::FObjectDiscreteExclusiveParameters Parameters = Object.GetDiscreteExclusive(ObjectElement);
 
-			UE_LOG(LogLearning, Display, TEXT("%s%s \"%s\" (%s) %i"), *Indentation, *Prefix, *Tag.ToString(), GetActionTypeString(Type), Parameters.DiscreteIndex);
+			UE_LOG(LogLearning, Display, TEXT("%s%s \"%s\" (%s) %i"), *Indentation, *Prefix, *Name.ToString(), GetActionTypeString(Type), Parameters.DiscreteIndex);
 			return;
 		}
 
@@ -361,7 +384,7 @@ namespace UE::Learning::Agents::Action::Private
 		{
 			const UE::Learning::Action::FObjectDiscreteInclusiveParameters Parameters = Object.GetDiscreteInclusive(ObjectElement);
 
-			UE_LOG(LogLearning, Display, TEXT("%s%s \"%s\" (%s) %s"), *Indentation, *Prefix, *Tag.ToString(), GetActionTypeString(Type), *UE::Learning::Array::FormatInt32(Parameters.DiscreteIndices));
+			UE_LOG(LogLearning, Display, TEXT("%s%s \"%s\" (%s) %s"), *Indentation, *Prefix, *Name.ToString(), GetActionTypeString(Type), *UE::Learning::Array::FormatInt32(Parameters.DiscreteIndices));
 			return;
 		}
 
@@ -369,10 +392,10 @@ namespace UE::Learning::Agents::Action::Private
 		{
 			const UE::Learning::Action::FObjectAndParameters Parameters = Object.GetAnd(ObjectElement);
 
-			UE_LOG(LogLearning, Display, TEXT("%s%s \"%s\" (%s)"), *Indentation, *Prefix, *Tag.ToString(), GetActionTypeString(Type));
+			UE_LOG(LogLearning, Display, TEXT("%s%s \"%s\" (%s)"), *Indentation, *Prefix, *Name.ToString(), GetActionTypeString(Type));
 			for (int32 SubElementIdx = 0; SubElementIdx < Parameters.Elements.Num(); SubElementIdx++)
 			{
-				LogAction(Object, Parameters.Elements[SubElementIdx], *(Indentation + TEXT("    ")), FString::Printf(TEXT("| \"%s\": "), *Parameters.ElementNames[SubElementIdx].ToString()));
+				LogAction(Object, Parameters.Elements[SubElementIdx], *(Indentation + TEXT("    ")), FString::Printf(TEXT("| \"%s\": "), *Parameters.ElementNames[SubElementIdx].ToString()), ObjectName);
 			}
 
 			return;
@@ -382,8 +405,8 @@ namespace UE::Learning::Agents::Action::Private
 		{
 			const UE::Learning::Action::FObjectOrExclusiveParameters Parameters = Object.GetOrExclusive(ObjectElement);
 
-			UE_LOG(LogLearning, Display, TEXT("%s%s \"%s\" (%s)"), *Indentation, *Prefix, *Tag.ToString(), GetActionTypeString(Type));
-			LogAction(Object, Parameters.Element, *(Indentation + TEXT("    ")), FString::Printf(TEXT("| \"%s\": "), *Parameters.ElementName.ToString()));
+			UE_LOG(LogLearning, Display, TEXT("%s%s \"%s\" (%s)"), *Indentation, *Prefix, *Name.ToString(), GetActionTypeString(Type));
+			LogAction(Object, Parameters.Element, *(Indentation + TEXT("    ")), FString::Printf(TEXT("| \"%s\": "), *Parameters.ElementName.ToString()), ObjectName);
 
 			return;
 		}
@@ -392,10 +415,10 @@ namespace UE::Learning::Agents::Action::Private
 		{
 			const UE::Learning::Action::FObjectOrInclusiveParameters Parameters = Object.GetOrInclusive(ObjectElement);
 
-			UE_LOG(LogLearning, Display, TEXT("%s%s \"%s\" (%s)"), *Indentation, *Prefix, *Tag.ToString(), GetActionTypeString(Type));
+			UE_LOG(LogLearning, Display, TEXT("%s%s \"%s\" (%s)"), *Indentation, *Prefix, *Name.ToString(), GetActionTypeString(Type));
 			for (int32 SubElementIdx = 0; SubElementIdx < Parameters.Elements.Num(); SubElementIdx++)
 			{
-				LogAction(Object, Parameters.Elements[SubElementIdx], *(Indentation + TEXT("    ")), FString::Printf(TEXT("| \"%s\": "), *Parameters.ElementNames[SubElementIdx].ToString()));
+				LogAction(Object, Parameters.Elements[SubElementIdx], *(Indentation + TEXT("    ")), FString::Printf(TEXT("| \"%s\": "), *Parameters.ElementNames[SubElementIdx].ToString()), ObjectName);
 			}
 
 			return;
@@ -405,10 +428,10 @@ namespace UE::Learning::Agents::Action::Private
 		{
 			const UE::Learning::Action::FObjectArrayParameters Parameters = Object.GetArray(ObjectElement);
 
-			UE_LOG(LogLearning, Display, TEXT("%s%s \"%s\" (%s)"), *Indentation, *Prefix, *Tag.ToString(), GetActionTypeString(Type));
+			UE_LOG(LogLearning, Display, TEXT("%s%s \"%s\" (%s)"), *Indentation, *Prefix, *Name.ToString(), GetActionTypeString(Type));
 			for (int32 SubElementIdx = 0; SubElementIdx < Parameters.Elements.Num(); SubElementIdx++)
 			{
-				LogAction(Object, Parameters.Elements[SubElementIdx], *(Indentation + TEXT("    ")), FString::Printf(TEXT("| %3i:"), SubElementIdx));
+				LogAction(Object, Parameters.Elements[SubElementIdx], *(Indentation + TEXT("    ")), FString::Printf(TEXT("| %3i:"), SubElementIdx), ObjectName);
 			}
 
 			return;
@@ -418,8 +441,8 @@ namespace UE::Learning::Agents::Action::Private
 		{
 			const UE::Learning::Action::FObjectEncodingParameters Parameters = Object.GetEncoding(ObjectElement);
 
-			UE_LOG(LogLearning, Display, TEXT("%s%s \"%s\" (%s)"), *Indentation, *Prefix, *Tag.ToString(), GetActionTypeString(Type));
-			LogAction(Object, Parameters.Element, *(Indentation + TEXT("    ")), TEXT("|"));
+			UE_LOG(LogLearning, Display, TEXT("%s%s \"%s\" (%s)"), *Indentation, *Prefix, *Name.ToString(), GetActionTypeString(Type));
+			LogAction(Object, Parameters.Element, *(Indentation + TEXT("    ")), TEXT("|"), ObjectName);
 
 			return;
 		}
@@ -442,14 +465,14 @@ namespace UE::Learning::Agents::Action::Private
 			FMath::Exp(V.Z));
 	}
 
-	static inline void NormalizeProbabilitiesExclusive(TArrayView<float> PriorProbabilities)
+	static inline void NormalizeProbabilitiesExclusive(TArrayView<float> PriorProbabilities, const FString& ObjectName)
 	{
 		float Total = 0.0f;
 		for (int32 Idx = 0; Idx < PriorProbabilities.Num(); Idx++)
 		{
 			if (PriorProbabilities[Idx] < 0.0f || PriorProbabilities[Idx] > 1.0f)
 			{
-				UE_LOG(LogLearning, Warning, TEXT("NormalizeProbabilitiesExclusive: Invalid Prior Probability Given (%f), must be in range 0 to 1."), PriorProbabilities[Idx]);
+				UE_LOG(LogLearning, Warning, TEXT("%s: Invalid Prior Probability Given (%f), must be in range 0 to 1."), *ObjectName, PriorProbabilities[Idx]);
 			}
 
 			PriorProbabilities[Idx] = FMath::Clamp(PriorProbabilities[Idx], 0.0f, 1.0f);
@@ -458,7 +481,7 @@ namespace UE::Learning::Agents::Action::Private
 
 		if (PriorProbabilities.Num() > 0 && FMath::Abs(Total) < UE_SMALL_NUMBER)
 		{
-			UE_LOG(LogLearning, Warning, TEXT("NormalizeProbabilitiesExclusive: Prior Probabilities are too small. Should sum to 1."));
+			UE_LOG(LogLearning, Warning, TEXT("%s: Prior Probabilities are too small. Should sum to 1."), *ObjectName);
 
 			for (int32 Idx = 0; Idx < PriorProbabilities.Num(); Idx++)
 			{
@@ -474,13 +497,13 @@ namespace UE::Learning::Agents::Action::Private
 		}
 	}
 
-	static inline void NormalizeProbabilitiesInclusive(TArrayView<float> PriorProbabilities)
+	static inline void NormalizeProbabilitiesInclusive(TArrayView<float> PriorProbabilities, const FString& ObjectName)
 	{
 		for (int32 Idx = 0; Idx < PriorProbabilities.Num(); Idx++)
 		{
 			if (PriorProbabilities[Idx] < 0.0f || PriorProbabilities[Idx] > 1.0f)
 			{
-				UE_LOG(LogLearning, Warning, TEXT("NormalizeProbabilitiesInclusive: Invalid Prior Probability Given (%f), must be in range 0 to 1."), PriorProbabilities[Idx]);
+				UE_LOG(LogLearning, Warning, TEXT("%s: Invalid Prior Probability Given (%f), must be in range 0 to 1."), *ObjectName, PriorProbabilities[Idx]);
 			}
 
 			PriorProbabilities[Idx] = FMath::Clamp(PriorProbabilities[Idx], 0.0f, 1.0f);
@@ -499,86 +522,56 @@ namespace UE::Learning::Agents::Action::Private
 	}
 }
 
-bool ULearningAgentsActions::ValidateObjectMatchesSchema(
-	const ULearningAgentsActionSchema* Schema,
+bool ULearningAgentsActionSchema::ValidateObjectMatchesSchema(
 	const FLearningAgentsActionSchemaElement SchemaElement,
 	const ULearningAgentsActionObject* Object,
-	const FLearningAgentsActionObjectElement ObjectElement)
+	const FLearningAgentsActionObjectElement ObjectElement) const
 {
-	if (!Schema)
-	{
-		UE_LOG(LogLearning, Error, TEXT("ValidateObjectMatchesSchema: Schema is nullptr."));
-		return false;
-	}
-
-	if (!Object)
-	{
-		UE_LOG(LogLearning, Error, TEXT("ValidateObjectMatchesSchema: Object is nullptr."));
-		return false;
-	}
-
 	return UE::Learning::Agents::Action::Private::ValidateObjectMatchesSchema(
-		Schema->ActionSchema,
+		ActionSchema,
 		SchemaElement.SchemaElement,
-		Object->ActionObject,
-		ObjectElement.ObjectElement);
+		Object->GetActionObject(),
+		ObjectElement.ObjectElement,
+		GetName());
 }
 
-FLearningAgentsActionSchemaElement ULearningAgentsActions::SpecifyNullAction(ULearningAgentsActionSchema* Schema, const FName Tag)
+FLearningAgentsActionSchemaElement ULearningAgentsActionSchema::SpecifyNullAction(const FName Name)
 {
-	if (!Schema)
-	{
-		UE_LOG(LogLearning, Error, TEXT("SpecifyNullAction: Schema is nullptr."));
-		return FLearningAgentsActionSchemaElement();
-	}
-
-	return { Schema->ActionSchema.CreateNull(Tag) };
+	return { ActionSchema.CreateNull(Name) };
 }
 
-FLearningAgentsActionSchemaElement ULearningAgentsActions::SpecifyContinuousAction(ULearningAgentsActionSchema* Schema, const int32 Size, const FName Tag)
+FLearningAgentsActionSchemaElement ULearningAgentsActionSchema::SpecifyContinuousAction(const int32 Size, const FName Name)
 {
-	if (!Schema)
-	{
-		UE_LOG(LogLearning, Error, TEXT("SpecifyContinuousAction: Schema is nullptr."));
-		return FLearningAgentsActionSchemaElement();
-	}
-
 	if (Size < 0)
 	{
-		UE_LOG(LogLearning, Error, TEXT("SpecifyContinuousAction: Invalid Continuous Action Size '%i'."), Size);
+		UE_LOG(LogLearning, Error, TEXT("%s: Invalid Continuous Action Size '%i'."), *GetName(), Size);
 		return FLearningAgentsActionSchemaElement();
 	}
 
 	if (Size == 0)
 	{
-		UE_LOG(LogLearning, Warning, TEXT("SpecifyContinuousAction: Specifying zero-sized Continuous Action."));
+		UE_LOG(LogLearning, Warning, TEXT("%s: Specifying zero-sized Continuous Action."), *GetName());
 	}
 
-	return { Schema->ActionSchema.CreateContinuous({ Size }, Tag) };
+	return { ActionSchema.CreateContinuous({ Size }, Name) };
 }
 
-FLearningAgentsActionSchemaElement ULearningAgentsActions::SpecifyExclusiveDiscreteAction(ULearningAgentsActionSchema* Schema, const int32 Size, const TArray<float>& PriorProbabilities, const FName Tag)
+FLearningAgentsActionSchemaElement ULearningAgentsActionSchema::SpecifyExclusiveDiscreteAction(const int32 Size, const TArray<float>& PriorProbabilities, const FName Name)
 {
-	return SpecifyExclusiveDiscreteActionFromArrayView(Schema, Size, PriorProbabilities, Tag);
+	return SpecifyExclusiveDiscreteActionFromArrayView(Size, PriorProbabilities, Name);
 }
 
-FLearningAgentsActionSchemaElement ULearningAgentsActions::SpecifyExclusiveDiscreteActionFromArrayView(ULearningAgentsActionSchema* Schema, const int32 Size, const TArrayView<const float> PriorProbabilities, const FName Tag)
+FLearningAgentsActionSchemaElement ULearningAgentsActionSchema::SpecifyExclusiveDiscreteActionFromArrayView(const int32 Size, const TArrayView<const float> PriorProbabilities, const FName Name)
 {
-	if (!Schema)
-	{
-		UE_LOG(LogLearning, Error, TEXT("SpecifyExclusiveDiscreteActionFromArrayView: Schema is nullptr."));
-		return FLearningAgentsActionSchemaElement();
-	}
-
 	if (Size < 0)
 	{
-		UE_LOG(LogLearning, Error, TEXT("SpecifyExclusiveDiscreteActionFromArrayView: Invalid DiscreteExclusive Action Size '%i'."), Size);
+		UE_LOG(LogLearning, Error, TEXT("%s: Invalid DiscreteExclusive Action Size '%i'."), *GetName(), Size);
 		return FLearningAgentsActionSchemaElement();
 	}
 
 	if (Size == 0)
 	{
-		UE_LOG(LogLearning, Warning, TEXT("SpecifyExclusiveDiscreteActionFromArrayView: Specifying zero-sized Exclusive Discrete Action."));
+		UE_LOG(LogLearning, Warning, TEXT("%s: Specifying zero-sized Exclusive Discrete Action."), *GetName());
 	}
 
 	TArray<float, TInlineAllocator<16>> NormalizedPriorProbabilities;
@@ -587,33 +580,27 @@ FLearningAgentsActionSchemaElement ULearningAgentsActions::SpecifyExclusiveDiscr
 	{
 		NormalizedPriorProbabilities[Idx] = PriorProbabilities[Idx];
 	}
-	UE::Learning::Agents::Action::Private::NormalizeProbabilitiesExclusive(NormalizedPriorProbabilities);
+	UE::Learning::Agents::Action::Private::NormalizeProbabilitiesExclusive(NormalizedPriorProbabilities, GetName());
 
-	return { Schema->ActionSchema.CreateDiscreteExclusive({ Size, MakeArrayView(NormalizedPriorProbabilities) }, Tag)};
+	return { ActionSchema.CreateDiscreteExclusive({ Size, MakeArrayView(NormalizedPriorProbabilities) }, Name)};
 }
 
-FLearningAgentsActionSchemaElement ULearningAgentsActions::SpecifyInclusiveDiscreteAction(ULearningAgentsActionSchema* Schema, const int32 Size, const TArray<float>& PriorProbabilities, const FName Tag)
+FLearningAgentsActionSchemaElement ULearningAgentsActionSchema::SpecifyInclusiveDiscreteAction(const int32 Size, const TArray<float>& PriorProbabilities, const FName Name)
 {
-	return SpecifyInclusiveDiscreteActionFromArrayView(Schema, Size, PriorProbabilities, Tag);
+	return SpecifyInclusiveDiscreteActionFromArrayView(Size, PriorProbabilities, Name);
 }
 
-FLearningAgentsActionSchemaElement ULearningAgentsActions::SpecifyInclusiveDiscreteActionFromArrayView(ULearningAgentsActionSchema* Schema, const int32 Size, const TArrayView<const float> PriorProbabilities, const FName Tag)
+FLearningAgentsActionSchemaElement ULearningAgentsActionSchema::SpecifyInclusiveDiscreteActionFromArrayView(const int32 Size, const TArrayView<const float> PriorProbabilities, const FName Name)
 {
-	if (!Schema)
-	{
-		UE_LOG(LogLearning, Error, TEXT("SpecifyInclusiveDiscreteActionFromArrayView: Schema is nullptr."));
-		return FLearningAgentsActionSchemaElement();
-	}
-
 	if (Size < 0)
 	{
-		UE_LOG(LogLearning, Error, TEXT("SpecifyInclusiveDiscreteActionFromArrayView: Invalid DiscreteInclusive Action Size '%i'."), Size);
+		UE_LOG(LogLearning, Error, TEXT("%s: Invalid DiscreteInclusive Action Size '%i'."), *GetName(), Size);
 		return FLearningAgentsActionSchemaElement();
 	}
 
 	if (Size == 0)
 	{
-		UE_LOG(LogLearning, Warning, TEXT("SpecifyInclusiveDiscreteActionFromArrayView: Specifying zero-sized Inclusive Discrete Action."));
+		UE_LOG(LogLearning, Warning, TEXT("%s: Specifying zero-sized Inclusive Discrete Action."), *GetName());
 	}
 
 	TArray<float, TInlineAllocator<16>> NormalizedPriorProbabilities;
@@ -622,16 +609,16 @@ FLearningAgentsActionSchemaElement ULearningAgentsActions::SpecifyInclusiveDiscr
 	{
 		NormalizedPriorProbabilities[Idx] = PriorProbabilities[Idx];
 	}
-	UE::Learning::Agents::Action::Private::NormalizeProbabilitiesInclusive(NormalizedPriorProbabilities);
+	UE::Learning::Agents::Action::Private::NormalizeProbabilitiesInclusive(NormalizedPriorProbabilities, GetName());
 
-	return { Schema->ActionSchema.CreateDiscreteInclusive({ Size, MakeArrayView(NormalizedPriorProbabilities) }, Tag) };
+	return { ActionSchema.CreateDiscreteInclusive({ Size, MakeArrayView(NormalizedPriorProbabilities) }, Name) };
 }
 
-FLearningAgentsActionSchemaElement ULearningAgentsActions::SpecifyStructAction(ULearningAgentsActionSchema* Schema, const TMap<FName, FLearningAgentsActionSchemaElement>& Elements, const FName Tag)
+FLearningAgentsActionSchemaElement ULearningAgentsActionSchema::SpecifyStructAction(const TMap<FName, FLearningAgentsActionSchemaElement>& Elements, const FName Name)
 {
 	if (Elements.Num() == 0)
 	{
-		UE_LOG(LogLearning, Warning, TEXT("SpecifyStructAction: Specifying zero-sized Struct Action."));
+		UE_LOG(LogLearning, Warning, TEXT("%s: Specifying zero-sized Struct Action."), *GetName());
 	}
 
 	const int32 SubElementNum = Elements.Num();
@@ -669,36 +656,30 @@ FLearningAgentsActionSchemaElement ULearningAgentsActions::SpecifyStructAction(U
 		SortedSubElements[Idx] = SubElements[SubElementIndices[Idx]];
 	}
 
-	return SpecifyStructActionFromArrayViews(Schema, SortedSubElementNames, SortedSubElements, Tag);
+	return SpecifyStructActionFromArrayViews(SortedSubElementNames, SortedSubElements, Name);
 }
 
-FLearningAgentsActionSchemaElement ULearningAgentsActions::SpecifyStructActionFromArrays(ULearningAgentsActionSchema* Schema, const TArray<FName>& ElementNames, const TArray<FLearningAgentsActionSchemaElement>& Elements, const FName Tag)
+FLearningAgentsActionSchemaElement ULearningAgentsActionSchema::SpecifyStructActionFromArrays(const TArray<FName>& ElementNames, const TArray<FLearningAgentsActionSchemaElement>& Elements, const FName Name)
 {
-	return SpecifyStructActionFromArrayViews(Schema, ElementNames, Elements, Tag);
+	return SpecifyStructActionFromArrayViews(ElementNames, Elements, Name);
 }
 
-FLearningAgentsActionSchemaElement ULearningAgentsActions::SpecifyStructActionFromArrayViews(ULearningAgentsActionSchema* Schema, const TArrayView<const FName> ElementNames, const TArrayView<const FLearningAgentsActionSchemaElement> Elements, const FName Tag)
+FLearningAgentsActionSchemaElement ULearningAgentsActionSchema::SpecifyStructActionFromArrayViews(const TArrayView<const FName> ElementNames, const TArrayView<const FLearningAgentsActionSchemaElement> Elements, const FName Name)
 {
-	if (!Schema)
-	{
-		UE_LOG(LogLearning, Error, TEXT("SpecifyStructActionFromArrayViews: Schema is nullptr."));
-		return FLearningAgentsActionSchemaElement();
-	}
-
 	if (Elements.Num() == 0)
 	{
-		UE_LOG(LogLearning, Warning, TEXT("SpecifyStructActionFromArrayViews: Specifying zero-sized Struct Action."));
+		UE_LOG(LogLearning, Warning, TEXT("%s: Specifying zero-sized Struct Action."), *GetName());
 	}
 
 	if (Elements.Num() != ElementNames.Num())
 	{
-		UE_LOG(LogLearning, Error, TEXT("SpecifyStructActionFromArrayViews: Number of elements (%i) must match number of names (%i)."), Elements.Num(), ElementNames.Num());
+		UE_LOG(LogLearning, Error, TEXT("%s: Number of elements (%i) must match number of names (%i)."), *GetName(), Elements.Num(), ElementNames.Num());
 		return FLearningAgentsActionSchemaElement();
 	}
 
 	if (UE::Learning::Agents::Action::Private::ContainsDuplicates(ElementNames))
 	{
-		UE_LOG(LogLearning, Error, TEXT("SpecifyStructActionFromArrayViews: Element Names contain duplicates."));
+		UE_LOG(LogLearning, Error, TEXT("%s: Element Names contain duplicates."), *GetName());
 		return FLearningAgentsActionSchemaElement();
 	}
 
@@ -707,23 +688,23 @@ FLearningAgentsActionSchemaElement ULearningAgentsActions::SpecifyStructActionFr
 
 	for (const FLearningAgentsActionSchemaElement& Element : Elements)
 	{
-		if (!Schema->ActionSchema.IsValid(Element.SchemaElement))
+		if (!ActionSchema.IsValid(Element.SchemaElement))
 		{
-			UE_LOG(LogLearning, Error, TEXT("SpecifyStructActionFromArrayViews: Invalid Action Object."));
+			UE_LOG(LogLearning, Error, TEXT("%s: Invalid Action Object."), *GetName());
 			return FLearningAgentsActionSchemaElement();
 		}
 
 		SubElements.Add(Element.SchemaElement);
 	}
 
-	return { Schema->ActionSchema.CreateAnd({ ElementNames, SubElements }, Tag) };
+	return { ActionSchema.CreateAnd({ ElementNames, SubElements }, Name) };
 }
 
-FLearningAgentsActionSchemaElement ULearningAgentsActions::SpecifyExclusiveUnionAction(ULearningAgentsActionSchema* Schema, const TMap<FName, FLearningAgentsActionSchemaElement>& Elements, const TMap<FName, float>& PriorProbabilities, const FName Tag)
+FLearningAgentsActionSchemaElement ULearningAgentsActionSchema::SpecifyExclusiveUnionAction(const TMap<FName, FLearningAgentsActionSchemaElement>& Elements, const TMap<FName, float>& PriorProbabilities, const FName Name)
 {
 	if (Elements.Num() == 0)
 	{
-		UE_LOG(LogLearning, Warning, TEXT("SpecifyExclusiveUnionAction: Specifying zero-sized Exclusive Union Action."));
+		UE_LOG(LogLearning, Warning, TEXT("%s: Specifying zero-sized Exclusive Union Action."), *GetName());
 	}
 
 	const int32 SubElementNum = Elements.Num();
@@ -768,36 +749,30 @@ FLearningAgentsActionSchemaElement ULearningAgentsActions::SpecifyExclusiveUnion
 		SortedSubElementPriorProbabilities[Idx] = SubElementPriorProbabilities[SubElementIndices[Idx]];
 	}
 
-	return SpecifyExclusiveUnionActionFromArrayViews(Schema, SortedSubElementNames, SortedSubElements, SortedSubElementPriorProbabilities, Tag);
+	return SpecifyExclusiveUnionActionFromArrayViews(SortedSubElementNames, SortedSubElements, SortedSubElementPriorProbabilities, Name);
 }
 
-FLearningAgentsActionSchemaElement ULearningAgentsActions::SpecifyExclusiveUnionActionFromArrays(ULearningAgentsActionSchema* Schema, const TArray<FName>& ElementNames, const TArray<FLearningAgentsActionSchemaElement>& Elements, const TArray<float>& PriorProbabilities, const FName Tag)
+FLearningAgentsActionSchemaElement ULearningAgentsActionSchema::SpecifyExclusiveUnionActionFromArrays(const TArray<FName>& ElementNames, const TArray<FLearningAgentsActionSchemaElement>& Elements, const TArray<float>& PriorProbabilities, const FName Name)
 {
-	return SpecifyExclusiveUnionActionFromArrayViews(Schema, ElementNames, Elements, PriorProbabilities, Tag);
+	return SpecifyExclusiveUnionActionFromArrayViews(ElementNames, Elements, PriorProbabilities, Name);
 }
 
-FLearningAgentsActionSchemaElement ULearningAgentsActions::SpecifyExclusiveUnionActionFromArrayViews(ULearningAgentsActionSchema* Schema, const TArrayView<const FName> ElementNames, const TArrayView<const FLearningAgentsActionSchemaElement> Elements, const TArrayView<const float> PriorProbabilities, const FName Tag)
+FLearningAgentsActionSchemaElement ULearningAgentsActionSchema::SpecifyExclusiveUnionActionFromArrayViews(const TArrayView<const FName> ElementNames, const TArrayView<const FLearningAgentsActionSchemaElement> Elements, const TArrayView<const float> PriorProbabilities, const FName Name)
 {
-	if (!Schema)
-	{
-		UE_LOG(LogLearning, Error, TEXT("SpecifyExclusiveUnionActionFromArrayViews: Schema is nullptr."));
-		return FLearningAgentsActionSchemaElement();
-	}
-
 	if (Elements.Num() == 0)
 	{
-		UE_LOG(LogLearning, Warning, TEXT("SpecifyExclusiveUnionActionFromArrayViews: Specifying zero-sized Exclusive Union Action."));
+		UE_LOG(LogLearning, Warning, TEXT("%s: Specifying zero-sized Exclusive Union Action."), *GetName());
 	}
 
 	if (Elements.Num() != ElementNames.Num())
 	{
-		UE_LOG(LogLearning, Error, TEXT("SpecifyExclusiveUnionActionFromArrayViews: Number of elements (%i) must match number of names (%i)."), Elements.Num(), ElementNames.Num());
+		UE_LOG(LogLearning, Error, TEXT("%s: Number of elements (%i) must match number of names (%i)."), *GetName(), Elements.Num(), ElementNames.Num());
 		return FLearningAgentsActionSchemaElement();
 	}
 
 	if (UE::Learning::Agents::Action::Private::ContainsDuplicates(ElementNames))
 	{
-		UE_LOG(LogLearning, Error, TEXT("SpecifyExclusiveUnionActionFromArrayViews: Element Names contain duplicates."));
+		UE_LOG(LogLearning, Error, TEXT("%s: Element Names contain duplicates."), *GetName());
 		return FLearningAgentsActionSchemaElement();
 	}
 
@@ -806,9 +781,9 @@ FLearningAgentsActionSchemaElement ULearningAgentsActions::SpecifyExclusiveUnion
 
 	for (const FLearningAgentsActionSchemaElement& Element : Elements)
 	{
-		if (!Schema->ActionSchema.IsValid(Element.SchemaElement))
+		if (!ActionSchema.IsValid(Element.SchemaElement))
 		{
-			UE_LOG(LogLearning, Error, TEXT("SpecifyExclusiveUnionActionFromArrayViews: Invalid Action Object."));
+			UE_LOG(LogLearning, Error, TEXT("%s: Invalid Action Object."), *GetName());
 			return FLearningAgentsActionSchemaElement();
 		}
 
@@ -821,16 +796,16 @@ FLearningAgentsActionSchemaElement ULearningAgentsActions::SpecifyExclusiveUnion
 	{
 		NormalizedPriorProbabilities[Idx] = PriorProbabilities[Idx];
 	}
-	UE::Learning::Agents::Action::Private::NormalizeProbabilitiesExclusive(NormalizedPriorProbabilities);
+	UE::Learning::Agents::Action::Private::NormalizeProbabilitiesExclusive(NormalizedPriorProbabilities, GetName());
 
-	return { Schema->ActionSchema.CreateOrExclusive({ ElementNames, SubElements, MakeArrayView(NormalizedPriorProbabilities) }, Tag)};
+	return { ActionSchema.CreateOrExclusive({ ElementNames, SubElements, MakeArrayView(NormalizedPriorProbabilities) }, Name)};
 }
 
-FLearningAgentsActionSchemaElement ULearningAgentsActions::SpecifyInclusiveUnionAction(ULearningAgentsActionSchema* Schema, const TMap<FName, FLearningAgentsActionSchemaElement>& Elements, const TMap<FName, float>& PriorProbabilities, const FName Tag)
+FLearningAgentsActionSchemaElement ULearningAgentsActionSchema::SpecifyInclusiveUnionAction(const TMap<FName, FLearningAgentsActionSchemaElement>& Elements, const TMap<FName, float>& PriorProbabilities, const FName Name)
 {
 	if (Elements.Num() == 0)
 	{
-		UE_LOG(LogLearning, Warning, TEXT("SpecifyInclusiveUnionAction: Specifying zero-sized Inclusive Union Action."));
+		UE_LOG(LogLearning, Warning, TEXT("%s: Specifying zero-sized Inclusive Union Action."), *GetName());
 	}
 
 	const int32 SubElementNum = Elements.Num();
@@ -875,36 +850,30 @@ FLearningAgentsActionSchemaElement ULearningAgentsActions::SpecifyInclusiveUnion
 		SortedSubElementPriorProbabilities[Idx] = SubElementPriorProbabilities[SubElementIndices[Idx]];
 	}
 
-	return SpecifyInclusiveUnionActionFromArrayViews(Schema, SortedSubElementNames, SortedSubElements, SortedSubElementPriorProbabilities, Tag);
+	return SpecifyInclusiveUnionActionFromArrayViews(SortedSubElementNames, SortedSubElements, SortedSubElementPriorProbabilities, Name);
 }
 
-FLearningAgentsActionSchemaElement ULearningAgentsActions::SpecifyInclusiveUnionActionFromArrays(ULearningAgentsActionSchema* Schema, const TArray<FName> ElementNames, const TArray<FLearningAgentsActionSchemaElement>& Elements, const TArray<float>& PriorProbabilities, const FName Tag)
+FLearningAgentsActionSchemaElement ULearningAgentsActionSchema::SpecifyInclusiveUnionActionFromArrays(const TArray<FName> ElementNames, const TArray<FLearningAgentsActionSchemaElement>& Elements, const TArray<float>& PriorProbabilities, const FName Name)
 {
-	return SpecifyInclusiveUnionActionFromArrayViews(Schema, ElementNames, Elements, PriorProbabilities, Tag);
+	return SpecifyInclusiveUnionActionFromArrayViews(ElementNames, Elements, PriorProbabilities, Name);
 }
 
-FLearningAgentsActionSchemaElement ULearningAgentsActions::SpecifyInclusiveUnionActionFromArrayViews(ULearningAgentsActionSchema* Schema, const TArrayView<const FName> ElementNames, const TArrayView<const FLearningAgentsActionSchemaElement> Elements, const TArrayView<const float> PriorProbabilities, const FName Tag)
+FLearningAgentsActionSchemaElement ULearningAgentsActionSchema::SpecifyInclusiveUnionActionFromArrayViews(const TArrayView<const FName> ElementNames, const TArrayView<const FLearningAgentsActionSchemaElement> Elements, const TArrayView<const float> PriorProbabilities, const FName Name)
 {
-	if (!Schema)
-	{
-		UE_LOG(LogLearning, Error, TEXT("SpecifyInclusiveUnionActionFromArrayViews: Schema is nullptr."));
-		return FLearningAgentsActionSchemaElement();
-	}
-
 	if (Elements.Num() == 0)
 	{
-		UE_LOG(LogLearning, Warning, TEXT("SpecifyInclusiveUnionActionFromArrayViews: Specifying zero-sized Inclusive Union Action."));
+		UE_LOG(LogLearning, Warning, TEXT("%s: Specifying zero-sized Inclusive Union Action."), *GetName());
 	}
 
 	if (Elements.Num() != ElementNames.Num())
 	{
-		UE_LOG(LogLearning, Error, TEXT("SpecifyInclusiveUnionActionFromArrayViews: Number of elements (%i) must match number of names (%i)."), Elements.Num(), ElementNames.Num());
+		UE_LOG(LogLearning, Error, TEXT("%s: Number of elements (%i) must match number of names (%i)."), *GetName(), Elements.Num(), ElementNames.Num());
 		return FLearningAgentsActionSchemaElement();
 	}
 
 	if (UE::Learning::Agents::Action::Private::ContainsDuplicates(ElementNames))
 	{
-		UE_LOG(LogLearning, Error, TEXT("SpecifyInclusiveUnionActionFromArrayViews: Element Names contain duplicates."));
+		UE_LOG(LogLearning, Error, TEXT("%s: Element Names contain duplicates."), *GetName());
 		return FLearningAgentsActionSchemaElement();
 	}
 
@@ -913,9 +882,9 @@ FLearningAgentsActionSchemaElement ULearningAgentsActions::SpecifyInclusiveUnion
 
 	for (const FLearningAgentsActionSchemaElement& Element : Elements)
 	{
-		if (!Schema->ActionSchema.IsValid(Element.SchemaElement))
+		if (!ActionSchema.IsValid(Element.SchemaElement))
 		{
-			UE_LOG(LogLearning, Error, TEXT("SpecifyInclusiveUnionActionFromArrayViews: Invalid Action Object."));
+			UE_LOG(LogLearning, Error, TEXT("%s: Invalid Action Object."), *GetName());
 			return FLearningAgentsActionSchemaElement();
 		}
 
@@ -928,49 +897,43 @@ FLearningAgentsActionSchemaElement ULearningAgentsActions::SpecifyInclusiveUnion
 	{
 		NormalizedPriorProbabilities[Idx] = PriorProbabilities[Idx];
 	}
-	UE::Learning::Agents::Action::Private::NormalizeProbabilitiesInclusive(NormalizedPriorProbabilities);
+	UE::Learning::Agents::Action::Private::NormalizeProbabilitiesInclusive(NormalizedPriorProbabilities, GetName());
 
-	return { Schema->ActionSchema.CreateOrInclusive({ ElementNames, SubElements, MakeArrayView(NormalizedPriorProbabilities) }, Tag)};
+	return { ActionSchema.CreateOrInclusive({ ElementNames, SubElements, MakeArrayView(NormalizedPriorProbabilities) }, Name)};
 }
 
-FLearningAgentsActionSchemaElement ULearningAgentsActions::SpecifyStaticArrayAction(ULearningAgentsActionSchema* Schema, const FLearningAgentsActionSchemaElement Element, const int32 Num, const FName Tag)
+FLearningAgentsActionSchemaElement ULearningAgentsActionSchema::SpecifyStaticArrayAction(const FLearningAgentsActionSchemaElement Element, const int32 Num, const FName Name)
 {
-	if (!Schema)
-	{
-		UE_LOG(LogLearning, Error, TEXT("SpecifyStaticArrayAction: Schema is nullptr."));
-		return FLearningAgentsActionSchemaElement();
-	}
-
 	if (Num < 0)
 	{
-		UE_LOG(LogLearning, Error, TEXT("SpecifyStaticArrayAction: Invalid Action Static Array Num %i."), Num);
+		UE_LOG(LogLearning, Error, TEXT("%s: Invalid Action Static Array Num %i."), *GetName(), Num);
 		return FLearningAgentsActionSchemaElement();
 	}
 
 	if (Num == 0)
 	{
-		UE_LOG(LogLearning, Warning, TEXT("SpecifyStaticArrayAction: Specifying zero-sized Static Array Action."));
+		UE_LOG(LogLearning, Warning, TEXT("%s: Specifying zero-sized Static Array Action."), *GetName());
 	}
 
-	if (!Schema->ActionSchema.IsValid(Element.SchemaElement))
+	if (!ActionSchema.IsValid(Element.SchemaElement))
 	{
-		UE_LOG(LogLearning, Error, TEXT("SpecifyStaticArrayAction: Invalid Action Object."));
+		UE_LOG(LogLearning, Error, TEXT("%s: Invalid Action Object."), *GetName());
 		return FLearningAgentsActionSchemaElement();
 	}
 
-	return { Schema->ActionSchema.CreateArray({ Element.SchemaElement, Num }, Tag) };
+	return { ActionSchema.CreateArray({ Element.SchemaElement, Num }, Name) };
 }
 
-FLearningAgentsActionSchemaElement ULearningAgentsActions::SpecifyPairAction(ULearningAgentsActionSchema* Schema, const FLearningAgentsActionSchemaElement Key, const FLearningAgentsActionSchemaElement Value, const FName Tag)
+FLearningAgentsActionSchemaElement ULearningAgentsActionSchema::SpecifyPairAction(const FLearningAgentsActionSchemaElement Key, const FLearningAgentsActionSchemaElement Value, const FName Name)
 {
-	return SpecifyStructActionFromArrayViews(Schema, { TEXT("Key"), TEXT("Value") }, { Key, Value }, Tag);
+	return SpecifyStructActionFromArrayViews({ TEXT("Key"), TEXT("Value") }, { Key, Value }, Name);
 }
 
-FLearningAgentsActionSchemaElement ULearningAgentsActions::SpecifyEnumAction(ULearningAgentsActionSchema* Schema, const UEnum* Enum, const TMap<uint8, float>& PriorProbabilities, const FName Tag)
+FLearningAgentsActionSchemaElement ULearningAgentsActionSchema::SpecifyEnumAction(const UEnum* Enum, const TMap<uint8, float>& PriorProbabilities, const FName Name)
 {
 	if (!Enum)
 	{
-		UE_LOG(LogLearning, Error, TEXT("SpecifyEnumAction: Enum is nullptr."));
+		UE_LOG(LogLearning, Error, TEXT("%s: Enum is nullptr."), *GetName());
 		return FLearningAgentsActionSchemaElement();
 	}
 
@@ -985,36 +948,36 @@ FLearningAgentsActionSchemaElement ULearningAgentsActions::SpecifyEnumAction(ULe
 		}
 	}
 
-	return SpecifyEnumActionFromArrayView(Schema, Enum, PriorProbabilitiesArray, Tag);
+	return SpecifyEnumActionFromArrayView(Enum, PriorProbabilitiesArray, Name);
 }
 
-FLearningAgentsActionSchemaElement ULearningAgentsActions::SpecifyEnumActionFromArray(ULearningAgentsActionSchema* Schema, const UEnum* Enum, const TArray<float>& PriorProbabilities, const FName Tag)
+FLearningAgentsActionSchemaElement ULearningAgentsActionSchema::SpecifyEnumActionFromArray(const UEnum* Enum, const TArray<float>& PriorProbabilities, const FName Name)
 {
-	return SpecifyEnumActionFromArrayView(Schema, Enum, PriorProbabilities, Tag);
+	return SpecifyEnumActionFromArrayView(Enum, PriorProbabilities, Name);
 }
 
-FLearningAgentsActionSchemaElement ULearningAgentsActions::SpecifyEnumActionFromArrayView(ULearningAgentsActionSchema* Schema, const UEnum* Enum, const TArrayView<const float> PriorProbabilities, const FName Tag)
+FLearningAgentsActionSchemaElement ULearningAgentsActionSchema::SpecifyEnumActionFromArrayView(const UEnum* Enum, const TArrayView<const float> PriorProbabilities, const FName Name)
 {
 	if (!Enum)
 	{
-		UE_LOG(LogLearning, Error, TEXT("SpecifyEnumActionFromArrayView: Enum is nullptr."));
+		UE_LOG(LogLearning, Error, TEXT("%s: Enum is nullptr."), *GetName());
 		return FLearningAgentsActionSchemaElement();
 	}
 
-	return SpecifyExclusiveDiscreteActionFromArrayView(Schema, Enum->NumEnums() - 1, PriorProbabilities, Tag);
+	return SpecifyExclusiveDiscreteActionFromArrayView(Enum->NumEnums() - 1, PriorProbabilities, Name);
 }
 
-FLearningAgentsActionSchemaElement ULearningAgentsActions::SpecifyBitmaskAction(ULearningAgentsActionSchema* Schema, const UEnum* Enum, const TMap<uint8, float>& PriorProbabilities, const FName Tag)
+FLearningAgentsActionSchemaElement ULearningAgentsActionSchema::SpecifyBitmaskAction(const UEnum* Enum, const TMap<uint8, float>& PriorProbabilities, const FName Name)
 {
 	if (!Enum)
 	{
-		UE_LOG(LogLearning, Error, TEXT("SpecifyBitmaskAction: Enum is nullptr."));
+		UE_LOG(LogLearning, Error, TEXT("%s: Enum is nullptr."), *GetName());
 		return FLearningAgentsActionSchemaElement();
 	}
 
 	if (Enum->NumEnums() - 1 > 32)
 	{
-		UE_LOG(LogLearning, Error, TEXT("SpecifyBitmaskAction: Too many values in Enum to use as Bitmask (%i)."), Enum->NumEnums() - 1);
+		UE_LOG(LogLearning, Error, TEXT("%s: Too many values in Enum to use as Bitmask (%i)."), *GetName(), Enum->NumEnums() - 1);
 		return FLearningAgentsActionSchemaElement();
 	}
 
@@ -1029,208 +992,181 @@ FLearningAgentsActionSchemaElement ULearningAgentsActions::SpecifyBitmaskAction(
 		}
 	}
 
-	return SpecifyBitmaskActionFromArrayView(Schema, Enum, PriorProbabilitiesArray, Tag);
+	return SpecifyBitmaskActionFromArrayView(Enum, PriorProbabilitiesArray, Name);
 }
 
-FLearningAgentsActionSchemaElement ULearningAgentsActions::SpecifyBitmaskActionFromArray(ULearningAgentsActionSchema* Schema, const UEnum* Enum, const TArray<float>& PriorProbabilities, const FName Tag)
+FLearningAgentsActionSchemaElement ULearningAgentsActionSchema::SpecifyBitmaskActionFromArray(const UEnum* Enum, const TArray<float>& PriorProbabilities, const FName Name)
 {
-	return SpecifyBitmaskActionFromArrayView(Schema, Enum, PriorProbabilities, Tag);
+	return SpecifyBitmaskActionFromArrayView(Enum, PriorProbabilities, Name);
 }
 
-FLearningAgentsActionSchemaElement ULearningAgentsActions::SpecifyBitmaskActionFromArrayView(ULearningAgentsActionSchema* Schema, const UEnum* Enum, const TArrayView<const float> PriorProbabilities, const FName Tag)
+FLearningAgentsActionSchemaElement ULearningAgentsActionSchema::SpecifyBitmaskActionFromArrayView(const UEnum* Enum, const TArrayView<const float> PriorProbabilities, const FName Name)
 {
 	if (!Enum)
 	{
-		UE_LOG(LogLearning, Error, TEXT("SpecifyBitmaskActionFromArrayView: Enum is nullptr."));
+		UE_LOG(LogLearning, Error, TEXT("%s: Enum is nullptr."), *GetName());
 		return FLearningAgentsActionSchemaElement();
 	}
 
 	if (Enum->NumEnums() - 1 > 32)
 	{
-		UE_LOG(LogLearning, Error, TEXT("SpecifyBitmaskActionFromArrayView: Too many values in Enum to use as Bitmask (%i)."), Enum->NumEnums() - 1);
+		UE_LOG(LogLearning, Error, TEXT("%s: Too many values in Enum to use as Bitmask (%i)."), *GetName(), Enum->NumEnums() - 1);
 		return FLearningAgentsActionSchemaElement();
 	}
 
-	return SpecifyInclusiveDiscreteActionFromArrayView(Schema, Enum->NumEnums() - 1, PriorProbabilities, Tag);
+	return SpecifyInclusiveDiscreteActionFromArrayView(Enum->NumEnums() - 1, PriorProbabilities, Name);
 }
 
-FLearningAgentsActionSchemaElement ULearningAgentsActions::SpecifyOptionalAction(ULearningAgentsActionSchema* Schema, const FLearningAgentsActionSchemaElement Element, const float PriorProbability, const FName Tag)
+FLearningAgentsActionSchemaElement ULearningAgentsActionSchema::SpecifyOptionalAction(const FLearningAgentsActionSchemaElement Element, const float PriorProbability, const FName Name)
 {
-	return SpecifyExclusiveUnionActionFromArrayViews(Schema, { TEXT("Null"), TEXT("Valid") }, { SpecifyNullAction(Schema), Element }, { 1.0f - PriorProbability, PriorProbability }, Tag);
+	return SpecifyExclusiveUnionActionFromArrayViews({ TEXT("Null"), TEXT("Valid") }, { SpecifyNullAction(), Element }, { 1.0f - PriorProbability, PriorProbability }, Name);
 }
 
-FLearningAgentsActionSchemaElement ULearningAgentsActions::SpecifyEitherAction(ULearningAgentsActionSchema* Schema, const FLearningAgentsActionSchemaElement A, const FLearningAgentsActionSchemaElement B, const float PriorProbabilityOfA, const FName Tag)
+FLearningAgentsActionSchemaElement ULearningAgentsActionSchema::SpecifyEitherAction(const FLearningAgentsActionSchemaElement A, const FLearningAgentsActionSchemaElement B, const float PriorProbabilityOfA, const FName Name)
 {
-	return SpecifyExclusiveUnionActionFromArrayViews(Schema, { TEXT("A"), TEXT("B") }, { A, B }, { 1.0f - PriorProbabilityOfA, PriorProbabilityOfA }, Tag);
+	return SpecifyExclusiveUnionActionFromArrayViews({ TEXT("A"), TEXT("B") }, { A, B }, { 1.0f - PriorProbabilityOfA, PriorProbabilityOfA }, Name);
 }
 
-FLearningAgentsActionSchemaElement ULearningAgentsActions::SpecifyEncodingAction(ULearningAgentsActionSchema* Schema, const FLearningAgentsActionSchemaElement Element, const int32 EncodingSize, const int32 LayerNum, const ELearningAgentsActivationFunction ActivationFunction, const FName Tag)
+FLearningAgentsActionSchemaElement ULearningAgentsActionSchema::SpecifyEncodingAction(const FLearningAgentsActionSchemaElement Element, const int32 EncodingSize, const int32 LayerNum, const ELearningAgentsActivationFunction ActivationFunction, const FName Name)
 {
-	if (!Schema)
-	{
-		UE_LOG(LogLearning, Error, TEXT("SpecifyEncodingAction: Schema is nullptr."));
-		return FLearningAgentsActionSchemaElement();
-	}
-
 	if (EncodingSize < 1)
 	{
-		UE_LOG(LogLearning, Error, TEXT("SpecifyEncodingAction: Invalid Action EncodingSize '%i' - must be greater than zero."), EncodingSize);
+		UE_LOG(LogLearning, Error, TEXT("%s: Invalid Action EncodingSize '%i' - must be greater than zero."), *GetName(), EncodingSize);
 		return FLearningAgentsActionSchemaElement();
 	}
 
 	if (LayerNum < 1)
 	{
-		UE_LOG(LogLearning, Error, TEXT("SpecifyEncodingAction: Invalid Action LayerNum '%i' - must be greater than zero."), LayerNum);
+		UE_LOG(LogLearning, Error, TEXT("%s: Invalid Action LayerNum '%i' - must be greater than zero."), *GetName(), LayerNum);
 		return FLearningAgentsActionSchemaElement();
 	}
 
-	if (!Schema->ActionSchema.IsValid(Element.SchemaElement))
+	if (!ActionSchema.IsValid(Element.SchemaElement))
 	{
-		UE_LOG(LogLearning, Error, TEXT("SpecifyEncodingAction: Invalid Action Object."));
+		UE_LOG(LogLearning, Error, TEXT("%s: Invalid Action Object."), *GetName());
 		return FLearningAgentsActionSchemaElement();
 	}
 
-	return { Schema->ActionSchema.CreateEncoding({ Element.SchemaElement, EncodingSize, LayerNum, UE::Learning::Agents::Action::Private::GetEncodingActivationFunction(ActivationFunction) }, Tag)};
+	return { ActionSchema.CreateEncoding({ Element.SchemaElement, EncodingSize, LayerNum, UE::Learning::Agents::Action::Private::GetEncodingActivationFunction(ActivationFunction) }, Name)};
 }
 
-FLearningAgentsActionSchemaElement ULearningAgentsActions::SpecifyBoolAction(ULearningAgentsActionSchema* Schema, const float PriorProbability, const FName Tag)
+FLearningAgentsActionSchemaElement ULearningAgentsActionSchema::SpecifyBoolAction(const float PriorProbability, const FName Name)
 {
-	return SpecifyExclusiveDiscreteActionFromArrayView(Schema, 2, { 1.0f - PriorProbability, PriorProbability }, Tag);
+	return SpecifyExclusiveDiscreteActionFromArrayView(2, { 1.0f - PriorProbability, PriorProbability }, Name);
 }
 
-FLearningAgentsActionSchemaElement ULearningAgentsActions::SpecifyFloatAction(ULearningAgentsActionSchema* Schema, const FName Tag)
+FLearningAgentsActionSchemaElement ULearningAgentsActionSchema::SpecifyFloatAction(const FName Name)
 {
-	return SpecifyContinuousAction(Schema, 1, Tag);
+	return SpecifyContinuousAction(1, Name);
 }
 
-FLearningAgentsActionSchemaElement ULearningAgentsActions::SpecifyLocationAction(ULearningAgentsActionSchema* Schema, const FName Tag)
+FLearningAgentsActionSchemaElement ULearningAgentsActionSchema::SpecifyLocationAction(const FName Name)
 {
-	return SpecifyContinuousAction(Schema, 3, Tag);
+	return SpecifyContinuousAction(3, Name);
 }
 
-FLearningAgentsActionSchemaElement ULearningAgentsActions::SpecifyRotationAction(ULearningAgentsActionSchema* Schema, const FName Tag)
+FLearningAgentsActionSchemaElement ULearningAgentsActionSchema::SpecifyRotationAction(const FName Name)
 {
-	return SpecifyContinuousAction(Schema, 3, Tag);
+	return SpecifyContinuousAction(3, Name);
 }
 
-FLearningAgentsActionSchemaElement ULearningAgentsActions::SpecifyScaleAction(ULearningAgentsActionSchema* Schema, const FName Tag)
+FLearningAgentsActionSchemaElement ULearningAgentsActionSchema::SpecifyScaleAction(const FName Name)
 {
-	return SpecifyContinuousAction(Schema, 3, Tag);
+	return SpecifyContinuousAction(3, Name);
 }
 
-FLearningAgentsActionSchemaElement ULearningAgentsActions::SpecifyTransformAction(ULearningAgentsActionSchema* Schema, const FName Tag)
+FLearningAgentsActionSchemaElement ULearningAgentsActionSchema::SpecifyTransformAction(const FName Name)
 {
-	return SpecifyStructActionFromArrayViews(Schema,
+	return SpecifyStructActionFromArrayViews(
 		{
 			TEXT("Location"),
 			TEXT("Rotation"),
 			TEXT("Scale")
 		},
 		{
-			SpecifyLocationAction(Schema),
-			SpecifyRotationAction(Schema),
-			SpecifyScaleAction(Schema)
+			SpecifyLocationAction(),
+			SpecifyRotationAction(),
+			SpecifyScaleAction()
 		}, 
-		Tag);
+		Name);
 }
 
-FLearningAgentsActionSchemaElement ULearningAgentsActions::SpecifyAngleAction(ULearningAgentsActionSchema* Schema, const FName Tag)
+FLearningAgentsActionSchemaElement ULearningAgentsActionSchema::SpecifyAngleAction(const FName Name)
 {
-	return SpecifyContinuousAction(Schema, 1, Tag);
+	return SpecifyContinuousAction(1, Name);
 }
 
-FLearningAgentsActionSchemaElement ULearningAgentsActions::SpecifyVelocityAction(ULearningAgentsActionSchema* Schema, const FName Tag)
+FLearningAgentsActionSchemaElement ULearningAgentsActionSchema::SpecifyVelocityAction(const FName Name)
 {
-	return SpecifyContinuousAction(Schema, 3, Tag);
+	return SpecifyContinuousAction(3, Name);
 }
 
-FLearningAgentsActionSchemaElement ULearningAgentsActions::SpecifyDirectionAction(ULearningAgentsActionSchema* Schema, const FName Tag)
+FLearningAgentsActionSchemaElement ULearningAgentsActionSchema::SpecifyDirectionAction(const FName Name)
 {
-	return SpecifyContinuousAction(Schema, 3, Tag);
+	return SpecifyContinuousAction(3, Name);
 }
 
-FLearningAgentsActionSchemaElement ULearningAgentsActions::SpecifySpeedAction(ULearningAgentsActionSchema* Schema, const FName Tag)
+FLearningAgentsActionSchemaElement ULearningAgentsActionSchema::SpecifySpeedAction(const FName Name)
 {
-	return SpecifyContinuousAction(Schema, 1, Tag);
+	return SpecifyContinuousAction(1, Name);
 }
 
-
-void ULearningAgentsActions::LogAction(const ULearningAgentsActionObject* Object, const FLearningAgentsActionObjectElement Element)
+const UE::Learning::Action::FObject& ULearningAgentsActionObject::GetActionObject() const
 {
-	if (!Object)
-	{
-		UE_LOG(LogLearning, Error, TEXT("LogAction: Object is nullptr."));
-		return;
-	}
-
-	UE::Learning::Agents::Action::Private::LogAction(Object->ActionObject, Element.ObjectElement, TEXT(""), TEXT(""));
+	return ActionObject;
 }
 
-FLearningAgentsActionObjectElement ULearningAgentsActions::MakeNullAction(ULearningAgentsActionObject* Object, const FName Tag)
+UE::Learning::Action::FObject& ULearningAgentsActionObject::GetActionObject()
 {
-	if (!Object)
-	{
-		UE_LOG(LogLearning, Error, TEXT("MakeNullAction: Object is nullptr."));
-		return FLearningAgentsActionObjectElement();
-	}
+	return ActionObject;
+}
 
-	return { Object->ActionObject.CreateNull(Tag) };
+void ULearningAgentsActionObject::LogAction(const FLearningAgentsActionObjectElement Element)
+{
+	UE::Learning::Agents::Action::Private::LogAction(GetActionObject(), Element.ObjectElement, TEXT(""), TEXT(""), GetName());
 }
 
-FLearningAgentsActionObjectElement ULearningAgentsActions::MakeContinuousAction(ULearningAgentsActionObject* Object, const TArray<float>& Values, const FName Tag)
+FLearningAgentsActionObjectElement ULearningAgentsActionObject::MakeNullAction(const FName Name)
 {
-	return MakeContinuousActionFromArrayView(Object, Values, Tag);
+	return { ActionObject.CreateNull(Name) };
 }
 
-FLearningAgentsActionObjectElement ULearningAgentsActions::MakeContinuousActionFromArrayView(ULearningAgentsActionObject* Object, const TArrayView<const float> Values, const FName Tag)
+FLearningAgentsActionObjectElement ULearningAgentsActionObject::MakeContinuousAction(const TArray<float>& Values, const FName Name)
 {
-	if (!Object)
-	{
-		UE_LOG(LogLearning, Error, TEXT("MakeContinuousActionFromArrayView: Object is nullptr."));
-		return FLearningAgentsActionObjectElement();
-	}
+	return MakeContinuousActionFromArrayView(Values, Name);
+}
 
+FLearningAgentsActionObjectElement ULearningAgentsActionObject::MakeContinuousActionFromArrayView(const TArrayView<const float> Values, const FName Name)
+{
 	if (Values.Num() == 0)
 	{
-		UE_LOG(LogLearning, Warning, TEXT("MakeContinuousActionFromArrayView: Creating zero-sized Continuous Action."));
+		UE_LOG(LogLearning, Warning, TEXT("%s: Creating zero-sized Continuous Action."), *GetName());
 	}
 
-	return { Object->ActionObject.CreateContinuous({ Values }, Tag) };
+	return { ActionObject.CreateContinuous({ Values }, Name) };
 }
 
-FLearningAgentsActionObjectElement ULearningAgentsActions::MakeExclusiveDiscreteAction(ULearningAgentsActionObject* Object, const int32 Index, const FName Tag)
+FLearningAgentsActionObjectElement ULearningAgentsActionObject::MakeExclusiveDiscreteAction(const int32 Index, const FName Name)
 {
-	if (!Object)
-	{
-		UE_LOG(LogLearning, Error, TEXT("MakeExclusiveDiscreteAction: Object is nullptr."));
-		return FLearningAgentsActionObjectElement();
-	}
-
 	if (Index < 0)
 	{
-		UE_LOG(LogLearning, Error, TEXT("MakeExclusiveDiscreteAction: Invalid Action Index %i."), Index);
+		UE_LOG(LogLearning, Error, TEXT("%s: Invalid Action Index %i."), *GetName(), Index);
 		return FLearningAgentsActionObjectElement();
 	}
 
-	return { Object->ActionObject.CreateDiscreteExclusive({ Index }, Tag) };
+	return { ActionObject.CreateDiscreteExclusive({ Index }, Name) };
 }
 
-FLearningAgentsActionObjectElement ULearningAgentsActions::MakeInclusiveDiscreteAction(ULearningAgentsActionObject* Object, const TArray<int32>& Indices, const FName Tag)
+FLearningAgentsActionObjectElement ULearningAgentsActionObject::MakeInclusiveDiscreteAction(const TArray<int32>& Indices, const FName Name)
 {
-	return MakeInclusiveDiscreteActionFromArrayView(Object, Indices, Tag);
+	return MakeInclusiveDiscreteActionFromArrayView(Indices, Name);
 }
 
-FLearningAgentsActionObjectElement ULearningAgentsActions::MakeInclusiveDiscreteActionFromArrayView(ULearningAgentsActionObject* Object, const TArrayView<const int32> Indices, const FName Tag)
+FLearningAgentsActionObjectElement ULearningAgentsActionObject::MakeInclusiveDiscreteActionFromArrayView(const TArrayView<const int32> Indices, const FName Name)
 {
-	if (!Object)
-	{
-		UE_LOG(LogLearning, Error, TEXT("MakeInclusiveDiscreteActionFromArrayView: Object is nullptr."));
-		return FLearningAgentsActionObjectElement();
-	}
-
 	if (UE::Learning::Agents::Action::Private::ContainsDuplicates(Indices))
 	{
-		UE_LOG(LogLearning, Error, TEXT("MakeInclusiveDiscreteActionFromArrayView: Indices contain duplicates."));
+		UE_LOG(LogLearning, Error, TEXT("%s: Indices contain duplicates."), *GetName());
 		return FLearningAgentsActionObjectElement();
 	}
 
@@ -1240,19 +1176,19 @@ FLearningAgentsActionObjectElement ULearningAgentsActions::MakeInclusiveDiscrete
 	{
 		if (Indices[IndexIdx] < 0)
 		{
-			UE_LOG(LogLearning, Error, TEXT("MakeInclusiveDiscreteActionFromArrayView: Invalid Action Index %i."), Indices[IndexIdx]);
+			UE_LOG(LogLearning, Error, TEXT("%s: Invalid Action Index %i."), *GetName(), Indices[IndexIdx]);
 			return FLearningAgentsActionObjectElement();
 		}
 	}
 
-	return { Object->ActionObject.CreateDiscreteInclusive({ Indices }, Tag) };
+	return { ActionObject.CreateDiscreteInclusive({ Indices }, Name) };
 }
 
-FLearningAgentsActionObjectElement ULearningAgentsActions::MakeStructAction(ULearningAgentsActionObject* Object, const TMap<FName, FLearningAgentsActionObjectElement>& Elements, const FName Tag)
+FLearningAgentsActionObjectElement ULearningAgentsActionObject::MakeStructAction(const TMap<FName, FLearningAgentsActionObjectElement>& Elements, const FName Name)
 {
 	if (Elements.Num() == 0)
 	{
-		UE_LOG(LogLearning, Warning, TEXT("MakeStructAction: Creating zero-sized Struct Action."));
+		UE_LOG(LogLearning, Warning, TEXT("%s: Creating zero-sized Struct Action."), *GetName());
 	}
 
 	const int32 SubElementNum = Elements.Num();
@@ -1268,36 +1204,30 @@ FLearningAgentsActionObjectElement ULearningAgentsActions::MakeStructAction(ULea
 		SubElementNames.Add(Element.Key);
 	}
 
-	return MakeStructActionFromArrayViews(Object, SubElementNames, SubElements);
+	return MakeStructActionFromArrayViews(SubElementNames, SubElements);
 }
 
-FLearningAgentsActionObjectElement ULearningAgentsActions::MakeStructActionFromArrays(ULearningAgentsActionObject* Object, const TArray<FName>& ElementNames, const TArray<FLearningAgentsActionObjectElement>& Elements, const FName Tag)
+FLearningAgentsActionObjectElement ULearningAgentsActionObject::MakeStructActionFromArrays(const TArray<FName>& ElementNames, const TArray<FLearningAgentsActionObjectElement>& Elements, const FName Name)
 {
-	return MakeStructActionFromArrayViews(Object, ElementNames, Elements, Tag);
+	return MakeStructActionFromArrayViews(ElementNames, Elements, Name);
 }
 
-FLearningAgentsActionObjectElement ULearningAgentsActions::MakeStructActionFromArrayViews(ULearningAgentsActionObject* Object, const TArrayView<const FName> ElementNames, const TArrayView<const FLearningAgentsActionObjectElement> Elements, const FName Tag)
+FLearningAgentsActionObjectElement ULearningAgentsActionObject::MakeStructActionFromArrayViews(const TArrayView<const FName> ElementNames, const TArrayView<const FLearningAgentsActionObjectElement> Elements, const FName Name)
 {
-	if (!Object)
-	{
-		UE_LOG(LogLearning, Error, TEXT("MakeStructActionFromArrayViews: Object is nullptr."));
-		return FLearningAgentsActionObjectElement();
-	}
-
 	if (Elements.Num() == 0)
 	{
-		UE_LOG(LogLearning, Warning, TEXT("MakeStructActionFromArrayViews: Creating zero-sized Struct Action."));
+		UE_LOG(LogLearning, Warning, TEXT("%s: Creating zero-sized Struct Action."), *GetName());
 	}
 
 	if (Elements.Num() != ElementNames.Num())
 	{
-		UE_LOG(LogLearning, Error, TEXT("MakeStructActionFromArrayViews: Number of elements (%i) must match number of names (%i)."), Elements.Num(), ElementNames.Num());
+		UE_LOG(LogLearning, Error, TEXT("%s: Number of elements (%i) must match number of names (%i)."), *GetName(), Elements.Num(), ElementNames.Num());
 		return FLearningAgentsActionObjectElement();
 	}
 
 	if (UE::Learning::Agents::Action::Private::ContainsDuplicates(ElementNames))
 	{
-		UE_LOG(LogLearning, Error, TEXT("MakeStructActionFromArrayViews: Element Names contain duplicates."));
+		UE_LOG(LogLearning, Error, TEXT("%s: Element Names contain duplicates."), *GetName());
 		return FLearningAgentsActionObjectElement();
 	}
 
@@ -1306,36 +1236,30 @@ FLearningAgentsActionObjectElement ULearningAgentsActions::MakeStructActionFromA
 
 	for (const FLearningAgentsActionObjectElement& Element : Elements)
 	{
-		if (!Object->ActionObject.IsValid(Element.ObjectElement))
+		if (!ActionObject.IsValid(Element.ObjectElement))
 		{
-			UE_LOG(LogLearning, Error, TEXT("MakeStructActionFromArrayViews: Invalid Action Object."));
+			UE_LOG(LogLearning, Error, TEXT("%s: Invalid Action Object."), *GetName());
 			return FLearningAgentsActionObjectElement();
 		}
 
 		SubElements.Add(Element.ObjectElement);
 	}
 
-	return { Object->ActionObject.CreateAnd({ ElementNames, SubElements }, Tag) };
+	return { ActionObject.CreateAnd({ ElementNames, SubElements }, Name) };
 }
 
-FLearningAgentsActionObjectElement ULearningAgentsActions::MakeExclusiveUnionAction(ULearningAgentsActionObject* Object, const FName ElementName, const FLearningAgentsActionObjectElement Element, const FName Tag)
+FLearningAgentsActionObjectElement ULearningAgentsActionObject::MakeExclusiveUnionAction(const FName ElementName, const FLearningAgentsActionObjectElement Element, const FName Name)
 {
-	if (!Object)
-	{
-		UE_LOG(LogLearning, Error, TEXT("MakeExclusiveUnionAction: Object is nullptr."));
-		return FLearningAgentsActionObjectElement();
-	}
-
-	if (!Object->ActionObject.IsValid(Element.ObjectElement))
+	if (!ActionObject.IsValid(Element.ObjectElement))
 	{
-		UE_LOG(LogLearning, Error, TEXT("MakeExclusiveUnionAction: Invalid Action Object."));
+		UE_LOG(LogLearning, Error, TEXT("%s: Invalid Action Object."), *GetName());
 		return FLearningAgentsActionObjectElement();
 	}
 
-	return { Object->ActionObject.CreateOrExclusive({ ElementName, Element.ObjectElement }, Tag) };
+	return { ActionObject.CreateOrExclusive({ ElementName, Element.ObjectElement }, Name) };
 }
 
-FLearningAgentsActionObjectElement ULearningAgentsActions::MakeInclusiveUnionAction(ULearningAgentsActionObject* Object, const TMap<FName, FLearningAgentsActionObjectElement>& Elements, const FName Tag)
+FLearningAgentsActionObjectElement ULearningAgentsActionObject::MakeInclusiveUnionAction(const TMap<FName, FLearningAgentsActionObjectElement>& Elements, const FName Name)
 {
 	const int32 SubElementNum = Elements.Num();
 
@@ -1350,31 +1274,25 @@ FLearningAgentsActionObjectElement ULearningAgentsActions::MakeInclusiveUnionAct
 		SubElementNames.Add(Element.Key);
 	}
 
-	return MakeInclusiveUnionActionFromArrayViews(Object, SubElementNames, SubElements, Tag);
+	return MakeInclusiveUnionActionFromArrayViews(SubElementNames, SubElements, Name);
 }
 
-FLearningAgentsActionObjectElement ULearningAgentsActions::MakeInclusiveUnionActionFromArrays(ULearningAgentsActionObject* Object, const TArray<FName>& ElementNames, const TArray<FLearningAgentsActionObjectElement>& Elements, const FName Tag)
+FLearningAgentsActionObjectElement ULearningAgentsActionObject::MakeInclusiveUnionActionFromArrays(const TArray<FName>& ElementNames, const TArray<FLearningAgentsActionObjectElement>& Elements, const FName Name)
 {
-	return MakeInclusiveUnionActionFromArrayViews(Object, ElementNames, Elements, Tag);
+	return MakeInclusiveUnionActionFromArrayViews(ElementNames, Elements, Name);
 }
 
-FLearningAgentsActionObjectElement ULearningAgentsActions::MakeInclusiveUnionActionFromArrayViews(ULearningAgentsActionObject* Object, const TArrayView<const FName> ElementNames, const TArrayView<const FLearningAgentsActionObjectElement> Elements, const FName Tag)
+FLearningAgentsActionObjectElement ULearningAgentsActionObject::MakeInclusiveUnionActionFromArrayViews(const TArrayView<const FName> ElementNames, const TArrayView<const FLearningAgentsActionObjectElement> Elements, const FName Name)
 {
-	if (!Object)
-	{
-		UE_LOG(LogLearning, Error, TEXT("MakeInclusiveUnionActionFromArrayViews: Object is nullptr."));
-		return FLearningAgentsActionObjectElement();
-	}
-
 	if (Elements.Num() != ElementNames.Num())
 	{
-		UE_LOG(LogLearning, Error, TEXT("MakeInclusiveUnionActionFromArrayViews: Number of elements (%i) must match number of names (%i)."), Elements.Num(), ElementNames.Num());
+		UE_LOG(LogLearning, Error, TEXT("%s: Number of elements (%i) must match number of names (%i)."), *GetName(), Elements.Num(), ElementNames.Num());
 		return FLearningAgentsActionObjectElement();
 	}
 
 	if (UE::Learning::Agents::Action::Private::ContainsDuplicates(ElementNames))
 	{
-		UE_LOG(LogLearning, Error, TEXT("MakeInclusiveUnionActionFromArrayViews: Element Names contain duplicates."));
+		UE_LOG(LogLearning, Error, TEXT("%s: Element Names contain duplicates."), *GetName());
 		return FLearningAgentsActionObjectElement();
 	}
 
@@ -1383,34 +1301,28 @@ FLearningAgentsActionObjectElement ULearningAgentsActions::MakeInclusiveUnionAct
 
 	for (const FLearningAgentsActionObjectElement& Element : Elements)
 	{
-		if (!Object->ActionObject.IsValid(Element.ObjectElement))
+		if (!ActionObject.IsValid(Element.ObjectElement))
 		{
-			UE_LOG(LogLearning, Error, TEXT("MakeInclusiveUnionActionFromArrayViews: Invalid Action Object."));
+			UE_LOG(LogLearning, Error, TEXT("%s: Invalid Action Object."), *GetName());
 			return FLearningAgentsActionObjectElement();
 		}
 
 		SubElements.Add(Element.ObjectElement);
 	}
 
-	return { Object->ActionObject.CreateOrInclusive({ ElementNames, SubElements }, Tag) };
+	return { ActionObject.CreateOrInclusive({ ElementNames, SubElements }, Name) };
 }
 
-FLearningAgentsActionObjectElement ULearningAgentsActions::MakeStaticArrayAction(ULearningAgentsActionObject* Object, const TArray<FLearningAgentsActionObjectElement>& Elements, const FName Tag)
+FLearningAgentsActionObjectElement ULearningAgentsActionObject::MakeStaticArrayAction(const TArray<FLearningAgentsActionObjectElement>& Elements, const FName Name)
 {
-	return MakeStaticArrayActionFromArrayView(Object, Elements, Tag);
+	return MakeStaticArrayActionFromArrayView(Elements, Name);
 }
 
-FLearningAgentsActionObjectElement ULearningAgentsActions::MakeStaticArrayActionFromArrayView(ULearningAgentsActionObject* Object, const TArrayView<const FLearningAgentsActionObjectElement> Elements, const FName Tag)
+FLearningAgentsActionObjectElement ULearningAgentsActionObject::MakeStaticArrayActionFromArrayView(const TArrayView<const FLearningAgentsActionObjectElement> Elements, const FName Name)
 {
-	if (!Object)
-	{
-		UE_LOG(LogLearning, Error, TEXT("MakeStaticArrayActionFromArrayView: Object is nullptr."));
-		return FLearningAgentsActionObjectElement();
-	}
-
 	if (Elements.Num() == 0)
 	{
-		UE_LOG(LogLearning, Warning, TEXT("MakeStaticArrayActionFromArrayView: Creating zero-sized Static Array Action."));
+		UE_LOG(LogLearning, Warning, TEXT("%s: Creating zero-sized Static Array Action."), *GetName());
 	}
 
 	TArray<UE::Learning::Action::FObjectElement, TInlineAllocator<16>> SubElements;
@@ -1418,28 +1330,28 @@ FLearningAgentsActionObjectElement ULearningAgentsActions::MakeStaticArrayAction
 
 	for (const FLearningAgentsActionObjectElement& Element : Elements)
 	{
-		if (!Object->ActionObject.IsValid(Element.ObjectElement))
+		if (!ActionObject.IsValid(Element.ObjectElement))
 		{
-			UE_LOG(LogLearning, Error, TEXT("MakeStaticArrayActionFromArrayView: Invalid Action Object."));
+			UE_LOG(LogLearning, Error, TEXT("%s: Invalid Action Object."), *GetName());
 			return FLearningAgentsActionObjectElement();
 		}
 
 		SubElements.Add(Element.ObjectElement);
 	}
 
-	return { Object->ActionObject.CreateArray({ SubElements }, Tag) };
+	return { ActionObject.CreateArray({ SubElements }, Name) };
 }
 
-FLearningAgentsActionObjectElement ULearningAgentsActions::MakePairAction(ULearningAgentsActionObject* Object, const FLearningAgentsActionObjectElement Key, const FLearningAgentsActionObjectElement Value, const FName Tag)
+FLearningAgentsActionObjectElement ULearningAgentsActionObject::MakePairAction(const FLearningAgentsActionObjectElement Key, const FLearningAgentsActionObjectElement Value, const FName Name)
 {
-	return MakeStructActionFromArrayViews(Object, { TEXT("Key"), TEXT("Value") }, { Key, Value }, Tag);
+	return MakeStructActionFromArrayViews({ TEXT("Key"), TEXT("Value") }, { Key, Value }, Name);
 }
 
-FLearningAgentsActionObjectElement ULearningAgentsActions::MakeEnumAction(ULearningAgentsActionObject* Object, const UEnum* Enum, const uint8 EnumValue, const FName Tag)
+FLearningAgentsActionObjectElement ULearningAgentsActionObject::MakeEnumAction(const UEnum* Enum, const uint8 EnumValue, const FName Name)
 {
 	if (!Enum)
 	{
-		UE_LOG(LogLearning, Error, TEXT("MakeEnumAction: Enum is nullptr."));
+		UE_LOG(LogLearning, Error, TEXT("%s: Enum is nullptr."), *GetName());
 		return FLearningAgentsActionObjectElement();
 	}
 
@@ -1447,24 +1359,24 @@ FLearningAgentsActionObjectElement ULearningAgentsActions::MakeEnumAction(ULearn
 
 	if (EnumValueIndex == INDEX_NONE || EnumValueIndex < 0 || EnumValueIndex >= Enum->NumEnums() - 1)
 	{
-		UE_LOG(LogLearning, Error, TEXT("MakeEnumAction: EnumValue %i not valid for Enum '%s'."), EnumValue , *Enum->GetName());
+		UE_LOG(LogLearning, Error, TEXT("%s: EnumValue %i not valid for Enum '%s'."), *GetName(), EnumValue , *Enum->GetName());
 		return FLearningAgentsActionObjectElement();
 	}
 
-	return MakeExclusiveDiscreteAction(Object, EnumValueIndex, Tag);
+	return MakeExclusiveDiscreteAction(EnumValueIndex, Name);
 }
 
-FLearningAgentsActionObjectElement ULearningAgentsActions::MakeBitmaskAction(ULearningAgentsActionObject* Object, const UEnum* Enum, const int32 BitmaskValue, const FName Tag)
+FLearningAgentsActionObjectElement ULearningAgentsActionObject::MakeBitmaskAction(const UEnum* Enum, const int32 BitmaskValue, const FName Name)
 {
 	if (!Enum)
 	{
-		UE_LOG(LogLearning, Error, TEXT("MakeBitmaskAction: Enum is nullptr."));
+		UE_LOG(LogLearning, Error, TEXT("%s: Enum is nullptr."), *GetName());
 		return FLearningAgentsActionObjectElement();
 	}
 
 	if (Enum->NumEnums() - 1 > 32)
 	{
-		UE_LOG(LogLearning, Error, TEXT("MakeBitmaskAction: Too many values in Enum to use as Bitmask (%i)."), Enum->NumEnums() - 1);
+		UE_LOG(LogLearning, Error, TEXT("%s: Too many values in Enum to use as Bitmask (%i)."), *GetName(), Enum->NumEnums() - 1);
 		return FLearningAgentsActionObjectElement();
 	}
 
@@ -1479,190 +1391,179 @@ FLearningAgentsActionObjectElement ULearningAgentsActions::MakeBitmaskAction(ULe
 		}
 	}
 
-	return MakeInclusiveDiscreteActionFromArrayView(Object, BitmaskIndices, Tag);
+	return MakeInclusiveDiscreteActionFromArrayView(BitmaskIndices, Name);
 }
 
-FLearningAgentsActionObjectElement ULearningAgentsActions::MakeOptionalAction(ULearningAgentsActionObject* Object, const FLearningAgentsActionObjectElement Element, const ELearningAgentsOptionalAction Option, const FName Tag)
+FLearningAgentsActionObjectElement ULearningAgentsActionObject::MakeOptionalAction(const FLearningAgentsActionObjectElement Element, const ELearningAgentsOptionalAction Option, const FName Name)
 {
-	return MakeExclusiveUnionAction(Object,
+	return MakeExclusiveUnionAction(
 		Option == ELearningAgentsOptionalAction::Null ? TEXT("Null") : TEXT("Valid"),
-		Option == ELearningAgentsOptionalAction::Null ? MakeNullAction(Object) : Element,
-		Tag);
+		Option == ELearningAgentsOptionalAction::Null ? MakeNullAction() : Element,
+		Name);
 }
 
-FLearningAgentsActionObjectElement ULearningAgentsActions::MakeOptionalNullAction(ULearningAgentsActionObject* Object, const FName Tag)
+FLearningAgentsActionObjectElement ULearningAgentsActionObject::MakeOptionalNullAction(const FName Name)
 {
-	return MakeExclusiveUnionAction(Object, TEXT("Null"), MakeNullAction(Object), Tag);
+	return MakeExclusiveUnionAction(TEXT("Null"), MakeNullAction(), Name);
 }
 
-FLearningAgentsActionObjectElement ULearningAgentsActions::MakeOptionalValidAction(ULearningAgentsActionObject* Object, const FLearningAgentsActionObjectElement Element, const FName Tag)
+FLearningAgentsActionObjectElement ULearningAgentsActionObject::MakeOptionalValidAction(const FLearningAgentsActionObjectElement Element, const FName Name)
 {
-	return MakeExclusiveUnionAction(Object, TEXT("Valid"), Element, Tag);
+	return MakeExclusiveUnionAction(TEXT("Valid"), Element, Name);
 }
 
-FLearningAgentsActionObjectElement ULearningAgentsActions::MakeEitherAction(ULearningAgentsActionObject* Object, const FLearningAgentsActionObjectElement Element, const ELearningAgentsEitherAction Either, const FName Tag)
+FLearningAgentsActionObjectElement ULearningAgentsActionObject::MakeEitherAction(const FLearningAgentsActionObjectElement Element, const ELearningAgentsEitherAction Either, const FName Name)
 {
-	return MakeExclusiveUnionAction(Object, Either == ELearningAgentsEitherAction::A ? TEXT("A") : TEXT("B"), Element, Tag);
+	return MakeExclusiveUnionAction(Either == ELearningAgentsEitherAction::A ? TEXT("A") : TEXT("B"), Element, Name);
 }
 
-FLearningAgentsActionObjectElement ULearningAgentsActions::MakeEitherAAction(ULearningAgentsActionObject* Object, const FLearningAgentsActionObjectElement A, const FName Tag)
+FLearningAgentsActionObjectElement ULearningAgentsActionObject::MakeEitherAAction(const FLearningAgentsActionObjectElement A, const FName Name)
 {
-	return MakeExclusiveUnionAction(Object, TEXT("A"), A, Tag);
+	return MakeExclusiveUnionAction(TEXT("A"), A, Name);
 }
 
-FLearningAgentsActionObjectElement ULearningAgentsActions::MakeEitherBAction(ULearningAgentsActionObject* Object, const FLearningAgentsActionObjectElement B, const FName Tag)
+FLearningAgentsActionObjectElement ULearningAgentsActionObject::MakeEitherBAction(const FLearningAgentsActionObjectElement B, const FName Name)
 {
-	return MakeExclusiveUnionAction(Object, TEXT("B"), B, Tag);
+	return MakeExclusiveUnionAction(TEXT("B"), B, Name);
 }
 
-FLearningAgentsActionObjectElement ULearningAgentsActions::MakeEncodingAction(ULearningAgentsActionObject* Object, const FLearningAgentsActionObjectElement Element, const FName Tag)
+FLearningAgentsActionObjectElement ULearningAgentsActionObject::MakeEncodingAction(const FLearningAgentsActionObjectElement Element, const FName Name)
 {
-	if (!Object)
-	{
-		UE_LOG(LogLearning, Error, TEXT("MakeEncodingAction: Object is nullptr."));
-		return FLearningAgentsActionObjectElement();
-	}
-
-	if (!Object->ActionObject.IsValid(Element.ObjectElement))
+	if (!ActionObject.IsValid(Element.ObjectElement))
 	{
-		UE_LOG(LogLearning, Error, TEXT("MakeEncodingAction: Invalid Action Object."));
+		UE_LOG(LogLearning, Error, TEXT("%s: Invalid Action Object."), *GetName());
 		return FLearningAgentsActionObjectElement();
 	}
 
-	return { Object->ActionObject.CreateEncoding({ Element.ObjectElement }, Tag) };
+	return { ActionObject.CreateEncoding({ Element.ObjectElement }, Name) };
 }
 
-FLearningAgentsActionObjectElement ULearningAgentsActions::MakeBoolAction(ULearningAgentsActionObject* Object, const bool bValue, const FName Tag)
+FLearningAgentsActionObjectElement ULearningAgentsActionObject::MakeBoolAction(const bool bValue, const FName Name)
 {
-	return MakeExclusiveDiscreteAction(Object, bValue ? 1 : 0, Tag);
+	return MakeExclusiveDiscreteAction(bValue ? 1 : 0, Name);
 }
 
-FLearningAgentsActionObjectElement ULearningAgentsActions::MakeFloatAction(ULearningAgentsActionObject* Object, const float Value, const float FloatScale, const FName Tag)
+FLearningAgentsActionObjectElement ULearningAgentsActionObject::MakeFloatAction(const float Value, const float FloatScale, const FName Name)
 {
-	return MakeContinuousActionFromArrayView(Object, { Value / FMath::Max(FloatScale, UE_SMALL_NUMBER) }, Tag);
+	return MakeContinuousActionFromArrayView({ Value / FMath::Max(FloatScale, UE_SMALL_NUMBER) }, Name);
 }
 
-FLearningAgentsActionObjectElement ULearningAgentsActions::MakeLocationAction(ULearningAgentsActionObject* Object, const FVector Location, const FTransform RelativeTransform, const float LocationScale, const FName Tag)
+FLearningAgentsActionObjectElement ULearningAgentsActionObject::MakeLocationAction(const FVector Location, const FTransform RelativeTransform, const float LocationScale, const FName Name)
 {
 	const FVector LocalLocation = RelativeTransform.InverseTransformPosition(Location);
 
-	return MakeContinuousActionFromArrayView(Object, {
+	return MakeContinuousActionFromArrayView({
 		(float)LocalLocation.X / FMath::Max(LocationScale, UE_SMALL_NUMBER),
 		(float)LocalLocation.Y / FMath::Max(LocationScale, UE_SMALL_NUMBER),
-		(float)LocalLocation.Z / FMath::Max(LocationScale, UE_SMALL_NUMBER) }, Tag);
+		(float)LocalLocation.Z / FMath::Max(LocationScale, UE_SMALL_NUMBER) }, Name);
 }
 
-FLearningAgentsActionObjectElement ULearningAgentsActions::MakeRotationAction(ULearningAgentsActionObject* Object, const FRotator Rotation, const FRotator RelativeRotation, const float RotationScale, const FName Tag)
+FLearningAgentsActionObjectElement ULearningAgentsActionObject::MakeRotationAction(const FRotator Rotation, const FRotator RelativeRotation, const float RotationScale, const FName Name)
 {
-	return MakeRotationActionFromQuat(Object, FQuat::MakeFromRotator(Rotation), FQuat::MakeFromRotator(RelativeRotation), RotationScale, Tag);
+	return MakeRotationActionFromQuat(FQuat::MakeFromRotator(Rotation), FQuat::MakeFromRotator(RelativeRotation), RotationScale, Name);
 }
 
-FLearningAgentsActionObjectElement ULearningAgentsActions::MakeRotationActionFromQuat(ULearningAgentsActionObject* Object, const FQuat Rotation, const FQuat RelativeRotation, const float RotationScale, const FName Tag)
+FLearningAgentsActionObjectElement ULearningAgentsActionObject::MakeRotationActionFromQuat(const FQuat Rotation, const FQuat RelativeRotation, const float RotationScale, const FName Name)
 {
 	FQuat LocalRotation = RelativeRotation.Inverse() * Rotation;
 	LocalRotation.EnforceShortestArcWith(FQuat::Identity);
 	const FVector RotationVector = LocalRotation.ToRotationVector();
 
-	return MakeContinuousActionFromArrayView(Object, {
+	return MakeContinuousActionFromArrayView({
 		(float)RotationVector.X / FMath::Max(FMath::DegreesToRadians(RotationScale), UE_SMALL_NUMBER),
 		(float)RotationVector.Y / FMath::Max(FMath::DegreesToRadians(RotationScale), UE_SMALL_NUMBER),
 		(float)RotationVector.Z / FMath::Max(FMath::DegreesToRadians(RotationScale), UE_SMALL_NUMBER),
-		}, Tag);
+		}, Name);
 }
 
-FLearningAgentsActionObjectElement ULearningAgentsActions::MakeScaleAction(ULearningAgentsActionObject* Object, const FVector Scale, const FVector RelativeScale, const FName Tag)
+FLearningAgentsActionObjectElement ULearningAgentsActionObject::MakeScaleAction(const FVector Scale, const FVector RelativeScale, const FName Name)
 {
 	const FVector LocalLogScale =
 		UE::Learning::Agents::Action::Private::VectorLogSafe(Scale) -
 		UE::Learning::Agents::Action::Private::VectorLogSafe(RelativeScale);
 
-	return MakeContinuousActionFromArrayView(Object, {
+	return MakeContinuousActionFromArrayView({
 		(float)LocalLogScale.X,
 		(float)LocalLogScale.Y,
 		(float)LocalLogScale.Z,
-		}, Tag);
+		}, Name);
 }
 
-FLearningAgentsActionObjectElement ULearningAgentsActions::MakeTransformAction(ULearningAgentsActionObject* Object, const FTransform Transform, const FTransform RelativeTransform, const float LocationScale, const FName Tag)
+FLearningAgentsActionObjectElement ULearningAgentsActionObject::MakeTransformAction(const FTransform Transform, const FTransform RelativeTransform, const float LocationScale, const FName Name)
 {
 	const FTransform LocalTransform = Transform * RelativeTransform.Inverse();
 
-	return MakeStructActionFromArrayViews(Object,
+	return MakeStructActionFromArrayViews(
 		{
 			TEXT("Location"),
 			TEXT("Rotation"),
 			TEXT("Scale")
 		},
 		{
-			MakeLocationAction(Object, LocalTransform.GetLocation(), FTransform::Identity, LocationScale),
-			MakeRotationActionFromQuat(Object, LocalTransform.GetRotation(), FQuat::Identity),
-			MakeScaleAction(Object, LocalTransform.GetScale3D(), FVector::OneVector)
+			MakeLocationAction(LocalTransform.GetLocation(), FTransform::Identity, LocationScale),
+			MakeRotationActionFromQuat(LocalTransform.GetRotation(), FQuat::Identity),
+			MakeScaleAction(LocalTransform.GetScale3D(), FVector::OneVector)
 		},
-		Tag);
+		Name);
 }
 
-FLearningAgentsActionObjectElement ULearningAgentsActions::MakeAngleAction(ULearningAgentsActionObject* Object, const float Angle, const float RelativeAngle, const float AngleScale, const FName Tag)
+FLearningAgentsActionObjectElement ULearningAgentsActionObject::MakeAngleAction(const float Angle, const float RelativeAngle, const float AngleScale, const FName Name)
 {
-	return MakeAngleActionRadians(Object, FMath::DegreesToRadians(Angle), FMath::DegreesToRadians(RelativeAngle), FMath::DegreesToRadians(AngleScale), Tag);
+	return MakeAngleActionRadians(FMath::DegreesToRadians(Angle), FMath::DegreesToRadians(RelativeAngle), FMath::DegreesToRadians(AngleScale), Name);
 }
 
-FLearningAgentsActionObjectElement ULearningAgentsActions::MakeAngleActionRadians(ULearningAgentsActionObject* Object, const float Angle, const float RelativeAngle, const float AngleScale, const FName Tag)
+FLearningAgentsActionObjectElement ULearningAgentsActionObject::MakeAngleActionRadians(const float Angle, const float RelativeAngle, const float AngleScale, const FName Name)
 {
 	const float LocalAngle = FMath::FindDeltaAngleRadians(RelativeAngle, Angle);
-	return MakeContinuousActionFromArrayView(Object, { LocalAngle / FMath::Max(AngleScale, UE_SMALL_NUMBER) }, Tag);
+	return MakeContinuousActionFromArrayView({ LocalAngle / FMath::Max(AngleScale, UE_SMALL_NUMBER) }, Name);
 }
 
-FLearningAgentsActionObjectElement ULearningAgentsActions::MakeVelocityAction(ULearningAgentsActionObject* Object, const FVector Velocity, const FTransform RelativeTransform, const float VelocityScale, const FName Tag)
+FLearningAgentsActionObjectElement ULearningAgentsActionObject::MakeVelocityAction(const FVector Velocity, const FTransform RelativeTransform, const float VelocityScale, const FName Name)
 {
 	const FVector LocalVelocity = RelativeTransform.InverseTransformVectorNoScale(Velocity);
 
-	return MakeContinuousActionFromArrayView(Object, {
+	return MakeContinuousActionFromArrayView({
 		(float)LocalVelocity.X / FMath::Max(VelocityScale, UE_SMALL_NUMBER),
 		(float)LocalVelocity.Y / FMath::Max(VelocityScale, UE_SMALL_NUMBER),
 		(float)LocalVelocity.Z / FMath::Max(VelocityScale, UE_SMALL_NUMBER),
-		}, Tag);
+		}, Name);
 }
 
-FLearningAgentsActionObjectElement ULearningAgentsActions::MakeDirectionAction(ULearningAgentsActionObject* Object, const FVector Direction, const FTransform RelativeTransform, const FName Tag)
+FLearningAgentsActionObjectElement ULearningAgentsActionObject::MakeDirectionAction(const FVector Direction, const FTransform RelativeTransform, const FName Name)
 {
 	const FVector LocalDirection = RelativeTransform.InverseTransformVectorNoScale(Direction).GetSafeNormal(UE_SMALL_NUMBER, FVector::ForwardVector);
 
-	return MakeContinuousActionFromArrayView(Object, {
+	return MakeContinuousActionFromArrayView({
 		(float)LocalDirection.X,
 		(float)LocalDirection.Y,
 		(float)LocalDirection.Z,
-		}, Tag);
+		}, Name);
 }
 
-FLearningAgentsActionObjectElement ULearningAgentsActions::MakeSpeedAction(ULearningAgentsActionObject* Object, const float Speed, const float SpeedScale, const FName Tag)
+FLearningAgentsActionObjectElement ULearningAgentsActionObject::MakeSpeedAction(const float Speed, const float SpeedScale, const FName Name)
 {
-	return MakeContinuousActionFromArrayView(Object, { Speed / FMath::Max(SpeedScale, UE_SMALL_NUMBER) });
+	return MakeContinuousActionFromArrayView({ Speed / FMath::Max(SpeedScale, UE_SMALL_NUMBER) });
 }
 
-bool ULearningAgentsActions::GetNullAction(const ULearningAgentsActionObject* Object, const FLearningAgentsActionObjectElement Element, const FName Tag)
+bool ULearningAgentsActionObject::GetNullAction(const FLearningAgentsActionObjectElement Element, const FName Name) const
 {
-	if (!Object)
-	{
-		UE_LOG(LogLearning, Error, TEXT("GetNullAction: Object is nullptr."));
-		return false;
-	}
-
-	if (!Object->ActionObject.IsValid(Element.ObjectElement))
+	if (!ActionObject.IsValid(Element.ObjectElement))
 	{
-		UE_LOG(LogLearning, Error, TEXT("GetNullAction: Invalid Action Object."));
+		UE_LOG(LogLearning, Error, TEXT("%s: Invalid Action Object."), *GetName());
 		return false;
 	}
 
-	if (Object->ActionObject.GetTag(Element.ObjectElement) != Tag)
+	if (ActionObject.GetName(Element.ObjectElement) != Name)
 	{
-		UE_LOG(LogLearning, Warning, TEXT("GetNullAction: Action tag does not match. Action is '%s' but asked for '%s'."), *Object->ActionObject.GetTag(Element.ObjectElement).ToString(), *Tag.ToString());
+		UE_LOG(LogLearning, Warning, TEXT("%s: Action name does not match. Action is '%s' but asked for '%s'."), *GetName(), *ActionObject.GetName(Element.ObjectElement).ToString(), *Name.ToString());
 	}
 
-	if (Object->ActionObject.GetType(Element.ObjectElement) != UE::Learning::Action::EType::Null)
+	if (ActionObject.GetType(Element.ObjectElement) != UE::Learning::Action::EType::Null)
 	{
-		UE_LOG(LogLearning, Error, TEXT("GetNullAction: Action '%s' type does not match. Action is '%s' but asked for '%s'."),
-			*Object->ActionObject.GetTag(Element.ObjectElement).ToString(),
-			UE::Learning::Agents::Action::Private::GetActionTypeString(Object->ActionObject.GetType(Element.ObjectElement)),
+		UE_LOG(LogLearning, Error, TEXT("%s: Action '%s' type does not match. Action is '%s' but asked for '%s'."),
+			*GetName(),
+			*ActionObject.GetName(Element.ObjectElement).ToString(),
+			UE::Learning::Agents::Action::Private::GetActionTypeString(ActionObject.GetType(Element.ObjectElement)),
 			UE::Learning::Agents::Action::Private::GetActionTypeString(UE::Learning::Action::EType::Null));
 		return false;
 	}
@@ -1670,45 +1571,39 @@ bool ULearningAgentsActions::GetNullAction(const ULearningAgentsActionObject* Ob
 	return true;
 }
 
-bool ULearningAgentsActions::GetContinuousActionNum(int32& OutNum, const ULearningAgentsActionObject* Object, const FLearningAgentsActionObjectElement Element, const FName Tag)
+bool ULearningAgentsActionObject::GetContinuousActionNum(int32& OutNum, const FLearningAgentsActionObjectElement Element, const FName Name) const
 {
-	if (!Object)
-	{
-		UE_LOG(LogLearning, Error, TEXT("GetContinuousActionNum: Object is nullptr."));
-		OutNum = 0;
-		return false;
-	}
-
-	if (!Object->ActionObject.IsValid(Element.ObjectElement))
+	if (!ActionObject.IsValid(Element.ObjectElement))
 	{
-		UE_LOG(LogLearning, Error, TEXT("GetContinuousActionNum: Invalid Action Object."));
+		UE_LOG(LogLearning, Error, TEXT("%s: Invalid Action Object."), *GetName());
 		OutNum = 0;
 		return false;
 	}
 
-	if (Object->ActionObject.GetTag(Element.ObjectElement) != Tag)
+	if (ActionObject.GetName(Element.ObjectElement) != Name)
 	{
-		UE_LOG(LogLearning, Warning, TEXT("GetContinuousActionNum: Action tag does not match. Action is '%s' but asked for '%s'."), *Object->ActionObject.GetTag(Element.ObjectElement).ToString(), *Tag.ToString());
+		UE_LOG(LogLearning, Warning, TEXT("%s: Action name does not match. Action is '%s' but asked for '%s'."), *GetName(), *ActionObject.GetName(Element.ObjectElement).ToString(), *Name.ToString());
 	}
 
-	if (Object->ActionObject.GetType(Element.ObjectElement) != UE::Learning::Action::EType::Continuous)
+	if (ActionObject.GetType(Element.ObjectElement) != UE::Learning::Action::EType::Continuous)
 	{
-		UE_LOG(LogLearning, Error, TEXT("GetContinuousActionNum: Action '%s' type does not match. Action is '%s' but asked for '%s'."),
-			*Object->ActionObject.GetTag(Element.ObjectElement).ToString(),
-			UE::Learning::Agents::Action::Private::GetActionTypeString(Object->ActionObject.GetType(Element.ObjectElement)),
+		UE_LOG(LogLearning, Error, TEXT("%s: Action '%s' type does not match. Action is '%s' but asked for '%s'."),
+			*GetName(),
+			*ActionObject.GetName(Element.ObjectElement).ToString(),
+			UE::Learning::Agents::Action::Private::GetActionTypeString(ActionObject.GetType(Element.ObjectElement)),
 			UE::Learning::Agents::Action::Private::GetActionTypeString(UE::Learning::Action::EType::Continuous));
 		OutNum = 0;
 		return false;
 	}
 
-	OutNum = Object->ActionObject.GetContinuous(Element.ObjectElement).Values.Num();
+	OutNum = ActionObject.GetContinuous(Element.ObjectElement).Values.Num();
 	return true;
 }
 
-bool ULearningAgentsActions::GetContinuousAction(TArray<float>& OutValues, const ULearningAgentsActionObject* Object, const FLearningAgentsActionObjectElement Element, const FName Tag)
+bool ULearningAgentsActionObject::GetContinuousAction(TArray<float>& OutValues, const FLearningAgentsActionObjectElement Element, const FName Name) const
 {
 	int32 OutValueNum = 0;
-	if (!GetContinuousActionNum(OutValueNum, Object, Element, Tag))
+	if (!GetContinuousActionNum(OutValueNum, Element, Name))
 	{
 		OutValues.Empty();
 		return false;
@@ -1716,7 +1611,7 @@ bool ULearningAgentsActions::GetContinuousAction(TArray<float>& OutValues, const
 
 	OutValues.SetNumUninitialized(OutValueNum);
 
-	if (!GetContinuousActionToArrayView(OutValues, Object, Element, Tag))
+	if (!GetContinuousActionToArrayView(OutValues, Element, Name))
 	{
 		OutValues.Empty();
 		return false;
@@ -1725,43 +1620,37 @@ bool ULearningAgentsActions::GetContinuousAction(TArray<float>& OutValues, const
 	return true;
 }
 
-bool ULearningAgentsActions::GetContinuousActionToArrayView(TArrayView<float> OutValues, const ULearningAgentsActionObject* Object, const FLearningAgentsActionObjectElement Element, const FName Tag)
+bool ULearningAgentsActionObject::GetContinuousActionToArrayView(TArrayView<float> OutValues, const FLearningAgentsActionObjectElement Element, const FName Name) const
 {
-	if (!Object)
-	{
-		UE_LOG(LogLearning, Error, TEXT("GetContinuousActionToArrayView: Object is nullptr."));
-		UE::Learning::Array::Zero<1, float>(OutValues);
-		return false;
-	}
-
-	if (!Object->ActionObject.IsValid(Element.ObjectElement))
+	if (!ActionObject.IsValid(Element.ObjectElement))
 	{
-		UE_LOG(LogLearning, Error, TEXT("GetContinuousActionToArrayView: Invalid Action Object."));
+		UE_LOG(LogLearning, Error, TEXT("%s: Invalid Action Object."), *GetName());
 		UE::Learning::Array::Zero<1, float>(OutValues);
 		return false;
 	}
 
-	if (Object->ActionObject.GetTag(Element.ObjectElement) != Tag)
+	if (ActionObject.GetName(Element.ObjectElement) != Name)
 	{
-		UE_LOG(LogLearning, Warning, TEXT("GetContinuousActionToArrayView: Action tag does not match. Action is '%s' but asked for '%s'."), *Object->ActionObject.GetTag(Element.ObjectElement).ToString(), *Tag.ToString());
+		UE_LOG(LogLearning, Warning, TEXT("%s: Action name does not match. Action is '%s' but asked for '%s'."), *GetName(), *ActionObject.GetName(Element.ObjectElement).ToString(), *Name.ToString());
 	}
 
-	if (Object->ActionObject.GetType(Element.ObjectElement) != UE::Learning::Action::EType::Continuous)
+	if (ActionObject.GetType(Element.ObjectElement) != UE::Learning::Action::EType::Continuous)
 	{
-		UE_LOG(LogLearning, Error, TEXT("GetContinuousActionToArrayView: Action '%s' type does not match. Action is '%s' but asked for '%s'."),
-			*Object->ActionObject.GetTag(Element.ObjectElement).ToString(),
-			UE::Learning::Agents::Action::Private::GetActionTypeString(Object->ActionObject.GetType(Element.ObjectElement)),
+		UE_LOG(LogLearning, Error, TEXT("%s: Action '%s' type does not match. Action is '%s' but asked for '%s'."),
+			*GetName(),
+			*ActionObject.GetName(Element.ObjectElement).ToString(),
+			UE::Learning::Agents::Action::Private::GetActionTypeString(ActionObject.GetType(Element.ObjectElement)),
 			UE::Learning::Agents::Action::Private::GetActionTypeString(UE::Learning::Action::EType::Continuous));
 		UE::Learning::Array::Zero<1, float>(OutValues);
 		return false;
 	}
 
-	const TArrayView<const float> Values = Object->ActionObject.GetContinuous(Element.ObjectElement).Values;
+	const TArrayView<const float> Values = ActionObject.GetContinuous(Element.ObjectElement).Values;
 
 	if (Values.Num() != OutValues.Num())
 	{
-		UE_LOG(LogLearning, Error, TEXT("GetContinuousActionToArrayView: Action '%s' size does not match. Action is '%i' values but asked for '%i'."),
-			*Object->ActionObject.GetTag(Element.ObjectElement).ToString(),
+		UE_LOG(LogLearning, Error, TEXT("%s: Action '%s' size does not match. Action is '%i' values but asked for '%i'."),
+			*GetName(), *ActionObject.GetName(Element.ObjectElement).ToString(),
 			Values.Num(), OutValues.Num());
 		UE::Learning::Array::Zero<1, float>(OutValues);
 		return false;
@@ -1771,80 +1660,68 @@ bool ULearningAgentsActions::GetContinuousActionToArrayView(TArrayView<float> Ou
 	return true;
 }
 
-bool ULearningAgentsActions::GetExclusiveDiscreteAction(int32& OutIndex, const ULearningAgentsActionObject* Object, const FLearningAgentsActionObjectElement Element, const FName Tag)
+bool ULearningAgentsActionObject::GetExclusiveDiscreteAction(int32& OutIndex, const FLearningAgentsActionObjectElement Element, const FName Name) const
 {
-	if (!Object)
-	{
-		UE_LOG(LogLearning, Error, TEXT("GetExclusiveDiscreteAction: Object is nullptr."));
-		OutIndex = 0;
-		return false;
-	}
-
-	if (!Object->ActionObject.IsValid(Element.ObjectElement))
+	if (!ActionObject.IsValid(Element.ObjectElement))
 	{
-		UE_LOG(LogLearning, Error, TEXT("GetExclusiveDiscreteAction: Invalid Action Object."));
+		UE_LOG(LogLearning, Error, TEXT("%s: Invalid Action Object."), *GetName());
 		OutIndex = 0;
 		return false;
 	}
 
-	if (Object->ActionObject.GetTag(Element.ObjectElement) != Tag)
+	if (ActionObject.GetName(Element.ObjectElement) != Name)
 	{
-		UE_LOG(LogLearning, Warning, TEXT("GetExclusiveDiscreteAction: Action tag does not match. Action is '%s' but asked for '%s'."), *Object->ActionObject.GetTag(Element.ObjectElement).ToString(), *Tag.ToString());
+		UE_LOG(LogLearning, Warning, TEXT("%s: Action name does not match. Action is '%s' but asked for '%s'."), *GetName(), *ActionObject.GetName(Element.ObjectElement).ToString(), *Name.ToString());
 	}
 
-	if (Object->ActionObject.GetType(Element.ObjectElement) != UE::Learning::Action::EType::DiscreteExclusive)
+	if (ActionObject.GetType(Element.ObjectElement) != UE::Learning::Action::EType::DiscreteExclusive)
 	{
-		UE_LOG(LogLearning, Error, TEXT("GetExclusiveDiscreteAction: Action '%s' type does not match. Action is '%s' but asked for '%s'."),
-			*Object->ActionObject.GetTag(Element.ObjectElement).ToString(),
-			UE::Learning::Agents::Action::Private::GetActionTypeString(Object->ActionObject.GetType(Element.ObjectElement)),
+		UE_LOG(LogLearning, Error, TEXT("%s: Action '%s' type does not match. Action is '%s' but asked for '%s'."),
+			*GetName(),
+			*ActionObject.GetName(Element.ObjectElement).ToString(),
+			UE::Learning::Agents::Action::Private::GetActionTypeString(ActionObject.GetType(Element.ObjectElement)),
 			UE::Learning::Agents::Action::Private::GetActionTypeString(UE::Learning::Action::EType::DiscreteExclusive));
 		OutIndex = 0;
 		return false;
 	}
 
-	OutIndex = Object->ActionObject.GetDiscreteExclusive(Element.ObjectElement).DiscreteIndex;
+	OutIndex = ActionObject.GetDiscreteExclusive(Element.ObjectElement).DiscreteIndex;
 	return true;
 }
 
-bool ULearningAgentsActions::GetInclusiveDiscreteActionNum(int32& OutNum, const ULearningAgentsActionObject* Object, const FLearningAgentsActionObjectElement Element, const FName Tag)
+bool ULearningAgentsActionObject::GetInclusiveDiscreteActionNum(int32& OutNum, const FLearningAgentsActionObjectElement Element, const FName Name) const
 {
-	if (!Object)
-	{
-		UE_LOG(LogLearning, Error, TEXT("GetInclusiveDiscreteActionNum: Object is nullptr."));
-		OutNum = 0;
-		return false;
-	}
-
-	if (!Object->ActionObject.IsValid(Element.ObjectElement))
+	if (!ActionObject.IsValid(Element.ObjectElement))
 	{
-		UE_LOG(LogLearning, Error, TEXT("GetInclusiveDiscreteActionNum: Invalid Action Object."));
+		UE_LOG(LogLearning, Error, TEXT("%s: Invalid Action Object."), *GetName());
 		OutNum = 0;
 		return false;
 	}
 
-	if (Object->ActionObject.GetTag(Element.ObjectElement) != Tag)
+	if (ActionObject.GetName(Element.ObjectElement) != Name)
 	{
-		UE_LOG(LogLearning, Warning, TEXT("GetInclusiveDiscreteActionNum: Action tag does not match. Action is '%s' but asked for '%s'."), *Object->ActionObject.GetTag(Element.ObjectElement).ToString(), *Tag.ToString());
+		UE_LOG(LogLearning, Warning, TEXT("%s: Action name does not match. Action is '%s' but asked for '%s'."), *GetName(), *ActionObject.GetName(Element.ObjectElement).ToString(), *Name.ToString());
 	}
 
-	if (Object->ActionObject.GetType(Element.ObjectElement) != UE::Learning::Action::EType::DiscreteInclusive)
+	if (ActionObject.GetType(Element.ObjectElement) != UE::Learning::Action::EType::DiscreteInclusive)
 	{
-		UE_LOG(LogLearning, Error, TEXT("GetInclusiveDiscreteActionNum: Action '%s' type does not match. Action is '%s' but asked for '%s'."),
-			*Object->ActionObject.GetTag(Element.ObjectElement).ToString(),
-			UE::Learning::Agents::Action::Private::GetActionTypeString(Object->ActionObject.GetType(Element.ObjectElement)),
+		UE_LOG(LogLearning, Error, TEXT("%s: Action '%s' type does not match. Action is '%s' but asked for '%s'."),
+			*GetName(),
+			*ActionObject.GetName(Element.ObjectElement).ToString(),
+			UE::Learning::Agents::Action::Private::GetActionTypeString(ActionObject.GetType(Element.ObjectElement)),
 			UE::Learning::Agents::Action::Private::GetActionTypeString(UE::Learning::Action::EType::DiscreteInclusive));
 		OutNum = 0;
 		return false;
 	}
 
-	OutNum = Object->ActionObject.GetDiscreteInclusive(Element.ObjectElement).DiscreteIndices.Num();
+	OutNum = ActionObject.GetDiscreteInclusive(Element.ObjectElement).DiscreteIndices.Num();
 	return true;
 }
 
-bool ULearningAgentsActions::GetInclusiveDiscreteAction(TArray<int32>& OutIndices, const ULearningAgentsActionObject* Object, const FLearningAgentsActionObjectElement Element, const FName Tag)
+bool ULearningAgentsActionObject::GetInclusiveDiscreteAction(TArray<int32>& OutIndices, const FLearningAgentsActionObjectElement Element, const FName Name) const
 {
 	int32 OutIndexNum = 0;
-	if (!GetInclusiveDiscreteActionNum(OutIndexNum, Object, Element, Tag))
+	if (!GetInclusiveDiscreteActionNum(OutIndexNum, Element, Name))
 	{
 		OutIndices.Empty();
 		return false;
@@ -1852,7 +1729,7 @@ bool ULearningAgentsActions::GetInclusiveDiscreteAction(TArray<int32>& OutIndice
 
 	OutIndices.SetNumUninitialized(OutIndexNum);
 
-	if (!GetInclusiveDiscreteActionToArrayView(OutIndices, Object, Element, Tag))
+	if (!GetInclusiveDiscreteActionToArrayView(OutIndices, Element, Name))
 	{
 		OutIndices.Empty();
 		return false;
@@ -1861,43 +1738,37 @@ bool ULearningAgentsActions::GetInclusiveDiscreteAction(TArray<int32>& OutIndice
 	return true;
 }
 
-bool ULearningAgentsActions::GetInclusiveDiscreteActionToArrayView(TArrayView<int32> OutIndices, const ULearningAgentsActionObject* Object, const FLearningAgentsActionObjectElement Element, const FName Tag)
+bool ULearningAgentsActionObject::GetInclusiveDiscreteActionToArrayView(TArrayView<int32> OutIndices, const FLearningAgentsActionObjectElement Element, const FName Name) const
 {
-	if (!Object)
-	{
-		UE_LOG(LogLearning, Error, TEXT("GetInclusiveDiscreteActionToArrayView: Object is nullptr."));
-		UE::Learning::Array::Zero<1, int32>(OutIndices);
-		return false;
-	}
-
-	if (!Object->ActionObject.IsValid(Element.ObjectElement))
+	if (!ActionObject.IsValid(Element.ObjectElement))
 	{
-		UE_LOG(LogLearning, Error, TEXT("GetInclusiveDiscreteActionToArrayView: Invalid Action Object."));
+		UE_LOG(LogLearning, Error, TEXT("%s: Invalid Action Object."), *GetName());
 		UE::Learning::Array::Zero<1, int32>(OutIndices);
 		return false;
 	}
 
-	if (Object->ActionObject.GetTag(Element.ObjectElement) != Tag)
+	if (ActionObject.GetName(Element.ObjectElement) != Name)
 	{
-		UE_LOG(LogLearning, Warning, TEXT("GetInclusiveDiscreteActionToArrayView: Action tag does not match. Action is '%s' but asked for '%s'."), *Object->ActionObject.GetTag(Element.ObjectElement).ToString(), *Tag.ToString());
+		UE_LOG(LogLearning, Warning, TEXT("%s: Action name does not match. Action is '%s' but asked for '%s'."), *GetName(), *ActionObject.GetName(Element.ObjectElement).ToString(), *Name.ToString());
 	}
 
-	if (Object->ActionObject.GetType(Element.ObjectElement) != UE::Learning::Action::EType::DiscreteInclusive)
+	if (ActionObject.GetType(Element.ObjectElement) != UE::Learning::Action::EType::DiscreteInclusive)
 	{
-		UE_LOG(LogLearning, Error, TEXT("GetInclusiveDiscreteActionToArrayView: Action '%s' type does not match. Action is '%s' but asked for '%s'."),
-			*Object->ActionObject.GetTag(Element.ObjectElement).ToString(),
-			UE::Learning::Agents::Action::Private::GetActionTypeString(Object->ActionObject.GetType(Element.ObjectElement)),
+		UE_LOG(LogLearning, Error, TEXT("%s: Action '%s' type does not match. Action is '%s' but asked for '%s'."),
+			*GetName(),
+			*ActionObject.GetName(Element.ObjectElement).ToString(),
+			UE::Learning::Agents::Action::Private::GetActionTypeString(ActionObject.GetType(Element.ObjectElement)),
 			UE::Learning::Agents::Action::Private::GetActionTypeString(UE::Learning::Action::EType::DiscreteInclusive));
 		UE::Learning::Array::Zero<1, int32>(OutIndices);
 		return false;
 	}
 
-	const TArrayView<const int32> Indices = Object->ActionObject.GetDiscreteInclusive(Element.ObjectElement).DiscreteIndices;
+	const TArrayView<const int32> Indices = ActionObject.GetDiscreteInclusive(Element.ObjectElement).DiscreteIndices;
 
 	if (Indices.Num() != OutIndices.Num())
 	{
-		UE_LOG(LogLearning, Error, TEXT("GetInclusiveDiscreteActionToArrayView: Action '%s' size does not match. Action is '%i' elements but asked for '%i'."),
-			*Object->ActionObject.GetTag(Element.ObjectElement).ToString(),
+		UE_LOG(LogLearning, Error, TEXT("%s: Action '%s' size does not match. Action is '%i' elements but asked for '%i'."),
+			*GetName(), *ActionObject.GetName(Element.ObjectElement).ToString(),
 			Indices.Num(), OutIndices.Num());
 		UE::Learning::Array::Zero<1, int32>(OutIndices);
 		return false;
@@ -1907,47 +1778,41 @@ bool ULearningAgentsActions::GetInclusiveDiscreteActionToArrayView(TArrayView<in
 	return true;
 }
 
-bool ULearningAgentsActions::GetStructActionNum(int32& OutNum, const ULearningAgentsActionObject* Object, const FLearningAgentsActionObjectElement Element, const FName Tag)
+bool ULearningAgentsActionObject::GetStructActionNum(int32& OutNum, const FLearningAgentsActionObjectElement Element, const FName Name) const
 {
-	if (!Object)
-	{
-		UE_LOG(LogLearning, Error, TEXT("GetStructActionNum: Object is nullptr."));
-		OutNum = 0;
-		return false;
-	}
-
-	if (!Object->ActionObject.IsValid(Element.ObjectElement))
+	if (!ActionObject.IsValid(Element.ObjectElement))
 	{
-		UE_LOG(LogLearning, Error, TEXT("GetStructActionNum: Invalid Action Object."));
+		UE_LOG(LogLearning, Error, TEXT("%s: Invalid Action Object."), *GetName());
 		OutNum = 0;
 		return false;
 	}
 
-	if (Object->ActionObject.GetTag(Element.ObjectElement) != Tag)
+	if (ActionObject.GetName(Element.ObjectElement) != Name)
 	{
-		UE_LOG(LogLearning, Warning, TEXT("GetStructActionNum: Action tag does not match. Action is '%s' but asked for '%s'."), *Object->ActionObject.GetTag(Element.ObjectElement).ToString(), *Tag.ToString());
+		UE_LOG(LogLearning, Warning, TEXT("%s: Action name does not match. Action is '%s' but asked for '%s'."), *GetName(), *ActionObject.GetName(Element.ObjectElement).ToString(), *Name.ToString());
 	}
 
-	if (Object->ActionObject.GetType(Element.ObjectElement) != UE::Learning::Action::EType::And)
+	if (ActionObject.GetType(Element.ObjectElement) != UE::Learning::Action::EType::And)
 	{
-		UE_LOG(LogLearning, Error, TEXT("GetStructActionNum: Action '%s' type does not match. Action is '%s' but asked for '%s'."),
-			*Object->ActionObject.GetTag(Element.ObjectElement).ToString(),
-			UE::Learning::Agents::Action::Private::GetActionTypeString(Object->ActionObject.GetType(Element.ObjectElement)),
+		UE_LOG(LogLearning, Error, TEXT("%s: Action '%s' type does not match. Action is '%s' but asked for '%s'."),
+			*GetName(),
+			*ActionObject.GetName(Element.ObjectElement).ToString(),
+			UE::Learning::Agents::Action::Private::GetActionTypeString(ActionObject.GetType(Element.ObjectElement)),
 			UE::Learning::Agents::Action::Private::GetActionTypeString(UE::Learning::Action::EType::And));
 		OutNum = 0;
 		return false;
 	}
 	
-	const UE::Learning::Action::FObjectAndParameters Parameters = Object->ActionObject.GetAnd(Element.ObjectElement);
+	const UE::Learning::Action::FObjectAndParameters Parameters = ActionObject.GetAnd(Element.ObjectElement);
 
 	OutNum = Parameters.Elements.Num();
 	return true;
 }
 
-bool ULearningAgentsActions::GetStructAction(TMap<FName, FLearningAgentsActionObjectElement>& OutElements, const ULearningAgentsActionObject* Object, const FLearningAgentsActionObjectElement Element, const FName Tag)
+bool ULearningAgentsActionObject::GetStructAction(TMap<FName, FLearningAgentsActionObjectElement>& OutElements, const FLearningAgentsActionObjectElement Element, const FName Name) const
 {
 	int32 OutElementNum = 0;
-	if (!GetStructActionNum(OutElementNum, Object, Element, Tag))
+	if (!GetStructActionNum(OutElementNum, Element, Name))
 	{
 		OutElements.Empty();
 		return false;
@@ -1958,7 +1823,7 @@ bool ULearningAgentsActions::GetStructAction(TMap<FName, FLearningAgentsActionOb
 	SubElementNames.SetNumUninitialized(OutElementNum);
 	SubElements.SetNumUninitialized(OutElementNum);
 
-	if (!GetStructActionToArrayViews(SubElementNames, SubElements, Object, Element, Tag))
+	if (!GetStructActionToArrayViews(SubElementNames, SubElements, Element, Name))
 	{
 		OutElements.Empty();
 		return false;
@@ -1973,10 +1838,10 @@ bool ULearningAgentsActions::GetStructAction(TMap<FName, FLearningAgentsActionOb
 	return true;
 }
 
-bool ULearningAgentsActions::GetStructActionToArrays(TArray<FName>& OutElementNames, TArray<FLearningAgentsActionObjectElement>& OutElements, const ULearningAgentsActionObject* Object, const FLearningAgentsActionObjectElement Element, const FName Tag)
+bool ULearningAgentsActionObject::GetStructActionToArrays(TArray<FName>& OutElementNames, TArray<FLearningAgentsActionObjectElement>& OutElements, const FLearningAgentsActionObjectElement Element, const FName Name) const
 {
 	int32 OutElementNum = 0;
-	if (!GetStructActionNum(OutElementNum, Object, Element, Tag))
+	if (!GetStructActionNum(OutElementNum, Element, Name))
 	{
 		OutElementNames.Empty();
 		OutElements.Empty();
@@ -1986,7 +1851,7 @@ bool ULearningAgentsActions::GetStructActionToArrays(TArray<FName>& OutElementNa
 	OutElementNames.SetNumUninitialized(OutElementNum);
 	OutElements.SetNumUninitialized(OutElementNum);
 
-	if (!GetStructActionToArrayViews(OutElementNames, OutElements, Object, Element, Tag))
+	if (!GetStructActionToArrayViews(OutElementNames, OutElements, Element, Name))
 	{
 		OutElementNames.Empty();
 		OutElements.Empty();
@@ -1996,51 +1861,45 @@ bool ULearningAgentsActions::GetStructActionToArrays(TArray<FName>& OutElementNa
 	return true;
 }
 
-bool ULearningAgentsActions::GetStructActionToArrayViews(TArrayView<FName> OutElementNames, TArrayView<FLearningAgentsActionObjectElement> OutElements, const ULearningAgentsActionObject* Object, const FLearningAgentsActionObjectElement Element, const FName Tag)
+bool ULearningAgentsActionObject::GetStructActionToArrayViews(TArrayView<FName> OutElementNames, TArrayView<FLearningAgentsActionObjectElement> OutElements, const FLearningAgentsActionObjectElement Element, const FName Name) const
 {
-	if (!Object)
-	{
-		UE_LOG(LogLearning, Error, TEXT("GetStructActionToArrayViews: Object is nullptr."));
-		UE::Learning::Array::Set<1, FName>(OutElementNames, NAME_None);
-		UE::Learning::Array::Set<1, FLearningAgentsActionObjectElement>(OutElements, FLearningAgentsActionObjectElement());
-		return false;
-	}
-
-	if (!Object->ActionObject.IsValid(Element.ObjectElement))
+	if (!ActionObject.IsValid(Element.ObjectElement))
 	{
-		UE_LOG(LogLearning, Error, TEXT("GetStructActionToArrayViews: Invalid Action Object."));
+		UE_LOG(LogLearning, Error, TEXT("%s: Invalid Action Object."), *GetName());
 		UE::Learning::Array::Set<1, FName>(OutElementNames, NAME_None);
 		UE::Learning::Array::Set<1, FLearningAgentsActionObjectElement>(OutElements, FLearningAgentsActionObjectElement());
 		return false;
 	}
 
-	if (Object->ActionObject.GetTag(Element.ObjectElement) != Tag)
+	if (ActionObject.GetName(Element.ObjectElement) != Name)
 	{
-		UE_LOG(LogLearning, Warning, TEXT("GetStructActionToArrayViews: Action tag does not match. Action is '%s' but asked for '%s'."), *Object->ActionObject.GetTag(Element.ObjectElement).ToString(), *Tag.ToString());
+		UE_LOG(LogLearning, Warning, TEXT("%s: Action name does not match. Action is '%s' but asked for '%s'."), *GetName(), *ActionObject.GetName(Element.ObjectElement).ToString(), *Name.ToString());
 	}
 
-	if (Object->ActionObject.GetType(Element.ObjectElement) != UE::Learning::Action::EType::And)
+	if (ActionObject.GetType(Element.ObjectElement) != UE::Learning::Action::EType::And)
 	{
-		UE_LOG(LogLearning, Error, TEXT("GetStructActionToArrayViews: Action '%s' type does not match. Action is '%s' but asked for '%s'."),
-			*Object->ActionObject.GetTag(Element.ObjectElement).ToString(),
-			UE::Learning::Agents::Action::Private::GetActionTypeString(Object->ActionObject.GetType(Element.ObjectElement)),
+		UE_LOG(LogLearning, Error, TEXT("%s: Action '%s' type does not match. Action is '%s' but asked for '%s'."),
+			*GetName(),
+			*ActionObject.GetName(Element.ObjectElement).ToString(),
+			UE::Learning::Agents::Action::Private::GetActionTypeString(ActionObject.GetType(Element.ObjectElement)),
 			UE::Learning::Agents::Action::Private::GetActionTypeString(UE::Learning::Action::EType::And));
 		UE::Learning::Array::Set<1, FName>(OutElementNames, NAME_None);
 		UE::Learning::Array::Set<1, FLearningAgentsActionObjectElement>(OutElements, FLearningAgentsActionObjectElement());
 		return false;
 	}
 
-	const UE::Learning::Action::FObjectAndParameters Parameters = Object->ActionObject.GetAnd(Element.ObjectElement);
+	const UE::Learning::Action::FObjectAndParameters Parameters = ActionObject.GetAnd(Element.ObjectElement);
 
 	if (Parameters.Elements.Num() == 0)
 	{
-		UE_LOG(LogLearning, Warning, TEXT("GetStructActionToArrayViews: Getting zero-sized And Action."));
+		UE_LOG(LogLearning, Warning, TEXT("%s: Getting zero-sized And Action."), *GetName());
 	}
 
 	if (Parameters.Elements.Num() != OutElements.Num())
 	{
-		UE_LOG(LogLearning, Error, TEXT("GetStructActionToArrayViews: Action '%s' size does not match. Action is '%i' elements but asked for '%i'."),
-			*Object->ActionObject.GetTag(Element.ObjectElement).ToString(),
+		UE_LOG(LogLearning, Error, TEXT("%s: Action '%s' size does not match. Action is '%i' elements but asked for '%i'."),
+			*GetName(),
+			*ActionObject.GetName(Element.ObjectElement).ToString(),
 			Parameters.Elements.Num(), OutElements.Num());
 		UE::Learning::Array::Set<1, FName>(OutElementNames, NAME_None);
 		UE::Learning::Array::Set<1, FLearningAgentsActionObjectElement>(OutElements, FLearningAgentsActionObjectElement());
@@ -2049,9 +1908,9 @@ bool ULearningAgentsActions::GetStructActionToArrayViews(TArrayView<FName> OutEl
 
 	for (int32 ElementIdx = 0; ElementIdx < Parameters.Elements.Num(); ElementIdx++)
 	{
-		if (!Object->ActionObject.IsValid(Parameters.Elements[ElementIdx]))
+		if (!ActionObject.IsValid(Parameters.Elements[ElementIdx]))
 		{
-			UE_LOG(LogLearning, Error, TEXT("GetStructActionToArrayViews: Invalid Action Object."));
+			UE_LOG(LogLearning, Error, TEXT("%s: Invalid Action Object."), *GetName());
 			UE::Learning::Array::Set<1, FName>(OutElementNames, NAME_None);
 			UE::Learning::Array::Set<1, FLearningAgentsActionObjectElement>(OutElements, FLearningAgentsActionObjectElement());
 			return false;
@@ -2064,88 +1923,75 @@ bool ULearningAgentsActions::GetStructActionToArrayViews(TArrayView<FName> OutEl
 	return true;
 }
 
-bool ULearningAgentsActions::GetExclusiveUnionAction(FName& OutElementName, FLearningAgentsActionObjectElement& OutElement, const ULearningAgentsActionObject* Object, const FLearningAgentsActionObjectElement Element, const FName Tag)
+bool ULearningAgentsActionObject::GetExclusiveUnionAction(FName& OutElementName, FLearningAgentsActionObjectElement& OutElement, const FLearningAgentsActionObjectElement Element, const FName Name) const
 {
-	if (!Object)
-	{
-		UE_LOG(LogLearning, Error, TEXT("GetExclusiveUnionAction: Object is nullptr."));
-		OutElementName = NAME_None;
-		OutElement = FLearningAgentsActionObjectElement();
-		return false;
-	}
-
-	if (!Object->ActionObject.IsValid(Element.ObjectElement))
+	if (!ActionObject.IsValid(Element.ObjectElement))
 	{
-		UE_LOG(LogLearning, Error, TEXT("GetExclusiveUnionAction: Invalid Action Object."));
+		UE_LOG(LogLearning, Error, TEXT("%s: Invalid Action Object."), *GetName());
 		OutElementName = NAME_None;
 		OutElement = FLearningAgentsActionObjectElement();
 		return false;
 	}
 
-	if (Object->ActionObject.GetTag(Element.ObjectElement) != Tag)
+	if (ActionObject.GetName(Element.ObjectElement) != Name)
 	{
-		UE_LOG(LogLearning, Warning, TEXT("GetExclusiveUnionAction: Action tag does not match. Action is '%s' but asked for '%s'."), *Object->ActionObject.GetTag(Element.ObjectElement).ToString(), *Tag.ToString());
+		UE_LOG(LogLearning, Warning, TEXT("%s: Action name does not match. Action is '%s' but asked for '%s'."), *GetName(), *ActionObject.GetName(Element.ObjectElement).ToString(), *Name.ToString());
 	}
 
-	if (Object->ActionObject.GetType(Element.ObjectElement) != UE::Learning::Action::EType::OrExclusive)
+	if (ActionObject.GetType(Element.ObjectElement) != UE::Learning::Action::EType::OrExclusive)
 	{
-		UE_LOG(LogLearning, Error, TEXT("GetExclusiveUnionAction: Action '%s' type does not match. Action is '%s' but asked for '%s'."),
-			*Object->ActionObject.GetTag(Element.ObjectElement).ToString(),
-			UE::Learning::Agents::Action::Private::GetActionTypeString(Object->ActionObject.GetType(Element.ObjectElement)),
+		UE_LOG(LogLearning, Error, TEXT("%s: Action '%s' type does not match. Action is '%s' but asked for '%s'."),
+			*GetName(),
+			*ActionObject.GetName(Element.ObjectElement).ToString(),
+			UE::Learning::Agents::Action::Private::GetActionTypeString(ActionObject.GetType(Element.ObjectElement)),
 			UE::Learning::Agents::Action::Private::GetActionTypeString(UE::Learning::Action::EType::OrExclusive));
 		OutElementName = NAME_None;
 		OutElement = FLearningAgentsActionObjectElement();
 		return false;
 	}
 
-	const UE::Learning::Action::FObjectOrExclusiveParameters Parameters = Object->ActionObject.GetOrExclusive(Element.ObjectElement);
+	const UE::Learning::Action::FObjectOrExclusiveParameters Parameters = ActionObject.GetOrExclusive(Element.ObjectElement);
 
 	OutElementName = Parameters.ElementName;
 	OutElement = { Parameters.Element };
 	return true;
 }
 
-bool ULearningAgentsActions::GetInclusiveUnionActionNum(int32& OutNum, const ULearningAgentsActionObject* Object, const FLearningAgentsActionObjectElement Element, const FName Tag)
+bool ULearningAgentsActionObject::GetInclusiveUnionActionNum(int32& OutNum, const FLearningAgentsActionObjectElement Element, const FName Name) const
 {
-	if (!Object)
-	{
-		UE_LOG(LogLearning, Error, TEXT("GetInclusiveUnionActionNum: Object is nullptr."));
-		OutNum = 0;
-		return false;
-	}
-
-	if (!Object->ActionObject.IsValid(Element.ObjectElement))
+	if (!ActionObject.IsValid(Element.ObjectElement))
 	{
-		UE_LOG(LogLearning, Error, TEXT("GetInclusiveUnionActionNum: Invalid Action Object."));
+		UE_LOG(LogLearning, Error, TEXT("%s: Invalid Action Object."), *GetName());
 		OutNum = 0;
 		return false;
 	}
 
-	if (Object->ActionObject.GetTag(Element.ObjectElement) != Tag)
+	if (ActionObject.GetName(Element.ObjectElement) != Name)
 	{
-		UE_LOG(LogLearning, Warning, TEXT("GetInclusiveUnionActionNum: Action tag does not match. Action is '%s' but asked for '%s'."), *Object->ActionObject.GetTag(Element.ObjectElement).ToString(), *Tag.ToString());
+		UE_LOG(LogLearning, Warning, TEXT("%s: Action name does not match. Action is '%s' but asked for '%s'."), *GetName(), *ActionObject.GetName(Element.ObjectElement).ToString(), *Name.ToString());
 	}
 
-	if (Object->ActionObject.GetType(Element.ObjectElement) != UE::Learning::Action::EType::OrInclusive)
+	if (ActionObject.GetType(Element.ObjectElement) != UE::Learning::Action::EType::OrInclusive)
 	{
-		UE_LOG(LogLearning, Error, TEXT("GetInclusiveUnionActionNum: Action '%s' type does not match. Action is '%s' but asked for '%s'."),
-			*Object->ActionObject.GetTag(Element.ObjectElement).ToString(),
-			UE::Learning::Agents::Action::Private::GetActionTypeString(Object->ActionObject.GetType(Element.ObjectElement)),
+		UE_LOG(LogLearning, Error, TEXT("%s: Action '%s' type does not match. Action is '%s' but asked for '%s'."),
+			*GetName(),
+			*ActionObject.GetName(Element.ObjectElement).ToString(),
+			UE::Learning::Agents::Action::Private::GetActionTypeString(ActionObject.GetType(Element.ObjectElement)),
 			UE::Learning::Agents::Action::Private::GetActionTypeString(UE::Learning::Action::EType::OrInclusive));
 		OutNum = 0;
 		return false;
 	}
 
-	const UE::Learning::Action::FObjectOrInclusiveParameters Parameters = Object->ActionObject.GetOrInclusive(Element.ObjectElement);
+	const UE::Learning::Action::FObjectOrInclusiveParameters Parameters = ActionObject.GetOrInclusive(Element.ObjectElement);
 
 	OutNum = Parameters.Elements.Num();
 	return true;
 }
 
-bool ULearningAgentsActions::GetInclusiveUnionAction(TMap<FName, FLearningAgentsActionObjectElement>& OutElements, const ULearningAgentsActionObject* Object, const FLearningAgentsActionObjectElement Element, const FName Tag)
+bool ULearningAgentsActionObject::GetInclusiveUnionAction(TMap<FName, FLearningAgentsActionObjectElement>& OutElements, const FLearningAgentsActionObjectElement Element, const FName Name) const
 {
 	int32 OutElementNum = 0;
-	if (!GetInclusiveUnionActionNum(OutElementNum, Object, Element, Tag))
+	if (!GetInclusiveUnionActionNum(OutElementNum, Element, Name))
 	{
 		OutElements.Empty();
 		return false;
@@ -2156,7 +2002,7 @@ bool ULearningAgentsActions::GetInclusiveUnionAction(TMap<FName, FLearningAgents
 	SubElementNames.SetNumUninitialized(OutElementNum);
 	SubElements.SetNumUninitialized(OutElementNum);
 
-	if (!GetInclusiveUnionActionToArrayViews(SubElementNames, SubElements, Object, Element, Tag))
+	if (!GetInclusiveUnionActionToArrayViews(SubElementNames, SubElements, Element, Name))
 	{
 		OutElements.Empty();
 		return false;
@@ -2171,10 +2017,10 @@ bool ULearningAgentsActions::GetInclusiveUnionAction(TMap<FName, FLearningAgents
 	return true;
 }
 
-bool ULearningAgentsActions::GetInclusiveUnionActionToArrays(TArray<FName>& OutElementNames, TArray<FLearningAgentsActionObjectElement>& OutElements, const ULearningAgentsActionObject* Object, const FLearningAgentsActionObjectElement Element, const FName Tag)
+bool ULearningAgentsActionObject::GetInclusiveUnionActionToArrays(TArray<FName>& OutElementNames, TArray<FLearningAgentsActionObjectElement>& OutElements, const FLearningAgentsActionObjectElement Element, const FName Name) const
 {
 	int32 OutElementNum = 0;
-	if (!GetInclusiveUnionActionNum(OutElementNum, Object, Element, Tag))
+	if (!GetInclusiveUnionActionNum(OutElementNum, Element, Name))
 	{
 		OutElementNames.Empty();
 		OutElements.Empty();
@@ -2184,7 +2030,7 @@ bool ULearningAgentsActions::GetInclusiveUnionActionToArrays(TArray<FName>& OutE
 	OutElementNames.SetNumUninitialized(OutElementNum);
 	OutElements.SetNumUninitialized(OutElementNum);
 
-	if (!GetInclusiveUnionActionToArrayViews(OutElementNames, OutElements, Object, Element, Tag))
+	if (!GetInclusiveUnionActionToArrayViews(OutElementNames, OutElements, Element, Name))
 	{
 		OutElementNames.Empty();
 		OutElements.Empty();
@@ -2194,46 +2040,40 @@ bool ULearningAgentsActions::GetInclusiveUnionActionToArrays(TArray<FName>& OutE
 	return true;
 }
 
-bool ULearningAgentsActions::GetInclusiveUnionActionToArrayViews(TArrayView<FName> OutElementNames, TArrayView<FLearningAgentsActionObjectElement> OutElements, const ULearningAgentsActionObject* Object, const FLearningAgentsActionObjectElement Element, const FName Tag)
+bool ULearningAgentsActionObject::GetInclusiveUnionActionToArrayViews(TArrayView<FName> OutElementNames, TArrayView<FLearningAgentsActionObjectElement> OutElements, const FLearningAgentsActionObjectElement Element, const FName Name) const
 {
-	if (!Object)
-	{
-		UE_LOG(LogLearning, Error, TEXT("GetInclusiveUnionActionToArrayViews: Object is nullptr."));
-		UE::Learning::Array::Set<1, FName>(OutElementNames, NAME_None);
-		UE::Learning::Array::Set<1, FLearningAgentsActionObjectElement>(OutElements, FLearningAgentsActionObjectElement());
-		return false;
-	}
-
-	if (!Object->ActionObject.IsValid(Element.ObjectElement))
+	if (!ActionObject.IsValid(Element.ObjectElement))
 	{
-		UE_LOG(LogLearning, Error, TEXT("GetInclusiveUnionActionToArrayViews: Invalid Action Object."));
+		UE_LOG(LogLearning, Error, TEXT("%s: Invalid Action Object."), *GetName());
 		UE::Learning::Array::Set<1, FName>(OutElementNames, NAME_None);
 		UE::Learning::Array::Set<1, FLearningAgentsActionObjectElement>(OutElements, FLearningAgentsActionObjectElement());
 		return false;
 	}
 
-	if (Object->ActionObject.GetTag(Element.ObjectElement) != Tag)
+	if (ActionObject.GetName(Element.ObjectElement) != Name)
 	{
-		UE_LOG(LogLearning, Warning, TEXT("GetInclusiveUnionActionToArrayViews: Action tag does not match. Action is '%s' but asked for '%s'."), *Object->ActionObject.GetTag(Element.ObjectElement).ToString(), *Tag.ToString());
+		UE_LOG(LogLearning, Warning, TEXT("%s: Action name does not match. Action is '%s' but asked for '%s'."), *GetName(), *ActionObject.GetName(Element.ObjectElement).ToString(), *Name.ToString());
 	}
 
-	if (Object->ActionObject.GetType(Element.ObjectElement) != UE::Learning::Action::EType::OrInclusive)
+	if (ActionObject.GetType(Element.ObjectElement) != UE::Learning::Action::EType::OrInclusive)
 	{
-		UE_LOG(LogLearning, Error, TEXT("GetInclusiveUnionActionToArrayViews: Action '%s' type does not match. Action is '%s' but asked for '%s'."),
-			*Object->ActionObject.GetTag(Element.ObjectElement).ToString(),
-			UE::Learning::Agents::Action::Private::GetActionTypeString(Object->ActionObject.GetType(Element.ObjectElement)),
+		UE_LOG(LogLearning, Error, TEXT("%s: Action '%s' type does not match. Action is '%s' but asked for '%s'."),
+			*GetName(),
+			*ActionObject.GetName(Element.ObjectElement).ToString(),
+			UE::Learning::Agents::Action::Private::GetActionTypeString(ActionObject.GetType(Element.ObjectElement)),
 			UE::Learning::Agents::Action::Private::GetActionTypeString(UE::Learning::Action::EType::OrInclusive));
 		UE::Learning::Array::Set<1, FName>(OutElementNames, NAME_None);
 		UE::Learning::Array::Set<1, FLearningAgentsActionObjectElement>(OutElements, FLearningAgentsActionObjectElement());
 		return false;
 	}
 
-	const UE::Learning::Action::FObjectOrInclusiveParameters Parameters = Object->ActionObject.GetOrInclusive(Element.ObjectElement);
+	const UE::Learning::Action::FObjectOrInclusiveParameters Parameters = ActionObject.GetOrInclusive(Element.ObjectElement);
 
 	if (Parameters.Elements.Num() != OutElements.Num())
 	{
-		UE_LOG(LogLearning, Error, TEXT("GetInclusiveUnionActionToArrayViews: Action '%s' size does not match. Action is '%i' elements but asked for '%i'."),
-			*Object->ActionObject.GetTag(Element.ObjectElement).ToString(),
+		UE_LOG(LogLearning, Error, TEXT("%s: Action '%s' size does not match. Action is '%i' elements but asked for '%i'."),
+			*GetName(),
+			*ActionObject.GetName(Element.ObjectElement).ToString(),
 			Parameters.Elements.Num(), OutElements.Num());
 		UE::Learning::Array::Set<1, FName>(OutElementNames, NAME_None);
 		UE::Learning::Array::Set<1, FLearningAgentsActionObjectElement>(OutElements, FLearningAgentsActionObjectElement());
@@ -2242,9 +2082,9 @@ bool ULearningAgentsActions::GetInclusiveUnionActionToArrayViews(TArrayView<FNam
 
 	for (int32 ElementIdx = 0; ElementIdx < Parameters.Elements.Num(); ElementIdx++)
 	{
-		if (!Object->ActionObject.IsValid(Parameters.Elements[ElementIdx]))
+		if (!ActionObject.IsValid(Parameters.Elements[ElementIdx]))
 		{
-			UE_LOG(LogLearning, Error, TEXT("GetInclusiveUnionActionToArrayViews: Invalid Action Object."));
+			UE_LOG(LogLearning, Error, TEXT("%s: Invalid Action Object."), *GetName());
 			UE::Learning::Array::Set<1, FName>(OutElementNames, NAME_None);
 			UE::Learning::Array::Set<1, FLearningAgentsActionObjectElement>(OutElements, FLearningAgentsActionObjectElement());
 			return false;
@@ -2257,45 +2097,39 @@ bool ULearningAgentsActions::GetInclusiveUnionActionToArrayViews(TArrayView<FNam
 	return true;
 }
 
-bool ULearningAgentsActions::GetStaticArrayActionNum(int32& OutNum, const ULearningAgentsActionObject* Object, const FLearningAgentsActionObjectElement Element, const FName Tag)
+bool ULearningAgentsActionObject::GetStaticArrayActionNum(int32& OutNum, const FLearningAgentsActionObjectElement Element, const FName Name) const
 {
-	if (!Object)
-	{
-		UE_LOG(LogLearning, Error, TEXT("GetStaticArrayActionNum: Object is nullptr."));
-		OutNum = 0;
-		return false;
-	}
-
-	if (!Object->ActionObject.IsValid(Element.ObjectElement))
+	if (!ActionObject.IsValid(Element.ObjectElement))
 	{
-		UE_LOG(LogLearning, Error, TEXT("GetStaticArrayActionNum: Invalid Action Object."));
+		UE_LOG(LogLearning, Error, TEXT("%s: Invalid Action Object."), *GetName());
 		OutNum = 0;
 		return false;
 	}
 
-	if (Object->ActionObject.GetTag(Element.ObjectElement) != Tag)
+	if (ActionObject.GetName(Element.ObjectElement) != Name)
 	{
-		UE_LOG(LogLearning, Warning, TEXT("GetStaticArrayActionNum: Action tag does not match. Action is '%s' but asked for '%s'."), *Object->ActionObject.GetTag(Element.ObjectElement).ToString(), *Tag.ToString());
+		UE_LOG(LogLearning, Warning, TEXT("%s: Action name does not match. Action is '%s' but asked for '%s'."), *GetName(), *ActionObject.GetName(Element.ObjectElement).ToString(), *Name.ToString());
 	}
 
-	if (Object->ActionObject.GetType(Element.ObjectElement) != UE::Learning::Action::EType::Array)
+	if (ActionObject.GetType(Element.ObjectElement) != UE::Learning::Action::EType::Array)
 	{
-		UE_LOG(LogLearning, Error, TEXT("GetStaticArrayActionNum: Action '%s' type does not match. Action is '%s' but asked for '%s'."),
-			*Object->ActionObject.GetTag(Element.ObjectElement).ToString(),
-			UE::Learning::Agents::Action::Private::GetActionTypeString(Object->ActionObject.GetType(Element.ObjectElement)),
+		UE_LOG(LogLearning, Error, TEXT("%s: Action '%s' type does not match. Action is '%s' but asked for '%s'."),
+			*GetName(),
+			*ActionObject.GetName(Element.ObjectElement).ToString(),
+			UE::Learning::Agents::Action::Private::GetActionTypeString(ActionObject.GetType(Element.ObjectElement)),
 			UE::Learning::Agents::Action::Private::GetActionTypeString(UE::Learning::Action::EType::Array));
 		OutNum = 0;
 		return false;
 	}
 
-	OutNum = Object->ActionObject.GetArray(Element.ObjectElement).Elements.Num();
+	OutNum = ActionObject.GetArray(Element.ObjectElement).Elements.Num();
 	return true;
 }
 
-bool ULearningAgentsActions::GetStaticArrayAction(TArray<FLearningAgentsActionObjectElement>& OutElements, const ULearningAgentsActionObject* Object, const FLearningAgentsActionObjectElement Element, const FName Tag)
+bool ULearningAgentsActionObject::GetStaticArrayAction(TArray<FLearningAgentsActionObjectElement>& OutElements, const FLearningAgentsActionObjectElement Element, const FName Name) const
 {
 	int32 OutElementNum = 0;
-	if (!GetStaticArrayActionNum(OutElementNum, Object, Element, Tag))
+	if (!GetStaticArrayActionNum(OutElementNum, Element, Name))
 	{
 		OutElements.Empty();
 		return false;
@@ -2303,7 +2137,7 @@ bool ULearningAgentsActions::GetStaticArrayAction(TArray<FLearningAgentsActionOb
 
 	OutElements.SetNumUninitialized(OutElementNum);
 
-	if (!GetStaticArrayActionToArrayView(OutElements, Object, Element, Tag))
+	if (!GetStaticArrayActionToArrayView(OutElements, Element, Name))
 	{
 		OutElements.Empty();
 		return false;
@@ -2312,48 +2146,42 @@ bool ULearningAgentsActions::GetStaticArrayAction(TArray<FLearningAgentsActionOb
 	return true;
 }
 
-bool ULearningAgentsActions::GetStaticArrayActionToArrayView(TArrayView<FLearningAgentsActionObjectElement> OutElements, const ULearningAgentsActionObject* Object, const FLearningAgentsActionObjectElement Element, const FName Tag)
+bool ULearningAgentsActionObject::GetStaticArrayActionToArrayView(TArrayView<FLearningAgentsActionObjectElement> OutElements, const FLearningAgentsActionObjectElement Element, const FName Name) const
 {
-	if (!Object)
+	if (!ActionObject.IsValid(Element.ObjectElement))
 	{
-		UE_LOG(LogLearning, Error, TEXT("GetStaticArrayActionToArrayView: Object is nullptr."));
+		UE_LOG(LogLearning, Error, TEXT("%s: Invalid Action Object."), *GetName());
 		UE::Learning::Array::Set<1, FLearningAgentsActionObjectElement>(OutElements, FLearningAgentsActionObjectElement());
 		return false;
 	}
 
-	if (!Object->ActionObject.IsValid(Element.ObjectElement))
+	if (ActionObject.GetName(Element.ObjectElement) != Name)
 	{
-		UE_LOG(LogLearning, Error, TEXT("GetStaticArrayActionToArrayView: Invalid Action Object."));
-		UE::Learning::Array::Set<1, FLearningAgentsActionObjectElement>(OutElements, FLearningAgentsActionObjectElement());
-		return false;
-	}
-
-	if (Object->ActionObject.GetTag(Element.ObjectElement) != Tag)
-	{
-		UE_LOG(LogLearning, Warning, TEXT("GetStaticArrayActionToArrayView: Action tag does not match. Action is '%s' but asked for '%s'."), *Object->ActionObject.GetTag(Element.ObjectElement).ToString(), *Tag.ToString());
+		UE_LOG(LogLearning, Warning, TEXT("%s: Action name does not match. Action is '%s' but asked for '%s'."), *GetName(), *ActionObject.GetName(Element.ObjectElement).ToString(), *Name.ToString());
 	}
 
-	if (Object->ActionObject.GetType(Element.ObjectElement) != UE::Learning::Action::EType::Array)
+	if (ActionObject.GetType(Element.ObjectElement) != UE::Learning::Action::EType::Array)
 	{
-		UE_LOG(LogLearning, Error, TEXT("GetStaticArrayActionToArrayView: Action '%s' type does not match. Action is '%s' but asked for '%s'."), 
-			*Object->ActionObject.GetTag(Element.ObjectElement).ToString(),
-			UE::Learning::Agents::Action::Private::GetActionTypeString(Object->ActionObject.GetType(Element.ObjectElement)),
+		UE_LOG(LogLearning, Error, TEXT("%s: Action '%s' type does not match. Action is '%s' but asked for '%s'."), 
+			*GetName(),
+			*ActionObject.GetName(Element.ObjectElement).ToString(),
+			UE::Learning::Agents::Action::Private::GetActionTypeString(ActionObject.GetType(Element.ObjectElement)),
 			UE::Learning::Agents::Action::Private::GetActionTypeString(UE::Learning::Action::EType::Array));
 		UE::Learning::Array::Set<1, FLearningAgentsActionObjectElement>(OutElements, FLearningAgentsActionObjectElement());
 		return false;
 	}
 
-	const TArrayView<const UE::Learning::Action::FObjectElement> SubElements = Object->ActionObject.GetArray(Element.ObjectElement).Elements;
+	const TArrayView<const UE::Learning::Action::FObjectElement> SubElements = ActionObject.GetArray(Element.ObjectElement).Elements;
 
 	if (SubElements.Num() == 0)
 	{
-		UE_LOG(LogLearning, Warning, TEXT("GetStaticArrayActionToArrayView: Getting zero-sized Array Action."));
+		UE_LOG(LogLearning, Warning, TEXT("%s: Getting zero-sized Array Action."), *GetName());
 	}
 
 	if (SubElements.Num() != OutElements.Num())
 	{
-		UE_LOG(LogLearning, Error, TEXT("GetStaticArrayActionToArrayView: Action '%s' size does not match. Action is '%i' elements but asked for '%i'."),
-			*Object->ActionObject.GetTag(Element.ObjectElement).ToString(),
+		UE_LOG(LogLearning, Error, TEXT("%s: Action '%s' size does not match. Action is '%i' elements but asked for '%i'."),
+			*GetName(), *ActionObject.GetName(Element.ObjectElement).ToString(),
 			SubElements.Num(), OutElements.Num());
 		UE::Learning::Array::Set<1, FLearningAgentsActionObjectElement>(OutElements, FLearningAgentsActionObjectElement());
 		return false;
@@ -2361,9 +2189,9 @@ bool ULearningAgentsActions::GetStaticArrayActionToArrayView(TArrayView<FLearnin
 
 	for (int32 ElementIdx = 0; ElementIdx < SubElements.Num(); ElementIdx++)
 	{
-		if (!Object->ActionObject.IsValid(SubElements[ElementIdx]))
+		if (!ActionObject.IsValid(SubElements[ElementIdx]))
 		{
-			UE_LOG(LogLearning, Error, TEXT("GetStaticArrayActionToArrayView: Invalid Action Object."));
+			UE_LOG(LogLearning, Error, TEXT("%s: Invalid Action Object."), *GetName());
 			UE::Learning::Array::Set<1, FLearningAgentsActionObjectElement>(OutElements, FLearningAgentsActionObjectElement());
 			return false;
 		}
@@ -2374,11 +2202,11 @@ bool ULearningAgentsActions::GetStaticArrayActionToArrayView(TArrayView<FLearnin
 	return true;
 }
 
-bool ULearningAgentsActions::GetPairAction(FLearningAgentsActionObjectElement& OutKey, FLearningAgentsActionObjectElement& OutValue, const ULearningAgentsActionObject* Object, const FLearningAgentsActionObjectElement Element, const FName Tag)
+bool ULearningAgentsActionObject::GetPairAction(FLearningAgentsActionObjectElement& OutKey, FLearningAgentsActionObjectElement& OutValue, const FLearningAgentsActionObjectElement Element, const FName Name) const
 {
 	TStaticArray<FName, 2> OutElementNames;
 	TStaticArray<FLearningAgentsActionObjectElement, 2> OutElements;
-	if (!GetStructActionToArrayViews(OutElementNames, OutElements, Object, Element, Tag))
+	if (!GetStructActionToArrayViews(OutElementNames, OutElements, Element, Name))
 	{
 		OutKey = FLearningAgentsActionObjectElement();
 		OutValue = FLearningAgentsActionObjectElement();
@@ -2390,17 +2218,17 @@ bool ULearningAgentsActions::GetPairAction(FLearningAgentsActionObjectElement& O
 	return true;
 }
 
-bool ULearningAgentsActions::GetEnumAction(uint8& OutEnumValue, const ULearningAgentsActionObject* Object, const FLearningAgentsActionObjectElement Element, const UEnum* Enum, const FName Tag)
+bool ULearningAgentsActionObject::GetEnumAction(uint8& OutEnumValue, const UEnum* Enum, const FLearningAgentsActionObjectElement Element, const FName Name) const
 {
 	if (!Enum)
 	{
-		UE_LOG(LogLearning, Error, TEXT("GetEnumAction: Enum is nullptr."));
+		UE_LOG(LogLearning, Error, TEXT("%s: Enum is nullptr."), *GetName());
 		OutEnumValue = 0;
 		return false;
 	}
 
 	int32 OutIndex = 0;
-	if (!GetExclusiveDiscreteAction(OutIndex, Object, Element, Tag))
+	if (!GetExclusiveDiscreteAction(OutIndex, Element, Name))
 	{
 		OutEnumValue = 0;
 		return false;
@@ -2408,7 +2236,7 @@ bool ULearningAgentsActions::GetEnumAction(uint8& OutEnumValue, const ULearningA
 
 	if (OutIndex >= Enum->NumEnums() - 1)
 	{
-		UE_LOG(LogLearning, Error, TEXT("GetEnumAction: EnumValue out of range for Enum '%s'. Expected %i or less, got %i."), *Enum->GetName(), Enum->NumEnums() - 1, OutIndex);
+		UE_LOG(LogLearning, Error, TEXT("%s: EnumValue out of range for Enum '%s'. Expected %i or less, got %i."), *GetName(), *Enum->GetName(), Enum->NumEnums() - 1, OutIndex);
 		OutEnumValue = 0;
 		return false;
 	}
@@ -2417,7 +2245,7 @@ bool ULearningAgentsActions::GetEnumAction(uint8& OutEnumValue, const ULearningA
 
 	if (EnumValue == INDEX_NONE)
 	{
-		UE_LOG(LogLearning, Error, TEXT("GetEnumAction: Enum Value not found for index %i."), OutIndex);
+		UE_LOG(LogLearning, Error, TEXT("%s: Enum Value not found for index %i."), *GetName(), OutIndex);
 		OutEnumValue = 0;
 		return false;
 	}
@@ -2426,24 +2254,24 @@ bool ULearningAgentsActions::GetEnumAction(uint8& OutEnumValue, const ULearningA
 	return true;
 }
 
-bool ULearningAgentsActions::GetBitmaskAction(int32& OutBitmaskValue, const ULearningAgentsActionObject* Object, const FLearningAgentsActionObjectElement Element, const UEnum* Enum, const FName Tag)
+bool ULearningAgentsActionObject::GetBitmaskAction(int32& OutBitmaskValue, const UEnum* Enum, const FLearningAgentsActionObjectElement Element, const FName Name) const
 {
 	if (!Enum)
 	{
-		UE_LOG(LogLearning, Error, TEXT("GetBitmaskAction: Enum is nullptr."));
+		UE_LOG(LogLearning, Error, TEXT("%s: Enum is nullptr."), *GetName());
 		OutBitmaskValue = 0;
 		return false;
 	}
 
 	if (Enum->NumEnums() - 1 > 32)
 	{
-		UE_LOG(LogLearning, Error, TEXT("GetBitmaskAction: Too many values in Enum to use as Bitmask (%i)."), Enum->NumEnums() - 1);
+		UE_LOG(LogLearning, Error, TEXT("%s: Too many values in Enum to use as Bitmask (%i)."), *GetName(), Enum->NumEnums() - 1);
 		OutBitmaskValue = 0;
 		return false;
 	}
 
 	int32 EnumValueNum;
-	if (!GetInclusiveDiscreteActionNum(EnumValueNum, Object, Element, Tag))
+	if (!GetInclusiveDiscreteActionNum(EnumValueNum, Element, Name))
 	{
 		OutBitmaskValue = 0;
 		return false;
@@ -2451,14 +2279,14 @@ bool ULearningAgentsActions::GetBitmaskAction(int32& OutBitmaskValue, const ULea
 
 	if (EnumValueNum > Enum->NumEnums() - 1)
 	{
-		UE_LOG(LogLearning, Error, TEXT("GetBitmaskAction: Too many values for Enum '%s'. Expected %i or less, got %i."), *Enum->GetName(), Enum->NumEnums() - 1, EnumValueNum);
+		UE_LOG(LogLearning, Error, TEXT("%s: Too many values for Enum '%s'. Expected %i or less, got %i."), *GetName(), *Enum->GetName(), Enum->NumEnums() - 1, EnumValueNum);
 		OutBitmaskValue = 0;
 		return false;
 	}
 
 	TArray<int32, TInlineAllocator<32>> OutIndices;
 	OutIndices.SetNumUninitialized(EnumValueNum);
-	if (!GetInclusiveDiscreteActionToArrayView(OutIndices, Object, Element, Tag))
+	if (!GetInclusiveDiscreteActionToArrayView(OutIndices, Element, Name))
 	{
 		OutBitmaskValue = 0;
 		return false;
@@ -2472,10 +2300,10 @@ bool ULearningAgentsActions::GetBitmaskAction(int32& OutBitmaskValue, const ULea
 	return true;
 }
 
-bool ULearningAgentsActions::GetOptionalAction(ELearningAgentsOptionalAction& OutOption, FLearningAgentsActionObjectElement& OutElement, const ULearningAgentsActionObject* Object, const FLearningAgentsActionObjectElement Element, const FName Tag)
+bool ULearningAgentsActionObject::GetOptionalAction(ELearningAgentsOptionalAction& OutOption, FLearningAgentsActionObjectElement& OutElement, const FLearningAgentsActionObjectElement Element, const FName Name) const
 {
 	FName OutName = NAME_None;
-	if (!GetExclusiveUnionAction(OutName, OutElement, Object, Element, Tag))
+	if (!GetExclusiveUnionAction(OutName, OutElement, Element, Name))
 	{
 		OutOption = ELearningAgentsOptionalAction::Null;
 		return false;
@@ -2485,10 +2313,10 @@ bool ULearningAgentsActions::GetOptionalAction(ELearningAgentsOptionalAction& Ou
 	return true;
 }
 
-bool ULearningAgentsActions::GetEitherAction(ELearningAgentsEitherAction& OutEither, FLearningAgentsActionObjectElement& OutElement, const ULearningAgentsActionObject* Object, const FLearningAgentsActionObjectElement Element, const FName Tag)
+bool ULearningAgentsActionObject::GetEitherAction(ELearningAgentsEitherAction& OutEither, FLearningAgentsActionObjectElement& OutElement, const FLearningAgentsActionObjectElement Element, const FName Name) const
 {
 	FName OutName = NAME_None;
-	if (!GetExclusiveUnionAction(OutName, OutElement, Object, Element, Tag))
+	if (!GetExclusiveUnionAction(OutName, OutElement, Element, Name))
 	{
 		OutEither = ELearningAgentsEitherAction::A;
 		return false;
@@ -2498,46 +2326,40 @@ bool ULearningAgentsActions::GetEitherAction(ELearningAgentsEitherAction& OutEit
 	return true;
 }
 
-bool ULearningAgentsActions::GetEncodingAction(FLearningAgentsActionObjectElement& OutElement, const ULearningAgentsActionObject* Object, const FLearningAgentsActionObjectElement Element, const FName Tag)
+bool ULearningAgentsActionObject::GetEncodingAction(FLearningAgentsActionObjectElement& OutElement, const FLearningAgentsActionObjectElement Element, const FName Name) const
 {
-	if (!Object)
-	{
-		UE_LOG(LogLearning, Error, TEXT("GetEncodingAction: Object is nullptr."));
-		OutElement = FLearningAgentsActionObjectElement();
-		return false;
-	}
-
-	if (!Object->ActionObject.IsValid(Element.ObjectElement))
+	if (!ActionObject.IsValid(Element.ObjectElement))
 	{
-		UE_LOG(LogLearning, Error, TEXT("GetEncodingAction: Invalid Action Object."));
+		UE_LOG(LogLearning, Error, TEXT("%s: Invalid Action Object."), *GetName());
 		OutElement = FLearningAgentsActionObjectElement();
 		return false;
 	}
 
-	if (Object->ActionObject.GetTag(Element.ObjectElement) != Tag)
+	if (ActionObject.GetName(Element.ObjectElement) != Name)
 	{
-		UE_LOG(LogLearning, Warning, TEXT("GetEncodingAction: Action tag does not match. Action is '%s' but asked for '%s'."), *Object->ActionObject.GetTag(Element.ObjectElement).ToString(), *Tag.ToString());
+		UE_LOG(LogLearning, Warning, TEXT("%s: Action name does not match. Action is '%s' but asked for '%s'."), *GetName(), *ActionObject.GetName(Element.ObjectElement).ToString(), *Name.ToString());
 	}
 
-	if (Object->ActionObject.GetType(Element.ObjectElement) != UE::Learning::Action::EType::Encoding)
+	if (ActionObject.GetType(Element.ObjectElement) != UE::Learning::Action::EType::Encoding)
 	{
-		UE_LOG(LogLearning, Error, TEXT("GetEncodingAction: Action '%s' type does not match. Action is '%s' but asked for '%s'."),
-			*Object->ActionObject.GetTag(Element.ObjectElement).ToString(),
-			UE::Learning::Agents::Action::Private::GetActionTypeString(Object->ActionObject.GetType(Element.ObjectElement)),
+		UE_LOG(LogLearning, Error, TEXT("%s: Action '%s' type does not match. Action is '%s' but asked for '%s'."),
+			*GetName(),
+			*ActionObject.GetName(Element.ObjectElement).ToString(),
+			UE::Learning::Agents::Action::Private::GetActionTypeString(ActionObject.GetType(Element.ObjectElement)),
 			UE::Learning::Agents::Action::Private::GetActionTypeString(UE::Learning::Action::EType::Encoding));
 		OutElement = FLearningAgentsActionObjectElement();
 		return false;
 	}
 
-	OutElement = { Object->ActionObject.GetEncoding(Element.ObjectElement).Element };
+	OutElement = { ActionObject.GetEncoding(Element.ObjectElement).Element };
 
 	return true;
 }
 
-bool ULearningAgentsActions::GetBoolAction(bool& bOutValue, const ULearningAgentsActionObject* Object, const FLearningAgentsActionObjectElement Element, const FName Tag)
+bool ULearningAgentsActionObject::GetBoolAction(bool& bOutValue, const FLearningAgentsActionObjectElement Element, const FName Name) const
 {
 	int32 OutIndex = 0;
-	if (!GetExclusiveDiscreteAction(OutIndex, Object, Element, Tag))
+	if (!GetExclusiveDiscreteAction(OutIndex, Element, Name))
 	{
 		bOutValue = false;
 		return false;
@@ -2547,10 +2369,10 @@ bool ULearningAgentsActions::GetBoolAction(bool& bOutValue, const ULearningAgent
 	return true;
 }
 
-bool ULearningAgentsActions::GetFloatAction(float& OutValue, const ULearningAgentsActionObject* Object, const FLearningAgentsActionObjectElement Element, const float FloatScale, const FName Tag)
+bool ULearningAgentsActionObject::GetFloatAction(float& OutValue, const FLearningAgentsActionObjectElement Element, const float FloatScale, const FName Name) const
 {
 	float OutValuesData;
-	if (!GetContinuousActionToArrayView(MakeArrayView(&OutValuesData, 1), Object, Element, Tag))
+	if (!GetContinuousActionToArrayView(MakeArrayView(&OutValuesData, 1), Element, Name))
 	{
 		OutValue = 0.0f;
 		return false;
@@ -2560,10 +2382,10 @@ bool ULearningAgentsActions::GetFloatAction(float& OutValue, const ULearningAgen
 	return true;
 }
 
-bool ULearningAgentsActions::GetLocationAction(FVector& OutLocation, const ULearningAgentsActionObject* Object, const FLearningAgentsActionObjectElement Element, const FTransform RelativeTransform, const float LocationScale, const FName Tag)
+bool ULearningAgentsActionObject::GetLocationAction(FVector& OutLocation, const FLearningAgentsActionObjectElement Element, const FTransform RelativeTransform, const float LocationScale, const FName Name) const
 {
 	TStaticArray<float, 3> OutValues;
-	if (!GetContinuousActionToArrayView(OutValues, Object, Element, Tag))
+	if (!GetContinuousActionToArrayView(OutValues, Element, Name))
 	{
 		OutLocation = FVector::ZeroVector;
 		return false;
@@ -2574,10 +2396,10 @@ bool ULearningAgentsActions::GetLocationAction(FVector& OutLocation, const ULear
 	return true;
 }
 
-bool ULearningAgentsActions::GetRotationAction(FRotator& OutRotation, const ULearningAgentsActionObject* Object, const FLearningAgentsActionObjectElement Element, const FRotator RelativeRotation, const float RotationScale, const FName Tag)
+bool ULearningAgentsActionObject::GetRotationAction(FRotator& OutRotation, const FLearningAgentsActionObjectElement Element, const FRotator RelativeRotation, const float RotationScale, const FName Name) const
 {
 	FQuat OutRotationQuat;
-	if (!GetRotationActionAsQuat(OutRotationQuat, Object, Element, FQuat::MakeFromRotator(RelativeRotation), RotationScale, Tag))
+	if (!GetRotationActionAsQuat(OutRotationQuat, Element, FQuat::MakeFromRotator(RelativeRotation), RotationScale, Name))
 	{
 		OutRotation = FRotator::ZeroRotator;
 		return false;
@@ -2587,10 +2409,10 @@ bool ULearningAgentsActions::GetRotationAction(FRotator& OutRotation, const ULea
 	return true;
 }
 
-bool ULearningAgentsActions::GetRotationActionAsQuat(FQuat& OutRotation, const ULearningAgentsActionObject* Object, const FLearningAgentsActionObjectElement Element, const FQuat RelativeRotation, const float RotationScale, const FName Tag)
+bool ULearningAgentsActionObject::GetRotationActionAsQuat(FQuat& OutRotation, const FLearningAgentsActionObjectElement Element, const FQuat RelativeRotation, const float RotationScale, const FName Name) const
 {
 	TStaticArray<float, 3> OutValues;
-	if (!GetContinuousActionToArrayView(OutValues, Object, Element, Tag))
+	if (!GetContinuousActionToArrayView(OutValues, Element, Name))
 	{
 		OutRotation = FQuat::Identity;
 		return false;
@@ -2601,10 +2423,10 @@ bool ULearningAgentsActions::GetRotationActionAsQuat(FQuat& OutRotation, const U
 	return true;
 }
 
-bool ULearningAgentsActions::GetScaleAction(FVector& OutScale, const ULearningAgentsActionObject* Object, const FLearningAgentsActionObjectElement Element, const FVector RelativeScale, const float Scale, const FName Tag)
+bool ULearningAgentsActionObject::GetScaleAction(FVector& OutScale, const FLearningAgentsActionObjectElement Element, const FVector RelativeScale, const float Scale, const FName Name) const
 {
 	TStaticArray<float, 3> OutValues;
-	if (!GetContinuousActionToArrayView(OutValues, Object, Element, Tag))
+	if (!GetContinuousActionToArrayView(OutValues, Element, Name))
 	{
 		OutScale = FVector::OneVector;
 		return false;
@@ -2615,32 +2437,32 @@ bool ULearningAgentsActions::GetScaleAction(FVector& OutScale, const ULearningAg
 	return true;
 }
 
-bool ULearningAgentsActions::GetTransformAction(FTransform& OutTransform, const ULearningAgentsActionObject* Object, const FLearningAgentsActionObjectElement Element, const FTransform RelativeTransform, const float LocationScale, const float RotationScale, const float ScaleScale, const FName Tag)
+bool ULearningAgentsActionObject::GetTransformAction(FTransform& OutTransform, const FLearningAgentsActionObjectElement Element, const FTransform RelativeTransform, const float LocationScale, const float RotationScale, const float ScaleScale, const FName Name) const
 {
 	TStaticArray<FName, 3> OutElementNames;
 	TStaticArray<FLearningAgentsActionObjectElement, 3> OutElements;
-	if (!GetStructActionToArrayViews(OutElementNames, OutElements, Object, Element, Tag))
+	if (!GetStructActionToArrayViews(OutElementNames, OutElements, Element, Name))
 	{
 		OutTransform = FTransform::Identity;
 		return false;
 	}
 
 	FVector OutLocation;
-	if (!GetLocationAction(OutLocation, Object, OutElements[MakeArrayView(OutElementNames).Find(TEXT("Location"))], RelativeTransform, LocationScale))
+	if (!GetLocationAction(OutLocation, OutElements[MakeArrayView(OutElementNames).Find(TEXT("Location"))], RelativeTransform, LocationScale))
 	{
 		OutTransform = FTransform::Identity;
 		return false;
 	}
 
 	FQuat OutRotation;
-	if (!GetRotationActionAsQuat(OutRotation, Object, OutElements[MakeArrayView(OutElementNames).Find(TEXT("Rotation"))], RelativeTransform.GetRotation(), RotationScale))
+	if (!GetRotationActionAsQuat(OutRotation, OutElements[MakeArrayView(OutElementNames).Find(TEXT("Rotation"))], RelativeTransform.GetRotation(), RotationScale))
 	{
 		OutTransform = FTransform::Identity;
 		return false;
 	}
 
 	FVector OutScale;
-	if (!GetScaleAction(OutScale, Object, OutElements[MakeArrayView(OutElementNames).Find(TEXT("Scale"))], RelativeTransform.GetScale3D(), ScaleScale))
+	if (!GetScaleAction(OutScale, OutElements[MakeArrayView(OutElementNames).Find(TEXT("Scale"))], RelativeTransform.GetScale3D(), ScaleScale))
 	{
 		OutTransform = FTransform::Identity;
 		return false;
@@ -2650,9 +2472,9 @@ bool ULearningAgentsActions::GetTransformAction(FTransform& OutTransform, const
 	return true;
 }
 
-bool ULearningAgentsActions::GetAngleAction(float& OutAngle, const ULearningAgentsActionObject* Object, const FLearningAgentsActionObjectElement Element, const float RelativeAngle, const float AngleScale, const FName Tag)
+bool ULearningAgentsActionObject::GetAngleAction(float& OutAngle, const FLearningAgentsActionObjectElement Element, const float RelativeAngle, const float AngleScale, const FName Name) const
 {
-	if (GetAngleActionRadians(OutAngle, Object, Element, FMath::DegreesToRadians(RelativeAngle), FMath::DegreesToRadians(AngleScale), Tag))
+	if (GetAngleActionRadians(OutAngle, Element, FMath::DegreesToRadians(RelativeAngle), FMath::DegreesToRadians(AngleScale), Name))
 	{
 		OutAngle = FMath::RadiansToDegrees(OutAngle);
 		return true;
@@ -2664,9 +2486,9 @@ bool ULearningAgentsActions::GetAngleAction(float& OutAngle, const ULearningAgen
 	}
 }
 
-bool ULearningAgentsActions::GetAngleActionRadians(float& OutAngle, const ULearningAgentsActionObject* Object, const FLearningAgentsActionObjectElement Element, const float RelativeAngle, const float AngleScale, const FName Tag)
+bool ULearningAgentsActionObject::GetAngleActionRadians(float& OutAngle, const FLearningAgentsActionObjectElement Element, const float RelativeAngle, const float AngleScale, const FName Name) const
 {
-	if (!GetContinuousActionToArrayView(MakeArrayView(&OutAngle, 1), Object, Element, Tag))
+	if (!GetContinuousActionToArrayView(MakeArrayView(&OutAngle, 1), Element, Name))
 	{
 		OutAngle = 0.0f;
 		return false;
@@ -2676,10 +2498,10 @@ bool ULearningAgentsActions::GetAngleActionRadians(float& OutAngle, const ULearn
 	return true;
 }
 
-bool ULearningAgentsActions::GetVelocityAction(FVector& OutVelocity, const ULearningAgentsActionObject* Object, const FLearningAgentsActionObjectElement Element, const FTransform RelativeTransform, const float VelocityScale, const FName Tag)
+bool ULearningAgentsActionObject::GetVelocityAction(FVector& OutVelocity, const FLearningAgentsActionObjectElement Element, const FTransform RelativeTransform, const float VelocityScale, const FName Name) const
 {
 	TStaticArray<float, 3> OutValues;
-	if (!GetContinuousActionToArrayView(OutValues, Object, Element, Tag))
+	if (!GetContinuousActionToArrayView(OutValues, Element, Name))
 	{
 		OutVelocity = FVector::OneVector;
 		return false;
@@ -2689,10 +2511,10 @@ bool ULearningAgentsActions::GetVelocityAction(FVector& OutVelocity, const ULear
 	return true;
 }
 
-bool ULearningAgentsActions::GetDirectionAction(FVector& OutDirection, const ULearningAgentsActionObject* Object, const FLearningAgentsActionObjectElement Element, const FTransform RelativeTransform, const FName Tag)
+bool ULearningAgentsActionObject::GetDirectionAction(FVector& OutDirection, const FLearningAgentsActionObjectElement Element, const FTransform RelativeTransform, const FName Name) const
 {
 	TStaticArray<float, 3> OutValues;
-	if (!GetContinuousActionToArrayView(OutValues, Object, Element, Tag))
+	if (!GetContinuousActionToArrayView(OutValues, Element, Name))
 	{
 		OutDirection = FVector::ForwardVector;
 		return false;
@@ -2702,9 +2524,9 @@ bool ULearningAgentsActions::GetDirectionAction(FVector& OutDirection, const ULe
 	return true;
 }
 
-bool ULearningAgentsActions::GetSpeedAction(float& OutSpeed, const ULearningAgentsActionObject* Object, const FLearningAgentsActionObjectElement Element, const float SpeedScale, const FName Tag)
+bool ULearningAgentsActionObject::GetSpeedAction(float& OutSpeed, const FLearningAgentsActionObjectElement Element, const float SpeedScale, const FName Name) const
 {
-	if (!GetContinuousActionToArrayView(MakeArrayView(&OutSpeed, 1), Object, Element, Tag))
+	if (!GetContinuousActionToArrayView(MakeArrayView(&OutSpeed, 1), Element, Name))
 	{
 		OutSpeed = 0.0f;
 		return false;
