[Backout] - CL31291642

diff --git a/Engine/Plugins/Experimental/LearningAgents/Source/LearningAgents/Private/LearningAgentsObservations.cpp b/Engine/Plugins/Experimental/LearningAgents/Source/LearningAgents/Private/LearningAgentsObservations.cpp
--- a/Engine/Plugins/Experimental/LearningAgents/Source/LearningAgents/Private/LearningAgentsObservations.cpp
+++ b/Engine/Plugins/Experimental/LearningAgents/Source/LearningAgents/Private/LearningAgentsObservations.cpp
@@ -24,6 +24,11 @@ uint32 GetTypeHash(const FLearningAgentsObservationObjectElement& Element)
 	return (uint32)Element.ObjectElement.Index;
 }
 
+const UE::Learning::Observation::FSchema& ULearningAgentsObservationSchema::GetObservationSchema() const
+{
+	return ObservationSchema;
+}
+
 namespace UE::Learning::Agents::Observation::Private
 {
 	static inline bool ContainsDuplicates(const TArrayView<const int32> Indices)
@@ -62,31 +67,32 @@ namespace UE::Learning::Agents::Observation::Private
 		const Learning::Observation::FSchema& Schema,
 		const Learning::Observation::FSchemaElement SchemaElement,
 		const Learning::Observation::FObject& Object,
-		const Learning::Observation::FObjectElement ObjectElement)
+		const Learning::Observation::FObjectElement ObjectElement,
+		const FString& ObjectName)
 	{
 		// Check Elements are Valid
 
 		if (!Schema.IsValid(SchemaElement))
 		{
-			UE_LOG(LogLearning, Error, TEXT("ValidateObjectMatchesSchema: Invalid Observation Schema Element."));
+			UE_LOG(LogLearning, Error, TEXT("%s: Invalid Observation Schema Element."), *ObjectName);
 			return false;
 		}
 
 		if (!Object.IsValid(ObjectElement))
 		{
-			UE_LOG(LogLearning, Error, TEXT("ValidateObjectMatchesSchema: Invalid Observation Object Element."));
+			UE_LOG(LogLearning, Error, TEXT("%s: Invalid Observation Object Element."), *ObjectName);
 			return false;
 		}
 
 		// Check Names Match
 
-		const FName ObservationSchemaElementTag = Schema.GetTag(SchemaElement);
-		const FName ObservationObjectElementTag = Object.GetTag(ObjectElement);
+		const FName ObservationSchemaElementName = Schema.GetName(SchemaElement);
+		const FName ObservationObjectElementName = Object.GetName(ObjectElement);
 
-		if (ObservationSchemaElementTag != ObservationObjectElementTag)
+		if (ObservationSchemaElementName != ObservationObjectElementName)
 		{
-			UE_LOG(LogLearning, Warning, TEXT("ValidateObjectMatchesSchema: Observation tag does not match Schema. Expected '%s', got '%s'."),
-				*ObservationSchemaElementTag.ToString(), *ObservationObjectElementTag.ToString());
+			UE_LOG(LogLearning, Warning, TEXT("%s: Observation name does not match Schema. Expected '%s', got '%s'."),
+				*ObjectName, *ObservationSchemaElementName.ToString(), *ObservationObjectElementName.ToString());
 		}
 
 		// Check Types Match
@@ -96,8 +102,9 @@ namespace UE::Learning::Agents::Observation::Private
 
 		if (ObservationSchemaElementType != ObservationObjectElementType)
 		{
-			UE_LOG(LogLearning, Error, TEXT("ValidateObjectMatchesSchema: Observation '%s' type does not match Schema. Expected type '%s', got type '%s'."),
-				*ObservationSchemaElementTag.ToString(),
+			UE_LOG(LogLearning, Error, TEXT("%s: Observation '%s' type does not match Schema. Expected type '%s', got type '%s'."),
+				*ObjectName,
+				*ObservationSchemaElementName.ToString(),
 				GetObservationTypeString(ObservationSchemaElementType),
 				GetObservationTypeString(ObservationObjectElementType));
 			return false;
@@ -116,8 +123,9 @@ namespace UE::Learning::Agents::Observation::Private
 
 			if (SchemaElementSize != ObjectElementSize)
 			{
-				UE_LOG(LogLearning, Error, TEXT("ValidateObjectMatchesSchema: Observation '%s' size does not match Schema. Expected '%i', got '%i'."),
-					*ObservationSchemaElementTag.ToString(),
+				UE_LOG(LogLearning, Error, TEXT("%s: Observation '%s' size does not match Schema. Expected '%i', got '%i'."),
+					*ObjectName,
+					*ObservationSchemaElementName.ToString(),
 					SchemaElementSize,
 					ObjectElementSize);
 				return false;
@@ -135,8 +143,9 @@ namespace UE::Learning::Agents::Observation::Private
 
 			if (SchemaParameters.Elements.Num() != ObjectParameters.Elements.Num())
 			{
-				UE_LOG(LogLearning, Error, TEXT("ValidateObjectMatchesSchema: Observation '%s' number of sub-elements does not match Schema. Expected '%i', got '%i'."),
-					*ObservationSchemaElementTag.ToString(),
+				UE_LOG(LogLearning, Error, TEXT("%s: Observation '%s' number of sub-elements does not match Schema. Expected '%i', got '%i'."),
+					*ObjectName,
+					*ObservationSchemaElementName.ToString(),
 					SchemaParameters.Elements.Num(),
 					ObjectParameters.Elements.Num());
 				return false;
@@ -148,8 +157,9 @@ namespace UE::Learning::Agents::Observation::Private
 
 				if (ObjectElementIdx == INDEX_NONE)
 				{
-					UE_LOG(LogLearning, Error, TEXT("ValidateObjectMatchesSchema: Observation '%s' does not include '%s' observation required by Schema."),
-						*ObservationSchemaElementTag.ToString(),
+					UE_LOG(LogLearning, Error, TEXT("%s: Observation '%s' does not include '%s' observation required by Schema."),
+						*ObjectName,
+						*ObservationSchemaElementName.ToString(),
 						*SchemaParameters.ElementNames[SchemaElementIdx].ToString());
 					return false;
 				}
@@ -158,7 +168,8 @@ namespace UE::Learning::Agents::Observation::Private
 					Schema,
 					SchemaParameters.Elements[SchemaElementIdx],
 					Object,
-					ObjectParameters.Elements[ObjectElementIdx]))
+					ObjectParameters.Elements[ObjectElementIdx],
+					ObjectName))
 				{
 					return false;
 				}
@@ -177,8 +188,9 @@ namespace UE::Learning::Agents::Observation::Private
 
 			if (SchemaSubElementIdx == INDEX_NONE)
 			{
-				UE_LOG(LogLearning, Error, TEXT("ValidateObjectMatchesSchema: Observation '%s' Schema does not include '%s' observation."),
-					*ObservationSchemaElementTag.ToString(),
+				UE_LOG(LogLearning, Error, TEXT("%s: Observation '%s' Schema does not include '%s' observation."),
+					*ObjectName,
+					*ObservationSchemaElementName.ToString(),
 					*ObjectParameters.ElementName.ToString());
 				return false;
 			}
@@ -187,7 +199,8 @@ namespace UE::Learning::Agents::Observation::Private
 				Schema,
 				SchemaParameters.Elements[SchemaSubElementIdx],
 				Object,
-				ObjectParameters.Element);
+				ObjectParameters.Element,
+				ObjectName);
 		}
 
 		case Learning::Observation::EType::OrInclusive:
@@ -197,8 +210,9 @@ namespace UE::Learning::Agents::Observation::Private
 
 			if (ObjectParameters.Elements.Num() > SchemaParameters.Elements.Num())
 			{
-				UE_LOG(LogLearning, Error, TEXT("ValidateObjectMatchesSchema: Observation '%s' too many sub-observations provided. Expected at most '%i', got '%i'."),
-					*ObservationSchemaElementTag.ToString(),
+				UE_LOG(LogLearning, Error, TEXT("%s: Observation '%s' too many sub-observations provided. Expected at most '%i', got '%i'."),
+					*ObjectName,
+					*ObservationSchemaElementName.ToString(),
 					SchemaParameters.Elements.Num(),
 					ObjectParameters.Elements.Num());
 				return false;
@@ -210,8 +224,9 @@ namespace UE::Learning::Agents::Observation::Private
 
 				if (SchemaSubElementIdx == INDEX_NONE)
 				{
-					UE_LOG(LogLearning, Error, TEXT("ValidateObjectMatchesSchema: Observation '%s' Schema does not include '%s' observation."),
-						*ObservationSchemaElementTag.ToString(),
+					UE_LOG(LogLearning, Error, TEXT("%s: Observation '%s' Schema does not include '%s' observation."),
+						*ObjectName,
+						*ObservationSchemaElementName.ToString(),
 						*ObjectParameters.ElementNames[ObjectSubElementIdx].ToString());
 					return false;
 				}
@@ -220,7 +235,8 @@ namespace UE::Learning::Agents::Observation::Private
 					Schema,
 					SchemaParameters.Elements[SchemaSubElementIdx],
 					Object,
-					ObjectParameters.Elements[ObjectSubElementIdx]))
+					ObjectParameters.Elements[ObjectSubElementIdx],
+					ObjectName))
 				{
 					return false;
 				}
@@ -236,8 +252,9 @@ namespace UE::Learning::Agents::Observation::Private
 
 			if (ObjectParameters.Elements.Num() != SchemaParameters.Num)
 			{
-				UE_LOG(LogLearning, Error, TEXT("ValidateObjectMatchesSchema: Observation '%s' array incorrect size. Expected '%i' elements, got '%i'."),
-					*ObservationSchemaElementTag.ToString(),
+				UE_LOG(LogLearning, Error, TEXT("%s: Observation '%s' array incorrect size. Expected '%i' elements, got '%i'."),
+					*ObjectName,
+					*ObservationSchemaElementName.ToString(),
 					SchemaParameters.Num,
 					ObjectParameters.Elements.Num());
 				return false;
@@ -249,7 +266,8 @@ namespace UE::Learning::Agents::Observation::Private
 					Schema,
 					SchemaParameters.Element,
 					Object,
-					ObjectParameters.Elements[ElementIdx]))
+					ObjectParameters.Elements[ElementIdx],
+					ObjectName))
 				{
 					return false;
 				}
@@ -265,8 +283,9 @@ namespace UE::Learning::Agents::Observation::Private
 
 			if (ObjectParameters.Elements.Num() > SchemaParameters.MaxNum)
 			{
-				UE_LOG(LogLearning, Error, TEXT("ValidateObjectMatchesSchema: Observation '%s' set too large. Expected at most '%i' elements, got '%i'."),
-					*ObservationSchemaElementTag.ToString(),
+				UE_LOG(LogLearning, Error, TEXT("%s: Observation '%s' set too large. Expected at most '%i' elements, got '%i'."),
+					*ObjectName,
+					*ObservationSchemaElementName.ToString(),
 					SchemaParameters.MaxNum,
 					ObjectParameters.Elements.Num());
 				return false;
@@ -278,7 +297,8 @@ namespace UE::Learning::Agents::Observation::Private
 					Schema,
 					SchemaParameters.Element,
 					Object,
-					ObjectParameters.Elements[ElementIdx]))
+					ObjectParameters.Elements[ElementIdx],
+					ObjectName))
 				{
 					return false;
 				}
@@ -296,7 +316,8 @@ namespace UE::Learning::Agents::Observation::Private
 				Schema,
 				SchemaParameters.Element,
 				Object,
-				ObjectParameters.Element);
+				ObjectParameters.Element,
+				ObjectName);
 		}
 
 		default:
@@ -311,22 +332,23 @@ namespace UE::Learning::Agents::Observation::Private
 		const UE::Learning::Observation::FObject& Object, 
 		const UE::Learning::Observation::FObjectElement ObjectElement,
 		const FString& Indentation,
-		const FString& Prefix)
+		const FString& Prefix,
+		const FString& ObjectName)
 	{
 		if (!Object.IsValid(ObjectElement))
 		{
-			UE_LOG(LogLearning, Error, TEXT("LogObservation: Invalid Observation Object Element."));
+			UE_LOG(LogLearning, Error, TEXT("%s: Invalid Observation Object Element."), *ObjectName);
 			return;
 		}
 
 		const UE::Learning::Observation::EType Type = Object.GetType(ObjectElement);
-		const FName Tag = Object.GetTag(ObjectElement);
+		const FName Name = Object.GetName(ObjectElement);
 
 		switch (Type)
 		{
 		case UE::Learning::Observation::EType::Null:
 		{
-			UE_LOG(LogLearning, Display, TEXT("%s%s \"%s\" (%s)"), *Indentation, *Prefix, *Tag.ToString(), GetObservationTypeString(Type));
+			UE_LOG(LogLearning, Display, TEXT("%s%s \"%s\" (%s)"), *Indentation, *Prefix, *Name.ToString(), GetObservationTypeString(Type));
 			return;
 		}
 
@@ -334,7 +356,7 @@ namespace UE::Learning::Agents::Observation::Private
 		{
 			const UE::Learning::Observation::FObjectContinuousParameters Parameters = Object.GetContinuous(ObjectElement);
 
-			UE_LOG(LogLearning, Display, TEXT("%s%s \"%s\" (%s) %s"), *Indentation, *Prefix, *Tag.ToString(), GetObservationTypeString(Type), *UE::Learning::Array::FormatFloat(Parameters.Values));
+			UE_LOG(LogLearning, Display, TEXT("%s%s \"%s\" (%s) %s"), *Indentation, *Prefix, *Name.ToString(), GetObservationTypeString(Type), *UE::Learning::Array::FormatFloat(Parameters.Values));
 			return;
 		}
 
@@ -342,10 +364,10 @@ namespace UE::Learning::Agents::Observation::Private
 		{
 			const UE::Learning::Observation::FObjectAndParameters Parameters = Object.GetAnd(ObjectElement);
 
-			UE_LOG(LogLearning, Display, TEXT("%s%s \"%s\" (%s)"), *Indentation, *Prefix, *Tag.ToString(), GetObservationTypeString(Type));
+			UE_LOG(LogLearning, Display, TEXT("%s%s \"%s\" (%s)"), *Indentation, *Prefix, *Name.ToString(), GetObservationTypeString(Type));
 			for (int32 SubElementIdx = 0; SubElementIdx < Parameters.Elements.Num(); SubElementIdx++)
 			{
-				LogObservation(Object, Parameters.Elements[SubElementIdx], *(Indentation + TEXT("    ")), FString::Printf(TEXT("| \"%s\": "), *Parameters.ElementNames[SubElementIdx].ToString()));
+				LogObservation(Object, Parameters.Elements[SubElementIdx], *(Indentation + TEXT("    ")), FString::Printf(TEXT("| \"%s\": "), *Parameters.ElementNames[SubElementIdx].ToString()), ObjectName);
 			}
 
 			return;
@@ -355,8 +377,8 @@ namespace UE::Learning::Agents::Observation::Private
 		{
 			const UE::Learning::Observation::FObjectOrExclusiveParameters Parameters = Object.GetOrExclusive(ObjectElement);
 
-			UE_LOG(LogLearning, Display, TEXT("%s%s \"%s\" (%s)"), *Indentation, *Prefix, *Tag.ToString(), GetObservationTypeString(Type));
-			LogObservation(Object, Parameters.Element, *(Indentation + TEXT("    ")), FString::Printf(TEXT("| \"%s\": "), *Parameters.ElementName.ToString()));
+			UE_LOG(LogLearning, Display, TEXT("%s%s \"%s\" (%s)"), *Indentation, *Prefix, *Name.ToString(), GetObservationTypeString(Type));
+			LogObservation(Object, Parameters.Element, *(Indentation + TEXT("    ")), FString::Printf(TEXT("| \"%s\": "), *Parameters.ElementName.ToString()), ObjectName);
 
 			return;
 		}
@@ -365,10 +387,10 @@ namespace UE::Learning::Agents::Observation::Private
 		{
 			const UE::Learning::Observation::FObjectOrInclusiveParameters Parameters = Object.GetOrInclusive(ObjectElement);
 
-			UE_LOG(LogLearning, Display, TEXT("%s%s \"%s\" (%s)"), *Indentation, *Prefix, *Tag.ToString(), GetObservationTypeString(Type));
+			UE_LOG(LogLearning, Display, TEXT("%s%s \"%s\" (%s)"), *Indentation, *Prefix, *Name.ToString(), GetObservationTypeString(Type));
 			for (int32 SubElementIdx = 0; SubElementIdx < Parameters.Elements.Num(); SubElementIdx++)
 			{
-				LogObservation(Object, Parameters.Elements[SubElementIdx], *(Indentation + TEXT("    ")), FString::Printf(TEXT("| \"%s\": "), *Parameters.ElementNames[SubElementIdx].ToString()));
+				LogObservation(Object, Parameters.Elements[SubElementIdx], *(Indentation + TEXT("    ")), FString::Printf(TEXT("| \"%s\": "), *Parameters.ElementNames[SubElementIdx].ToString()), ObjectName);
 			}
 
 			return;
@@ -378,10 +400,10 @@ namespace UE::Learning::Agents::Observation::Private
 		{
 			const UE::Learning::Observation::FObjectArrayParameters Parameters = Object.GetArray(ObjectElement);
 
-			UE_LOG(LogLearning, Display, TEXT("%s%s \"%s\" (%s)"), *Indentation, *Prefix, *Tag.ToString(), GetObservationTypeString(Type));
+			UE_LOG(LogLearning, Display, TEXT("%s%s \"%s\" (%s)"), *Indentation, *Prefix, *Name.ToString(), GetObservationTypeString(Type));
 			for (int32 SubElementIdx = 0; SubElementIdx < Parameters.Elements.Num(); SubElementIdx++)
 			{
-				LogObservation(Object, Parameters.Elements[SubElementIdx], *(Indentation + TEXT("    ")), FString::Printf(TEXT("| %3i:"), SubElementIdx));
+				LogObservation(Object, Parameters.Elements[SubElementIdx], *(Indentation + TEXT("    ")), FString::Printf(TEXT("| %3i:"), SubElementIdx), ObjectName);
 			}
 
 			return;
@@ -391,10 +413,10 @@ namespace UE::Learning::Agents::Observation::Private
 		{
 			const UE::Learning::Observation::FObjectSetParameters Parameters = Object.GetSet(ObjectElement);
 
-			UE_LOG(LogLearning, Display, TEXT("%s%s \"%s\" (%s)"), *Indentation, *Prefix, *Tag.ToString(), GetObservationTypeString(Type));
+			UE_LOG(LogLearning, Display, TEXT("%s%s \"%s\" (%s)"), *Indentation, *Prefix, *Name.ToString(), GetObservationTypeString(Type));
 			for (int32 SubElementIdx = 0; SubElementIdx < Parameters.Elements.Num(); SubElementIdx++)
 			{
-				LogObservation(Object, Parameters.Elements[SubElementIdx], *(Indentation + TEXT("    ")), FString::Printf(TEXT("| %3i:"), SubElementIdx));
+				LogObservation(Object, Parameters.Elements[SubElementIdx], *(Indentation + TEXT("    ")), FString::Printf(TEXT("| %3i:"), SubElementIdx), ObjectName);
 			}
 
 			return;
@@ -404,8 +426,8 @@ namespace UE::Learning::Agents::Observation::Private
 		{
 			const UE::Learning::Observation::FObjectEncodingParameters Parameters = Object.GetEncoding(ObjectElement);
 
-			UE_LOG(LogLearning, Display, TEXT("%s%s \"%s\" (%s)"), *Indentation, *Prefix, *Tag.ToString(), GetObservationTypeString(Type));
-			LogObservation(Object, Parameters.Element, *(Indentation + TEXT("    ")), TEXT("|"));
+			UE_LOG(LogLearning, Display, TEXT("%s%s \"%s\" (%s)"), *Indentation, *Prefix, *Name.ToString(), GetObservationTypeString(Type));
+			LogObservation(Object, Parameters.Element, *(Indentation + TEXT("    ")), TEXT("|"), ObjectName);
 
 			return;
 		}
@@ -440,7 +462,7 @@ namespace UE::Learning::Agents::Observation::Private
 	}
 }
 
-FTransform ULearningAgentsObservations::ProjectTransformOntoGroundPlane(const FTransform Transform, const FVector LocalForwardVector, const float GroundPlaneHeight)
+FTransform ULearningAgentsObservationFunctions::ProjectTransformOntoGroundPlane(const FTransform Transform, const FVector LocalForwardVector, const float GroundPlaneHeight)
 {
 	FVector Position = Transform.GetLocation();
 	Position.Z = GroundPlaneHeight;
@@ -450,94 +472,70 @@ FTransform ULearningAgentsObservations::ProjectTransformOntoGroundPlane(const FT
 	return FTransform(FQuat::FindBetweenNormals(FVector::ForwardVector, Direction), Position, Transform.GetScale3D());
 }
 
-UEnum* ULearningAgentsObservations::FindEnumByName(const FString& Name)
+UEnum* ULearningAgentsObservationFunctions::FindEnumByName(const FString& Name)
 {
 	return FindObject<UEnum>(nullptr, *Name);
 }
 
-bool ULearningAgentsObservations::ValidateObjectMatchesSchema(
-	const ULearningAgentsObservationSchema* Schema,
+bool ULearningAgentsObservationSchema::ValidateObjectMatchesSchema(
 	const FLearningAgentsObservationSchemaElement SchemaElement,
 	const ULearningAgentsObservationObject* Object,
-	const FLearningAgentsObservationObjectElement ObjectElement)
+	const FLearningAgentsObservationObjectElement ObjectElement) const
 {
-	if (!Schema)
-	{
-		UE_LOG(LogLearning, Error, TEXT("ValidateObjectMatchesSchema: Schema is nullptr."));
-		return false;
-	}
-
-	if (!Object)
-	{
-		UE_LOG(LogLearning, Error, TEXT("ValidateObjectMatchesSchema: Object is nullptr."));
-		return false;
-	}
-
 	return UE::Learning::Agents::Observation::Private::ValidateObjectMatchesSchema(
-		Schema->ObservationSchema,
+		ObservationSchema,
 		SchemaElement.SchemaElement,
-		Object->ObservationObject,
-		ObjectElement.ObjectElement);
+		Object->GetObservationObject(),
+		ObjectElement.ObjectElement,
+		GetName());
 }
 
-FLearningAgentsObservationSchemaElement ULearningAgentsObservations::SpecifyNullObservation(ULearningAgentsObservationSchema* Schema, const FName Tag)
+FLearningAgentsObservationSchemaElement ULearningAgentsObservationSchema::SpecifyNullObservation(const FName Name)
 {
-	if (!Schema)
-	{
-		UE_LOG(LogLearning, Error, TEXT("SpecifyNullObservation: Schema is nullptr."));
-		return FLearningAgentsObservationSchemaElement();
-	}
-
-	return { Schema->ObservationSchema.CreateNull(Tag)};
+	return { ObservationSchema.CreateNull(Name) };
 }
 
-FLearningAgentsObservationSchemaElement ULearningAgentsObservations::SpecifyContinuousObservation(ULearningAgentsObservationSchema* Schema, const int32 Size, const FName Tag)
+FLearningAgentsObservationSchemaElement ULearningAgentsObservationSchema::SpecifyContinuousObservation(const int32 Size, const FName Name)
 {
-	if (!Schema)
-	{
-		UE_LOG(LogLearning, Error, TEXT("SpecifyContinuousObservation: Schema is nullptr."));
-		return FLearningAgentsObservationSchemaElement();
-	}
-
 	if (Size < 0)
 	{
-		UE_LOG(LogLearning, Error, TEXT("SpecifyContinuousObservation: Invalid Continuous Observation Size '%i'."), Size);
+		UE_LOG(LogLearning, Error, TEXT("%s: Invalid Continuous Observation Size '%i'."), *GetName(), Size);
 		return FLearningAgentsObservationSchemaElement();
 	}
 	
 	if (Size == 0)
 	{
-		UE_LOG(LogLearning, Warning, TEXT("SpecifyContinuousObservation: Specifying zero-sized Continuous Observation."));
+		UE_LOG(LogLearning, Warning, TEXT("%s: Specifying zero-sized Continuous Observation."), *GetName());
 	}
 
-	return { Schema->ObservationSchema.CreateContinuous({ Size }, Tag) };
+	return { ObservationSchema.CreateContinuous({ Size }, Name) };
 }
 
-FLearningAgentsObservationSchemaElement ULearningAgentsObservations::SpecifyExclusiveDiscreteObservation(ULearningAgentsObservationSchema* Schema, const int32 Size, const FName Tag)
+FLearningAgentsObservationSchemaElement ULearningAgentsObservationSchema::SpecifyExclusiveDiscreteObservation(const int32 Size, const FName Name)
 {
-	return SpecifyContinuousObservation(Schema, Size, Tag);
+	return SpecifyContinuousObservation(Size, Name);
 }
 
-FLearningAgentsObservationSchemaElement ULearningAgentsObservations::SpecifyInclusiveDiscreteObservation(ULearningAgentsObservationSchema* Schema, const int32 Size, const FName Tag)
+FLearningAgentsObservationSchemaElement ULearningAgentsObservationSchema::SpecifyInclusiveDiscreteObservation(const int32 Size, const FName Name)
 {
-	return SpecifyContinuousObservation(Schema, Size, Tag);
+	return SpecifyContinuousObservation(Size, Name);
 }
 
-FLearningAgentsObservationSchemaElement ULearningAgentsObservations::SpecifyIndexObservation(ULearningAgentsObservationSchema* Schema, const int32 Size, const FName Tag)
+FLearningAgentsObservationSchemaElement ULearningAgentsObservationSchema::SpecifyIndexObservation(const int32 Size, const FName Name)
 {
-	return SpecifyContinuousObservation(Schema, Size, Tag);
+	return SpecifyContinuousObservation(Size, Name);
 }
 
-FLearningAgentsObservationSchemaElement ULearningAgentsObservations::SpecifyCountObservation(ULearningAgentsObservationSchema* Schema, const FName Tag)
+FLearningAgentsObservationSchemaElement ULearningAgentsObservationSchema::SpecifyCountObservation(const FName Name)
 {
-	return SpecifyContinuousObservation(Schema, 1, Tag);
+	return SpecifyContinuousObservation(1, Name);
 }
 
-FLearningAgentsObservationSchemaElement ULearningAgentsObservations::SpecifyStructObservation(ULearningAgentsObservationSchema* Schema, const TMap<FName, FLearningAgentsObservationSchemaElement>& Elements, const FName Tag)
+FLearningAgentsObservationSchemaElement ULearningAgentsObservationSchema::SpecifyStructObservation(const TMap<FName, FLearningAgentsObservationSchemaElement>& Elements, const FName Name)
 {
 	if (Elements.Num() == 0)
 	{
-		UE_LOG(LogLearning, Warning, TEXT("SpecifyStructObservation: Specifying zero-sized Struct Observation."));
+		UE_LOG(LogLearning, Warning, TEXT("%s: Specifying zero-sized Struct Observation."), *GetName());
 	}
 
 	const int32 SubElementNum = Elements.Num();
@@ -575,36 +573,30 @@ FLearningAgentsObservationSchemaElement ULearningAgentsObservations::SpecifyStru
 		SortedSubElements[Idx] = SubElements[SubElementIndices[Idx]];
 	}
 
-	return SpecifyStructObservationFromArrayViews(Schema, SortedSubElementNames, SortedSubElements, Tag);
+	return SpecifyStructObservationFromArrayViews(SortedSubElementNames, SortedSubElements, Name);
 }
 
-FLearningAgentsObservationSchemaElement ULearningAgentsObservations::SpecifyStructObservationFromArrays(ULearningAgentsObservationSchema* Schema, const TArray<FName>& ElementNames, const TArray<FLearningAgentsObservationSchemaElement>& Elements, const FName Tag)
+FLearningAgentsObservationSchemaElement ULearningAgentsObservationSchema::SpecifyStructObservationFromArrays(const TArray<FName>& ElementNames, const TArray<FLearningAgentsObservationSchemaElement>& Elements, const FName Name)
 {
-	return SpecifyStructObservationFromArrayViews(Schema, ElementNames, Elements, Tag);
+	return SpecifyStructObservationFromArrayViews(ElementNames, Elements, Name);
 }
 
-FLearningAgentsObservationSchemaElement ULearningAgentsObservations::SpecifyStructObservationFromArrayViews(ULearningAgentsObservationSchema* Schema, const TArrayView<const FName> ElementNames, const TArrayView<const FLearningAgentsObservationSchemaElement> Elements, const FName Tag)
+FLearningAgentsObservationSchemaElement ULearningAgentsObservationSchema::SpecifyStructObservationFromArrayViews(const TArrayView<const FName> ElementNames, const TArrayView<const FLearningAgentsObservationSchemaElement> Elements, const FName Name)
 {
-	if (!Schema)
-	{
-		UE_LOG(LogLearning, Error, TEXT("SpecifyStructObservationFromArrayViews: Schema is nullptr."));
-		return FLearningAgentsObservationSchemaElement();
-	}
-
 	if (Elements.Num() == 0)
 	{
-		UE_LOG(LogLearning, Warning, TEXT("SpecifyStructObservationFromArrayViews: Specifying zero-sized Struct Observation."));
+		UE_LOG(LogLearning, Warning, TEXT("%s: Specifying zero-sized Struct Observation."), *GetName());
 	}
 
 	if (Elements.Num() != ElementNames.Num())
 	{
-		UE_LOG(LogLearning, Error, TEXT("SpecifyStructObservationFromArrayViews: Number of elements (%i) must match number of names (%i)."), Elements.Num(), ElementNames.Num());
+		UE_LOG(LogLearning, Error, TEXT("%s: Number of elements (%i) must match number of names (%i)."), *GetName(), Elements.Num(), ElementNames.Num());
 		return FLearningAgentsObservationSchemaElement();
 	}
 
 	if (UE::Learning::Agents::Observation::Private::ContainsDuplicates(ElementNames))
 	{
-		UE_LOG(LogLearning, Error, TEXT("SpecifyStructObservationFromArrayViews: Element Names contain duplicates."));
+		UE_LOG(LogLearning, Error, TEXT("%s: Element Names contain duplicates."), *GetName());
 		return FLearningAgentsObservationSchemaElement();
 	}
 
@@ -613,29 +605,29 @@ FLearningAgentsObservationSchemaElement ULearningAgentsObservations::SpecifyStru
 
 	for (const FLearningAgentsObservationSchemaElement& Element : Elements)
 	{
-		if (!Schema->ObservationSchema.IsValid(Element.SchemaElement))
+		if (!ObservationSchema.IsValid(Element.SchemaElement))
 		{
-			UE_LOG(LogLearning, Error, TEXT("SpecifyStructObservationFromArrayViews: Invalid Observation Object."));
+			UE_LOG(LogLearning, Error, TEXT("%s: Invalid Observation Object."), *GetName());
 			return FLearningAgentsObservationSchemaElement();
 		}
 
 		SubElements.Add(Element.SchemaElement);
 	}
 
-	return { Schema->ObservationSchema.CreateAnd({ ElementNames, SubElements }, Tag) };
+	return { ObservationSchema.CreateAnd({ ElementNames, SubElements }, Name) };
 }
 
-FLearningAgentsObservationSchemaElement ULearningAgentsObservations::SpecifyExclusiveUnionObservation(ULearningAgentsObservationSchema* Schema, const TMap<FName, FLearningAgentsObservationSchemaElement>& Elements, const int32 EncodingSize, const FName Tag)
+FLearningAgentsObservationSchemaElement ULearningAgentsObservationSchema::SpecifyExclusiveUnionObservation(const TMap<FName, FLearningAgentsObservationSchemaElement>& Elements, const int32 EncodingSize, const FName Name)
 {
 	if (EncodingSize < 1)
 	{
-		UE_LOG(LogLearning, Error, TEXT("SpecifyExclusiveUnionObservation: Invalid Observation EncodingSize '%i' - must be greater than zero."), EncodingSize);
+		UE_LOG(LogLearning, Error, TEXT("%s: Invalid Observation EncodingSize '%i' - must be greater than zero."), *GetName(), EncodingSize);
 		return FLearningAgentsObservationSchemaElement();
 	}
 
 	if (Elements.Num() == 0)
 	{
-		UE_LOG(LogLearning, Warning, TEXT("SpecifyExclusiveUnionObservation: Specifying zero-sized Exclusive Union Observation."));
+		UE_LOG(LogLearning, Warning, TEXT("%s: Specifying zero-sized Exclusive Union Observation."), *GetName());
 	}
 
 	const int32 SubElementNum = Elements.Num();
@@ -673,42 +665,36 @@ FLearningAgentsObservationSchemaElement ULearningAgentsObservations::SpecifyExcl
 		SortedSubElements[Idx] = SubElements[SubElementIndices[Idx]];
 	}
 
-	return SpecifyExclusiveUnionObservationFromArrayViews(Schema, SortedSubElementNames, SortedSubElements, EncodingSize, Tag);
+	return SpecifyExclusiveUnionObservationFromArrayViews(SortedSubElementNames, SortedSubElements, EncodingSize, Name);
 }
 
-FLearningAgentsObservationSchemaElement ULearningAgentsObservations::SpecifyExclusiveUnionObservationFromArrays(ULearningAgentsObservationSchema* Schema, const TArray<FName>& ElementNames, const TArray<FLearningAgentsObservationSchemaElement>& Elements, const int32 EncodingSize, const FName Tag)
+FLearningAgentsObservationSchemaElement ULearningAgentsObservationSchema::SpecifyExclusiveUnionObservationFromArrays(const TArray<FName>& ElementNames, const TArray<FLearningAgentsObservationSchemaElement>& Elements, const int32 EncodingSize, const FName Name)
 {
-	return SpecifyExclusiveUnionObservationFromArrayViews(Schema, ElementNames, Elements, EncodingSize, Tag);
+	return SpecifyExclusiveUnionObservationFromArrayViews(ElementNames, Elements, EncodingSize, Name);
 }
 
-FLearningAgentsObservationSchemaElement ULearningAgentsObservations::SpecifyExclusiveUnionObservationFromArrayViews(ULearningAgentsObservationSchema* Schema, const TArrayView<const FName> ElementNames, const TArrayView<const FLearningAgentsObservationSchemaElement> Elements, const int32 EncodingSize, const FName Tag)
+FLearningAgentsObservationSchemaElement ULearningAgentsObservationSchema::SpecifyExclusiveUnionObservationFromArrayViews(const TArrayView<const FName> ElementNames, const TArrayView<const FLearningAgentsObservationSchemaElement> Elements, const int32 EncodingSize, const FName Name)
 {
-	if (!Schema)
-	{
-		UE_LOG(LogLearning, Error, TEXT("SpecifyExclusiveUnionObservationFromArrayViews: Schema is nullptr."));
-		return FLearningAgentsObservationSchemaElement();
-	}
-
 	if (EncodingSize < 1)
 	{
-		UE_LOG(LogLearning, Error, TEXT("SpecifyExclusiveUnionObservationFromArrayViews: Invalid Observation EncodingSize '%i' - must be greater than zero."), EncodingSize);
+		UE_LOG(LogLearning, Error, TEXT("%s: Invalid Observation EncodingSize '%i' - must be greater than zero."), *GetName(), EncodingSize);
 		return FLearningAgentsObservationSchemaElement();
 	}
 
 	if (Elements.Num() == 0)
 	{
-		UE_LOG(LogLearning, Warning, TEXT("SpecifyExclusiveUnionObservationFromArrayViews: Specifying zero-sized Exclusive Union Observation."));
+		UE_LOG(LogLearning, Warning, TEXT("%s: Specifying zero-sized Exclusive Union Observation."), *GetName());
 	}
 
 	if (Elements.Num() != ElementNames.Num())
 	{
-		UE_LOG(LogLearning, Error, TEXT("SpecifyExclusiveUnionObservationFromArrayViews: Number of elements (%i) must match number of names (%i)."), Elements.Num(), ElementNames.Num());
+		UE_LOG(LogLearning, Error, TEXT("%s: Number of elements (%i) must match number of names (%i)."), *GetName(), Elements.Num(), ElementNames.Num());
 		return FLearningAgentsObservationSchemaElement();
 	}
 
 	if (UE::Learning::Agents::Observation::Private::ContainsDuplicates(ElementNames))
 	{
-		UE_LOG(LogLearning, Error, TEXT("SpecifyExclusiveUnionObservationFromArrayViews: Element Names contain duplicates."));
+		UE_LOG(LogLearning, Error, TEXT("%s: Element Names contain duplicates."), *GetName());
 		return FLearningAgentsObservationSchemaElement();
 	}
 
@@ -717,29 +703,29 @@ FLearningAgentsObservationSchemaElement ULearningAgentsObservations::SpecifyExcl
 
 	for (const FLearningAgentsObservationSchemaElement& Element : Elements)
 	{
-		if (!Schema->ObservationSchema.IsValid(Element.SchemaElement))
+		if (!ObservationSchema.IsValid(Element.SchemaElement))
 		{
-			UE_LOG(LogLearning, Error, TEXT("SpecifyExclusiveUnionObservationFromArrayViews: Invalid Observation Object."));
+			UE_LOG(LogLearning, Error, TEXT("%s: Invalid Observation Object."), *GetName());
 			return FLearningAgentsObservationSchemaElement();
 		}
 
 		SubElements.Add(Element.SchemaElement);
 	}
 
-	return { Schema->ObservationSchema.CreateOrExclusive({ ElementNames, SubElements, EncodingSize }, Tag) };
+	return { ObservationSchema.CreateOrExclusive({ ElementNames, SubElements, EncodingSize }, Name) };
 }
 
-FLearningAgentsObservationSchemaElement ULearningAgentsObservations::SpecifyInclusiveUnionObservation(ULearningAgentsObservationSchema* Schema, const TMap<FName, FLearningAgentsObservationSchemaElement>& Elements, const int32 AttentionEncodingSize, const int32 AttentionHeadNum, const int32 ValueEncodingSize, const FName Tag)
+FLearningAgentsObservationSchemaElement ULearningAgentsObservationSchema::SpecifyInclusiveUnionObservation(const TMap<FName, FLearningAgentsObservationSchemaElement>& Elements, const int32 AttentionEncodingSize, const int32 AttentionHeadNum, const int32 ValueEncodingSize, const FName Name)
 {
 	if (AttentionEncodingSize < 1 || AttentionHeadNum < 1 || ValueEncodingSize < 1)
 	{
-		UE_LOG(LogLearning, Error, TEXT("SpecifyInclusiveUnionObservation: Invalid Observation Parameters: AttentionEncodingSize: %i, AttentionHeadNum: %i, ValueEncodingSize: %i - must be greater than zero."), AttentionEncodingSize, AttentionHeadNum, ValueEncodingSize);
+		UE_LOG(LogLearning, Error, TEXT("%s: Invalid Observation Parameters: AttentionEncodingSize: %i, AttentionHeadNum: %i, ValueEncodingSize: %i - must be greater than zero."), *GetName(), AttentionEncodingSize, AttentionHeadNum, ValueEncodingSize);
 		return FLearningAgentsObservationSchemaElement();
 	}
 
 	if (Elements.Num() == 0)
 	{
-		UE_LOG(LogLearning, Warning, TEXT("SpecifyInclusiveUnionObservation: Specifying zero-sized Inclusive Union Observation."));
+		UE_LOG(LogLearning, Warning, TEXT("%s: Specifying zero-sized Inclusive Union Observation."), *GetName());
 	}
 
 	const int32 SubElementNum = Elements.Num();
@@ -777,42 +763,36 @@ FLearningAgentsObservationSchemaElement ULearningAgentsObservations::SpecifyIncl
 		SortedSubElements[Idx] = SubElements[SubElementIndices[Idx]];
 	}
 
-	return SpecifyInclusiveUnionObservationFromArrayViews(Schema, SortedSubElementNames, SortedSubElements, AttentionEncodingSize, AttentionHeadNum, ValueEncodingSize, Tag);
+	return SpecifyInclusiveUnionObservationFromArrayViews(SortedSubElementNames, SortedSubElements, AttentionEncodingSize, AttentionHeadNum, ValueEncodingSize, Name);
 }
 
-FLearningAgentsObservationSchemaElement ULearningAgentsObservations::SpecifyInclusiveUnionObservationFromArrays(ULearningAgentsObservationSchema* Schema, const TArray<FName>& ElementNames, const TArray<FLearningAgentsObservationSchemaElement>& Elements, const int32 AttentionEncodingSize, const int32 AttentionHeadNum, const int32 ValueEncodingSize, const FName Tag)
+FLearningAgentsObservationSchemaElement ULearningAgentsObservationSchema::SpecifyInclusiveUnionObservationFromArrays(const TArray<FName>& ElementNames, const TArray<FLearningAgentsObservationSchemaElement>& Elements, const int32 AttentionEncodingSize, const int32 AttentionHeadNum, const int32 ValueEncodingSize, const FName Name)
 {
-	return SpecifyInclusiveUnionObservationFromArrayViews(Schema, ElementNames, Elements, AttentionEncodingSize, AttentionHeadNum, ValueEncodingSize, Tag);
+	return SpecifyInclusiveUnionObservationFromArrayViews(ElementNames, Elements, AttentionEncodingSize, AttentionHeadNum, ValueEncodingSize, Name);
 }
 
-FLearningAgentsObservationSchemaElement ULearningAgentsObservations::SpecifyInclusiveUnionObservationFromArrayViews(ULearningAgentsObservationSchema* Schema, const TArrayView<const FName> ElementNames, const TArrayView<const FLearningAgentsObservationSchemaElement> Elements, const int32 AttentionEncodingSize, const int32 AttentionHeadNum, const int32 ValueEncodingSize, const FName Tag)
+FLearningAgentsObservationSchemaElement ULearningAgentsObservationSchema::SpecifyInclusiveUnionObservationFromArrayViews(const TArrayView<const FName> ElementNames, const TArrayView<const FLearningAgentsObservationSchemaElement> Elements, const int32 AttentionEncodingSize, const int32 AttentionHeadNum, const int32 ValueEncodingSize, const FName Name)
 {
-	if (!Schema)
-	{
-		UE_LOG(LogLearning, Error, TEXT("SpecifyInclusiveUnionObservationFromArrayViews: Schema is nullptr."));
-		return FLearningAgentsObservationSchemaElement();
-	}
-
 	if (AttentionEncodingSize < 1 || AttentionHeadNum < 1 || ValueEncodingSize < 1)
 	{
-		UE_LOG(LogLearning, Error, TEXT("SpecifyInclusiveUnionObservationFromArrayViews: Invalid Observation Parameters: AttentionEncodingSize: %i, AttentionHeadNum: %i, ValueEncodingSize: %i - must be greater than zero."), AttentionEncodingSize, AttentionHeadNum, ValueEncodingSize);
+		UE_LOG(LogLearning, Error, TEXT("%s: Invalid Observation Parameters: AttentionEncodingSize: %i, AttentionHeadNum: %i, ValueEncodingSize: %i - must be greater than zero."), *GetName(), AttentionEncodingSize, AttentionHeadNum, ValueEncodingSize);
 		return FLearningAgentsObservationSchemaElement();
 	}
 
 	if (Elements.Num() == 0)
 	{
-		UE_LOG(LogLearning, Warning, TEXT("SpecifyInclusiveUnionObservationFromArrayViews: Specifying zero-sized Inclusive Union Observation."));
+		UE_LOG(LogLearning, Warning, TEXT("%s: Specifying zero-sized Inclusive Union Observation."), *GetName());
 	}
 
 	if (Elements.Num() != ElementNames.Num())
 	{
-		UE_LOG(LogLearning, Error, TEXT("SpecifyInclusiveUnionObservationFromArrayViews: Number of elements (%i) must match number of names (%i)."), Elements.Num(), ElementNames.Num());
+		UE_LOG(LogLearning, Error, TEXT("%s: Number of elements (%i) must match number of names (%i)."), *GetName(), Elements.Num(), ElementNames.Num());
 		return FLearningAgentsObservationSchemaElement();
 	}
 
 	if (UE::Learning::Agents::Observation::Private::ContainsDuplicates(ElementNames))
 	{
-		UE_LOG(LogLearning, Error, TEXT("SpecifyInclusiveUnionObservationFromArrayViews: Element Names contain duplicates."));
+		UE_LOG(LogLearning, Error, TEXT("%s: Element Names contain duplicates."), *GetName());
 		return FLearningAgentsObservationSchemaElement();
 	}
 
@@ -821,336 +801,310 @@ FLearningAgentsObservationSchemaElement ULearningAgentsObservations::SpecifyIncl
 
 	for (const FLearningAgentsObservationSchemaElement& Element : Elements)
 	{
-		if (!Schema->ObservationSchema.IsValid(Element.SchemaElement))
+		if (!ObservationSchema.IsValid(Element.SchemaElement))
 		{
-			UE_LOG(LogLearning, Error, TEXT("SpecifyInclusiveUnionObservationFromArrayViews: Invalid Observation Object."));
+			UE_LOG(LogLearning, Error, TEXT("%s: Invalid Observation Object."), *GetName());
 			return FLearningAgentsObservationSchemaElement();
 		}
 
 		SubElements.Add(Element.SchemaElement);
 	}
 
-	return { Schema->ObservationSchema.CreateOrInclusive({ ElementNames, SubElements, AttentionEncodingSize, AttentionHeadNum, ValueEncodingSize }, Tag) };
+	return { ObservationSchema.CreateOrInclusive({ ElementNames, SubElements, AttentionEncodingSize, AttentionHeadNum, ValueEncodingSize }, Name) };
 }
 
-FLearningAgentsObservationSchemaElement ULearningAgentsObservations::SpecifyStaticArrayObservation(ULearningAgentsObservationSchema* Schema, const FLearningAgentsObservationSchemaElement Element, const int32 Num, const FName Tag)
+FLearningAgentsObservationSchemaElement ULearningAgentsObservationSchema::SpecifyStaticArrayObservation(const FLearningAgentsObservationSchemaElement Element, const int32 Num, const FName Name)
 {
-	if (!Schema)
-	{
-		UE_LOG(LogLearning, Error, TEXT("SpecifyStaticArrayObservation: Schema is nullptr."));
-		return FLearningAgentsObservationSchemaElement();
-	}
-
 	if (Num < 0)
 	{
-		UE_LOG(LogLearning, Error, TEXT("SpecifyStaticArrayObservation: Invalid Observation Array Num %i."), Num);
+		UE_LOG(LogLearning, Error, TEXT("%s: Invalid Observation Array Num %i."), *GetName(), Num);
 		return FLearningAgentsObservationSchemaElement();
 	}
 
 	if (Num == 0)
 	{
-		UE_LOG(LogLearning, Warning, TEXT("SpecifyStaticArrayObservation: Specifying zero-sized Static Array Observation."));
+		UE_LOG(LogLearning, Warning, TEXT("%s: Specifying zero-sized Static Array Observation."), *GetName());
 	}
 
-	if (!Schema->ObservationSchema.IsValid(Element.SchemaElement))
+	if (!ObservationSchema.IsValid(Element.SchemaElement))
 	{
-		UE_LOG(LogLearning, Error, TEXT("SpecifyStaticArrayObservation: Invalid Observation Object."));
+		UE_LOG(LogLearning, Error, TEXT("%s: Invalid Observation Object."), *GetName());
 		return FLearningAgentsObservationSchemaElement();
 	}
 
-	return { Schema->ObservationSchema.CreateArray({ Element.SchemaElement, Num }, Tag) };
+	return { ObservationSchema.CreateArray({ Element.SchemaElement, Num }, Name) };
 }
 
-FLearningAgentsObservationSchemaElement ULearningAgentsObservations::SpecifySetObservation(ULearningAgentsObservationSchema* Schema, const FLearningAgentsObservationSchemaElement Element, const int32 MaxNum, const int32 AttentionEncodingSize, const int32 AttentionHeadNum, const int32 ValueEncodingSize, const FName Tag)
+FLearningAgentsObservationSchemaElement ULearningAgentsObservationSchema::SpecifySetObservation(const FLearningAgentsObservationSchemaElement Element, const int32 MaxNum, const int32 AttentionEncodingSize, const int32 AttentionHeadNum, const int32 ValueEncodingSize, const FName Name)
 {
-	if (!Schema)
-	{
-		UE_LOG(LogLearning, Error, TEXT("SpecifySetObservation: Schema is nullptr."));
-		return FLearningAgentsObservationSchemaElement();
-	}
-
 	if (MaxNum < 0)
 	{
-		UE_LOG(LogLearning, Error, TEXT("SpecifySetObservation: Invalid Observation Set MaxNum %i - must be greater than or equal to zero."), MaxNum);
+		UE_LOG(LogLearning, Error, TEXT("%s: Invalid Observation Set MaxNum %i - must be greater than or equal to zero."), *GetName(), MaxNum);
 		return FLearningAgentsObservationSchemaElement();
 	}
 
 	if (MaxNum == 0)
 	{
-		UE_LOG(LogLearning, Warning, TEXT("SpecifySetObservation: Specifying zero-sized Set Observation."));
+		UE_LOG(LogLearning, Warning, TEXT("%s: Specifying zero-sized Set Observation."), *GetName());
 	}
 
 	if (AttentionEncodingSize < 1 || AttentionHeadNum < 1 || ValueEncodingSize < 1)
 	{
-		UE_LOG(LogLearning, Error, TEXT("SpecifySetObservation: Invalid Observation Parameters: AttentionEncodingSize: %i, AttentionHeadNum: %i, ValueEncodingSize: %i - must be greater than zero."), AttentionEncodingSize, AttentionHeadNum, ValueEncodingSize);
+		UE_LOG(LogLearning, Error, TEXT("%s: Invalid Observation Parameters: AttentionEncodingSize: %i, AttentionHeadNum: %i, ValueEncodingSize: %i - must be greater than zero."), *GetName(), AttentionEncodingSize, AttentionHeadNum, ValueEncodingSize);
 		return FLearningAgentsObservationSchemaElement();
 	}
 
-	if (!Schema->ObservationSchema.IsValid(Element.SchemaElement))
+	if (!ObservationSchema.IsValid(Element.SchemaElement))
 	{
-		UE_LOG(LogLearning, Error, TEXT("SpecifySetObservation: Invalid Observation Object."));
+		UE_LOG(LogLearning, Error, TEXT("%s: Invalid Observation Object."), *GetName());
 		return FLearningAgentsObservationSchemaElement();
 	}
 
-	return { Schema->ObservationSchema.CreateSet({ Element.SchemaElement, MaxNum, AttentionEncodingSize, AttentionHeadNum, ValueEncodingSize }, Tag) };
+	return { ObservationSchema.CreateSet({ Element.SchemaElement, MaxNum, AttentionEncodingSize, AttentionHeadNum, ValueEncodingSize }, Name) };
 }
 
-FLearningAgentsObservationSchemaElement ULearningAgentsObservations::SpecifyPairObservation(ULearningAgentsObservationSchema* Schema, const FLearningAgentsObservationSchemaElement Element0, const FLearningAgentsObservationSchemaElement Element1, const FName Tag)
+FLearningAgentsObservationSchemaElement ULearningAgentsObservationSchema::SpecifyPairObservation(const FLearningAgentsObservationSchemaElement Element0, const FLearningAgentsObservationSchemaElement Element1, const FName Name)
 {
-	return SpecifyStructObservationFromArrayViews(Schema, { TEXT("Key"), TEXT("Value") }, { Element0, Element1 }, Tag);
+	return SpecifyStructObservationFromArrayViews({ TEXT("Key"), TEXT("Value") }, { Element0, Element1 }, Name);
 }
 
-FLearningAgentsObservationSchemaElement ULearningAgentsObservations::SpecifyArrayObservation(ULearningAgentsObservationSchema* Schema, const FLearningAgentsObservationSchemaElement Element, const int32 MaxNum, const int32 AttentionEncodingSize, const int32 AttentionHeadNum, const int32 ValueEncodingSize, const FName Tag)
+FLearningAgentsObservationSchemaElement ULearningAgentsObservationSchema::SpecifyArrayObservation(const FLearningAgentsObservationSchemaElement Element, const int32 MaxNum, const int32 AttentionEncodingSize, const int32 AttentionHeadNum, const int32 ValueEncodingSize, const FName Name)
 {
-	return SpecifySetObservation(Schema, SpecifyPairObservation(Schema, SpecifyIndexObservation(Schema, MaxNum), Element), MaxNum, AttentionEncodingSize, AttentionHeadNum, ValueEncodingSize);
+	return SpecifySetObservation(SpecifyPairObservation(SpecifyIndexObservation(MaxNum), Element), MaxNum, AttentionEncodingSize, AttentionHeadNum, ValueEncodingSize);
 }
 
-FLearningAgentsObservationSchemaElement ULearningAgentsObservations::SpecifyMapObservation(ULearningAgentsObservationSchema* Schema, const FLearningAgentsObservationSchemaElement KeyElement, const FLearningAgentsObservationSchemaElement ValueElement, const int32 MaxNum, const int32 AttentionEncodingSize, const int32 AttentionHeadNum, const int32 ValueEncodingSize, const FName Tag)
+FLearningAgentsObservationSchemaElement ULearningAgentsObservationSchema::SpecifyMapObservation(const FLearningAgentsObservationSchemaElement KeyElement, const FLearningAgentsObservationSchemaElement ValueElement, const int32 MaxNum, const int32 AttentionEncodingSize, const int32 AttentionHeadNum, const int32 ValueEncodingSize, const FName Name)
 {
-	return SpecifySetObservation(Schema, SpecifyPairObservation(Schema, KeyElement, ValueElement), MaxNum, AttentionEncodingSize, AttentionHeadNum, ValueEncodingSize);
+	return SpecifySetObservation(SpecifyPairObservation(KeyElement, ValueElement), MaxNum, AttentionEncodingSize, AttentionHeadNum, ValueEncodingSize);
 }
 
-FLearningAgentsObservationSchemaElement ULearningAgentsObservations::SpecifyEnumObservation(ULearningAgentsObservationSchema* Schema, const UEnum* Enum, const FName Tag)
+FLearningAgentsObservationSchemaElement ULearningAgentsObservationSchema::SpecifyEnumObservation(const UEnum* Enum, const FName Name)
 {
 	if (!Enum)
 	{
-		UE_LOG(LogLearning, Error, TEXT("SpecifyEnumObservation: Enum is nullptr."));
+		UE_LOG(LogLearning, Error, TEXT("%s: Enum is nullptr."), *GetName());
 		return FLearningAgentsObservationSchemaElement();
 	}
 
-	return SpecifyContinuousObservation(Schema, Enum->NumEnums() - 1, Tag);
+	return SpecifyContinuousObservation(Enum->NumEnums() - 1, Name);
 }
 
-FLearningAgentsObservationSchemaElement ULearningAgentsObservations::SpecifyBitmaskObservation(ULearningAgentsObservationSchema* Schema, const UEnum* Enum, const FName Tag)
+FLearningAgentsObservationSchemaElement ULearningAgentsObservationSchema::SpecifyBitmaskObservation(const UEnum* Enum, const FName Name)
 {
 	if (!Enum)
 	{
-		UE_LOG(LogLearning, Error, TEXT("SpecifyBitmaskObservation: Enum is nullptr."));
+		UE_LOG(LogLearning, Error, TEXT("%s: Enum is nullptr."), *GetName());
 		return FLearningAgentsObservationSchemaElement();
 	}
 
 	if (Enum->NumEnums() - 1 > 32)
 	{
-		UE_LOG(LogLearning, Error, TEXT("SpecifyBitmaskObservation: Too many values in Enum to use as Bitmask (%i)."), Enum->NumEnums() - 1);
+		UE_LOG(LogLearning, Error, TEXT("%s: Too many values in Enum to use as Bitmask (%i)."), *GetName(), Enum->NumEnums() - 1);
 		return FLearningAgentsObservationSchemaElement();
 	}
 
-	return SpecifyContinuousObservation(Schema, Enum->NumEnums() - 1, Tag);
+	return SpecifyContinuousObservation(Enum->NumEnums() - 1, Name);
 }
 
-FLearningAgentsObservationSchemaElement ULearningAgentsObservations::SpecifyOptionalObservation(ULearningAgentsObservationSchema* Schema, const FLearningAgentsObservationSchemaElement Element, const int32 EncodingSize, const FName Tag)
+FLearningAgentsObservationSchemaElement ULearningAgentsObservationSchema::SpecifyOptionalObservation(const FLearningAgentsObservationSchemaElement Element, const int32 EncodingSize, const FName Name)
 {
-	return SpecifyExclusiveUnionObservationFromArrayViews(Schema, { TEXT("Null"), TEXT("Valid") }, { SpecifyNullObservation(Schema), Element }, EncodingSize, Tag);
+	return SpecifyExclusiveUnionObservationFromArrayViews({ TEXT("Null"), TEXT("Valid") }, { SpecifyNullObservation(), Element }, EncodingSize, Name);
 }
 
-FLearningAgentsObservationSchemaElement ULearningAgentsObservations::SpecifyEitherObservation(ULearningAgentsObservationSchema* Schema, const FLearningAgentsObservationSchemaElement A, const FLearningAgentsObservationSchemaElement B, const int32 EncodingSize, const FName Tag)
+FLearningAgentsObservationSchemaElement ULearningAgentsObservationSchema::SpecifyEitherObservation(const FLearningAgentsObservationSchemaElement A, const FLearningAgentsObservationSchemaElement B, const int32 EncodingSize, const FName Name)
 {
-	return SpecifyExclusiveUnionObservationFromArrayViews(Schema, { TEXT("A"), TEXT("B") }, { A, B }, EncodingSize, Tag);
+	return SpecifyExclusiveUnionObservationFromArrayViews({ TEXT("A"), TEXT("B") }, { A, B }, EncodingSize, Name);
 }
 
-FLearningAgentsObservationSchemaElement ULearningAgentsObservations::SpecifyEncodingObservation(ULearningAgentsObservationSchema* Schema, const FLearningAgentsObservationSchemaElement Element, const int32 EncodingSize, const int32 LayerNum, const ELearningAgentsActivationFunction ActivationFunction, const FName Tag)
+FLearningAgentsObservationSchemaElement ULearningAgentsObservationSchema::SpecifyEncodingObservation(const FLearningAgentsObservationSchemaElement Element, const int32 EncodingSize, const int32 LayerNum, const ELearningAgentsActivationFunction ActivationFunction, const FName Name)
 {
-	if (!Schema)
-	{
-		UE_LOG(LogLearning, Error, TEXT("SpecifyEncodingObservation: Schema is nullptr."));
-		return FLearningAgentsObservationSchemaElement();
-	}
-
 	if (EncodingSize < 1)
 	{
-		UE_LOG(LogLearning, Error, TEXT("SpecifyEncodingObservation: Invalid Observation EncodingSize '%i' - must be greater than zero."), EncodingSize);
+		UE_LOG(LogLearning, Error, TEXT("%s: Invalid Observation EncodingSize '%i' - must be greater than zero."), *GetName(), EncodingSize);
 		return FLearningAgentsObservationSchemaElement();
 	}
 
 	if (LayerNum < 1)
 	{
-		UE_LOG(LogLearning, Error, TEXT("SpecifyEncodingObservation: Invalid Observation LayerNum '%i' - must be greater than zero."), LayerNum);
+		UE_LOG(LogLearning, Error, TEXT("%s: Invalid Observation LayerNum '%i' - must be greater than zero."), *GetName(), LayerNum);
 		return FLearningAgentsObservationSchemaElement();
 	}
 
-	if (!Schema->ObservationSchema.IsValid(Element.SchemaElement))
+	if (!ObservationSchema.IsValid(Element.SchemaElement))
 	{
-		UE_LOG(LogLearning, Error, TEXT("SpecifyEncodingObservation: Invalid Observation Object."));
+		UE_LOG(LogLearning, Error, TEXT("%s: Invalid Observation Object."), *GetName());
 		return FLearningAgentsObservationSchemaElement();
 	}
 
-	return { Schema->ObservationSchema.CreateEncoding({ Element.SchemaElement, EncodingSize, LayerNum, UE::Learning::Agents::Observation::Private::GetEncodingActivationFunction(ActivationFunction) }, Tag)};
+	return { ObservationSchema.CreateEncoding({ Element.SchemaElement, EncodingSize, LayerNum, UE::Learning::Agents::Observation::Private::GetEncodingActivationFunction(ActivationFunction) }, Name)};
 }
 
-FLearningAgentsObservationSchemaElement ULearningAgentsObservations::SpecifyBoolObservation(ULearningAgentsObservationSchema* Schema, const FName Tag)
+FLearningAgentsObservationSchemaElement ULearningAgentsObservationSchema::SpecifyBoolObservation(const FName Name)
 {
-	return SpecifyContinuousObservation(Schema, 1, Tag);
+	return SpecifyContinuousObservation(1, Name);
 }
 
-FLearningAgentsObservationSchemaElement ULearningAgentsObservations::SpecifyFloatObservation(ULearningAgentsObservationSchema* Schema, const FName Tag)
+FLearningAgentsObservationSchemaElement ULearningAgentsObservationSchema::SpecifyFloatObservation(const FName Name)
 {
-	return SpecifyContinuousObservation(Schema, 1, Tag);
+	return SpecifyContinuousObservation(1, Name);
 }
 
-FLearningAgentsObservationSchemaElement ULearningAgentsObservations::SpecifyLocationObservation(ULearningAgentsObservationSchema* Schema, const FName Tag)
+FLearningAgentsObservationSchemaElement ULearningAgentsObservationSchema::SpecifyLocationObservation(const FName Name)
 {
-	return SpecifyContinuousObservation(Schema, 3, Tag);
+	return SpecifyContinuousObservation(3, Name);
 }
 
-FLearningAgentsObservationSchemaElement ULearningAgentsObservations::SpecifyRotationObservation(ULearningAgentsObservationSchema* Schema, const FName Tag)
+FLearningAgentsObservationSchemaElement ULearningAgentsObservationSchema::SpecifyRotationObservation(const FName Name)
 {
-	return SpecifyContinuousObservation(Schema, 6, Tag);
+	return SpecifyContinuousObservation(6, Name);
 }
 
-FLearningAgentsObservationSchemaElement ULearningAgentsObservations::SpecifyScaleObservation(ULearningAgentsObservationSchema* Schema, const FName Tag)
+FLearningAgentsObservationSchemaElement ULearningAgentsObservationSchema::SpecifyScaleObservation(const FName Name)
 {
-	return SpecifyContinuousObservation(Schema, 3, Tag);
+	return SpecifyContinuousObservation(3, Name);
 }
 
-FLearningAgentsObservationSchemaElement ULearningAgentsObservations::SpecifyTransformObservation(ULearningAgentsObservationSchema* Schema, const FName Tag)
+FLearningAgentsObservationSchemaElement ULearningAgentsObservationSchema::SpecifyTransformObservation(const FName Name)
 {
-	return SpecifyStructObservationFromArrayViews(Schema,
+	return SpecifyStructObservationFromArrayViews(
 		{
 			TEXT("Location"),
 			TEXT("Rotation"),
 			TEXT("Scale")
 		},
 		{
-			SpecifyLocationObservation(Schema),
-			SpecifyRotationObservation(Schema),
-			SpecifyScaleObservation(Schema)
+			SpecifyLocationObservation(),
+			SpecifyRotationObservation(),
+			SpecifyScaleObservation()
 		}, 
-		Tag);
+		Name);
 }
 
-FLearningAgentsObservationSchemaElement ULearningAgentsObservations::SpecifyAngleObservation(ULearningAgentsObservationSchema* Schema, const FName Tag)
+FLearningAgentsObservationSchemaElement ULearningAgentsObservationSchema::SpecifyAngleObservation(const FName Name)
 {
-	return SpecifyContinuousObservation(Schema, 2, Tag);
+	return SpecifyContinuousObservation(2, Name);
 }
 
-FLearningAgentsObservationSchemaElement ULearningAgentsObservations::SpecifyVelocityObservation(ULearningAgentsObservationSchema* Schema, const FName Tag)
+FLearningAgentsObservationSchemaElement ULearningAgentsObservationSchema::SpecifyVelocityObservation(const FName Name)
 {
-	return SpecifyContinuousObservation(Schema, 3, Tag);
+	return SpecifyContinuousObservation(3, Name);
 }
 
-FLearningAgentsObservationSchemaElement ULearningAgentsObservations::SpecifyDirectionObservation(ULearningAgentsObservationSchema* Schema, const FName Tag)
+FLearningAgentsObservationSchemaElement ULearningAgentsObservationSchema::SpecifyDirectionObservation(const FName Name)
 {
-	return SpecifyContinuousObservation(Schema, 3, Tag);
+	return SpecifyContinuousObservation(3, Name);
 }
 
-FLearningAgentsObservationSchemaElement ULearningAgentsObservations::SpecifyLocationAlongSplineObservation(ULearningAgentsObservationSchema* Schema, const FName Tag)
+FLearningAgentsObservationSchemaElement ULearningAgentsObservationSchema::SpecifyLocationAlongSplineObservation(const FName Name)
 {
-	return SpecifyLocationObservation(Schema, Tag);
+	return SpecifyLocationObservation(Name);
 }
 
-FLearningAgentsObservationSchemaElement ULearningAgentsObservations::SpecifyProportionAlongSplineObservation(ULearningAgentsObservationSchema* Schema, const FName Tag)
+FLearningAgentsObservationSchemaElement ULearningAgentsObservationSchema::SpecifyProportionAlongSplineObservation(const FName Name)
 {
-	return SpecifyExclusiveUnionObservationFromArrayViews(Schema,
+	return SpecifyExclusiveUnionObservationFromArrayViews(
 		{
 			TEXT("Angle"),
 			TEXT("Proportion")
 		},
 		{
-			SpecifyAngleObservation(Schema),
-			SpecifyFloatObservation(Schema)
+			SpecifyAngleObservation(),
+			SpecifyFloatObservation()
 		},
 		8,
-		Tag);
+		Name);
 }
 
-FLearningAgentsObservationSchemaElement ULearningAgentsObservations::SpecifyDirectionAlongSplineObservation(ULearningAgentsObservationSchema* Schema, const FName Tag)
+FLearningAgentsObservationSchemaElement ULearningAgentsObservationSchema::SpecifyDirectionAlongSplineObservation(const FName Name)
 {
-	return SpecifyDirectionObservation(Schema, Tag);
+	return SpecifyDirectionObservation(Name);
 }
 
-FLearningAgentsObservationSchemaElement ULearningAgentsObservations::SpecifyPropertiesAlongSplineObservation(ULearningAgentsObservationSchema* Schema, const FName Tag)
+FLearningAgentsObservationSchemaElement ULearningAgentsObservationSchema::SpecifyPropertiesAlongSplineObservation(const FName Name)
 {
-	return SpecifyStructObservationFromArrayViews(Schema,
+	return SpecifyStructObservationFromArrayViews(
 		{
 			TEXT("Location"),
 			TEXT("Proportion"),
 			TEXT("Direction")
 		},
 		{
-			SpecifyLocationAlongSplineObservation(Schema),
-			SpecifyProportionAlongSplineObservation(Schema),
-			SpecifyDirectionAlongSplineObservation(Schema),
+			SpecifyLocationAlongSplineObservation(),
+			SpecifyProportionAlongSplineObservation(),
+			SpecifyDirectionAlongSplineObservation(),
 		},
-		Tag);
+		Name);
 }
 
-FLearningAgentsObservationSchemaElement ULearningAgentsObservations::SpecifyProportionAlongRayObservation(ULearningAgentsObservationSchema* Schema, const FName Tag)
+FLearningAgentsObservationSchemaElement ULearningAgentsObservationSchema::SpecifyProportionAlongRayObservation(const FName Name)
 {
-	return SpecifyFloatObservation(Schema, Tag);
+	return SpecifyFloatObservation(Name);
 }
 
-FLearningAgentsObservationSchemaElement ULearningAgentsObservations::SpecifyProportionAlongRaysObservation(ULearningAgentsObservationSchema* Schema, const int32 Num, const FName Tag)
+FLearningAgentsObservationSchemaElement ULearningAgentsObservationSchema::SpecifyProportionAlongRaysObservation(const int32 Num, const FName Name)
 {
-	return SpecifyStaticArrayObservation(Schema, SpecifyProportionAlongRayObservation(Schema), Num, Tag);
+	return SpecifyStaticArrayObservation(SpecifyProportionAlongRayObservation(), Num, Name);
 }
 
-void ULearningAgentsObservations::LogObservation(const ULearningAgentsObservationObject* Object, const FLearningAgentsObservationObjectElement Element)
+const UE::Learning::Observation::FObject& ULearningAgentsObservationObject::GetObservationObject() const
 {
-	if (!Object)
-	{
-		UE_LOG(LogLearning, Error, TEXT("LogObservation: Object is nullptr."));
-		return;
-	}
-
-	UE::Learning::Agents::Observation::Private::LogObservation(Object->ObservationObject, Element.ObjectElement, TEXT(""), TEXT(""));
+	return ObservationObject;
 }
 
-FLearningAgentsObservationObjectElement ULearningAgentsObservations::MakeNullObservation(ULearningAgentsObservationObject* Object, const FName Tag)
+UE::Learning::Observation::FObject& ULearningAgentsObservationObject::GetObservationObject()
 {
-	if (!Object)
-	{
-		UE_LOG(LogLearning, Error, TEXT("MakeNullObservation: Object is nullptr."));
-		return FLearningAgentsObservationObjectElement();
-	}
+	return ObservationObject;
+}
 
-	return { Object->ObservationObject.CreateNull(Tag) };
+void ULearningAgentsObservationObject::LogObservation(const FLearningAgentsObservationObjectElement Element)
+{
+	UE::Learning::Agents::Observation::Private::LogObservation(GetObservationObject(), Element.ObjectElement, TEXT(""), TEXT(""), GetName());
 }
 
-FLearningAgentsObservationObjectElement ULearningAgentsObservations::MakeContinuousObservation(ULearningAgentsObservationObject* Object, const TArray<float>& Values, const FName Tag)
+FLearningAgentsObservationObjectElement ULearningAgentsObservationObject::MakeNullObservation(const FName Name)
 {
-	return MakeContinuousObservationFromArrayView(Object, Values, Tag);
+	return { ObservationObject.CreateNull(Name) };
 }
 
-FLearningAgentsObservationObjectElement ULearningAgentsObservations::MakeContinuousObservationFromArrayView(ULearningAgentsObservationObject* Object, const TArrayView<const float> Values, const FName Tag)
+FLearningAgentsObservationObjectElement ULearningAgentsObservationObject::MakeContinuousObservation(const TArray<float>& Values, const FName Name)
 {
-	if (!Object)
-	{
-		UE_LOG(LogLearning, Error, TEXT("MakeContinuousObservationFromArrayView: Object is nullptr."));
-		return FLearningAgentsObservationObjectElement();
-	}
+	return MakeContinuousObservationFromArrayView(Values, Name);
+}
 
+FLearningAgentsObservationObjectElement ULearningAgentsObservationObject::MakeContinuousObservationFromArrayView(const TArrayView<const float> Values, const FName Name)
+{
 	if (Values.Num() == 0)
 	{
-		UE_LOG(LogLearning, Warning, TEXT("MakeContinuousObservationFromArrayView: Creating zero-sized Continuous Observation."));
+		UE_LOG(LogLearning, Warning, TEXT("%s: Creating zero-sized Continuous Observation."), *GetName());
 	}
 
-	return { Object->ObservationObject.CreateContinuous({ Values }, Tag) };
+	return { ObservationObject.CreateContinuous({ Values }, Name) };
 }
 
-FLearningAgentsObservationObjectElement ULearningAgentsObservations::MakeExclusiveDiscreteObservation(ULearningAgentsObservationObject* Object, const int32 DiscreteIndex, const int32 Size, const FName Tag)
+FLearningAgentsObservationObjectElement ULearningAgentsObservationObject::MakeExclusiveDiscreteObservation(const int32 DiscreteIndex, const int32 Size, const FName Name)
 {
 	if (DiscreteIndex < 0 || DiscreteIndex >= Size)
 	{
-		UE_LOG(LogLearning, Error, TEXT("MakeExclusiveDiscreteObservation: Discrete index out of range: Got %i, expected <= %i."), DiscreteIndex, Size);
+		UE_LOG(LogLearning, Error, TEXT("%s: Discrete index out of range: Got %i, expected <= %i."), *GetName(), DiscreteIndex, Size);
 		return FLearningAgentsObservationObjectElement();
 	}
 
 	TArray<float, TInlineAllocator<32>> Values;
 	Values.Init(0.0f, Size);
 	Values[DiscreteIndex] = 1.0f;
-	return MakeContinuousObservationFromArrayView(Object, Values, Tag);
+	return MakeContinuousObservationFromArrayView(Values, Name);
 }
 
-FLearningAgentsObservationObjectElement ULearningAgentsObservations::MakeInclusiveDiscreteObservation(ULearningAgentsObservationObject* Object, const TArray<int32>& DiscreteIndices, const int32 Size, const FName Tag)
+FLearningAgentsObservationObjectElement ULearningAgentsObservationObject::MakeInclusiveDiscreteObservation(const TArray<int32>& DiscreteIndices, const int32 Size, const FName Name)
 {
-	return MakeInclusiveDiscreteObservationFromArrayView(Object, DiscreteIndices, Size, Tag);
+	return MakeInclusiveDiscreteObservationFromArrayView(DiscreteIndices, Size, Name);
 }
 
-FLearningAgentsObservationObjectElement ULearningAgentsObservations::MakeInclusiveDiscreteObservationFromArrayView(ULearningAgentsObservationObject* Object, const TArrayView<const int32> DiscreteIndices, const int32 Size, const FName Tag)
+FLearningAgentsObservationObjectElement ULearningAgentsObservationObject::MakeInclusiveDiscreteObservationFromArrayView(const TArrayView<const int32> DiscreteIndices, const int32 Size, const FName Name)
 {
 	if (UE::Learning::Agents::Observation::Private::ContainsDuplicates(DiscreteIndices))
 	{
-		UE_LOG(LogLearning, Error, TEXT("MakeInclusiveDiscreteObservationFromArrayView: Indices contain duplicates."));
+		UE_LOG(LogLearning, Error, TEXT("%s: Indices contain duplicates."), *GetName());
 		return FLearningAgentsObservationObjectElement();
 	}
 
@@ -1161,21 +1115,21 @@ FLearningAgentsObservationObjectElement ULearningAgentsObservations::MakeInclusi
 	{
 		if (DiscreteIndices[Idx] < 0 || DiscreteIndices[Idx] >= Size)
 		{
-			UE_LOG(LogLearning, Error, TEXT("MakeInclusiveDiscreteObservationFromArrayView: Discrete index out of range: Got %i, expected <= %i."), DiscreteIndices[Idx], Size);
+			UE_LOG(LogLearning, Error, TEXT("%s: Discrete index out of range: Got %i, expected <= %i."), *GetName(), DiscreteIndices[Idx], Size);
 			return FLearningAgentsObservationObjectElement();
 		}
 
 		Values[DiscreteIndices[Idx]] = 1.0f;
 	}
 
-	return MakeContinuousObservationFromArrayView(Object, Values, Tag);
+	return MakeContinuousObservationFromArrayView(Values, Name);
 }
 
-FLearningAgentsObservationObjectElement ULearningAgentsObservations::MakeIndexObservation(ULearningAgentsObservationObject* Object, const int32 Index, const int32 Size, const FName Tag)
+FLearningAgentsObservationObjectElement ULearningAgentsObservationObject::MakeIndexObservation(const int32 Index, const int32 Size, const FName Name)
 {
 	if (Index < 0 || Index >= Size)
 	{
-		UE_LOG(LogLearning, Error, TEXT("MakeIndexObservation: Discrete index out of range: Got %i, expected <= %i."), Index, Size);
+		UE_LOG(LogLearning, Error, TEXT("%s: Discrete index out of range: Got %i, expected <= %i."), *GetName(), Index, Size);
 		return FLearningAgentsObservationObjectElement();
 	}
 
@@ -1187,25 +1141,25 @@ FLearningAgentsObservationObjectElement ULearningAgentsObservations::MakeIndexOb
 		Values[Idx] = 1.0f;
 	}
 
-	return MakeContinuousObservationFromArrayView(Object, Values, Tag);
+	return MakeContinuousObservationFromArrayView(Values, Name);
 }
 
-FLearningAgentsObservationObjectElement ULearningAgentsObservations::MakeCountObservation(ULearningAgentsObservationObject* Object, const int32 Num, const int32 MaxNum, const FName Tag)
+FLearningAgentsObservationObjectElement ULearningAgentsObservationObject::MakeCountObservation(const int32 Num, const int32 MaxNum, const FName Name)
 {
 	if (MaxNum == 0)
 	{
-		UE_LOG(LogLearning, Error, TEXT("MakeCountObservation: MaxNum must not be zero for Count Observation."));
+		UE_LOG(LogLearning, Error, TEXT("%s: MaxNum must not be zero for Count Observation."), *GetName());
 		return FLearningAgentsObservationObjectElement();
 	}
 
-	return MakeContinuousObservationFromArrayView(Object, { (float)Num / (float)MaxNum }, Tag);
+	return MakeContinuousObservationFromArrayView({ (float)Num / (float)MaxNum }, Name);
 }
 
-FLearningAgentsObservationObjectElement ULearningAgentsObservations::MakeStructObservation(ULearningAgentsObservationObject* Object, const TMap<FName, FLearningAgentsObservationObjectElement>& Elements, const FName Tag)
+FLearningAgentsObservationObjectElement ULearningAgentsObservationObject::MakeStructObservation(const TMap<FName, FLearningAgentsObservationObjectElement>& Elements, const FName Name)
 {
 	if (Elements.Num() == 0)
 	{
-		UE_LOG(LogLearning, Warning, TEXT("MakeStructObservation: Creating zero-sized Struct Observation."));
+		UE_LOG(LogLearning, Warning, TEXT("%s: Creating zero-sized Struct Observation."), *GetName());
 	}
 
 	const int32 SubElementNum = Elements.Num();
@@ -1221,36 +1175,30 @@ FLearningAgentsObservationObjectElement ULearningAgentsObservations::MakeStructO
 		SubElements.Add(Element.Value);
 	}
 
-	return MakeStructObservationFromArrayViews(Object, SubElementNames, SubElements, Tag);
+	return MakeStructObservationFromArrayViews(SubElementNames, SubElements, Name);
 }
 
-FLearningAgentsObservationObjectElement ULearningAgentsObservations::MakeStructObservationFromArrays(ULearningAgentsObservationObject* Object, const TArray<FName>& ElementNames, const TArray<FLearningAgentsObservationObjectElement>& Elements, const FName Tag)
+FLearningAgentsObservationObjectElement ULearningAgentsObservationObject::MakeStructObservationFromArrays(const TArray<FName>& ElementNames, const TArray<FLearningAgentsObservationObjectElement>& Elements, const FName Name)
 {
-	return MakeStructObservationFromArrayViews(Object, ElementNames, Elements, Tag);
+	return MakeStructObservationFromArrayViews(ElementNames, Elements, Name);
 }
 
-FLearningAgentsObservationObjectElement ULearningAgentsObservations::MakeStructObservationFromArrayViews(ULearningAgentsObservationObject* Object, const TArrayView<const FName> ElementNames, const TArrayView<const FLearningAgentsObservationObjectElement> Elements, const FName Tag)
+FLearningAgentsObservationObjectElement ULearningAgentsObservationObject::MakeStructObservationFromArrayViews(const TArrayView<const FName> ElementNames, const TArrayView<const FLearningAgentsObservationObjectElement> Elements, const FName Name)
 {
-	if (!Object)
-	{
-		UE_LOG(LogLearning, Error, TEXT("MakeStructObservationFromArrayViews: Object is nullptr."));
-		return FLearningAgentsObservationObjectElement();
-	}
-
 	if (Elements.Num() == 0)
 	{
-		UE_LOG(LogLearning, Warning, TEXT("MakeStructObservationFromArrayViews: Creating zero-sized Struct Observation."));
+		UE_LOG(LogLearning, Warning, TEXT("%s: Creating zero-sized Struct Observation."), *GetName());
 	}
 
 	if (Elements.Num() != ElementNames.Num())
 	{
-		UE_LOG(LogLearning, Error, TEXT("MakeStructObservationFromArrayViews: Number of elements (%i) must match number of names (%i)."), Elements.Num(), ElementNames.Num());
+		UE_LOG(LogLearning, Error, TEXT("%s: Number of elements (%i) must match number of names (%i)."), *GetName(), Elements.Num(), ElementNames.Num());
 		return FLearningAgentsObservationObjectElement();
 	}
 
 	if (UE::Learning::Agents::Observation::Private::ContainsDuplicates(ElementNames))
 	{
-		UE_LOG(LogLearning, Error, TEXT("MakeStructObservationFromArrayViews: Element Names contain duplicates."));
+		UE_LOG(LogLearning, Error, TEXT("%s: Element Names contain duplicates."), *GetName());
 		return FLearningAgentsObservationObjectElement();
 	}
 
@@ -1259,36 +1207,30 @@ FLearningAgentsObservationObjectElement ULearningAgentsObservations::MakeStructO
 
 	for (const FLearningAgentsObservationObjectElement& Element : Elements)
 	{
-		if (!Object->ObservationObject.IsValid(Element.ObjectElement))
+		if (!ObservationObject.IsValid(Element.ObjectElement))
 		{
-			UE_LOG(LogLearning, Error, TEXT("MakeStructObservationFromArrayViews: Invalid Observation Object."));
+			UE_LOG(LogLearning, Error, TEXT("%s: Invalid Observation Object."), *GetName());
 			return FLearningAgentsObservationObjectElement();
 		}
 
 		SubElements.Add(Element.ObjectElement);
 	}
 
-	return { Object->ObservationObject.CreateAnd({ ElementNames, SubElements }, Tag)};
+	return { ObservationObject.CreateAnd({ ElementNames, SubElements }, Name)};
 }
 
-FLearningAgentsObservationObjectElement ULearningAgentsObservations::MakeExclusiveUnionObservation(ULearningAgentsObservationObject* Object, const FName ElementName, const FLearningAgentsObservationObjectElement Element, const FName Tag)
+FLearningAgentsObservationObjectElement ULearningAgentsObservationObject::MakeExclusiveUnionObservation(const FName ElementName, const FLearningAgentsObservationObjectElement Element, const FName Name)
 {
-	if (!Object)
+	if (!ObservationObject.IsValid(Element.ObjectElement))
 	{
-		UE_LOG(LogLearning, Error, TEXT("MakeExclusiveUnionObservation: Object is nullptr."));
+		UE_LOG(LogLearning, Error, TEXT("%s: Invalid Observation Object."), *GetName());
 		return FLearningAgentsObservationObjectElement();
 	}
 
-	if (!Object->ObservationObject.IsValid(Element.ObjectElement))
-	{
-		UE_LOG(LogLearning, Error, TEXT("MakeExclusiveUnionObservation: Invalid Observation Object."));
-		return FLearningAgentsObservationObjectElement();
-	}
-
-	return { Object->ObservationObject.CreateOrExclusive({ ElementName, Element.ObjectElement }, Tag) };
+	return { ObservationObject.CreateOrExclusive({ ElementName, Element.ObjectElement }, Name) };
 }
 
-FLearningAgentsObservationObjectElement ULearningAgentsObservations::MakeInclusiveUnionObservation(ULearningAgentsObservationObject* Object, const TMap<FName, FLearningAgentsObservationObjectElement>& Elements, const FName Tag)
+FLearningAgentsObservationObjectElement ULearningAgentsObservationObject::MakeInclusiveUnionObservation(const TMap<FName, FLearningAgentsObservationObjectElement>& Elements, const FName Name)
 {
 	const int32 SubElementNum = Elements.Num();
 
@@ -1303,31 +1245,25 @@ FLearningAgentsObservationObjectElement ULearningAgentsObservations::MakeInclusi
 		SubElements.Add(Element.Value);
 	}
 
-	return MakeInclusiveUnionObservationFromArrayViews(Object, SubElementNames, SubElements, Tag);
+	return MakeInclusiveUnionObservationFromArrayViews(SubElementNames, SubElements, Name);
 }
 
-FLearningAgentsObservationObjectElement ULearningAgentsObservations::MakeInclusiveUnionObservationFromArrays(ULearningAgentsObservationObject* Object, const TArray<FName>& ElementNames, const TArray<FLearningAgentsObservationObjectElement>& Elements, const FName Tag)
+FLearningAgentsObservationObjectElement ULearningAgentsObservationObject::MakeInclusiveUnionObservationFromArrays(const TArray<FName>& ElementNames, const TArray<FLearningAgentsObservationObjectElement>& Elements, const FName Name)
 {
-	return MakeInclusiveUnionObservationFromArrayViews(Object, ElementNames, Elements, Tag);
+	return MakeInclusiveUnionObservationFromArrayViews(ElementNames, Elements, Name);
 }
 
-FLearningAgentsObservationObjectElement ULearningAgentsObservations::MakeInclusiveUnionObservationFromArrayViews(ULearningAgentsObservationObject* Object, const TArrayView<const FName> ElementNames, const TArrayView<const FLearningAgentsObservationObjectElement> Elements, const FName Tag)
+FLearningAgentsObservationObjectElement ULearningAgentsObservationObject::MakeInclusiveUnionObservationFromArrayViews(const TArrayView<const FName> ElementNames, const TArrayView<const FLearningAgentsObservationObjectElement> Elements, const FName Name)
 {
-	if (!Object)
-	{
-		UE_LOG(LogLearning, Error, TEXT("MakeInclusiveUnionObservationFromArrayViews: Object is nullptr."));
-		return FLearningAgentsObservationObjectElement();
-	}
-
 	if (Elements.Num() != ElementNames.Num())
 	{
-		UE_LOG(LogLearning, Error, TEXT("MakeInclusiveUnionObservationFromArrayViews: Number of elements (%i) must match number of names (%i)."), Elements.Num(), ElementNames.Num());
+		UE_LOG(LogLearning, Error, TEXT("%s: Number of elements (%i) must match number of names (%i)."), *GetName(), Elements.Num(), ElementNames.Num());
 		return FLearningAgentsObservationObjectElement();
 	}
 
 	if (UE::Learning::Agents::Observation::Private::ContainsDuplicates(ElementNames))
 	{
-		UE_LOG(LogLearning, Error, TEXT("MakeInclusiveUnionObservationFromArrayViews: Element Names contain duplicates."));
+		UE_LOG(LogLearning, Error, TEXT("%s: Element Names contain duplicates."), *GetName());
 		return FLearningAgentsObservationObjectElement();
 	}
 
@@ -1336,34 +1272,28 @@ FLearningAgentsObservationObjectElement ULearningAgentsObservations::MakeInclusi
 
 	for (const FLearningAgentsObservationObjectElement& Element : Elements)
 	{
-		if (!Object->ObservationObject.IsValid(Element.ObjectElement))
+		if (!ObservationObject.IsValid(Element.ObjectElement))
 		{
-			UE_LOG(LogLearning, Error, TEXT("MakeInclusiveUnionObservationFromArrayViews: Invalid Observation Object."));
+			UE_LOG(LogLearning, Error, TEXT("%s: Invalid Observation Object."), *GetName());
 			return FLearningAgentsObservationObjectElement();
 		}
 
 		SubElements.Add(Element.ObjectElement);
 	}
 
-	return { Object->ObservationObject.CreateOrInclusive({ ElementNames, SubElements }, Tag) };
+	return { ObservationObject.CreateOrInclusive({ ElementNames, SubElements }, Name) };
 }
 
-FLearningAgentsObservationObjectElement ULearningAgentsObservations::MakeStaticArrayObservation(ULearningAgentsObservationObject* Object, const TArray<FLearningAgentsObservationObjectElement>& Elements, const FName Tag)
+FLearningAgentsObservationObjectElement ULearningAgentsObservationObject::MakeStaticArrayObservation(const TArray<FLearningAgentsObservationObjectElement>& Elements, const FName Name)
 {
-	return MakeStaticArrayObservationFromArrayView(Object, Elements, Tag);
+	return MakeStaticArrayObservationFromArrayView(Elements, Name);
 }
 
-FLearningAgentsObservationObjectElement ULearningAgentsObservations::MakeStaticArrayObservationFromArrayView(ULearningAgentsObservationObject* Object, const TArrayView<const FLearningAgentsObservationObjectElement> Elements, const FName Tag)
+FLearningAgentsObservationObjectElement ULearningAgentsObservationObject::MakeStaticArrayObservationFromArrayView(const TArrayView<const FLearningAgentsObservationObjectElement> Elements, const FName Name)
 {
-	if (!Object)
-	{
-		UE_LOG(LogLearning, Error, TEXT("MakeStaticArrayObservationFromArrayView: Object is nullptr."));
-		return FLearningAgentsObservationObjectElement();
-	}
-
 	if (Elements.Num() == 0)
 	{
-		UE_LOG(LogLearning, Warning, TEXT("MakeStaticArrayObservationFromArrayView: Creating zero-sized Static Array Observation."));
+		UE_LOG(LogLearning, Warning, TEXT("%s: Creating zero-sized Static Array Observation."), *GetName());
 	}
 
 	TArray<UE::Learning::Observation::FObjectElement, TInlineAllocator<16>> SubElements;
@@ -1371,91 +1301,73 @@ FLearningAgentsObservationObjectElement ULearningAgentsObservations::MakeStaticA
 
 	for (const FLearningAgentsObservationObjectElement& Element : Elements)
 	{
-		if (!Object->ObservationObject.IsValid(Element.ObjectElement))
+		if (!ObservationObject.IsValid(Element.ObjectElement))
 		{
-			UE_LOG(LogLearning, Error, TEXT("MakeStaticArrayObservationFromArrayView: Invalid Observation Object."));
+			UE_LOG(LogLearning, Error, TEXT("%s: Invalid Observation Object."), *GetName());
 			return FLearningAgentsObservationObjectElement();
 		}
 
 		SubElements.Add(Element.ObjectElement);
 	}
 
-	return { Object->ObservationObject.CreateArray({ SubElements }, Tag) };
+	return { ObservationObject.CreateArray({ SubElements }, Name) };
 }
 
-FLearningAgentsObservationObjectElement ULearningAgentsObservations::MakeSetObservation(ULearningAgentsObservationObject* Object, const TSet<FLearningAgentsObservationObjectElement>& Elements, const FName Tag)
+FLearningAgentsObservationObjectElement ULearningAgentsObservationObject::MakeSetObservation(const TSet<FLearningAgentsObservationObjectElement>& Elements, const FName Name)
 {
-	if (!Object)
-	{
-		UE_LOG(LogLearning, Error, TEXT("MakeSetObservation: Object is nullptr."));
-		return FLearningAgentsObservationObjectElement();
-	}
-
 	TArray<UE::Learning::Observation::FObjectElement, TInlineAllocator<16>> SubElements;
 	SubElements.Empty(Elements.Num());
 
 	for (const FLearningAgentsObservationObjectElement& Element : Elements)
 	{
-		if (!Object->ObservationObject.IsValid(Element.ObjectElement))
+		if (!ObservationObject.IsValid(Element.ObjectElement))
 		{
-			UE_LOG(LogLearning, Error, TEXT("MakeSetObservation: Invalid Observation Object."));
+			UE_LOG(LogLearning, Error, TEXT("%s: Invalid Observation Object."), *GetName());
 			return FLearningAgentsObservationObjectElement();
 		}
 
 		SubElements.Add(Element.ObjectElement);
 	}
 
-	return { Object->ObservationObject.CreateSet({ SubElements }, Tag) };
+	return { ObservationObject.CreateSet({ SubElements }, Name) };
 }
 
-FLearningAgentsObservationObjectElement ULearningAgentsObservations::MakeSetObservationFromArray(ULearningAgentsObservationObject* Object, const TArray<FLearningAgentsObservationObjectElement>& Elements, const FName Tag)
+FLearningAgentsObservationObjectElement ULearningAgentsObservationObject::MakeSetObservationFromArray(const TArray<FLearningAgentsObservationObjectElement>& Elements, const FName Name)
 {
-	return MakeSetObservationFromArrayView(Object, Elements, Tag);
+	return MakeSetObservationFromArrayView(Elements, Name);
 }
 
-FLearningAgentsObservationObjectElement ULearningAgentsObservations::MakeSetObservationFromArrayView(ULearningAgentsObservationObject* Object, const TArrayView<const FLearningAgentsObservationObjectElement> Elements, const FName Tag)
+FLearningAgentsObservationObjectElement ULearningAgentsObservationObject::MakeSetObservationFromArrayView(const TArrayView<const FLearningAgentsObservationObjectElement> Elements, const FName Name)
 {
-	if (!Object)
-	{
-		UE_LOG(LogLearning, Error, TEXT("MakeSetObservationFromArrayView: Object is nullptr."));
-		return FLearningAgentsObservationObjectElement();
-	}
-
 	TArray<UE::Learning::Observation::FObjectElement, TInlineAllocator<16>> SubElements;
 	SubElements.Empty(Elements.Num());
 
 	for (const FLearningAgentsObservationObjectElement& Element : Elements)
 	{
-		if (!Object->ObservationObject.IsValid(Element.ObjectElement))
+		if (!ObservationObject.IsValid(Element.ObjectElement))
 		{
-			UE_LOG(LogLearning, Error, TEXT("MakeSetObservationFromArrayView: Invalid Observation Object."));
+			UE_LOG(LogLearning, Error, TEXT("%s: Invalid Observation Object."), *GetName());
 			return FLearningAgentsObservationObjectElement();
 		}
 
 		SubElements.Add(Element.ObjectElement);
 	}
 
-	return { Object->ObservationObject.CreateSet({ SubElements }, Tag) };
+	return { ObservationObject.CreateSet({ SubElements }, Name) };
 }
 
-FLearningAgentsObservationObjectElement ULearningAgentsObservations::MakePairObservation(ULearningAgentsObservationObject* Object, const FLearningAgentsObservationObjectElement Key, const FLearningAgentsObservationObjectElement Value, const FName Tag)
+FLearningAgentsObservationObjectElement ULearningAgentsObservationObject::MakePairObservation(const FLearningAgentsObservationObjectElement Key, const FLearningAgentsObservationObjectElement Value, const FName Name)
 {
-	return MakeStructObservationFromArrayViews(Object, { TEXT("Key"), TEXT("Value") }, { Key, Value }, Tag);
+	return MakeStructObservationFromArrayViews({ TEXT("Key"), TEXT("Value") }, { Key, Value }, Name);
 }
 
-FLearningAgentsObservationObjectElement ULearningAgentsObservations::MakeArrayObservation(ULearningAgentsObservationObject* Object, const TArray<FLearningAgentsObservationObjectElement>& Elements, const FName Tag)
+FLearningAgentsObservationObjectElement ULearningAgentsObservationObject::MakeArrayObservation(const TArray<FLearningAgentsObservationObjectElement>& Elements, const FName Name)
 {
-	return MakeArrayObservationFromArrayView(Object, Elements, Tag);
+	return MakeArrayObservationFromArrayView(Elements, Name);
 }
 
-FLearningAgentsObservationObjectElement ULearningAgentsObservations::MakeArrayObservationFromArrayView(ULearningAgentsObservationObject* Object, const TArrayView<const FLearningAgentsObservationObjectElement> Elements, const FName Tag)
+FLearningAgentsObservationObjectElement ULearningAgentsObservationObject::MakeArrayObservationFromArrayView(const TArrayView<const FLearningAgentsObservationObjectElement> Elements, const FName Name)
 {
-	if (!Object)
-	{
-		UE_LOG(LogLearning, Error, TEXT("MakeArrayObservationFromArrayView: Object is nullptr."));
-		return FLearningAgentsObservationObjectElement();
-	}
-
 	TArray<FLearningAgentsObservationObjectElement, TInlineAllocator<16>> SubElements;
 	SubElements.Empty(Elements.Num());
 
@@ -1463,59 +1375,47 @@ FLearningAgentsObservationObjectElement ULearningAgentsObservations::MakeArrayOb
 	{
 		const FLearningAgentsObservationObjectElement Element = Elements[ElementIdx];
 
-		if (!Object->ObservationObject.IsValid(Element.ObjectElement))
+		if (!ObservationObject.IsValid(Element.ObjectElement))
 		{
-			UE_LOG(LogLearning, Error, TEXT("MakeArrayObservationFromArrayView: Invalid Observation Object."));
+			UE_LOG(LogLearning, Error, TEXT("%s: Invalid Observation Object."), *GetName());
 			return FLearningAgentsObservationObjectElement();
 		}
 
-		SubElements.Add(MakePairObservation(Object, MakeIndexObservation(Object, ElementIdx, Elements.Num()), Element));
+		SubElements.Add(MakePairObservation(MakeIndexObservation(ElementIdx, Elements.Num()), Element));
 	}
 
-	return MakeSetObservationFromArrayView(Object, SubElements);
+	return MakeSetObservationFromArrayView(SubElements);
 }
 
-FLearningAgentsObservationObjectElement ULearningAgentsObservations::MakeMapObservation(ULearningAgentsObservationObject* Object, const TMap<FLearningAgentsObservationObjectElement, FLearningAgentsObservationObjectElement>& Map, const FName Tag)
+FLearningAgentsObservationObjectElement ULearningAgentsObservationObject::MakeMapObservation(const TMap<FLearningAgentsObservationObjectElement, FLearningAgentsObservationObjectElement>& Map, const FName Name)
 {
-	if (!Object)
-	{
-		UE_LOG(LogLearning, Error, TEXT("MakeMapObservation: Object is nullptr."));
-		return FLearningAgentsObservationObjectElement();
-	}
-
 	TArray<FLearningAgentsObservationObjectElement, TInlineAllocator<16>> SubElements;
 	SubElements.Empty(Map.Num());
 
 	for (TPair<FLearningAgentsObservationObjectElement, FLearningAgentsObservationObjectElement> Item : Map)
 	{
-		if (!Object->ObservationObject.IsValid(Item.Key.ObjectElement) || !Object->ObservationObject.IsValid(Item.Value.ObjectElement))
+		if (!ObservationObject.IsValid(Item.Key.ObjectElement) || !ObservationObject.IsValid(Item.Value.ObjectElement))
 		{
-			UE_LOG(LogLearning, Error, TEXT("MakeMapObservation: Invalid Observation Object."));
+			UE_LOG(LogLearning, Error, TEXT("%s: Invalid Observation Object."), *GetName());
 			return FLearningAgentsObservationObjectElement();
 		}
 
-		SubElements.Add(MakePairObservation(Object, Item.Key, Item.Value));
+		SubElements.Add(MakePairObservation(Item.Key, Item.Value));
 	}
 
-	return MakeSetObservationFromArrayView(Object, SubElements);
+	return MakeSetObservationFromArrayView(SubElements);
 }
 
-FLearningAgentsObservationObjectElement ULearningAgentsObservations::MakeMapObservationFromArrays(ULearningAgentsObservationObject* Object, const TArray<FLearningAgentsObservationObjectElement>& Keys, const TArray<FLearningAgentsObservationObjectElement>& Values, const FName Tag)
+FLearningAgentsObservationObjectElement ULearningAgentsObservationObject::MakeMapObservationFromArrays(const TArray<FLearningAgentsObservationObjectElement>& Keys, const TArray<FLearningAgentsObservationObjectElement>& Values, const FName Name)
 {
-	return MakeMapObservationFromArrayViews(Object, Keys, Values, Tag);
+	return MakeMapObservationFromArrayViews(Keys, Values, Name);
 }
 
-FLearningAgentsObservationObjectElement ULearningAgentsObservations::MakeMapObservationFromArrayViews(ULearningAgentsObservationObject* Object, const TArrayView<const FLearningAgentsObservationObjectElement> Keys, const TArrayView<const FLearningAgentsObservationObjectElement> Values, const FName Tag)
+FLearningAgentsObservationObjectElement ULearningAgentsObservationObject::MakeMapObservationFromArrayViews(const TArrayView<const FLearningAgentsObservationObjectElement> Keys, const TArrayView<const FLearningAgentsObservationObjectElement> Values, const FName Name)
 {
-	if (!Object)
-	{
-		UE_LOG(LogLearning, Error, TEXT("MakeMapObservationFromArrayViews: Object is nullptr."));
-		return FLearningAgentsObservationObjectElement();
-	}
-
 	if (Keys.Num() != Values.Num())
 	{
-		UE_LOG(LogLearning, Error, TEXT("MakeMapObservationFromArrayViews: Number of keys (%i) must match number of values (%i)."), Keys.Num(), Values.Num());
+		UE_LOG(LogLearning, Error, TEXT("%s: Number of keys (%i) must match number of values (%i)."), *GetName(), Keys.Num(), Values.Num());
 		return FLearningAgentsObservationObjectElement();
 	}
 
@@ -1524,24 +1424,23 @@ FLearningAgentsObservationObjectElement ULearningAgentsObservations::MakeMapObse
 
 	for (int32 ElementIdx = 0; ElementIdx < Keys.Num(); ElementIdx++)
 	{
-		if (!Object->ObservationObject.IsValid(Keys[ElementIdx].ObjectElement) || 
-			!Object->ObservationObject.IsValid(Values[ElementIdx].ObjectElement))
+		if(!ObservationObject.IsValid(Keys[ElementIdx].ObjectElement) || !ObservationObject.IsValid(Values[ElementIdx].ObjectElement))
 		{
-			UE_LOG(LogLearning, Error, TEXT("MakeMapObservationFromArrayViews: Invalid Observation Object."));
+			UE_LOG(LogLearning, Error, TEXT("%s: Invalid Observation Object."), *GetName());
 			return FLearningAgentsObservationObjectElement();
 		}
 
-		SubElements.Add(MakePairObservation(Object, Keys[ElementIdx], Values[ElementIdx]));
+		SubElements.Add(MakePairObservation(Keys[ElementIdx], Values[ElementIdx]));
 	}
 
-	return MakeSetObservationFromArrayView(Object, SubElements);
+	return MakeSetObservationFromArrayView(SubElements);
 }
 
-FLearningAgentsObservationObjectElement ULearningAgentsObservations::MakeEnumObservation(ULearningAgentsObservationObject* Object, const UEnum* Enum, const uint8 EnumValue, const FName Tag)
+FLearningAgentsObservationObjectElement ULearningAgentsObservationObject::MakeEnumObservation(const UEnum* Enum, const uint8 EnumValue, const FName Name)
 {
 	if (!Enum)
 	{
-		UE_LOG(LogLearning, Error, TEXT("MakeEnumObservation: Enum is nullptr."));
+		UE_LOG(LogLearning, Error, TEXT("%s: Enum is nullptr."), *GetName());
 		return FLearningAgentsObservationObjectElement();
 	}
 
@@ -1549,7 +1448,7 @@ FLearningAgentsObservationObjectElement ULearningAgentsObservations::MakeEnumObs
 
 	if (EnumValueIndex == INDEX_NONE || EnumValueIndex < 0 || EnumValueIndex >= Enum->NumEnums() - 1)
 	{
-		UE_LOG(LogLearning, Error, TEXT("MakeEnumObservation: EnumValue %i not valid for Enum '%s'."), EnumValue , *Enum->GetName());
+		UE_LOG(LogLearning, Error, TEXT("%s: EnumValue %i not valid for Enum '%s'."), *GetName(), EnumValue , *Enum->GetName());
 		return FLearningAgentsObservationObjectElement();
 	}
 
@@ -1557,20 +1456,20 @@ FLearningAgentsObservationObjectElement ULearningAgentsObservations::MakeEnumObs
 	OneHot.Init(0.0f, Enum->NumEnums() - 1);
 	OneHot[EnumValueIndex] = 1.0f;
 
-	return MakeContinuousObservationFromArrayView(Object, OneHot, Tag);
+	return MakeContinuousObservationFromArrayView(OneHot, Name);
 }
 
-FLearningAgentsObservationObjectElement ULearningAgentsObservations::MakeBitmaskObservation(ULearningAgentsObservationObject* Object, const UEnum* Enum, const int32 BitmaskValue, const FName Tag)
+FLearningAgentsObservationObjectElement ULearningAgentsObservationObject::MakeBitmaskObservation(const UEnum* Enum, const int32 BitmaskValue, const FName Name)
 {
 	if (!Enum)
 	{
-		UE_LOG(LogLearning, Error, TEXT("MakeBitmaskObservation: Enum is nullptr."));
+		UE_LOG(LogLearning, Error, TEXT("%s: Enum is nullptr."), *GetName());
 		return FLearningAgentsObservationObjectElement();
 	}
 
 	if (Enum->NumEnums() - 1 > 32)
 	{
-		UE_LOG(LogLearning, Error, TEXT("MakeBitmaskObservation: Too many values in Enum to use as Bitmask (%i)."), Enum->NumEnums() - 1);
+		UE_LOG(LogLearning, Error, TEXT("%s: Too many values in Enum to use as Bitmask (%i)."), *GetName(), Enum->NumEnums() - 1);
 		return FLearningAgentsObservationObjectElement();
 	}
 
@@ -1585,64 +1484,56 @@ FLearningAgentsObservationObjectElement ULearningAgentsObservations::MakeBitmask
 		}
 	}
 
-	return MakeContinuousObservationFromArrayView(Object, OneHot, Tag);
+	return MakeContinuousObservationFromArrayView(OneHot, Name);
 }
 
-FLearningAgentsObservationObjectElement ULearningAgentsObservations::MakeOptionalObservation(ULearningAgentsObservationObject* Object, const FLearningAgentsObservationObjectElement Element, const ELearningAgentsOptionalObservation Option, const FName Tag)
+FLearningAgentsObservationObjectElement ULearningAgentsObservationObject::MakeOptionalObservation(const FLearningAgentsObservationObjectElement Element, const ELearningAgentsOptionalObservation Option, const FName Name)
 {
 	return MakeExclusiveUnionObservation(
-		Object,
 		Option == ELearningAgentsOptionalObservation::Null ? TEXT("Null") : TEXT("Valid"),
-		Option == ELearningAgentsOptionalObservation::Null ? MakeNullObservation(Object) : Element,
-		Tag);
+		Option == ELearningAgentsOptionalObservation::Null ? MakeNullObservation() : Element,
+		Name);
 }
 
-FLearningAgentsObservationObjectElement ULearningAgentsObservations::MakeOptionalNullObservation(ULearningAgentsObservationObject* Object, const FName Tag)
+FLearningAgentsObservationObjectElement ULearningAgentsObservationObject::MakeOptionalNullObservation(const FName Name)
 {
-	return MakeExclusiveUnionObservation(Object, TEXT("Null"), MakeNullObservation(Object), Tag);
+	return MakeExclusiveUnionObservation(TEXT("Null"), MakeNullObservation(), Name);
 }
 
-FLearningAgentsObservationObjectElement ULearningAgentsObservations::MakeOptionalValidObservation(ULearningAgentsObservationObject* Object, const FLearningAgentsObservationObjectElement Element, const FName Tag)
+FLearningAgentsObservationObjectElement ULearningAgentsObservationObject::MakeOptionalValidObservation(const FLearningAgentsObservationObjectElement Element, const FName Name)
 {
-	return MakeExclusiveUnionObservation(Object, TEXT("Valid"), Element, Tag);
+	return MakeExclusiveUnionObservation(TEXT("Valid"), Element, Name);
 }
 
-FLearningAgentsObservationObjectElement ULearningAgentsObservations::MakeEitherObservation(ULearningAgentsObservationObject* Object, const FLearningAgentsObservationObjectElement Element, const ELearningAgentsEitherObservation Either, const FName Tag)
+FLearningAgentsObservationObjectElement ULearningAgentsObservationObject::MakeEitherObservation(const FLearningAgentsObservationObjectElement Element, const ELearningAgentsEitherObservation Either, const FName Name)
 {
-	return MakeExclusiveUnionObservation(Object, Either == ELearningAgentsEitherObservation::A ? TEXT("A") : TEXT("B"), Element, Tag);
+	return MakeExclusiveUnionObservation(Either == ELearningAgentsEitherObservation::A ? TEXT("A") : TEXT("B"), Element, Name);
 }
 
-FLearningAgentsObservationObjectElement ULearningAgentsObservations::MakeEitherAObservation(ULearningAgentsObservationObject* Object, const FLearningAgentsObservationObjectElement A, const FName Tag)
+FLearningAgentsObservationObjectElement ULearningAgentsObservationObject::MakeEitherAObservation(const FLearningAgentsObservationObjectElement A, const FName Name)
 {
-	return MakeExclusiveUnionObservation(Object, TEXT("A"), A, Tag);
+	return MakeExclusiveUnionObservation(TEXT("A"), A, Name);
 }
 
-FLearningAgentsObservationObjectElement ULearningAgentsObservations::MakeEitherBObservation(ULearningAgentsObservationObject* Object, const FLearningAgentsObservationObjectElement B, const FName Tag)
+FLearningAgentsObservationObjectElement ULearningAgentsObservationObject::MakeEitherBObservation(const FLearningAgentsObservationObjectElement B, const FName Name)
 {
-	return MakeExclusiveUnionObservation(Object, TEXT("B"), B, Tag);
+	return MakeExclusiveUnionObservation(TEXT("B"), B, Name);
 }
 
-FLearningAgentsObservationObjectElement ULearningAgentsObservations::MakeEncodingObservation(ULearningAgentsObservationObject* Object, const FLearningAgentsObservationObjectElement Element, const FName Tag)
+FLearningAgentsObservationObjectElement ULearningAgentsObservationObject::MakeEncodingObservation(const FLearningAgentsObservationObjectElement Element, const FName Name)
 {
-	if (!Object)
-	{
-		UE_LOG(LogLearning, Error, TEXT("MakeEncodingObservation: Object is nullptr."));
-		return FLearningAgentsObservationObjectElement();
-	}
-
-	if (!Object->ObservationObject.IsValid(Element.ObjectElement))
+	if (!ObservationObject.IsValid(Element.ObjectElement))
 	{
-		UE_LOG(LogLearning, Error, TEXT("MakeEncodingObservation: Invalid Observation Object."));
+		UE_LOG(LogLearning, Error, TEXT("%s: Invalid Observation Object."), *GetName());
 		return FLearningAgentsObservationObjectElement();
 	}
 
-	return { Object->ObservationObject.CreateEncoding({ Element.ObjectElement }, Tag) };
+	return { ObservationObject.CreateEncoding({ Element.ObjectElement }, Name) };
 }
 
-FLearningAgentsObservationObjectElement ULearningAgentsObservations::MakeBoolObservation(
-	ULearningAgentsObservationObject* Object,
+FLearningAgentsObservationObjectElement ULearningAgentsObservationObject::MakeBoolObservation(
 	const bool bValue,
-	const FName Tag,
+	const FName Name,
 	const bool bVisualLoggerEnabled,
 	ULearningAgentsManagerListener* VisualLoggerListener,
 	const int32 VisualLoggerAgentId,
@@ -1654,27 +1545,26 @@ FLearningAgentsObservationObjectElement ULearningAgentsObservations::MakeBoolObs
 #if UE_LEARNING_AGENTS_ENABLE_VISUAL_LOG
 	if (bVisualLoggerEnabled && VisualLoggerListener)
 	{
-		const ULearningAgentsVisualLoggerObject* VisualLoggerObject = VisualLoggerListener->GetOrAddVisualLoggerObject(Tag);
+		const ULearningAgentsVisualLoggerObject* VisualLoggerObject = VisualLoggerListener->GetOrAddVisualLoggerObject(Name);
 
 		UE_LEARNING_AGENTS_VLOG_STRING(VisualLoggerObject, LogLearning, Display, VisualLoggerLocation,
 			VisualLoggerColor.ToFColor(true),
-			TEXT("Listener: %s\nTag: %s\nAgent Id: % 3i\nValue: [%s]\nEncoded: [% 6.2f]"),
+			TEXT("Listener: %s\nName: %s\nAgent Id: % 3i\nValue: [%s]\nEncoded: [% 6.2f]"),
 			*VisualLoggerListener->GetName(),
-			*Tag.ToString(),
+			*Name.ToString(),
 			VisualLoggerAgentId,
 			bValue ? TEXT("true") : TEXT("false"),
 			EncodedBool);
 	}
 #endif
 
-	return MakeContinuousObservationFromArrayView(Object, { EncodedBool }, Tag);
+	return MakeContinuousObservationFromArrayView({ EncodedBool }, Name);
 }
 
-FLearningAgentsObservationObjectElement ULearningAgentsObservations::MakeFloatObservation(
-	ULearningAgentsObservationObject* Object,
+FLearningAgentsObservationObjectElement ULearningAgentsObservationObject::MakeFloatObservation(
 	const float Value,
 	const float FloatScale,
-	const FName Tag,
+	const FName Name,
 	const bool bVisualLoggerEnabled,
 	ULearningAgentsManagerListener* VisualLoggerListener,
 	const int32 VisualLoggerAgentId,
@@ -1686,13 +1576,13 @@ FLearningAgentsObservationObjectElement ULearningAgentsObservations::MakeFloatOb
 #if UE_LEARNING_AGENTS_ENABLE_VISUAL_LOG
 	if (bVisualLoggerEnabled && VisualLoggerListener)
 	{
-		const ULearningAgentsVisualLoggerObject* VisualLoggerObject = VisualLoggerListener->GetOrAddVisualLoggerObject(Tag);
+		const ULearningAgentsVisualLoggerObject* VisualLoggerObject = VisualLoggerListener->GetOrAddVisualLoggerObject(Name);
 
 		UE_LEARNING_AGENTS_VLOG_STRING(VisualLoggerObject, LogLearning, Display, VisualLoggerLocation,
 			VisualLoggerColor.ToFColor(true),
-			TEXT("Listener: %s\nTag: %s\nAgent Id: % 3i\nValue: [% 6.1f]\nScale: [% 6.2f]\nEncoded: [% 6.2f]"),
+			TEXT("Listener: %s\nName: %s\nAgent Id: % 3i\nValue: [% 6.1f]\nScale: [% 6.2f]\nEncoded: [% 6.2f]"),
 			*VisualLoggerListener->GetName(),
-			*Tag.ToString(),
+			*Name.ToString(),
 			VisualLoggerAgentId,
 			Value,
 			FloatScale,
@@ -1700,15 +1590,14 @@ FLearningAgentsObservationObjectElement ULearningAgentsObservations::MakeFloatOb
 	}
 #endif
 
-	return MakeContinuousObservationFromArrayView(Object, { EncodedValue }, Tag);
+	return MakeContinuousObservationFromArrayView({ EncodedValue }, Name);
 }
 
-FLearningAgentsObservationObjectElement ULearningAgentsObservations::MakeLocationObservation(
-	ULearningAgentsObservationObject* Object,
+FLearningAgentsObservationObjectElement ULearningAgentsObservationObject::MakeLocationObservation(
 	const FVector Location,
 	const FTransform RelativeTransform,
 	const float LocationScale,
-	const FName Tag,
+	const FName Name,
 	const bool bVisualLoggerEnabled,
 	ULearningAgentsManagerListener* VisualLoggerListener,
 	const int32 VisualLoggerAgentId,
@@ -1724,7 +1613,7 @@ FLearningAgentsObservationObjectElement ULearningAgentsObservations::MakeLocatio
 #if UE_LEARNING_AGENTS_ENABLE_VISUAL_LOG
 	if (bVisualLoggerEnabled && VisualLoggerListener)
 	{
-		const ULearningAgentsVisualLoggerObject* VisualLoggerObject = VisualLoggerListener->GetOrAddVisualLoggerObject(Tag);
+		const ULearningAgentsVisualLoggerObject* VisualLoggerObject = VisualLoggerListener->GetOrAddVisualLoggerObject(Name);
 
 		UE_LEARNING_AGENTS_VLOG_LOCATION(VisualLoggerObject, LogLearning, Display,
 			Location,
@@ -1746,9 +1635,9 @@ FLearningAgentsObservationObjectElement ULearningAgentsObservations::MakeLocatio
 
 		UE_LEARNING_AGENTS_VLOG_STRING(VisualLoggerObject, LogLearning, Display, VisualLoggerLocation,
 			VisualLoggerColor.ToFColor(true),
-			TEXT("Listener: %s\nTag: %s\nAgent Id: % 3i\nLocation: [% 6.1f % 6.1f % 6.1f]\nLocal Location: [% 6.1f % 6.1f % 6.1f]\nScale: [% 6.2f]\nEncoded: [% 6.2f % 6.2f % 6.2f]"),
+			TEXT("Listener: %s\nName: %s\nAgent Id: % 3i\nLocation: [% 6.1f % 6.1f % 6.1f]\nLocal Location: [% 6.1f % 6.1f % 6.1f]\nScale: [% 6.2f]\nEncoded: [% 6.2f % 6.2f % 6.2f]"),
 			*VisualLoggerListener->GetName(),
-			*Tag.ToString(),
+			*Name.ToString(),
 			VisualLoggerAgentId,
 			Location.X, Location.Y, Location.Z,
 			LocalLocation.X, LocalLocation.Y, LocalLocation.Z,
@@ -1757,18 +1646,17 @@ FLearningAgentsObservationObjectElement ULearningAgentsObservations::MakeLocatio
 	}
 #endif
 
-	return MakeContinuousObservationFromArrayView(Object, {
+	return MakeContinuousObservationFromArrayView({
 		(float)EncodedLocation.X,
 		(float)EncodedLocation.Y,
 		(float)EncodedLocation.Z,
-		}, Tag);
+		}, Name);
 }
 
-FLearningAgentsObservationObjectElement ULearningAgentsObservations::MakeRotationObservation(
-	ULearningAgentsObservationObject* Object,
+FLearningAgentsObservationObjectElement ULearningAgentsObservationObject::MakeRotationObservation(
 	const FRotator Rotation,
 	const FRotator RelativeRotation,
-	const FName Tag,
+	const FName Name,
 	const bool bVisualLoggerEnabled,
 	ULearningAgentsManagerListener* VisualLoggerListener,
 	const int32 VisualLoggerAgentId,
@@ -1779,16 +1667,14 @@ FLearningAgentsObservationObjectElement ULearningAgentsObservations::MakeRotatio
 	// Visual logging is handled in the MakeRotationObservationFromQuat
 
 	return MakeRotationObservationFromQuat(
-		Object,
-		FQuat::MakeFromRotator(Rotation), FQuat::MakeFromRotator(RelativeRotation), Tag,
+		FQuat::MakeFromRotator(Rotation), FQuat::MakeFromRotator(RelativeRotation), Name,
 		bVisualLoggerEnabled, VisualLoggerListener, VisualLoggerAgentId, VisualLoggerRotationLocation, VisualLoggerLocation, VisualLoggerColor);
 }
 
-FLearningAgentsObservationObjectElement ULearningAgentsObservations::MakeRotationObservationFromQuat(
-	ULearningAgentsObservationObject* Object,
+FLearningAgentsObservationObjectElement ULearningAgentsObservationObject::MakeRotationObservationFromQuat(
 	const FQuat Rotation,
 	const FQuat RelativeRotation,
-	const FName Tag,
+	const FName Name,
 	const bool bVisualLoggerEnabled,
 	ULearningAgentsManagerListener* VisualLoggerListener,
 	const int32 VisualLoggerAgentId,
@@ -1803,7 +1689,7 @@ FLearningAgentsObservationObjectElement ULearningAgentsObservations::MakeRotatio
 #if UE_LEARNING_AGENTS_ENABLE_VISUAL_LOG
 	if (bVisualLoggerEnabled && VisualLoggerListener)
 	{
-		const ULearningAgentsVisualLoggerObject* VisualLoggerObject = VisualLoggerListener->GetOrAddVisualLoggerObject(Tag);
+		const ULearningAgentsVisualLoggerObject* VisualLoggerObject = VisualLoggerListener->GetOrAddVisualLoggerObject(Name);
 
 		UE_LEARNING_AGENTS_VLOG_TRANSFORM(VisualLoggerObject, LogLearning, Display,
 			VisualLoggerRotationLocation,
@@ -1813,9 +1699,9 @@ FLearningAgentsObservationObjectElement ULearningAgentsObservations::MakeRotatio
 
 		UE_LEARNING_AGENTS_VLOG_STRING(VisualLoggerObject, LogLearning, Display, VisualLoggerLocation,
 			VisualLoggerColor.ToFColor(true),
-			TEXT("Listener: %s\nTag: %s\nAgent Id: % 3i\nRotation: [% 6.1f % 6.1f % 6.1f % 6.1f]\nLocal Rotation: [% 6.1f % 6.1f % 6.1f % 6.1f]\nEncoded Forward: [% 6.2f % 6.2f % 6.2f]\nEncoded Right: [% 6.2f % 6.2f % 6.2f]"),
+			TEXT("Listener: %s\nName: %s\nAgent Id: % 3i\nRotation: [% 6.1f % 6.1f % 6.1f % 6.1f]\nLocal Rotation: [% 6.1f % 6.1f % 6.1f % 6.1f]\nEncoded Forward: [% 6.2f % 6.2f % 6.2f]\nEncoded Right: [% 6.2f % 6.2f % 6.2f]"),
 			*VisualLoggerListener->GetName(),
-			*Tag.ToString(),
+			*Name.ToString(),
 			VisualLoggerAgentId,
 			Rotation.W, Rotation.X, Rotation.Y, Rotation.Z,
 			LocalRotation.W, LocalRotation.X, LocalRotation.Y, LocalRotation.Z,
@@ -1824,21 +1710,20 @@ FLearningAgentsObservationObjectElement ULearningAgentsObservations::MakeRotatio
 	}
 #endif
 
-	return MakeContinuousObservationFromArrayView(Object, {
+	return MakeContinuousObservationFromArrayView({
 		(float)LocalAxisForward.X,
 		(float)LocalAxisForward.Y,
 		(float)LocalAxisForward.Z,
 		(float)LocalAxisRight.X,
 		(float)LocalAxisRight.Y,
 		(float)LocalAxisRight.Z,
-		}, Tag);
+		}, Name);
 }
 
-FLearningAgentsObservationObjectElement ULearningAgentsObservations::MakeScaleObservation(
-	ULearningAgentsObservationObject* Object,
+FLearningAgentsObservationObjectElement ULearningAgentsObservationObject::MakeScaleObservation(
 	const FVector Scale,
 	const FVector RelativeScale,
-	const FName Tag,
+	const FName Name,
 	const bool bVisualLoggerEnabled,
 	ULearningAgentsManagerListener* VisualLoggerListener,
 	const int32 VisualLoggerAgentId,
@@ -1853,32 +1738,31 @@ FLearningAgentsObservationObjectElement ULearningAgentsObservations::MakeScaleOb
 #if UE_LEARNING_AGENTS_ENABLE_VISUAL_LOG
 	if (bVisualLoggerEnabled && VisualLoggerListener)
 	{
-		const ULearningAgentsVisualLoggerObject* VisualLoggerObject = VisualLoggerListener->GetOrAddVisualLoggerObject(Tag);
+		const ULearningAgentsVisualLoggerObject* VisualLoggerObject = VisualLoggerListener->GetOrAddVisualLoggerObject(Name);
 
 		UE_LEARNING_AGENTS_VLOG_STRING(VisualLoggerObject, LogLearning, Display, VisualLoggerLocation,
 			VisualLoggerColor.ToFColor(true),
-			TEXT("Listener: %s\nTag: %s\nAgent Id: % 3i\nScale: [% 6.1f % 6.1f % 6.1f]\nEncoded: [% 6.2f % 6.2f % 6.2f]"),
+			TEXT("Listener: %s\nName: %s\nAgent Id: % 3i\nScale: [% 6.1f % 6.1f % 6.1f]\nEncoded: [% 6.2f % 6.2f % 6.2f]"),
 			*VisualLoggerListener->GetName(),
-			*Tag.ToString(),
+			*Name.ToString(),
 			VisualLoggerAgentId,
 			Scale.X, Scale.Y, Scale.Z,
 			LocalLogScale.X, LocalLogScale.Y, LocalLogScale.Z);
 	}
 #endif
 
-	return MakeContinuousObservationFromArrayView(Object, {
+	return MakeContinuousObservationFromArrayView({
 		(float)LocalLogScale.X,
 		(float)LocalLogScale.Y,
 		(float)LocalLogScale.Z,
-		}, Tag);
+		}, Name);
 }
 
-FLearningAgentsObservationObjectElement ULearningAgentsObservations::MakeTransformObservation(
-	ULearningAgentsObservationObject* Object,
+FLearningAgentsObservationObjectElement ULearningAgentsObservationObject::MakeTransformObservation(
 	const FTransform Transform,
 	const FTransform RelativeTransform,
 	const float LocationScale,
-	const FName Tag,
+	const FName Name,
 	const bool bVisualLoggerEnabled,
 	ULearningAgentsManagerListener* VisualLoggerListener,
 	const int32 VisualLoggerAgentId,
@@ -1890,7 +1774,7 @@ FLearningAgentsObservationObjectElement ULearningAgentsObservations::MakeTransfo
 #if UE_LEARNING_AGENTS_ENABLE_VISUAL_LOG
 	if (bVisualLoggerEnabled && VisualLoggerListener)
 	{
-		const ULearningAgentsVisualLoggerObject* VisualLoggerObject = VisualLoggerListener->GetOrAddVisualLoggerObject(Tag);
+		const ULearningAgentsVisualLoggerObject* VisualLoggerObject = VisualLoggerListener->GetOrAddVisualLoggerObject(Name);
 
 		const FVector LocalLocation = LocalTransform.GetLocation();
 		const FRotator LocalRotation = LocalTransform.Rotator();
@@ -1908,9 +1792,9 @@ FLearningAgentsObservationObjectElement ULearningAgentsObservations::MakeTransfo
 
 		UE_LEARNING_AGENTS_VLOG_STRING(VisualLoggerObject, LogLearning, Display, VisualLoggerLocation,
 			VisualLoggerColor.ToFColor(true),
-			TEXT("Listener: %s\nTag: %s\nAgent Id: % 3i\nLocation: [% 6.1f % 6.1f % 6.1f]\nLocal Location: [% 6.1f % 6.1f % 6.1f]\nRotation: [% 6.1f % 6.1f % 6.1f]\nLocal Rotation: [% 6.1f % 6.1f % 6.1f]\nScale: [% 6.1f % 6.1f % 6.1f]\nLocal Scale: [% 6.1f % 6.1f % 6.1f]\nLocation Scale: [% 6.1f]"),
+			TEXT("Listener: %s\nName: %s\nAgent Id: % 3i\nLocation: [% 6.1f % 6.1f % 6.1f]\nLocal Location: [% 6.1f % 6.1f % 6.1f]\nRotation: [% 6.1f % 6.1f % 6.1f]\nLocal Rotation: [% 6.1f % 6.1f % 6.1f]\nScale: [% 6.1f % 6.1f % 6.1f]\nLocal Scale: [% 6.1f % 6.1f % 6.1f]\nLocation Scale: [% 6.1f]"),
 			*VisualLoggerListener->GetName(),
-			*Tag.ToString(),
+			*Name.ToString(),
 			VisualLoggerAgentId,
 			Location.X, Location.Y, Location.Z,
 			LocalLocation.X, LocalLocation.Y, LocalLocation.Z,
@@ -1922,51 +1806,50 @@ FLearningAgentsObservationObjectElement ULearningAgentsObservations::MakeTransfo
 	}
 #endif
 
-	return MakeStructObservationFromArrayViews(Object,
+	return MakeStructObservationFromArrayViews(
 		{
 			TEXT("Location"),
 			TEXT("Rotation"),
 			TEXT("Scale")
 		},
 		{
-			MakeLocationObservation(Object, LocalTransform.GetLocation(), FTransform::Identity, LocationScale),
-			MakeRotationObservationFromQuat(Object, LocalTransform.GetRotation(), FQuat::Identity),
-			MakeScaleObservation(Object, LocalTransform.GetScale3D(), FVector::OneVector)
+			MakeLocationObservation(LocalTransform.GetLocation(), FTransform::Identity, LocationScale),
+			MakeRotationObservationFromQuat(LocalTransform.GetRotation(), FQuat::Identity),
+			MakeScaleObservation(LocalTransform.GetScale3D(), FVector::OneVector)
 		},
-		Tag);
+		Name);
 }
 
-FLearningAgentsObservationObjectElement ULearningAgentsObservations::MakeAngleObservation(ULearningAgentsObservationObject* Object, const float Angle, const float RelativeAngle, const FName Tag)
+FLearningAgentsObservationObjectElement ULearningAgentsObservationObject::MakeAngleObservation(const float Angle, const float RelativeAngle, const FName Name)
 {
-	return MakeAngleObservationRadians(Object, FMath::DegreesToRadians(Angle), FMath::DegreesToRadians(RelativeAngle), Tag);
+	return MakeAngleObservationRadians(FMath::DegreesToRadians(Angle), FMath::DegreesToRadians(RelativeAngle), Name);
 }
 
-FLearningAgentsObservationObjectElement ULearningAgentsObservations::MakeAngleObservationRadians(ULearningAgentsObservationObject* Object, const float Angle, const float RelativeAngle, const FName Tag)
+FLearningAgentsObservationObjectElement ULearningAgentsObservationObject::MakeAngleObservationRadians(const float Angle, const float RelativeAngle, const FName Name)
 {
 	const float LocalAngle = FMath::FindDeltaAngleRadians(RelativeAngle, Angle);
 
-	return MakeContinuousObservationFromArrayView(Object, {
+	return MakeContinuousObservationFromArrayView({
 		(float)FMath::Sin(Angle),
 		(float)FMath::Cos(Angle),
-		}, Tag);
+		}, Name);
 }
 
-FLearningAgentsObservationObjectElement ULearningAgentsObservations::MakeVelocityObservation(ULearningAgentsObservationObject* Object, const FVector Velocity, const FTransform RelativeTransform, const float VelocityScale, const FName Tag)
+FLearningAgentsObservationObjectElement ULearningAgentsObservationObject::MakeVelocityObservation(const FVector Velocity, const FTransform RelativeTransform, const float VelocityScale, const FName Name)
 {
 	const FVector LocalVelocity = RelativeTransform.InverseTransformVectorNoScale(Velocity);
 
-	return MakeContinuousObservationFromArrayView(Object, {
+	return MakeContinuousObservationFromArrayView({
 		(float)LocalVelocity.X / FMath::Max(VelocityScale, UE_SMALL_NUMBER),
 		(float)LocalVelocity.Y / FMath::Max(VelocityScale, UE_SMALL_NUMBER),
 		(float)LocalVelocity.Z / FMath::Max(VelocityScale, UE_SMALL_NUMBER),
-		}, Tag);
+		}, Name);
 }
 
-FLearningAgentsObservationObjectElement ULearningAgentsObservations::MakeDirectionObservation(
-	ULearningAgentsObservationObject* Object,
+FLearningAgentsObservationObjectElement ULearningAgentsObservationObject::MakeDirectionObservation(
 	const FVector Direction,
 	const FTransform RelativeTransform,
-	const FName Tag,
+	const FName Name,
 	const bool bVisualLoggerEnabled,
 	ULearningAgentsManagerListener* VisualLoggerListener,
 	const int32 VisualLoggerAgentId,
@@ -1980,7 +1863,7 @@ FLearningAgentsObservationObjectElement ULearningAgentsObservations::MakeDirecti
 #if UE_LEARNING_AGENTS_ENABLE_VISUAL_LOG
 	if (bVisualLoggerEnabled && VisualLoggerListener)
 	{
-		const ULearningAgentsVisualLoggerObject* VisualLoggerObject = VisualLoggerListener->GetOrAddVisualLoggerObject(Tag);
+		const ULearningAgentsVisualLoggerObject* VisualLoggerObject = VisualLoggerListener->GetOrAddVisualLoggerObject(Name);
 
 		UE_LEARNING_AGENTS_VLOG_ARROW(VisualLoggerObject, LogLearning, Display,
 			VisualLoggerDirectionLocation,
@@ -1996,9 +1879,9 @@ FLearningAgentsObservationObjectElement ULearningAgentsObservations::MakeDirecti
 
 		UE_LEARNING_AGENTS_VLOG_STRING(VisualLoggerObject, LogLearning, Display, VisualLoggerLocation,
 			VisualLoggerColor.ToFColor(true),
-			TEXT("Listener: %s\nTag: %s\nAgent Id: % 3i\nDirection: [% 6.1f % 6.1f % 6.1f]\nLocal Direction: [% 6.1f % 6.1f % 6.1f]\nEncoded: [% 6.2f % 6.2f % 6.2f]"),
+			TEXT("Listener: %s\nName: %s\nAgent Id: % 3i\nDirection: [% 6.1f % 6.1f % 6.1f]\nLocal Direction: [% 6.1f % 6.1f % 6.1f]\nEncoded: [% 6.2f % 6.2f % 6.2f]"),
 			*VisualLoggerListener->GetName(),
-			*Tag.ToString(),
+			*Name.ToString(),
 			VisualLoggerAgentId,
 			Direction.X, Direction.Y, Direction.Z,
 			LocalDirection.X, LocalDirection.Y, LocalDirection.Z,
@@ -2006,20 +1889,19 @@ FLearningAgentsObservationObjectElement ULearningAgentsObservations::MakeDirecti
 	}
 #endif
 
-	return MakeContinuousObservationFromArrayView(Object, {
+	return MakeContinuousObservationFromArrayView({
 		(float)LocalDirection.X,
 		(float)LocalDirection.Y,
 		(float)LocalDirection.Z,
-		}, Tag);
+		}, Name);
 }
 
-FLearningAgentsObservationObjectElement ULearningAgentsObservations::MakeLocationAlongSplineObservation(
-	ULearningAgentsObservationObject* Object,
+FLearningAgentsObservationObjectElement ULearningAgentsObservationObject::MakeLocationAlongSplineObservation(
 	const USplineComponent* SplineComponent,
 	const float DistanceAlongSpline,
 	const FTransform RelativeTransform,
 	const float LocationScale,
-	const FName Tag,
+	const FName Name,
 	const bool bVisualLoggerEnabled,
 	ULearningAgentsManagerListener* VisualLoggerListener,
 	const int32 VisualLoggerAgentId,
@@ -2032,10 +1914,10 @@ FLearningAgentsObservationObjectElement ULearningAgentsObservations::MakeLocatio
 		return FLearningAgentsObservationObjectElement();
 	}
 
-	return MakeLocationObservation(Object, SplineComponent->GetLocationAtDistanceAlongSpline(DistanceAlongSpline, ESplineCoordinateSpace::World), RelativeTransform, LocationScale, Tag, bVisualLoggerEnabled, VisualLoggerListener, VisualLoggerAgentId, VisualLoggerLocation, VisualLoggerColor);
+	return MakeLocationObservation(SplineComponent->GetLocationAtDistanceAlongSpline(DistanceAlongSpline, ESplineCoordinateSpace::World), RelativeTransform, LocationScale, Name, bVisualLoggerEnabled, VisualLoggerListener, VisualLoggerAgentId, VisualLoggerLocation, VisualLoggerColor);
 }
 
-FLearningAgentsObservationObjectElement ULearningAgentsObservations::MakeProportionAlongSplineObservation(ULearningAgentsObservationObject* Object, const USplineComponent* SplineComponent, const float DistanceAlongSpline, const FName Tag)
+FLearningAgentsObservationObjectElement ULearningAgentsObservationObject::MakeProportionAlongSplineObservation(const USplineComponent* SplineComponent, const float DistanceAlongSpline, const FName Name)
 {
 	if (!SplineComponent)
 	{
@@ -2049,21 +1931,20 @@ FLearningAgentsObservationObjectElement ULearningAgentsObservations::MakeProport
 		const float WrapDistance = FMath::Wrap(DistanceAlongSpline, 0.0f, TotalDistance);
 		const float Proportion = WrapDistance / FMath::Max(TotalDistance, UE_SMALL_NUMBER);
 		const float Angle = FMath::Wrap(UE_TWO_PI * Proportion, -UE_PI, UE_PI);
-		return MakeExclusiveUnionObservation(Object, TEXT("Angle"), MakeAngleObservation(Object, Angle), Tag);
+		return MakeExclusiveUnionObservation(TEXT("Angle"), MakeAngleObservation(Angle), Name);
 	}
 	else
 	{
 		const float Proportion = FMath::Clamp(DistanceAlongSpline / FMath::Max(SplineComponent->GetSplineLength(), UE_SMALL_NUMBER), 0.0f, 1.0f);
-		return MakeExclusiveUnionObservation(Object, TEXT("Proportion"), MakeFloatObservation(Object, Proportion), Tag);
+		return MakeExclusiveUnionObservation(TEXT("Proportion"), MakeFloatObservation(Proportion), Name);
 	}
 }
 
-FLearningAgentsObservationObjectElement ULearningAgentsObservations::MakeDirectionAlongSplineObservation(
-	ULearningAgentsObservationObject* Object,
+FLearningAgentsObservationObjectElement ULearningAgentsObservationObject::MakeDirectionAlongSplineObservation(
 	const USplineComponent* SplineComponent,
 	const float DistanceAlongSpline,
 	const FTransform RelativeTransform,
-	const FName Tag,
+	const FName Name,
 	const bool bVisualLoggerEnabled,
 	ULearningAgentsManagerListener* VisualLoggerListener,
 	const int32 VisualLoggerAgentId,
@@ -2086,10 +1967,10 @@ FLearningAgentsObservationObjectElement ULearningAgentsObservations::MakeDirecti
 	const FVector VisualLoggerDirectionLocation = FVector::ZeroVector;
 #endif
 
-	return MakeDirectionObservation(Object, SplineComponent->GetDirectionAtDistanceAlongSpline(DistanceAlongSpline, ESplineCoordinateSpace::World), RelativeTransform, Tag, bVisualLoggerEnabled, VisualLoggerListener, VisualLoggerAgentId, VisualLoggerDirectionLocation, VisualLoggerLocation, VisualLoggerArrowLength, VisualLoggerColor);
+	return MakeDirectionObservation(SplineComponent->GetDirectionAtDistanceAlongSpline(DistanceAlongSpline, ESplineCoordinateSpace::World), RelativeTransform, Name, bVisualLoggerEnabled, VisualLoggerListener, VisualLoggerAgentId, VisualLoggerDirectionLocation, VisualLoggerLocation, VisualLoggerArrowLength, VisualLoggerColor);
 }
 
-FLearningAgentsObservationObjectElement ULearningAgentsObservations::MakePropertiesAlongSplineObservation(ULearningAgentsObservationObject* Object, const USplineComponent* SplineComponent, const float DistanceAlongSpline, const FTransform RelativeTransform, const float LocationScale, const FName Tag)
+FLearningAgentsObservationObjectElement ULearningAgentsObservationObject::MakePropertiesAlongSplineObservation(const USplineComponent* SplineComponent, const float DistanceAlongSpline, const FTransform RelativeTransform, const float LocationScale, const FName Name)
 {
 	if (!SplineComponent)
 	{
@@ -2097,20 +1978,20 @@ FLearningAgentsObservationObjectElement ULearningAgentsObservations::MakePropert
 		return FLearningAgentsObservationObjectElement();
 	}
 
-	return MakeStructObservationFromArrayViews(Object,
+	return MakeStructObservationFromArrayViews(
 		{
 			TEXT("Location"),
 			TEXT("Proportion"),
 			TEXT("Direction")
 		},
 		{
-			MakeLocationAlongSplineObservation(Object, SplineComponent, DistanceAlongSpline, RelativeTransform, LocationScale),
-			MakeProportionAlongSplineObservation(Object, SplineComponent, DistanceAlongSpline),
-			MakeDirectionAlongSplineObservation(Object, SplineComponent, DistanceAlongSpline, RelativeTransform),
-		}, Tag);
+			MakeLocationAlongSplineObservation(SplineComponent, DistanceAlongSpline, RelativeTransform, LocationScale),
+			MakeProportionAlongSplineObservation(SplineComponent, DistanceAlongSpline),
+			MakeDirectionAlongSplineObservation(SplineComponent, DistanceAlongSpline, RelativeTransform),
+		}, Name);
 }
 
-FLearningAgentsObservationObjectElement ULearningAgentsObservations::MakeProportionAlongRayObservation(ULearningAgentsObservationObject* Object, const FVector RayStart, const FVector RayEnd, const FTransform RayTransform, const ECollisionChannel CollisionChannel, const FName Tag)
+FLearningAgentsObservationObjectElement ULearningAgentsObservationObject::MakeProportionAlongRayObservation(const FVector RayStart, const FVector RayEnd, const FTransform RayTransform, const ECollisionChannel CollisionChannel, const FName Name)
 {
 	const FVector RayStartWorld = RayTransform.TransformPosition(RayStart);
 	const FVector RayEndWorld = RayTransform.TransformPosition(RayEnd);
@@ -2121,19 +2002,19 @@ FLearningAgentsObservationObjectElement ULearningAgentsObservations::MakeProport
 	ObjectQueryParams.AddObjectTypesToQuery(CollisionChannel);
 
 	FHitResult TraceHit;
-	const bool bHit = Object->GetWorld()->LineTraceSingleByObjectType(TraceHit, RayStartWorld, RayEndWorld, ObjectQueryParams);
+	const bool bHit = GetWorld()->LineTraceSingleByObjectType(TraceHit, RayStartWorld, RayEndWorld, ObjectQueryParams);
 
 	if (bHit)
 	{
-		return MakeFloatObservation(Object, 1.0f - TraceHit.Time, 1.0f, Tag);
+		return MakeFloatObservation(1.0f - TraceHit.Time, 1.0f, Name);
 	}
 	else
 	{
-		return MakeFloatObservation(Object, 0.0f, 1.0f, Tag);
+		return MakeFloatObservation(0.0f, 1.0f, Name);
 	}
 }
 
-FLearningAgentsObservationObjectElement ULearningAgentsObservations::MakeProportionAlongRaysObservation(ULearningAgentsObservationObject* Object, const TArray<FVector>& RayStarts, const TArray<FVector>& RayEnds, const FTransform RayTransform, const ECollisionChannel CollisionChannel, const FName Tag)
+FLearningAgentsObservationObjectElement ULearningAgentsObservationObject::MakeProportionAlongRaysObservation(const TArray<FVector>& RayStarts, const TArray<FVector>& RayEnds, const FTransform RayTransform, const ECollisionChannel CollisionChannel, const FName Name)
 {
 	if (RayStarts.Num() != RayEnds.Num())
 	{
@@ -2148,37 +2029,31 @@ FLearningAgentsObservationObjectElement ULearningAgentsObservations::MakeProport
 
 	for (int32 RayIdx = 0; RayIdx < RayNum; RayIdx++)
 	{
-		RayElements.Add(MakeProportionAlongRayObservation(Object, RayStarts[RayIdx], RayEnds[RayIdx], RayTransform, CollisionChannel));
+		RayElements.Add(MakeProportionAlongRayObservation(RayStarts[RayIdx], RayEnds[RayIdx], RayTransform, CollisionChannel));
 	}
 
-	return MakeStaticArrayObservationFromArrayView(Object, RayElements, Tag);
+	return MakeStaticArrayObservationFromArrayView(RayElements, Name);
 }
 
-
-bool ULearningAgentsObservations::GetNullObservation(const ULearningAgentsObservationObject* Object, const FLearningAgentsObservationObjectElement Element, const FName Tag)
+bool ULearningAgentsObservationObject::GetNullObservation(const FLearningAgentsObservationObjectElement Element, const FName Name) const
 {
-	if (!Object)
-	{
-		UE_LOG(LogLearning, Error, TEXT("GetNullObservation: Object is nullptr."));
-		return false;
-	}
-
-	if (!Object->ObservationObject.IsValid(Element.ObjectElement))
+	if (!ObservationObject.IsValid(Element.ObjectElement))
 	{
-		UE_LOG(LogLearning, Error, TEXT("GetNullObservation: Invalid Observation Object."));
+		UE_LOG(LogLearning, Error, TEXT("%s: Invalid Observation Object."), *GetName());
 		return false;
 	}
 
-	if (Object->ObservationObject.GetTag(Element.ObjectElement) != Tag)
+	if (ObservationObject.GetName(Element.ObjectElement) != Name)
 	{
-		UE_LOG(LogLearning, Warning, TEXT("GetNullObservation: Observation tag does not match. Observation is '%s' but asked for '%s'."), *Object->ObservationObject.GetTag(Element.ObjectElement).ToString(), *Tag.ToString());
+		UE_LOG(LogLearning, Warning, TEXT("%s: Observation name does not match. Observation is '%s' but asked for '%s'."), *GetName(), *ObservationObject.GetName(Element.ObjectElement).ToString(), *Name.ToString());
 	}
 
-	if (Object->ObservationObject.GetType(Element.ObjectElement) != UE::Learning::Observation::EType::Null)
+	if (ObservationObject.GetType(Element.ObjectElement) != UE::Learning::Observation::EType::Null)
 	{
-		UE_LOG(LogLearning, Error, TEXT("GetNullObservation: Observation '%s' type does not match. Observation is '%s' but asked for '%s'."),
-			*Object->ObservationObject.GetTag(Element.ObjectElement).ToString(),
-			UE::Learning::Agents::Observation::Private::GetObservationTypeString(Object->ObservationObject.GetType(Element.ObjectElement)),
+		UE_LOG(LogLearning, Error, TEXT("%s: Observation '%s' type does not match. Observation is '%s' but asked for '%s'."),
+			*GetName(),
+			*ObservationObject.GetName(Element.ObjectElement).ToString(),
+			UE::Learning::Agents::Observation::Private::GetObservationTypeString(ObservationObject.GetType(Element.ObjectElement)),
 			UE::Learning::Agents::Observation::Private::GetObservationTypeString(UE::Learning::Observation::EType::Null));
 		return false;
 	}
@@ -2186,45 +2061,39 @@ bool ULearningAgentsObservations::GetNullObservation(const ULearningAgentsObserv
 	return true;
 }
 
-bool ULearningAgentsObservations::GetContinuousObservationNum(int32& OutNum, const ULearningAgentsObservationObject* Object, const FLearningAgentsObservationObjectElement Element, const FName Tag)
+bool ULearningAgentsObservationObject::GetContinuousObservationNum(int32& OutNum, const FLearningAgentsObservationObjectElement Element, const FName Name) const
 {
-	if (!Object)
+	if (!ObservationObject.IsValid(Element.ObjectElement))
 	{
-		UE_LOG(LogLearning, Error, TEXT("GetContinuousObservationNum: Object is nullptr."));
+		UE_LOG(LogLearning, Error, TEXT("%s: Invalid Observation Object."), *GetName());
 		OutNum = 0;
 		return false;
 	}
 
-	if (!Object->ObservationObject.IsValid(Element.ObjectElement))
+	if (ObservationObject.GetName(Element.ObjectElement) != Name)
 	{
-		UE_LOG(LogLearning, Error, TEXT("GetContinuousObservationNum: Invalid Observation Object."));
-		OutNum = 0;
-		return false;
-	}
-
-	if (Object->ObservationObject.GetTag(Element.ObjectElement) != Tag)
-	{
-		UE_LOG(LogLearning, Warning, TEXT("GetContinuousObservationNum: Observation tag does not match. Observation is '%s' but asked for '%s'."), *Object->ObservationObject.GetTag(Element.ObjectElement).ToString(), *Tag.ToString());
+		UE_LOG(LogLearning, Warning, TEXT("%s: Observation name does not match. Observation is '%s' but asked for '%s'."), *GetName(), *ObservationObject.GetName(Element.ObjectElement).ToString(), *Name.ToString());
 	}
 
-	if (Object->ObservationObject.GetType(Element.ObjectElement) != UE::Learning::Observation::EType::Continuous)
+	if (ObservationObject.GetType(Element.ObjectElement) != UE::Learning::Observation::EType::Continuous)
 	{
-		UE_LOG(LogLearning, Error, TEXT("GetContinuousObservationNum: Observation '%s' type does not match. Observation is '%s' but asked for '%s'."),
-			*Object->ObservationObject.GetTag(Element.ObjectElement).ToString(),
-			UE::Learning::Agents::Observation::Private::GetObservationTypeString(Object->ObservationObject.GetType(Element.ObjectElement)),
+		UE_LOG(LogLearning, Error, TEXT("%s: Observation '%s' type does not match. Observation is '%s' but asked for '%s'."),
+			*GetName(),
+			*ObservationObject.GetName(Element.ObjectElement).ToString(),
+			UE::Learning::Agents::Observation::Private::GetObservationTypeString(ObservationObject.GetType(Element.ObjectElement)),
 			UE::Learning::Agents::Observation::Private::GetObservationTypeString(UE::Learning::Observation::EType::Continuous));
 		OutNum = 0;
 		return false;
 	}
 
-	OutNum = Object->ObservationObject.GetContinuous(Element.ObjectElement).Values.Num();
+	OutNum = ObservationObject.GetContinuous(Element.ObjectElement).Values.Num();
 	return true;
 }
 
-bool ULearningAgentsObservations::GetContinuousObservation(TArray<float>& OutValues, const ULearningAgentsObservationObject* Object, const FLearningAgentsObservationObjectElement Element, const FName Tag)
+bool ULearningAgentsObservationObject::GetContinuousObservation(TArray<float>& OutValues, const FLearningAgentsObservationObjectElement Element, const FName Name) const
 {
 	int32 OutValueNum = 0;
-	if (!GetContinuousObservationNum(OutValueNum, Object, Element, Tag))
+	if (!GetContinuousObservationNum(OutValueNum, Element, Name))
 	{
 		OutValues.Empty();
 		return false;
@@ -2232,7 +2101,7 @@ bool ULearningAgentsObservations::GetContinuousObservation(TArray<float>& OutVal
 
 	OutValues.SetNumUninitialized(OutValueNum);
 
-	if (!GetContinuousObservationToArrayView(OutValues, Object, Element, Tag))
+	if (!GetContinuousObservationToArrayView(OutValues, Element, Name))
 	{
 		OutValues.Empty();
 		return false;
@@ -2241,43 +2110,37 @@ bool ULearningAgentsObservations::GetContinuousObservation(TArray<float>& OutVal
 	return true;
 }
 
-bool ULearningAgentsObservations::GetContinuousObservationToArrayView(TArrayView<float> OutValues, const ULearningAgentsObservationObject* Object, const FLearningAgentsObservationObjectElement Element, const FName Tag)
+bool ULearningAgentsObservationObject::GetContinuousObservationToArrayView(TArrayView<float> OutValues, const FLearningAgentsObservationObjectElement Element, const FName Name) const
 {
-	if (!Object)
+	if (!ObservationObject.IsValid(Element.ObjectElement))
 	{
-		UE_LOG(LogLearning, Error, TEXT("GetContinuousObservationToArrayView: Object is nullptr."));
+		UE_LOG(LogLearning, Error, TEXT("%s: Invalid Observation Object."), *GetName());
 		UE::Learning::Array::Zero<1, float>(OutValues);
 		return false;
 	}
 
-	if (!Object->ObservationObject.IsValid(Element.ObjectElement))
+	if (ObservationObject.GetName(Element.ObjectElement) != Name)
 	{
-		UE_LOG(LogLearning, Error, TEXT("GetContinuousObservationToArrayView: Invalid Observation Object."));
-		UE::Learning::Array::Zero<1, float>(OutValues);
-		return false;
+		UE_LOG(LogLearning, Warning, TEXT("%s: Observation name does not match. Observation is '%s' but asked for '%s'."), *GetName(), *ObservationObject.GetName(Element.ObjectElement).ToString(), *Name.ToString());
 	}
 
-	if (Object->ObservationObject.GetTag(Element.ObjectElement) != Tag)
+	if (ObservationObject.GetType(Element.ObjectElement) != UE::Learning::Observation::EType::Continuous)
 	{
-		UE_LOG(LogLearning, Warning, TEXT("GetContinuousObservationToArrayView: Observation tag does not match. Observation is '%s' but asked for '%s'."), *Object->ObservationObject.GetTag(Element.ObjectElement).ToString(), *Tag.ToString());
-	}
-
-	if (Object->ObservationObject.GetType(Element.ObjectElement) != UE::Learning::Observation::EType::Continuous)
-	{
-		UE_LOG(LogLearning, Error, TEXT("GetContinuousObservationToArrayView: Observation '%s' type does not match. Observation is '%s' but asked for '%s'."),
-			*Object->ObservationObject.GetTag(Element.ObjectElement).ToString(),
-			UE::Learning::Agents::Observation::Private::GetObservationTypeString(Object->ObservationObject.GetType(Element.ObjectElement)),
+		UE_LOG(LogLearning, Error, TEXT("%s: Observation '%s' type does not match. Observation is '%s' but asked for '%s'."),
+			*GetName(),
+			*ObservationObject.GetName(Element.ObjectElement).ToString(),
+			UE::Learning::Agents::Observation::Private::GetObservationTypeString(ObservationObject.GetType(Element.ObjectElement)),
 			UE::Learning::Agents::Observation::Private::GetObservationTypeString(UE::Learning::Observation::EType::Continuous));
 		UE::Learning::Array::Zero<1, float>(OutValues);
 		return false;
 	}
 
-	const TArrayView<const float> Values = Object->ObservationObject.GetContinuous(Element.ObjectElement).Values;
+	const TArrayView<const float> Values = ObservationObject.GetContinuous(Element.ObjectElement).Values;
 
 	if (Values.Num() != OutValues.Num())
 	{
-		UE_LOG(LogLearning, Error, TEXT("GetContinuousObservationToArrayView: Observation '%s' size does not match. Observation is '%i' values but asked for '%i'."),
-			*Object->ObservationObject.GetTag(Element.ObjectElement).ToString(),
+		UE_LOG(LogLearning, Error, TEXT("%s: Observation '%s' size does not match. Observation is '%i' values but asked for '%i'."),
+			*GetName(), *ObservationObject.GetName(Element.ObjectElement).ToString(),
 			Values.Num(), OutValues.Num());
 		UE::Learning::Array::Zero<1, float>(OutValues);
 		return false;
@@ -2287,10 +2150,10 @@ bool ULearningAgentsObservations::GetContinuousObservationToArrayView(TArrayView
 	return true;
 }
 
-bool ULearningAgentsObservations::GetExclusiveDiscreteObservation(int32& OutIndex, const ULearningAgentsObservationObject* Object, const FLearningAgentsObservationObjectElement Element, const FName Tag)
+bool ULearningAgentsObservationObject::GetExclusiveDiscreteObservation(int32& OutIndex, const FLearningAgentsObservationObjectElement Element, const FName Name) const
 {
 	int32 DiscreteValueNum;
-	if (!GetContinuousObservationNum(DiscreteValueNum, Object, Element, Tag))
+	if (!GetContinuousObservationNum(DiscreteValueNum, Element, Name))
 	{
 		OutIndex = INDEX_NONE;
 		return false;
@@ -2298,7 +2161,7 @@ bool ULearningAgentsObservations::GetExclusiveDiscreteObservation(int32& OutInde
 
 	TArray<float, TInlineAllocator<32>> OneHot;
 	OneHot.SetNumUninitialized(DiscreteValueNum);
-	if (!GetContinuousObservationToArrayView(OneHot, Object, Element, Tag))
+	if (!GetContinuousObservationToArrayView(OneHot, Element, Name))
 	{
 		OutIndex = INDEX_NONE;
 		return false;
@@ -2317,10 +2180,10 @@ bool ULearningAgentsObservations::GetExclusiveDiscreteObservation(int32& OutInde
 	return OutIndex != INDEX_NONE;
 }
 
-bool ULearningAgentsObservations::GetInclusiveDiscreteObservationNum(int32& OutNum, const ULearningAgentsObservationObject* Object, const FLearningAgentsObservationObjectElement Element, const FName Tag)
+bool ULearningAgentsObservationObject::GetInclusiveDiscreteObservationNum(int32& OutNum, const FLearningAgentsObservationObjectElement Element, const FName Name) const
 {
 	int32 DiscreteValueNum;
-	if (!GetContinuousObservationNum(DiscreteValueNum, Object, Element, Tag))
+	if (!GetContinuousObservationNum(DiscreteValueNum, Element, Name))
 	{
 		OutNum = 0;
 		return false;
@@ -2328,7 +2191,7 @@ bool ULearningAgentsObservations::GetInclusiveDiscreteObservationNum(int32& OutN
 
 	TArray<float, TInlineAllocator<32>> OneHot;
 	OneHot.SetNumUninitialized(DiscreteValueNum);
-	if (!GetContinuousObservationToArrayView(OneHot, Object, Element, Tag))
+	if (!GetContinuousObservationToArrayView(OneHot, Element, Name))
 	{
 		OutNum = 0;
 		return false;
@@ -2343,36 +2206,30 @@ bool ULearningAgentsObservations::GetInclusiveDiscreteObservationNum(int32& OutN
 	return true;
 }
 
-bool ULearningAgentsObservations::GetInclusiveDiscreteObservation(TArray<int32>& OutIndices, const ULearningAgentsObservationObject* Object, const FLearningAgentsObservationObjectElement Element, const FName Tag)
+bool ULearningAgentsObservationObject::GetInclusiveDiscreteObservation(TArray<int32>& OutIndices, const FLearningAgentsObservationObjectElement Element, const FName Name) const
 {
 	int32 OutIndicesNum;
-	if (!GetInclusiveDiscreteObservationNum(OutIndicesNum, Object, Element, Tag))
+	if (!GetInclusiveDiscreteObservationNum(OutIndicesNum, Element, Name))
 	{
 		OutIndices.Empty();
 		return false;
 	}
 
 	OutIndices.SetNumUninitialized(OutIndicesNum);
-	return GetInclusiveDiscreteObservationToArrayView(OutIndices, Object, Element, Tag);
+	return GetInclusiveDiscreteObservationToArrayView(OutIndices, Element, Name);
 }
 
-bool ULearningAgentsObservations::GetInclusiveDiscreteObservationToArrayView(TArrayView<int32> OutIndices, const ULearningAgentsObservationObject* Object, const FLearningAgentsObservationObjectElement Element, const FName Tag)
+bool ULearningAgentsObservationObject::GetInclusiveDiscreteObservationToArrayView(TArrayView<int32> OutIndices, const FLearningAgentsObservationObjectElement Element, const FName Name) const
 {
-	if (!Object)
-	{
-		UE_LOG(LogLearning, Error, TEXT("GetInclusiveDiscreteObservationToArrayView: Object is nullptr."));
-		return false;
-	}
-
 	int32 DiscreteValueNum;
-	if (!GetContinuousObservationNum(DiscreteValueNum, Object, Element, Tag))
+	if (!GetContinuousObservationNum(DiscreteValueNum, Element, Name))
 	{
 		return false;
 	}
 
 	TArray<float, TInlineAllocator<32>> OneHot;
 	OneHot.SetNumUninitialized(DiscreteValueNum);
-	if (!GetContinuousObservationToArrayView(OneHot, Object, Element, Tag))
+	if (!GetContinuousObservationToArrayView(OneHot, Element, Name))
 	{
 		return false;
 	}
@@ -2385,8 +2242,8 @@ bool ULearningAgentsObservations::GetInclusiveDiscreteObservationToArrayView(TAr
 
 	if (IndicesNum != OutIndices.Num())
 	{
-		UE_LOG(LogLearning, Error, TEXT("GetInclusiveDiscreteObservationToArrayView: Observation '%s' size does not match. Observation is '%i' indices but asked for '%i'."),
-			*Object->ObservationObject.GetTag(Element.ObjectElement).ToString(),
+		UE_LOG(LogLearning, Error, TEXT("%s: Observation '%s' size does not match. Observation is '%i' indices but asked for '%i'."),
+			*GetName(), *ObservationObject.GetName(Element.ObjectElement).ToString(),
 			IndicesNum, OutIndices.Num());
 		UE::Learning::Array::Zero<1, int32>(OutIndices);
 		return false;
@@ -2401,10 +2258,10 @@ bool ULearningAgentsObservations::GetInclusiveDiscreteObservationToArrayView(TAr
 	return true;
 }
 
-bool ULearningAgentsObservations::GetIndexObservation(int32& OutIndex, const ULearningAgentsObservationObject* Object, const FLearningAgentsObservationObjectElement Element, const FName Tag)
+bool ULearningAgentsObservationObject::GetIndexObservation(int32& OutIndex, const FLearningAgentsObservationObjectElement Element, const FName Name) const
 {
 	int32 DiscreteValueNum;
-	if (!GetContinuousObservationNum(DiscreteValueNum, Object, Element, Tag))
+	if (!GetContinuousObservationNum(DiscreteValueNum, Element, Name))
 	{
 		OutIndex = INDEX_NONE;
 		return false;
@@ -2412,7 +2269,7 @@ bool ULearningAgentsObservations::GetIndexObservation(int32& OutIndex, const ULe
 
 	TArray<float, TInlineAllocator<32>> OneHot;
 	OneHot.SetNumUninitialized(DiscreteValueNum);
-	if (!GetContinuousObservationToArrayView(OneHot, Object, Element, Tag))
+	if (!GetContinuousObservationToArrayView(OneHot, Element, Name))
 	{
 		OutIndex = INDEX_NONE;
 		return false;
@@ -2427,10 +2284,10 @@ bool ULearningAgentsObservations::GetIndexObservation(int32& OutIndex, const ULe
 	return true;
 }
 
-bool ULearningAgentsObservations::GetCountObservation(int32& OutNum, const ULearningAgentsObservationObject* Object, const FLearningAgentsObservationObjectElement Element, const int32 MaxNum, const FName Tag)
+bool ULearningAgentsObservationObject::GetCountObservation(int32& OutNum, const FLearningAgentsObservationObjectElement Element, const int32 MaxNum, const FName Name) const
 {
 	float FloatNum = 0.0f;
-	if (!GetContinuousObservationToArrayView(MakeArrayView(&FloatNum, 1), Object, Element, Tag))
+	if (!GetContinuousObservationToArrayView(MakeArrayView(&FloatNum, 1), Element, Name))
 	{
 		OutNum = -1;
 		return false;
@@ -2440,47 +2297,41 @@ bool ULearningAgentsObservations::GetCountObservation(int32& OutNum, const ULear
 	return true;
 }
 
-bool ULearningAgentsObservations::GetStructObservationNum(int32& OutNum, const ULearningAgentsObservationObject* Object, const FLearningAgentsObservationObjectElement Element, const FName Tag)
+bool ULearningAgentsObservationObject::GetStructObservationNum(int32& OutNum, const FLearningAgentsObservationObjectElement Element, const FName Name) const
 {
-	if (!Object)
+	if (!ObservationObject.IsValid(Element.ObjectElement))
 	{
-		UE_LOG(LogLearning, Error, TEXT("GetStructObservationNum: Object is nullptr."));
+		UE_LOG(LogLearning, Error, TEXT("%s: Invalid Observation Object."), *GetName());
 		OutNum = 0;
 		return false;
 	}
 
-	if (!Object->ObservationObject.IsValid(Element.ObjectElement))
+	if (ObservationObject.GetName(Element.ObjectElement) != Name)
 	{
-		UE_LOG(LogLearning, Error, TEXT("GetStructObservationNum: Invalid Observation Object."));
-		OutNum = 0;
-		return false;
+		UE_LOG(LogLearning, Warning, TEXT("%s: Observation name does not match. Observation is '%s' but asked for '%s'."), *GetName(), *ObservationObject.GetName(Element.ObjectElement).ToString(), *Name.ToString());
 	}
 
-	if (Object->ObservationObject.GetTag(Element.ObjectElement) != Tag)
+	if (ObservationObject.GetType(Element.ObjectElement) != UE::Learning::Observation::EType::And)
 	{
-		UE_LOG(LogLearning, Warning, TEXT("GetStructObservationNum: Observation tag does not match. Observation is '%s' but asked for '%s'."), *Object->ObservationObject.GetTag(Element.ObjectElement).ToString(), *Tag.ToString());
-	}
-
-	if (Object->ObservationObject.GetType(Element.ObjectElement) != UE::Learning::Observation::EType::And)
-	{
-		UE_LOG(LogLearning, Error, TEXT("GetStructObservationNum: Observation '%s' type does not match. Observation is '%s' but asked for '%s'."),
-			*Object->ObservationObject.GetTag(Element.ObjectElement).ToString(),
-			UE::Learning::Agents::Observation::Private::GetObservationTypeString(Object->ObservationObject.GetType(Element.ObjectElement)),
+		UE_LOG(LogLearning, Error, TEXT("%s: Observation '%s' type does not match. Observation is '%s' but asked for '%s'."),
+			*GetName(),
+			*ObservationObject.GetName(Element.ObjectElement).ToString(),
+			UE::Learning::Agents::Observation::Private::GetObservationTypeString(ObservationObject.GetType(Element.ObjectElement)),
 			UE::Learning::Agents::Observation::Private::GetObservationTypeString(UE::Learning::Observation::EType::And));
 		OutNum = 0;
 		return false;
 	}
 
-	const UE::Learning::Observation::FObjectAndParameters Parameters = Object->ObservationObject.GetAnd(Element.ObjectElement);
+	const UE::Learning::Observation::FObjectAndParameters Parameters = ObservationObject.GetAnd(Element.ObjectElement);
 
 	OutNum = Parameters.Elements.Num();
 	return true;
 }
 
-bool ULearningAgentsObservations::GetStructObservation(TMap<FName, FLearningAgentsObservationObjectElement>& OutElements, const ULearningAgentsObservationObject* Object, const FLearningAgentsObservationObjectElement Element, const FName Tag)
+bool ULearningAgentsObservationObject::GetStructObservation(TMap<FName, FLearningAgentsObservationObjectElement>& OutElements, const FLearningAgentsObservationObjectElement Element, const FName Name) const
 {
 	int32 OutElementNum = 0;
-	if (!GetStructObservationNum(OutElementNum, Object, Element, Tag))
+	if (!GetStructObservationNum(OutElementNum, Element, Name))
 	{
 		OutElements.Empty();
 		return false;
@@ -2491,7 +2342,7 @@ bool ULearningAgentsObservations::GetStructObservation(TMap<FName, FLearningAgen
 	SubElementNames.SetNumUninitialized(OutElementNum);
 	SubElements.SetNumUninitialized(OutElementNum);
 
-	if (!GetStructObservationToArrayViews(SubElementNames, SubElements, Object, Element, Tag))
+	if (!GetStructObservationToArrayViews(SubElementNames, SubElements, Element, Name))
 	{
 		OutElements.Empty();
 		return false;
@@ -2506,10 +2357,10 @@ bool ULearningAgentsObservations::GetStructObservation(TMap<FName, FLearningAgen
 	return true;
 }
 
-bool ULearningAgentsObservations::GetStructObservationToArrays(TArray<FName>& OutElementNames, TArray<FLearningAgentsObservationObjectElement>& OutElements, const ULearningAgentsObservationObject* Object, const FLearningAgentsObservationObjectElement Element, const FName Tag)
+bool ULearningAgentsObservationObject::GetStructObservationToArrays(TArray<FName>& OutElementNames, TArray<FLearningAgentsObservationObjectElement>& OutElements, const FLearningAgentsObservationObjectElement Element, const FName Name) const
 {
 	int32 OutElementNum = 0;
-	if (!GetStructObservationNum(OutElementNum, Object, Element, Tag))
+	if (!GetStructObservationNum(OutElementNum, Element, Name))
 	{
 		OutElementNames.Empty();
 		OutElements.Empty();
@@ -2519,7 +2370,7 @@ bool ULearningAgentsObservations::GetStructObservationToArrays(TArray<FName>& Ou
 	OutElementNames.SetNumUninitialized(OutElementNum);
 	OutElements.SetNumUninitialized(OutElementNum);
 
-	if (!GetStructObservationToArrayViews(OutElementNames, OutElements, Object, Element, Tag))
+	if (!GetStructObservationToArrayViews(OutElementNames, OutElements, Element, Name))
 	{
 		OutElementNames.Empty();
 		OutElements.Empty();
@@ -2529,51 +2380,45 @@ bool ULearningAgentsObservations::GetStructObservationToArrays(TArray<FName>& Ou
 	return true;
 }
 
-bool ULearningAgentsObservations::GetStructObservationToArrayViews(TArrayView<FName> OutElementNames, TArrayView<FLearningAgentsObservationObjectElement> OutElements, const ULearningAgentsObservationObject* Object, const FLearningAgentsObservationObjectElement Element, const FName Tag)
+bool ULearningAgentsObservationObject::GetStructObservationToArrayViews(TArrayView<FName> OutElementNames, TArrayView<FLearningAgentsObservationObjectElement> OutElements, const FLearningAgentsObservationObjectElement Element, const FName Name) const
 {
-	if (!Object)
-	{
-		UE_LOG(LogLearning, Error, TEXT("GetStructObservationToArrayViews: Object is nullptr."));
-		UE::Learning::Array::Set<1, FName>(OutElementNames, NAME_None);
-		UE::Learning::Array::Set<1, FLearningAgentsObservationObjectElement>(OutElements, FLearningAgentsObservationObjectElement());
-		return false;
-	}
-
-	if (!Object->ObservationObject.IsValid(Element.ObjectElement))
+	if (!ObservationObject.IsValid(Element.ObjectElement))
 	{
-		UE_LOG(LogLearning, Error, TEXT("GetStructObservationToArrayViews: Invalid Observation Object."));
+		UE_LOG(LogLearning, Error, TEXT("%s: Invalid Observation Object."), *GetName());
 		UE::Learning::Array::Set<1, FName>(OutElementNames, NAME_None);
 		UE::Learning::Array::Set<1, FLearningAgentsObservationObjectElement>(OutElements, FLearningAgentsObservationObjectElement());
 		return false;
 	}
 
-	if (Object->ObservationObject.GetTag(Element.ObjectElement) != Tag)
+	if (ObservationObject.GetName(Element.ObjectElement) != Name)
 	{
-		UE_LOG(LogLearning, Warning, TEXT("GetStructObservationToArrayViews: Observation tag does not match. Observation is '%s' but asked for '%s'."), *Object->ObservationObject.GetTag(Element.ObjectElement).ToString(), *Tag.ToString());
+		UE_LOG(LogLearning, Warning, TEXT("%s: Observation name does not match. Observation is '%s' but asked for '%s'."), *GetName(), *ObservationObject.GetName(Element.ObjectElement).ToString(), *Name.ToString());
 	}
 
-	if (Object->ObservationObject.GetType(Element.ObjectElement) != UE::Learning::Observation::EType::And)
+	if (ObservationObject.GetType(Element.ObjectElement) != UE::Learning::Observation::EType::And)
 	{
-		UE_LOG(LogLearning, Error, TEXT("GetStructObservationToArrayViews: Observation '%s' type does not match. Observation is '%s' but asked for '%s'."),
-			*Object->ObservationObject.GetTag(Element.ObjectElement).ToString(),
-			UE::Learning::Agents::Observation::Private::GetObservationTypeString(Object->ObservationObject.GetType(Element.ObjectElement)),
+		UE_LOG(LogLearning, Error, TEXT("%s: Observation '%s' type does not match. Observation is '%s' but asked for '%s'."),
+			*GetName(),
+			*ObservationObject.GetName(Element.ObjectElement).ToString(),
+			UE::Learning::Agents::Observation::Private::GetObservationTypeString(ObservationObject.GetType(Element.ObjectElement)),
 			UE::Learning::Agents::Observation::Private::GetObservationTypeString(UE::Learning::Observation::EType::And));
 		UE::Learning::Array::Set<1, FName>(OutElementNames, NAME_None);
 		UE::Learning::Array::Set<1, FLearningAgentsObservationObjectElement>(OutElements, FLearningAgentsObservationObjectElement());
 		return false;
 	}
 
-	const UE::Learning::Observation::FObjectAndParameters Parameters = Object->ObservationObject.GetAnd(Element.ObjectElement);
+	const UE::Learning::Observation::FObjectAndParameters Parameters = ObservationObject.GetAnd(Element.ObjectElement);
 
 	if (Parameters.Elements.Num() == 0)
 	{
-		UE_LOG(LogLearning, Warning, TEXT("GetStructObservationToArrayViews: Getting zero-sized And Observation."));
+		UE_LOG(LogLearning, Warning, TEXT("%s: Getting zero-sized And Observation."), *GetName());
 	}
 
 	if (Parameters.Elements.Num() != OutElements.Num())
 	{
-		UE_LOG(LogLearning, Error, TEXT("GetStructObservationToArrayViews: Observation '%s' size does not match. Observation is '%i' elements but asked for '%i'."),
-			*Object->ObservationObject.GetTag(Element.ObjectElement).ToString(),
+		UE_LOG(LogLearning, Error, TEXT("%s: Observation '%s' size does not match. Observation is '%i' elements but asked for '%i'."),
+			*GetName(),
+			*ObservationObject.GetName(Element.ObjectElement).ToString(),
 			Parameters.Elements.Num(), OutElements.Num());
 		UE::Learning::Array::Set<1, FName>(OutElementNames, NAME_None);
 		UE::Learning::Array::Set<1, FLearningAgentsObservationObjectElement>(OutElements, FLearningAgentsObservationObjectElement());
@@ -2582,9 +2427,9 @@ bool ULearningAgentsObservations::GetStructObservationToArrayViews(TArrayView<FN
 
 	for (int32 ElementIdx = 0; ElementIdx < Parameters.Elements.Num(); ElementIdx++)
 	{
-		if (!Object->ObservationObject.IsValid(Parameters.Elements[ElementIdx]))
+		if (!ObservationObject.IsValid(Parameters.Elements[ElementIdx]))
 		{
-			UE_LOG(LogLearning, Error, TEXT("GetStructObservationToArrayViews: Invalid Observation Object."));
+			UE_LOG(LogLearning, Error, TEXT("%s: Invalid Observation Object."), *GetName());
 			UE::Learning::Array::Set<1, FName>(OutElementNames, NAME_None);
 			UE::Learning::Array::Set<1, FLearningAgentsObservationObjectElement>(OutElements, FLearningAgentsObservationObjectElement());
 			return false;
@@ -2597,87 +2442,74 @@ bool ULearningAgentsObservations::GetStructObservationToArrayViews(TArrayView<FN
 	return true;
 }
 
-bool ULearningAgentsObservations::GetExclusiveUnionObservation(FName& OutElementName, FLearningAgentsObservationObjectElement& OutElement, const ULearningAgentsObservationObject* Object, const FLearningAgentsObservationObjectElement Element, const FName Tag)
+bool ULearningAgentsObservationObject::GetExclusiveUnionObservation(FName& OutElementName, FLearningAgentsObservationObjectElement& OutElement, const FLearningAgentsObservationObjectElement Element, const FName Name) const
 {
-	if (!Object)
+	if (!ObservationObject.IsValid(Element.ObjectElement))
 	{
-		UE_LOG(LogLearning, Error, TEXT("GetStructObservationToArrayViews: Object is nullptr."));
+		UE_LOG(LogLearning, Error, TEXT("%s: Invalid Observation Object."), *GetName());
 		OutElementName = NAME_None;
 		OutElement = FLearningAgentsObservationObjectElement();
 		return false;
 	}
 
-	if (!Object->ObservationObject.IsValid(Element.ObjectElement))
+	if (ObservationObject.GetName(Element.ObjectElement) != Name)
 	{
-		UE_LOG(LogLearning, Error, TEXT("GetStructObservationToArrayViews: Invalid Observation Object."));
-		OutElementName = NAME_None;
-		OutElement = FLearningAgentsObservationObjectElement();
-		return false;
+		UE_LOG(LogLearning, Warning, TEXT("%s: Observation name does not match. Observation is '%s' but asked for '%s'."), *GetName(), *ObservationObject.GetName(Element.ObjectElement).ToString(), *Name.ToString());
 	}
 
-	if (Object->ObservationObject.GetTag(Element.ObjectElement) != Tag)
+	if (ObservationObject.GetType(Element.ObjectElement) != UE::Learning::Observation::EType::OrExclusive)
 	{
-		UE_LOG(LogLearning, Warning, TEXT("GetStructObservationToArrayViews: Observation tag does not match. Observation is '%s' but asked for '%s'."), *Object->ObservationObject.GetTag(Element.ObjectElement).ToString(), *Tag.ToString());
-	}
-
-	if (Object->ObservationObject.GetType(Element.ObjectElement) != UE::Learning::Observation::EType::OrExclusive)
-	{
-		UE_LOG(LogLearning, Error, TEXT("GetStructObservationToArrayViews: Observation '%s' type does not match. Observation is '%s' but asked for '%s'."),
-			*Object->ObservationObject.GetTag(Element.ObjectElement).ToString(),
-			UE::Learning::Agents::Observation::Private::GetObservationTypeString(Object->ObservationObject.GetType(Element.ObjectElement)),
+		UE_LOG(LogLearning, Error, TEXT("%s: Observation '%s' type does not match. Observation is '%s' but asked for '%s'."),
+			*GetName(),
+			*ObservationObject.GetName(Element.ObjectElement).ToString(),
+			UE::Learning::Agents::Observation::Private::GetObservationTypeString(ObservationObject.GetType(Element.ObjectElement)),
 			UE::Learning::Agents::Observation::Private::GetObservationTypeString(UE::Learning::Observation::EType::OrExclusive));
 		OutElementName = NAME_None;
 		OutElement = FLearningAgentsObservationObjectElement();
 		return false;
 	}
 
-	const UE::Learning::Observation::FObjectOrExclusiveParameters Parameters = Object->ObservationObject.GetOrExclusive(Element.ObjectElement);
+	const UE::Learning::Observation::FObjectOrExclusiveParameters Parameters = ObservationObject.GetOrExclusive(Element.ObjectElement);
 	OutElementName = Parameters.ElementName;
 	OutElement = { Parameters.Element };
 	return true;
 }
 
-bool ULearningAgentsObservations::GetInclusiveUnionObservationNum(int32& OutNum, const ULearningAgentsObservationObject* Object, const FLearningAgentsObservationObjectElement Element, const FName Tag)
+bool ULearningAgentsObservationObject::GetInclusiveUnionObservationNum(int32& OutNum, const FLearningAgentsObservationObjectElement Element, const FName Name) const
 {
-	if (!Object)
+	if (!ObservationObject.IsValid(Element.ObjectElement))
 	{
-		UE_LOG(LogLearning, Error, TEXT("GetInclusiveUnionObservationNum: Object is nullptr."));
+		UE_LOG(LogLearning, Error, TEXT("%s: Invalid Observation Object."), *GetName());
 		OutNum = 0;
 		return false;
 	}
 
-	if (!Object->ObservationObject.IsValid(Element.ObjectElement))
+	if (ObservationObject.GetName(Element.ObjectElement) != Name)
 	{
-		UE_LOG(LogLearning, Error, TEXT("GetInclusiveUnionObservationNum: Invalid Observation Object."));
-		OutNum = 0;
-		return false;
+		UE_LOG(LogLearning, Warning, TEXT("%s: Observation name does not match. Observation is '%s' but asked for '%s'."), *GetName(), *ObservationObject.GetName(Element.ObjectElement).ToString(), *Name.ToString());
 	}
 
-	if (Object->ObservationObject.GetTag(Element.ObjectElement) != Tag)
+	if (ObservationObject.GetType(Element.ObjectElement) != UE::Learning::Observation::EType::OrInclusive)
 	{
-		UE_LOG(LogLearning, Warning, TEXT("GetInclusiveUnionObservationNum: Observation tag does not match. Observation is '%s' but asked for '%s'."), *Object->ObservationObject.GetTag(Element.ObjectElement).ToString(), *Tag.ToString());
-	}
-
-	if (Object->ObservationObject.GetType(Element.ObjectElement) != UE::Learning::Observation::EType::OrInclusive)
-	{
-		UE_LOG(LogLearning, Error, TEXT("GetInclusiveUnionObservationNum: Observation '%s' type does not match. Observation is '%s' but asked for '%s'."),
-			*Object->ObservationObject.GetTag(Element.ObjectElement).ToString(),
-			UE::Learning::Agents::Observation::Private::GetObservationTypeString(Object->ObservationObject.GetType(Element.ObjectElement)),
+		UE_LOG(LogLearning, Error, TEXT("%s: Observation '%s' type does not match. Observation is '%s' but asked for '%s'."),
+			*GetName(),
+			*ObservationObject.GetName(Element.ObjectElement).ToString(),
+			UE::Learning::Agents::Observation::Private::GetObservationTypeString(ObservationObject.GetType(Element.ObjectElement)),
 			UE::Learning::Agents::Observation::Private::GetObservationTypeString(UE::Learning::Observation::EType::OrInclusive));
 		OutNum = 0;
 		return false;
 	}
 
-	const UE::Learning::Observation::FObjectOrInclusiveParameters Parameters = Object->ObservationObject.GetOrInclusive(Element.ObjectElement);
+	const UE::Learning::Observation::FObjectOrInclusiveParameters Parameters = ObservationObject.GetOrInclusive(Element.ObjectElement);
 
 	OutNum = Parameters.Elements.Num();
 	return true;
 }
 
-bool ULearningAgentsObservations::GetInclusiveUnionObservation(TMap<FName, FLearningAgentsObservationObjectElement>& OutElements, const ULearningAgentsObservationObject* Object, const FLearningAgentsObservationObjectElement Element, const FName Tag)
+bool ULearningAgentsObservationObject::GetInclusiveUnionObservation(TMap<FName, FLearningAgentsObservationObjectElement>& OutElements, const FLearningAgentsObservationObjectElement Element, const FName Name) const
 {
 	int32 OutElementNum = 0;
-	if (!GetInclusiveUnionObservationNum(OutElementNum, Object, Element, Tag))
+	if (!GetInclusiveUnionObservationNum(OutElementNum, Element, Name))
 	{
 		OutElements.Empty();
 		return false;
@@ -2688,7 +2520,7 @@ bool ULearningAgentsObservations::GetInclusiveUnionObservation(TMap<FName, FLear
 	SubElementNames.SetNumUninitialized(OutElementNum);
 	SubElements.SetNumUninitialized(OutElementNum);
 
-	if (!GetInclusiveUnionObservationToArrayViews(SubElementNames, SubElements, Object, Element, Tag))
+	if (!GetInclusiveUnionObservationToArrayViews(SubElementNames, SubElements, Element, Name))
 	{
 		OutElements.Empty();
 		return false;
@@ -2703,10 +2535,10 @@ bool ULearningAgentsObservations::GetInclusiveUnionObservation(TMap<FName, FLear
 	return true;
 }
 
-bool ULearningAgentsObservations::GetInclusiveUnionObservationToArrays(TArray<FName>& OutElementNames, TArray<FLearningAgentsObservationObjectElement>& OutElements, const ULearningAgentsObservationObject* Object, const FLearningAgentsObservationObjectElement Element, const FName Tag)
+bool ULearningAgentsObservationObject::GetInclusiveUnionObservationToArrays(TArray<FName>& OutElementNames, TArray<FLearningAgentsObservationObjectElement>& OutElements, const FLearningAgentsObservationObjectElement Element, const FName Name) const
 {
 	int32 OutElementNum = 0;
-	if (!GetInclusiveUnionObservationNum(OutElementNum, Object, Element, Tag))
+	if (!GetInclusiveUnionObservationNum(OutElementNum, Element, Name))
 	{
 		OutElementNames.Empty();
 		OutElements.Empty();
@@ -2716,7 +2548,7 @@ bool ULearningAgentsObservations::GetInclusiveUnionObservationToArrays(TArray<FN
 	OutElementNames.SetNumUninitialized(OutElementNum);
 	OutElements.SetNumUninitialized(OutElementNum);
 
-	if (!GetInclusiveUnionObservationToArrayViews(OutElementNames, OutElements, Object, Element, Tag))
+	if (!GetInclusiveUnionObservationToArrayViews(OutElementNames, OutElements, Element, Name))
 	{
 		OutElementNames.Empty();
 		OutElements.Empty();
@@ -2726,46 +2558,40 @@ bool ULearningAgentsObservations::GetInclusiveUnionObservationToArrays(TArray<FN
 	return true;
 }
 
-bool ULearningAgentsObservations::GetInclusiveUnionObservationToArrayViews(TArrayView<FName> OutElementNames, TArrayView<FLearningAgentsObservationObjectElement> OutElements, const ULearningAgentsObservationObject* Object, const FLearningAgentsObservationObjectElement Element, const FName Tag)
+bool ULearningAgentsObservationObject::GetInclusiveUnionObservationToArrayViews(TArrayView<FName> OutElementNames, TArrayView<FLearningAgentsObservationObjectElement> OutElements, const FLearningAgentsObservationObjectElement Element, const FName Name) const
 {
-	if (!Object)
+	if (!ObservationObject.IsValid(Element.ObjectElement))
 	{
-		UE_LOG(LogLearning, Error, TEXT("GetInclusiveUnionObservationToArrayViews: Object is nullptr."));
+		UE_LOG(LogLearning, Error, TEXT("%s: Invalid Observation Object."), *GetName());
 		UE::Learning::Array::Set<1, FName>(OutElementNames, NAME_None);
 		UE::Learning::Array::Set<1, FLearningAgentsObservationObjectElement>(OutElements, FLearningAgentsObservationObjectElement());
 		return false;
 	}
 
-	if (!Object->ObservationObject.IsValid(Element.ObjectElement))
+	if (ObservationObject.GetName(Element.ObjectElement) != Name)
 	{
-		UE_LOG(LogLearning, Error, TEXT("GetInclusiveUnionObservationToArrayViews: Invalid Observation Object."));
-		UE::Learning::Array::Set<1, FName>(OutElementNames, NAME_None);
-		UE::Learning::Array::Set<1, FLearningAgentsObservationObjectElement>(OutElements, FLearningAgentsObservationObjectElement());
-		return false;
+		UE_LOG(LogLearning, Warning, TEXT("%s: Observation name does not match. Observation is '%s' but asked for '%s'."), *GetName(), *ObservationObject.GetName(Element.ObjectElement).ToString(), *Name.ToString());
 	}
 
-	if (Object->ObservationObject.GetTag(Element.ObjectElement) != Tag)
+	if (ObservationObject.GetType(Element.ObjectElement) != UE::Learning::Observation::EType::OrInclusive)
 	{
-		UE_LOG(LogLearning, Warning, TEXT("GetInclusiveUnionObservationToArrayViews: Observation tag does not match. Observation is '%s' but asked for '%s'."), *Object->ObservationObject.GetTag(Element.ObjectElement).ToString(), *Tag.ToString());
-	}
-
-	if (Object->ObservationObject.GetType(Element.ObjectElement) != UE::Learning::Observation::EType::OrInclusive)
-	{
-		UE_LOG(LogLearning, Error, TEXT("GetInclusiveUnionObservationToArrayViews: Observation '%s' type does not match. Observation is '%s' but asked for '%s'."),
-			*Object->ObservationObject.GetTag(Element.ObjectElement).ToString(),
-			UE::Learning::Agents::Observation::Private::GetObservationTypeString(Object->ObservationObject.GetType(Element.ObjectElement)),
+		UE_LOG(LogLearning, Error, TEXT("%s: Observation '%s' type does not match. Observation is '%s' but asked for '%s'."),
+			*GetName(),
+			*ObservationObject.GetName(Element.ObjectElement).ToString(),
+			UE::Learning::Agents::Observation::Private::GetObservationTypeString(ObservationObject.GetType(Element.ObjectElement)),
 			UE::Learning::Agents::Observation::Private::GetObservationTypeString(UE::Learning::Observation::EType::OrInclusive));
 		UE::Learning::Array::Set<1, FName>(OutElementNames, NAME_None);
 		UE::Learning::Array::Set<1, FLearningAgentsObservationObjectElement>(OutElements, FLearningAgentsObservationObjectElement());
 		return false;
 	}
 
-	const UE::Learning::Observation::FObjectOrInclusiveParameters Parameters = Object->ObservationObject.GetOrInclusive(Element.ObjectElement);
+	const UE::Learning::Observation::FObjectOrInclusiveParameters Parameters = ObservationObject.GetOrInclusive(Element.ObjectElement);
 
 	if (Parameters.Elements.Num() != OutElements.Num())
 	{
-		UE_LOG(LogLearning, Error, TEXT("GetInclusiveUnionObservationToArrayViews: Observation '%s' size does not match. Observation is '%i' elements but asked for '%i'."),
-			*Object->ObservationObject.GetTag(Element.ObjectElement).ToString(),
+		UE_LOG(LogLearning, Error, TEXT("%s: Observation '%s' size does not match. Observation is '%i' elements but asked for '%i'."),
+			*GetName(),
+			*ObservationObject.GetName(Element.ObjectElement).ToString(),
 			Parameters.Elements.Num(), OutElements.Num());
 		UE::Learning::Array::Set<1, FName>(OutElementNames, NAME_None);
 		UE::Learning::Array::Set<1, FLearningAgentsObservationObjectElement>(OutElements, FLearningAgentsObservationObjectElement());
@@ -2774,9 +2600,9 @@ bool ULearningAgentsObservations::GetInclusiveUnionObservationToArrayViews(TArra
 
 	for (int32 ElementIdx = 0; ElementIdx < Parameters.Elements.Num(); ElementIdx++)
 	{
-		if (!Object->ObservationObject.IsValid(Parameters.Elements[ElementIdx]))
+		if (!ObservationObject.IsValid(Parameters.Elements[ElementIdx]))
 		{
-			UE_LOG(LogLearning, Error, TEXT("GetInclusiveUnionObservationToArrayViews: Invalid Observation Object."));
+			UE_LOG(LogLearning, Error, TEXT("%s: Invalid Observation Object."), *GetName());
 			UE::Learning::Array::Set<1, FLearningAgentsObservationObjectElement>(OutElements, FLearningAgentsObservationObjectElement());
 			UE::Learning::Array::Set<1, FName>(OutElementNames, NAME_None);
 			return false;
@@ -2789,45 +2615,39 @@ bool ULearningAgentsObservations::GetInclusiveUnionObservationToArrayViews(TArra
 	return true;
 }
 
-bool ULearningAgentsObservations::GetStaticArrayObservationNum(int32& OutNum, const ULearningAgentsObservationObject* Object, const FLearningAgentsObservationObjectElement Element, const FName Tag)
+bool ULearningAgentsObservationObject::GetStaticArrayObservationNum(int32& OutNum, const FLearningAgentsObservationObjectElement Element, const FName Name) const
 {
-	if (!Object)
-	{
-		UE_LOG(LogLearning, Error, TEXT("GetStaticArrayObservationNum: Object is nullptr."));
-		OutNum = 0;
-		return false;
-	}
-
-	if (!Object->ObservationObject.IsValid(Element.ObjectElement))
+	if (!ObservationObject.IsValid(Element.ObjectElement))
 	{
-		UE_LOG(LogLearning, Error, TEXT("GetStaticArrayObservationNum: Invalid Observation Object."));
+		UE_LOG(LogLearning, Error, TEXT("%s: Invalid Observation Object."), *GetName());
 		OutNum = 0;
 		return false;
 	}
 
-	if (Object->ObservationObject.GetTag(Element.ObjectElement) != Tag)
+	if (ObservationObject.GetName(Element.ObjectElement) != Name)
 	{
-		UE_LOG(LogLearning, Warning, TEXT("GetStaticArrayObservationNum: Observation tag does not match. Observation is '%s' but asked for '%s'."), *Object->ObservationObject.GetTag(Element.ObjectElement).ToString(), *Tag.ToString());
+		UE_LOG(LogLearning, Warning, TEXT("%s: Observation name does not match. Observation is '%s' but asked for '%s'."), *GetName(), *ObservationObject.GetName(Element.ObjectElement).ToString(), *Name.ToString());
 	}
 
-	if (Object->ObservationObject.GetType(Element.ObjectElement) != UE::Learning::Observation::EType::Array)
+	if (ObservationObject.GetType(Element.ObjectElement) != UE::Learning::Observation::EType::Array)
 	{
-		UE_LOG(LogLearning, Error, TEXT("GetStaticArrayObservationNum: Observation '%s' type does not match. Observation is '%s' but asked for '%s'."),
-			*Object->ObservationObject.GetTag(Element.ObjectElement).ToString(),
-			UE::Learning::Agents::Observation::Private::GetObservationTypeString(Object->ObservationObject.GetType(Element.ObjectElement)),
+		UE_LOG(LogLearning, Error, TEXT("%s: Observation '%s' type does not match. Observation is '%s' but asked for '%s'."),
+			*GetName(),
+			*ObservationObject.GetName(Element.ObjectElement).ToString(),
+			UE::Learning::Agents::Observation::Private::GetObservationTypeString(ObservationObject.GetType(Element.ObjectElement)),
 			UE::Learning::Agents::Observation::Private::GetObservationTypeString(UE::Learning::Observation::EType::Array));
 		OutNum = 0;
 		return false;
 	}
 
-	OutNum = Object->ObservationObject.GetArray(Element.ObjectElement).Elements.Num();
+	OutNum = ObservationObject.GetArray(Element.ObjectElement).Elements.Num();
 	return true;
 }
 
-bool ULearningAgentsObservations::GetStaticArrayObservation(TArray<FLearningAgentsObservationObjectElement>& OutElements, const ULearningAgentsObservationObject* Object, const FLearningAgentsObservationObjectElement Element, const FName Tag)
+bool ULearningAgentsObservationObject::GetStaticArrayObservation(TArray<FLearningAgentsObservationObjectElement>& OutElements, const FLearningAgentsObservationObjectElement Element, const FName Name) const
 {
 	int32 OutElementNum = 0;
-	if (!GetStaticArrayObservationNum(OutElementNum, Object, Element, Tag))
+	if (!GetStaticArrayObservationNum(OutElementNum, Element, Name))
 	{
 		OutElements.Empty();
 		return false;
@@ -2835,7 +2655,7 @@ bool ULearningAgentsObservations::GetStaticArrayObservation(TArray<FLearningAgen
 
 	OutElements.SetNumUninitialized(OutElementNum);
 
-	if (!GetStaticArrayObservationToArrayView(OutElements, Object, Element, Tag))
+	if (!GetStaticArrayObservationToArrayView(OutElements, Element, Name))
 	{
 		OutElements.Empty();
 		return false;
@@ -2844,48 +2664,42 @@ bool ULearningAgentsObservations::GetStaticArrayObservation(TArray<FLearningAgen
 	return true;
 }
 
-bool ULearningAgentsObservations::GetStaticArrayObservationToArrayView(TArrayView<FLearningAgentsObservationObjectElement> OutElements, const ULearningAgentsObservationObject* Object, const FLearningAgentsObservationObjectElement Element, const FName Tag)
+bool ULearningAgentsObservationObject::GetStaticArrayObservationToArrayView(TArrayView<FLearningAgentsObservationObjectElement> OutElements, const FLearningAgentsObservationObjectElement Element, const FName Name) const
 {
-	if (!Object)
+	if (!ObservationObject.IsValid(Element.ObjectElement))
 	{
-		UE_LOG(LogLearning, Error, TEXT("GetStaticArrayObservationNum: Object is nullptr."));
+		UE_LOG(LogLearning, Error, TEXT("%s: Invalid Observation Object."), *GetName());
 		UE::Learning::Array::Set<1, FLearningAgentsObservationObjectElement>(OutElements, FLearningAgentsObservationObjectElement());
 		return false;
 	}
 
-	if (!Object->ObservationObject.IsValid(Element.ObjectElement))
+	if (ObservationObject.GetName(Element.ObjectElement) != Name)
 	{
-		UE_LOG(LogLearning, Error, TEXT("GetStaticArrayObservationToArrayView: Invalid Observation Object."));
-		UE::Learning::Array::Set<1, FLearningAgentsObservationObjectElement>(OutElements, FLearningAgentsObservationObjectElement());
-		return false;
+		UE_LOG(LogLearning, Warning, TEXT("%s: Observation name does not match. Observation is '%s' but asked for '%s'."), *GetName(), *ObservationObject.GetName(Element.ObjectElement).ToString(), *Name.ToString());
 	}
 
-	if (Object->ObservationObject.GetTag(Element.ObjectElement) != Tag)
+	if (ObservationObject.GetType(Element.ObjectElement) != UE::Learning::Observation::EType::Array)
 	{
-		UE_LOG(LogLearning, Warning, TEXT("GetStaticArrayObservationToArrayView: Observation tag does not match. Observation is '%s' but asked for '%s'."), *Object->ObservationObject.GetTag(Element.ObjectElement).ToString(), *Tag.ToString());
-	}
-
-	if (Object->ObservationObject.GetType(Element.ObjectElement) != UE::Learning::Observation::EType::Array)
-	{
-		UE_LOG(LogLearning, Error, TEXT("GetStaticArrayObservationToArrayView: Observation '%s' type does not match. Observation is '%s' but asked for '%s'."),
-			*Object->ObservationObject.GetTag(Element.ObjectElement).ToString(),
-			UE::Learning::Agents::Observation::Private::GetObservationTypeString(Object->ObservationObject.GetType(Element.ObjectElement)),
+		UE_LOG(LogLearning, Error, TEXT("%s: Observation '%s' type does not match. Observation is '%s' but asked for '%s'."),
+			*GetName(),
+			*ObservationObject.GetName(Element.ObjectElement).ToString(),
+			UE::Learning::Agents::Observation::Private::GetObservationTypeString(ObservationObject.GetType(Element.ObjectElement)),
 			UE::Learning::Agents::Observation::Private::GetObservationTypeString(UE::Learning::Observation::EType::Array));
 		UE::Learning::Array::Set<1, FLearningAgentsObservationObjectElement>(OutElements, FLearningAgentsObservationObjectElement());
 		return false;
 	}
 
-	const UE::Learning::Observation::FObjectArrayParameters Parameters = Object->ObservationObject.GetArray(Element.ObjectElement);
+	const UE::Learning::Observation::FObjectArrayParameters Parameters = ObservationObject.GetArray(Element.ObjectElement);
 
 	if (Parameters.Elements.Num() == 0)
 	{
-		UE_LOG(LogLearning, Warning, TEXT("GetStaticArrayObservationToArrayView: Getting zero-sized Static Array Observation."));
+		UE_LOG(LogLearning, Warning, TEXT("%s: Getting zero-sized Static Array Observation."), *GetName());
 	}
 
 	if (Parameters.Elements.Num() != OutElements.Num())
 	{
-		UE_LOG(LogLearning, Error, TEXT("GetStaticArrayObservationToArrayView: Observation '%s' size does not match. Observation is '%i' elements but asked for '%i'."),
-			*Object->ObservationObject.GetTag(Element.ObjectElement).ToString(),
+		UE_LOG(LogLearning, Error, TEXT("%s: Observation '%s' size does not match. Observation is '%i' elements but asked for '%i'."),
+			*GetName(), *ObservationObject.GetName(Element.ObjectElement).ToString(),
 			Parameters.Elements.Num(), OutElements.Num());
 		UE::Learning::Array::Set<1, FLearningAgentsObservationObjectElement>(OutElements, FLearningAgentsObservationObjectElement());
 		return false;
@@ -2893,9 +2707,9 @@ bool ULearningAgentsObservations::GetStaticArrayObservationToArrayView(TArrayVie
 
 	for (int32 ElementIdx = 0; ElementIdx < Parameters.Elements.Num(); ElementIdx++)
 	{
-		if (!Object->ObservationObject.IsValid(Parameters.Elements[ElementIdx]))
+		if (!ObservationObject.IsValid(Parameters.Elements[ElementIdx]))
 		{
-			UE_LOG(LogLearning, Error, TEXT("GetStaticArrayObservationToArrayView: Invalid Observation Object."));
+			UE_LOG(LogLearning, Error, TEXT("%s: Invalid Observation Object."), *GetName());
 			UE::Learning::Array::Set<1, FLearningAgentsObservationObjectElement>(OutElements, FLearningAgentsObservationObjectElement());
 			return false;
 		}
@@ -2906,85 +2720,73 @@ bool ULearningAgentsObservations::GetStaticArrayObservationToArrayView(TArrayVie
 	return true;
 }
 
-bool ULearningAgentsObservations::GetSetObservationNum(int32& OutNum, const ULearningAgentsObservationObject* Object, const FLearningAgentsObservationObjectElement Element, const FName Tag)
+bool ULearningAgentsObservationObject::GetSetObservationNum(int32& OutNum, const FLearningAgentsObservationObjectElement Element, const FName Name) const
 {
-	if (!Object)
-	{
-		UE_LOG(LogLearning, Error, TEXT("GetSetObservationNum: Object is nullptr."));
-		OutNum = 0;
-		return false;
-	}
-
-	if (!Object->ObservationObject.IsValid(Element.ObjectElement))
+	if (!ObservationObject.IsValid(Element.ObjectElement))
 	{
-		UE_LOG(LogLearning, Error, TEXT("GetSetObservationNum: Invalid Observation Object."));
+		UE_LOG(LogLearning, Error, TEXT("%s: Invalid Observation Object."), *GetName());
 		OutNum = 0;
 		return false;
 	}
 
-	if (Object->ObservationObject.GetTag(Element.ObjectElement) != Tag)
+	if (ObservationObject.GetName(Element.ObjectElement) != Name)
 	{
-		UE_LOG(LogLearning, Warning, TEXT("GetSetObservationNum: Observation tag does not match. Observation is '%s' but asked for '%s'."), *Object->ObservationObject.GetTag(Element.ObjectElement).ToString(), *Tag.ToString());
+		UE_LOG(LogLearning, Warning, TEXT("%s: Observation name does not match. Observation is '%s' but asked for '%s'."), *GetName(), *ObservationObject.GetName(Element.ObjectElement).ToString(), *Name.ToString());
 	}
 
-	if (Object->ObservationObject.GetType(Element.ObjectElement) != UE::Learning::Observation::EType::Set)
+	if (ObservationObject.GetType(Element.ObjectElement) != UE::Learning::Observation::EType::Set)
 	{
-		UE_LOG(LogLearning, Error, TEXT("GetSetObservationNum: Observation '%s' type does not match. Observation is '%s' but asked for '%s'."),
-			*Object->ObservationObject.GetTag(Element.ObjectElement).ToString(),
-			UE::Learning::Agents::Observation::Private::GetObservationTypeString(Object->ObservationObject.GetType(Element.ObjectElement)),
+		UE_LOG(LogLearning, Error, TEXT("%s: Observation '%s' type does not match. Observation is '%s' but asked for '%s'."),
+			*GetName(),
+			*ObservationObject.GetName(Element.ObjectElement).ToString(),
+			UE::Learning::Agents::Observation::Private::GetObservationTypeString(ObservationObject.GetType(Element.ObjectElement)),
 			UE::Learning::Agents::Observation::Private::GetObservationTypeString(UE::Learning::Observation::EType::Set));
 		OutNum = 0;
 		return false;
 	}
 
-	OutNum = Object->ObservationObject.GetSet(Element.ObjectElement).Elements.Num();
+	OutNum = ObservationObject.GetSet(Element.ObjectElement).Elements.Num();
 	return true;
 }
 
-bool ULearningAgentsObservations::GetSetObservation(TSet<FLearningAgentsObservationObjectElement>& OutElements, const ULearningAgentsObservationObject* Object, const FLearningAgentsObservationObjectElement Element, const FName Tag)
+bool ULearningAgentsObservationObject::GetSetObservation(TSet<FLearningAgentsObservationObjectElement>& OutElements, const FLearningAgentsObservationObjectElement Element, const FName Name) const
 {
 	int32 OutElementNum = 0;
-	if (!GetSetObservationNum(OutElementNum, Object, Element, Tag))
+	if (!GetSetObservationNum(OutElementNum, Element, Name))
 	{
 		OutElements.Empty();
 		return false;
 	}
 
-	if (!Object)
+	if (!ObservationObject.IsValid(Element.ObjectElement))
 	{
-		UE_LOG(LogLearning, Error, TEXT("GetSetObservation: Object is nullptr."));
+		UE_LOG(LogLearning, Error, TEXT("%s: Invalid Observation Object."), *GetName());
 		OutElements.Empty();
 		return false;
 	}
 
-	if (!Object->ObservationObject.IsValid(Element.ObjectElement))
+	if (ObservationObject.GetName(Element.ObjectElement) != Name)
 	{
-		UE_LOG(LogLearning, Error, TEXT("GetSetObservation: Invalid Observation Object."));
-		OutElements.Empty();
-		return false;
+		UE_LOG(LogLearning, Warning, TEXT("%s: Observation name does not match. Observation is '%s' but asked for '%s'."), *GetName(), *ObservationObject.GetName(Element.ObjectElement).ToString(), *Name.ToString());
 	}
 
-	if (Object->ObservationObject.GetTag(Element.ObjectElement) != Tag)
+	if (ObservationObject.GetType(Element.ObjectElement) != UE::Learning::Observation::EType::Set)
 	{
-		UE_LOG(LogLearning, Warning, TEXT("GetSetObservation: Observation tag does not match. Observation is '%s' but asked for '%s'."), *Object->ObservationObject.GetTag(Element.ObjectElement).ToString(), *Tag.ToString());
-	}
-
-	if (Object->ObservationObject.GetType(Element.ObjectElement) != UE::Learning::Observation::EType::Set)
-	{
-		UE_LOG(LogLearning, Error, TEXT("GetSetObservation: Observation '%s' type does not match. Observation is '%s' but asked for '%s'."),
-			*Object->ObservationObject.GetTag(Element.ObjectElement).ToString(),
-			UE::Learning::Agents::Observation::Private::GetObservationTypeString(Object->ObservationObject.GetType(Element.ObjectElement)),
+		UE_LOG(LogLearning, Error, TEXT("%s: Observation '%s' type does not match. Observation is '%s' but asked for '%s'."),
+			*GetName(),
+			*ObservationObject.GetName(Element.ObjectElement).ToString(),
+			UE::Learning::Agents::Observation::Private::GetObservationTypeString(ObservationObject.GetType(Element.ObjectElement)),
 			UE::Learning::Agents::Observation::Private::GetObservationTypeString(UE::Learning::Observation::EType::Set));
 		OutElements.Empty();
 		return false;
 	}
 
-	const UE::Learning::Observation::FObjectSetParameters Parameters = Object->ObservationObject.GetSet(Element.ObjectElement);
+	const UE::Learning::Observation::FObjectSetParameters Parameters = ObservationObject.GetSet(Element.ObjectElement);
 
 	if (Parameters.Elements.Num() != OutElements.Num())
 	{
-		UE_LOG(LogLearning, Error, TEXT("GetSetObservation: Observation '%s' size does not match. Observation is '%i' elements but asked for '%i'."),
-			*Object->ObservationObject.GetTag(Element.ObjectElement).ToString(),
+		UE_LOG(LogLearning, Error, TEXT("%s: Observation '%s' size does not match. Observation is '%i' elements but asked for '%i'."),
+			*GetName(), *ObservationObject.GetName(Element.ObjectElement).ToString(),
 			Parameters.Elements.Num(), OutElements.Num());
 		OutElements.Empty();
 		return false;
@@ -2993,9 +2795,9 @@ bool ULearningAgentsObservations::GetSetObservation(TSet<FLearningAgentsObservat
 	OutElements.Empty(Parameters.Elements.Num());
 	for (int32 ElementIdx = 0; ElementIdx < Parameters.Elements.Num(); ElementIdx++)
 	{
-		if (!Object->ObservationObject.IsValid(Parameters.Elements[ElementIdx]))
+		if (!ObservationObject.IsValid(Parameters.Elements[ElementIdx]))
 		{
-			UE_LOG(LogLearning, Error, TEXT("GetSetObservation: Invalid Observation Object."));
+			UE_LOG(LogLearning, Error, TEXT("%s: Invalid Observation Object."), *GetName());
 			OutElements.Empty();
 			return false;
 		}
@@ -3006,10 +2808,10 @@ bool ULearningAgentsObservations::GetSetObservation(TSet<FLearningAgentsObservat
 	return true;
 }
 
-bool ULearningAgentsObservations::GetSetObservationToArray(TArray<FLearningAgentsObservationObjectElement>& OutElements, const ULearningAgentsObservationObject* Object, const FLearningAgentsObservationObjectElement Element, const FName Tag)
+bool ULearningAgentsObservationObject::GetSetObservationToArray(TArray<FLearningAgentsObservationObjectElement>& OutElements, const FLearningAgentsObservationObjectElement Element, const FName Name) const
 {
 	int32 OutElementNum = 0;
-	if (!GetSetObservationNum(OutElementNum, Object, Element, Tag))
+	if (!GetSetObservationNum(OutElementNum, Element, Name))
 	{
 		OutElements.Empty();
 		return false;
@@ -3017,7 +2819,7 @@ bool ULearningAgentsObservations::GetSetObservationToArray(TArray<FLearningAgent
 
 	OutElements.SetNumUninitialized(OutElementNum);
 
-	if (!GetSetObservationToArrayView(OutElements, Object, Element, Tag))
+	if (!GetSetObservationToArrayView(OutElements, Element, Name))
 	{
 		OutElements.Empty();
 		return false;
@@ -3026,43 +2828,37 @@ bool ULearningAgentsObservations::GetSetObservationToArray(TArray<FLearningAgent
 	return true;
 }
 
-bool ULearningAgentsObservations::GetSetObservationToArrayView(TArrayView<FLearningAgentsObservationObjectElement> OutElements, const ULearningAgentsObservationObject* Object, const FLearningAgentsObservationObjectElement Element, const FName Tag)
+bool ULearningAgentsObservationObject::GetSetObservationToArrayView(TArrayView<FLearningAgentsObservationObjectElement> OutElements, const FLearningAgentsObservationObjectElement Element, const FName Name) const
 {
-	if (!Object)
-	{
-		UE_LOG(LogLearning, Error, TEXT("GetSetObservationToArrayView: Object is nullptr."));
-		UE::Learning::Array::Set<1, FLearningAgentsObservationObjectElement>(OutElements, FLearningAgentsObservationObjectElement());
-		return false;
-	}
-
-	if (!Object->ObservationObject.IsValid(Element.ObjectElement))
+	if (!ObservationObject.IsValid(Element.ObjectElement))
 	{
-		UE_LOG(LogLearning, Error, TEXT("GetSetObservationToArrayView: Invalid Observation Object."));
+		UE_LOG(LogLearning, Error, TEXT("%s: Invalid Observation Object."), *GetName());
 		UE::Learning::Array::Set<1, FLearningAgentsObservationObjectElement>(OutElements, FLearningAgentsObservationObjectElement());
 		return false;
 	}
 
-	if (Object->ObservationObject.GetTag(Element.ObjectElement) != Tag)
+	if (ObservationObject.GetName(Element.ObjectElement) != Name)
 	{
-		UE_LOG(LogLearning, Warning, TEXT("GetSetObservationToArrayView: Observation tag does not match. Observation is '%s' but asked for '%s'."), *Object->ObservationObject.GetTag(Element.ObjectElement).ToString(), *Tag.ToString());
+		UE_LOG(LogLearning, Warning, TEXT("%s: Observation name does not match. Observation is '%s' but asked for '%s'."), *GetName(), *ObservationObject.GetName(Element.ObjectElement).ToString(), *Name.ToString());
 	}
 
-	if (Object->ObservationObject.GetType(Element.ObjectElement) != UE::Learning::Observation::EType::Set)
+	if (ObservationObject.GetType(Element.ObjectElement) != UE::Learning::Observation::EType::Set)
 	{
-		UE_LOG(LogLearning, Error, TEXT("GetSetObservationToArrayView: Observation '%s' type does not match. Observation is '%s' but asked for '%s'."),
-			*Object->ObservationObject.GetTag(Element.ObjectElement).ToString(),
-			UE::Learning::Agents::Observation::Private::GetObservationTypeString(Object->ObservationObject.GetType(Element.ObjectElement)),
+		UE_LOG(LogLearning, Error, TEXT("%s: Observation '%s' type does not match. Observation is '%s' but asked for '%s'."),
+			*GetName(),
+			*ObservationObject.GetName(Element.ObjectElement).ToString(),
+			UE::Learning::Agents::Observation::Private::GetObservationTypeString(ObservationObject.GetType(Element.ObjectElement)),
 			UE::Learning::Agents::Observation::Private::GetObservationTypeString(UE::Learning::Observation::EType::Set));
 		UE::Learning::Array::Set<1, FLearningAgentsObservationObjectElement>(OutElements, FLearningAgentsObservationObjectElement());
 		return false;
 	}
 
-	const UE::Learning::Observation::FObjectSetParameters Parameters = Object->ObservationObject.GetSet(Element.ObjectElement);
+	const UE::Learning::Observation::FObjectSetParameters Parameters = ObservationObject.GetSet(Element.ObjectElement);
 
 	if (Parameters.Elements.Num() != OutElements.Num())
 	{
-		UE_LOG(LogLearning, Error, TEXT("GetSetObservationToArrayView: Observation '%s' size does not match. Observation is '%i' elements but asked for '%i'."),
-			*Object->ObservationObject.GetTag(Element.ObjectElement).ToString(),
+		UE_LOG(LogLearning, Error, TEXT("%s: Observation '%s' size does not match. Observation is '%i' elements but asked for '%i'."),
+			*GetName(), *ObservationObject.GetName(Element.ObjectElement).ToString(),
 			Parameters.Elements.Num(), OutElements.Num());
 		UE::Learning::Array::Set<1, FLearningAgentsObservationObjectElement>(OutElements, FLearningAgentsObservationObjectElement());
 		return false;
@@ -3070,9 +2866,9 @@ bool ULearningAgentsObservations::GetSetObservationToArrayView(TArrayView<FLearn
 
 	for (int32 ElementIdx = 0; ElementIdx < Parameters.Elements.Num(); ElementIdx++)
 	{
-		if (!Object->ObservationObject.IsValid(Parameters.Elements[ElementIdx]))
+		if (!ObservationObject.IsValid(Parameters.Elements[ElementIdx]))
 		{
-			UE_LOG(LogLearning, Error, TEXT("GetSetObservationToArrayView: Invalid Observation Object."));
+			UE_LOG(LogLearning, Error, TEXT("%s: Invalid Observation Object."), *GetName());
 			UE::Learning::Array::Set<1, FLearningAgentsObservationObjectElement>(OutElements, FLearningAgentsObservationObjectElement());
 			return false;
 		}
@@ -3083,11 +2879,11 @@ bool ULearningAgentsObservations::GetSetObservationToArrayView(TArrayView<FLearn
 	return true;
 }
 
-bool ULearningAgentsObservations::GetPairObservation(FLearningAgentsObservationObjectElement& OutKey, FLearningAgentsObservationObjectElement& OutValue, const ULearningAgentsObservationObject* Object, const FLearningAgentsObservationObjectElement Element, const FName Tag)
+bool ULearningAgentsObservationObject::GetPairObservation(FLearningAgentsObservationObjectElement& OutKey, FLearningAgentsObservationObjectElement& OutValue, const FLearningAgentsObservationObjectElement Element, const FName Name) const
 {
 	TStaticArray<FName, 2> OutElementNames;
 	TStaticArray<FLearningAgentsObservationObjectElement, 2> OutElements;
-	if (!GetStructObservationToArrayViews(OutElementNames, OutElements, Object, Element, Tag))
+	if (!GetStructObservationToArrayViews(OutElementNames, OutElements, Element, Name))
 	{
 		OutKey = FLearningAgentsObservationObjectElement();
 		OutValue = FLearningAgentsObservationObjectElement();
@@ -3099,15 +2895,15 @@ bool ULearningAgentsObservations::GetPairObservation(FLearningAgentsObservationO
 	return true;
 }
 
-bool ULearningAgentsObservations::GetArrayObservationNum(int32& OutNum, const ULearningAgentsObservationObject* Object, const FLearningAgentsObservationObjectElement Element, const FName Tag)
+bool ULearningAgentsObservationObject::GetArrayObservationNum(int32& OutNum, const FLearningAgentsObservationObjectElement Element, const FName Name) const
 {
-	return GetSetObservationNum(OutNum, Object, Element, Tag);
+	return GetSetObservationNum(OutNum, Element, Name);
 }
 
-bool ULearningAgentsObservations::GetArrayObservation(TArray<FLearningAgentsObservationObjectElement>& OutElements, const ULearningAgentsObservationObject* Object, const FLearningAgentsObservationObjectElement Element, const FName Tag)
+bool ULearningAgentsObservationObject::GetArrayObservation(TArray<FLearningAgentsObservationObjectElement>& OutElements, const FLearningAgentsObservationObjectElement Element, const FName Name) const
 {
 	int32 OutElementNum = 0;
-	if (!GetArrayObservationNum(OutElementNum, Object, Element, Tag))
+	if (!GetArrayObservationNum(OutElementNum, Element, Name))
 	{
 		OutElements.Empty();
 		return false;
@@ -3115,7 +2911,7 @@ bool ULearningAgentsObservations::GetArrayObservation(TArray<FLearningAgentsObse
 
 	OutElements.SetNumUninitialized(OutElementNum);
 
-	if (!GetArrayObservationToArrayView(OutElements, Object, Element, Tag))
+	if (!GetArrayObservationToArrayView(OutElements, Element, Name))
 	{
 		OutElements.Empty();
 		return false;
@@ -3124,11 +2920,11 @@ bool ULearningAgentsObservations::GetArrayObservation(TArray<FLearningAgentsObse
 	return true;
 }
 
-bool ULearningAgentsObservations::GetArrayObservationToArrayView(TArrayView<FLearningAgentsObservationObjectElement> OutElements, const ULearningAgentsObservationObject* Object, const FLearningAgentsObservationObjectElement Element, const FName Tag)
+bool ULearningAgentsObservationObject::GetArrayObservationToArrayView(TArrayView<FLearningAgentsObservationObjectElement> OutElements, const FLearningAgentsObservationObjectElement Element, const FName Name) const
 {
 	TArray<FLearningAgentsObservationObjectElement, TInlineAllocator<16>> Pairs;
 	Pairs.SetNumUninitialized(OutElements.Num());
-	if (!GetSetObservationToArrayView(Pairs, Object, Element, Tag))
+	if (!GetSetObservationToArrayView(Pairs, Element, Name))
 	{
 		UE::Learning::Array::Set<1, FLearningAgentsObservationObjectElement>(OutElements, FLearningAgentsObservationObjectElement());
 		return false;
@@ -3137,14 +2933,14 @@ bool ULearningAgentsObservations::GetArrayObservationToArrayView(TArrayView<FLea
 	for (int32 PairIdx = 0; PairIdx < Pairs.Num(); PairIdx++)
 	{
 		FLearningAgentsObservationObjectElement Key, Value;
-		if (!GetPairObservation(Key, Value, Object, Pairs[PairIdx]))
+		if (!GetPairObservation(Pairs[PairIdx], Key, Value))
 		{
 			UE::Learning::Array::Set<1, FLearningAgentsObservationObjectElement>(OutElements, FLearningAgentsObservationObjectElement());
 			return false;
 		}
 
 		int32 Index = INDEX_NONE;
-		if (!GetIndexObservation(Index, Object, Key))
+		if (!GetIndexObservation(Index, Key))
 		{
 			UE::Learning::Array::Set<1, FLearningAgentsObservationObjectElement>(OutElements, FLearningAgentsObservationObjectElement());
 			return false;
@@ -3156,15 +2952,15 @@ bool ULearningAgentsObservations::GetArrayObservationToArrayView(TArrayView<FLea
 	return true;
 }
 
-bool ULearningAgentsObservations::GetMapObservationNum(int32& OutNum, const ULearningAgentsObservationObject* Object, const FLearningAgentsObservationObjectElement Element, const FName Tag)
+bool ULearningAgentsObservationObject::GetMapObservationNum(int32& OutNum, const FLearningAgentsObservationObjectElement Element, const FName Name) const
 {
-	return GetSetObservationNum(OutNum, Object, Element, Tag);
+	return GetSetObservationNum(OutNum, Element, Name);
 }
 
-bool ULearningAgentsObservations::GetMapObservation(TMap<FLearningAgentsObservationObjectElement, FLearningAgentsObservationObjectElement>& OutElements, const ULearningAgentsObservationObject* Object, const FLearningAgentsObservationObjectElement Element, const FName Tag)
+bool ULearningAgentsObservationObject::GetMapObservation(TMap<FLearningAgentsObservationObjectElement, FLearningAgentsObservationObjectElement>& OutElements, const FLearningAgentsObservationObjectElement Element, const FName Name) const
 {
 	int32 OutElementNum = 0;
-	if (!GetMapObservationNum(OutElementNum, Object, Element, Tag))
+	if (!GetMapObservationNum(OutElementNum, Element, Name))
 	{
 		OutElements.Empty();
 		return false;
@@ -3172,7 +2968,7 @@ bool ULearningAgentsObservations::GetMapObservation(TMap<FLearningAgentsObservat
 
 	TArray<FLearningAgentsObservationObjectElement, TInlineAllocator<16>> Pairs;
 	Pairs.SetNumUninitialized(OutElementNum);
-	if (!GetSetObservationToArrayView(Pairs, Object, Element, Tag))
+	if (!GetSetObservationToArrayView(Pairs, Element, Name))
 	{
 		OutElements.Empty();
 		return false;
@@ -3182,7 +2978,7 @@ bool ULearningAgentsObservations::GetMapObservation(TMap<FLearningAgentsObservat
 	for (int32 PairIdx = 0; PairIdx < OutElementNum; PairIdx++)
 	{
 		FLearningAgentsObservationObjectElement Key, Value;
-		if (!GetPairObservation(Key, Value, Object, Pairs[PairIdx]))
+		if (!GetPairObservation(Pairs[PairIdx], Key, Value))
 		{
 			OutElements.Empty();
 			return false;
@@ -3194,10 +2990,10 @@ bool ULearningAgentsObservations::GetMapObservation(TMap<FLearningAgentsObservat
 	return true;
 }
 
-bool ULearningAgentsObservations::GetMapObservationToArrays(TArray<FLearningAgentsObservationObjectElement>& OutKeys, TArray<FLearningAgentsObservationObjectElement>& OutValues, const ULearningAgentsObservationObject* Object, const FLearningAgentsObservationObjectElement Element, const FName Tag)
+bool ULearningAgentsObservationObject::GetMapObservationToArrays(TArray<FLearningAgentsObservationObjectElement>& OutKeys, TArray<FLearningAgentsObservationObjectElement>& OutValues, const FLearningAgentsObservationObjectElement Element, const FName Name) const
 {
 	int32 OutElementNum = 0;
-	if (!GetMapObservationNum(OutElementNum, Object, Element, Tag))
+	if (!GetMapObservationNum(OutElementNum, Element, Name))
 	{
 		OutKeys.Empty();
 		OutValues.Empty();
@@ -3206,7 +3002,7 @@ bool ULearningAgentsObservations::GetMapObservationToArrays(TArray<FLearningAgen
 
 	OutKeys.SetNumUninitialized(OutElementNum);
 	OutValues.SetNumUninitialized(OutElementNum);
-	if (!GetMapObservationToArrayViews(OutKeys, OutValues, Object, Element, Tag))
+	if (!GetMapObservationToArrayViews(OutKeys, OutValues, Element, Name))
 	{
 		OutKeys.Empty();
 		OutValues.Empty();
@@ -3216,11 +3012,11 @@ bool ULearningAgentsObservations::GetMapObservationToArrays(TArray<FLearningAgen
 	return true;
 }
 
-bool ULearningAgentsObservations::GetMapObservationToArrayViews(TArrayView<FLearningAgentsObservationObjectElement> OutKeys, TArrayView<FLearningAgentsObservationObjectElement> OutValues, const ULearningAgentsObservationObject* Object, const FLearningAgentsObservationObjectElement Element, const FName Tag)
+bool ULearningAgentsObservationObject::GetMapObservationToArrayViews(TArrayView<FLearningAgentsObservationObjectElement> OutKeys, TArrayView<FLearningAgentsObservationObjectElement> OutValues, const FLearningAgentsObservationObjectElement Element, const FName Name) const
 {
 	TArray<FLearningAgentsObservationObjectElement, TInlineAllocator<16>> Pairs;
 	Pairs.SetNumUninitialized(OutKeys.Num());
-	if (!GetSetObservationToArrayView(Pairs, Object, Element, Tag))
+	if (!GetSetObservationToArrayView(Pairs, Element, Name))
 	{
 		UE::Learning::Array::Set<1, FLearningAgentsObservationObjectElement>(OutKeys, FLearningAgentsObservationObjectElement());
 		UE::Learning::Array::Set<1, FLearningAgentsObservationObjectElement>(OutValues, FLearningAgentsObservationObjectElement());
@@ -3230,7 +3026,7 @@ bool ULearningAgentsObservations::GetMapObservationToArrayViews(TArrayView<FLear
 	for (int32 PairIdx = 0; PairIdx < Pairs.Num(); PairIdx++)
 	{
 		FLearningAgentsObservationObjectElement Key, Value;
-		if (!GetPairObservation(Key, Value, Object, Pairs[PairIdx]))
+		if (!GetPairObservation(Pairs[PairIdx], Key, Value))
 		{
 			UE::Learning::Array::Set<1, FLearningAgentsObservationObjectElement>(OutKeys, FLearningAgentsObservationObjectElement());
 			UE::Learning::Array::Set<1, FLearningAgentsObservationObjectElement>(OutValues, FLearningAgentsObservationObjectElement());
@@ -3244,17 +3040,17 @@ bool ULearningAgentsObservations::GetMapObservationToArrayViews(TArrayView<FLear
 	return true;
 }
 
-bool ULearningAgentsObservations::GetEnumObservation(uint8& OutEnumValue, const ULearningAgentsObservationObject* Object, const FLearningAgentsObservationObjectElement Element, const UEnum* Enum, const FName Tag)
+bool ULearningAgentsObservationObject::GetEnumObservation(uint8& OutEnumValue, const UEnum* Enum, const FLearningAgentsObservationObjectElement Element, const FName Name) const
 {
 	if (!Enum)
 	{
-		UE_LOG(LogLearning, Error, TEXT("GetEnumObservation: Enum is nullptr."));
+		UE_LOG(LogLearning, Error, TEXT("%s: Enum is nullptr."), *GetName());
 		OutEnumValue = 0;
 		return false;
 	}
 
 	int32 EnumValueNum;
-	if (!GetContinuousObservationNum(EnumValueNum, Object, Element, Tag))
+	if (!GetContinuousObservationNum(EnumValueNum, Element, Name))
 	{
 		OutEnumValue = 0;
 		return false;
@@ -3262,14 +3058,14 @@ bool ULearningAgentsObservations::GetEnumObservation(uint8& OutEnumValue, const
 	
 	if (EnumValueNum != Enum->NumEnums() - 1)
 	{
-		UE_LOG(LogLearning, Error, TEXT("GetEnumObservation: Too many values for Enum '%s'. Expected %i, got %i."), *Enum->GetName(), Enum->NumEnums() - 1, EnumValueNum);
+		UE_LOG(LogLearning, Error, TEXT("%s: Too many values for Enum '%s'. Expected %i, got %i."), *GetName(), *Enum->GetName(), Enum->NumEnums() - 1, EnumValueNum);
 		OutEnumValue = 0;
 		return false;
 	}
 
 	TArray<float, TInlineAllocator<32>> OneHot;
 	OneHot.SetNumUninitialized(EnumValueNum);
-	if (!GetContinuousObservationToArrayView(OneHot, Object, Element, Tag))
+	if (!GetContinuousObservationToArrayView(OneHot, Element, Name))
 	{
 		OutEnumValue = 0;
 		return false;
@@ -3287,7 +3083,7 @@ bool ULearningAgentsObservations::GetEnumObservation(uint8& OutEnumValue, const
 
 	if (EnumValueIndex == INDEX_NONE)
 	{
-		UE_LOG(LogLearning, Error, TEXT("GetEnumObservation: Index not found."));
+		UE_LOG(LogLearning, Error, TEXT("%s: Index not found."), *GetName());
 		OutEnumValue = 0;
 		return false;
 	}
@@ -3296,7 +3092,7 @@ bool ULearningAgentsObservations::GetEnumObservation(uint8& OutEnumValue, const
 
 	if (EnumValue == INDEX_NONE)
 	{
-		UE_LOG(LogLearning, Error, TEXT("GetEnumObservation: Enum Value not found for index %i."), EnumValueIndex);
+		UE_LOG(LogLearning, Error, TEXT("%s: Enum Value not found for index %i."), *GetName(), EnumValueIndex);
 		OutEnumValue = 0;
 		return false;
 	}
@@ -3305,24 +3101,24 @@ bool ULearningAgentsObservations::GetEnumObservation(uint8& OutEnumValue, const
 	return true;
 }
 
-bool ULearningAgentsObservations::GetBitmaskObservation(int32& OutBitmaskValue, const ULearningAgentsObservationObject* Object, const FLearningAgentsObservationObjectElement Element, const UEnum* Enum, const FName Tag)
+bool ULearningAgentsObservationObject::GetBitmaskObservation(int32& OutBitmaskValue, const UEnum* Enum, const FLearningAgentsObservationObjectElement Element, const FName Name) const
 {
 	if (!Enum)
 	{
-		UE_LOG(LogLearning, Error, TEXT("GetBitmaskObservation: Enum is nullptr."));
+		UE_LOG(LogLearning, Error, TEXT("%s: Enum is nullptr."), *GetName());
 		OutBitmaskValue = 0;
 		return false;
 	}
 
 	if (Enum->NumEnums() - 1 > 32)
 	{
-		UE_LOG(LogLearning, Error, TEXT("GetBitmaskObservation: Too many values in Enum to use as Bitmask (%i)."), Enum->NumEnums() - 1);
+		UE_LOG(LogLearning, Error, TEXT("%s: Too many values in Enum to use as Bitmask (%i)."), *GetName(), Enum->NumEnums() - 1);
 		OutBitmaskValue = 0;
 		return false;
 	}
 
 	int32 EnumValueNum;
-	if (!GetContinuousObservationNum(EnumValueNum, Object, Element, Tag))
+	if (!GetContinuousObservationNum(EnumValueNum, Element, Name))
 	{
 		OutBitmaskValue = 0;
 		return false;
@@ -3330,14 +3126,14 @@ bool ULearningAgentsObservations::GetBitmaskObservation(int32& OutBitmaskValue,
 
 	if (EnumValueNum != Enum->NumEnums() - 1)
 	{
-		UE_LOG(LogLearning, Error, TEXT("GetBitmaskObservation: Too many values for Enum '%s'. Expected %i, got %i."), *Enum->GetName(), Enum->NumEnums() - 1, EnumValueNum);
+		UE_LOG(LogLearning, Error, TEXT("%s: Too many values for Enum '%s'. Expected %i, got %i."), *GetName(), *Enum->GetName(), Enum->NumEnums() - 1, EnumValueNum);
 		OutBitmaskValue = 0;
 		return false;
 	}
 
 	TArray<float, TInlineAllocator<32>> OneHot;
 	OneHot.Init(0.0f, EnumValueNum);
-	if (!GetContinuousObservationToArrayView(OneHot, Object, Element, Tag))
+	if (!GetContinuousObservationToArrayView(OneHot, Element, Name))
 	{
 		OutBitmaskValue = 0;
 		return false;
@@ -3355,10 +3151,10 @@ bool ULearningAgentsObservations::GetBitmaskObservation(int32& OutBitmaskValue,
 }
 
 
-bool ULearningAgentsObservations::GetOptionalObservation(ELearningAgentsOptionalObservation& OutOption, FLearningAgentsObservationObjectElement& OutElement, const ULearningAgentsObservationObject* Object, const FLearningAgentsObservationObjectElement Element, const FName Tag)
+bool ULearningAgentsObservationObject::GetOptionalObservation(ELearningAgentsOptionalObservation& OutOption, FLearningAgentsObservationObjectElement& OutElement, const FLearningAgentsObservationObjectElement Element, const FName Name) const
 {
 	FName OutName = NAME_None;
-	if (!GetExclusiveUnionObservation(OutName, OutElement, Object, Element, Tag))
+	if (!GetExclusiveUnionObservation(OutName, OutElement, Element, Name))
 	{
 		OutOption = ELearningAgentsOptionalObservation::Null;
 		return false;
@@ -3368,10 +3164,10 @@ bool ULearningAgentsObservations::GetOptionalObservation(ELearningAgentsOptional
 	return true;
 }
 
-bool ULearningAgentsObservations::GetEitherObservation(ELearningAgentsEitherObservation& OutEither, FLearningAgentsObservationObjectElement& OutElement, const ULearningAgentsObservationObject* Object, const FLearningAgentsObservationObjectElement Element, const FName Tag)
+bool ULearningAgentsObservationObject::GetEitherObservation(ELearningAgentsEitherObservation& OutEither, FLearningAgentsObservationObjectElement& OutElement, const FLearningAgentsObservationObjectElement Element, const FName Name) const
 {
 	FName OutName = NAME_None;
-	if (!GetExclusiveUnionObservation(OutName, OutElement, Object, Element, Tag))
+	if (!GetExclusiveUnionObservation(OutName, OutElement, Element, Name))
 	{
 		OutEither = ELearningAgentsEitherObservation::A;
 		return false;
@@ -3381,46 +3177,40 @@ bool ULearningAgentsObservations::GetEitherObservation(ELearningAgentsEitherObse
 	return true;
 }
 
-bool ULearningAgentsObservations::GetEncodingObservation(FLearningAgentsObservationObjectElement& OutElement, const ULearningAgentsObservationObject* Object, const FLearningAgentsObservationObjectElement Element, const FName Tag)
+bool ULearningAgentsObservationObject::GetEncodingObservation(FLearningAgentsObservationObjectElement& OutElement, const FLearningAgentsObservationObjectElement Element, const FName Name) const
 {
-	if (!Object)
-	{
-		UE_LOG(LogLearning, Error, TEXT("GetEncodingObservation: Object is nullptr."));
-		OutElement = FLearningAgentsObservationObjectElement();
-		return false;
-	}
-
-	if (!Object->ObservationObject.IsValid(Element.ObjectElement))
+	if (!ObservationObject.IsValid(Element.ObjectElement))
 	{
-		UE_LOG(LogLearning, Error, TEXT("GetEncodingObservation: Invalid Observation Object."));
+		UE_LOG(LogLearning, Error, TEXT("%s: Invalid Observation Object."), *GetName());
 		OutElement = FLearningAgentsObservationObjectElement();
 		return false;
 	}
 
-	if (Object->ObservationObject.GetTag(Element.ObjectElement) != Tag)
+	if (ObservationObject.GetName(Element.ObjectElement) != Name)
 	{
-		UE_LOG(LogLearning, Warning, TEXT("GetEncodingObservation: Observation tag does not match. Observation is '%s' but asked for '%s'."), *Object->ObservationObject.GetTag(Element.ObjectElement).ToString(), *Tag.ToString());
+		UE_LOG(LogLearning, Warning, TEXT("%s: Observation name does not match. Observation is '%s' but asked for '%s'."), *GetName(), *ObservationObject.GetName(Element.ObjectElement).ToString(), *Name.ToString());
 	}
 
-	if (Object->ObservationObject.GetType(Element.ObjectElement) != UE::Learning::Observation::EType::Encoding)
+	if (ObservationObject.GetType(Element.ObjectElement) != UE::Learning::Observation::EType::Encoding)
 	{
-		UE_LOG(LogLearning, Error, TEXT("GetEncodingObservation: Observation '%s' type does not match. Observation is '%s' but asked for '%s'."),
-			*Object->ObservationObject.GetTag(Element.ObjectElement).ToString(),
-			UE::Learning::Agents::Observation::Private::GetObservationTypeString(Object->ObservationObject.GetType(Element.ObjectElement)),
+		UE_LOG(LogLearning, Error, TEXT("%s: Observation '%s' type does not match. Observation is '%s' but asked for '%s'."),
+			*GetName(),
+			*ObservationObject.GetName(Element.ObjectElement).ToString(),
+			UE::Learning::Agents::Observation::Private::GetObservationTypeString(ObservationObject.GetType(Element.ObjectElement)),
 			UE::Learning::Agents::Observation::Private::GetObservationTypeString(UE::Learning::Observation::EType::Encoding));
 		OutElement = FLearningAgentsObservationObjectElement();
 		return false;
 	}
 
-	OutElement = { Object->ObservationObject.GetEncoding(Element.ObjectElement).Element };
+	OutElement = { ObservationObject.GetEncoding(Element.ObjectElement).Element };
 
 	return true;
 }
 
-bool ULearningAgentsObservations::GetBoolObservation(bool& bOutValue, const ULearningAgentsObservationObject* Object, const FLearningAgentsObservationObjectElement Element, const FName Tag)
+bool ULearningAgentsObservationObject::GetBoolObservation(bool& bOutValue, const FLearningAgentsObservationObjectElement Element, const FName Name) const
 {
 	float OutValue = 0.0f;
-	if (!GetFloatObservation(OutValue, Object, Element, 1.0f, Tag))
+	if (!GetFloatObservation(OutValue, Element, 1.0f, Name))
 	{
 		bOutValue = false;
 		return false;
@@ -3430,10 +3220,10 @@ bool ULearningAgentsObservations::GetBoolObservation(bool& bOutValue, const ULea
 	return true;
 }
 
-bool ULearningAgentsObservations::GetFloatObservation(float& OutValue, const ULearningAgentsObservationObject* Object, const FLearningAgentsObservationObjectElement Element, const float FloatScale, const FName Tag)
+bool ULearningAgentsObservationObject::GetFloatObservation(float& OutValue, const FLearningAgentsObservationObjectElement Element, const float FloatScale, const FName Name) const
 {
 	float OutValuesData;
-	if (!GetContinuousObservationToArrayView(MakeArrayView(&OutValuesData, 1), Object, Element, Tag))
+	if (!GetContinuousObservationToArrayView(MakeArrayView(&OutValuesData, 1), Element, Name))
 	{
 		OutValue = 0.0f;
 		return false;
@@ -3443,10 +3233,10 @@ bool ULearningAgentsObservations::GetFloatObservation(float& OutValue, const ULe
 	return true;
 }
 
-bool ULearningAgentsObservations::GetLocationObservation(FVector& OutLocation, const ULearningAgentsObservationObject* Object, const FLearningAgentsObservationObjectElement Element, const FTransform RelativeTransform, const float LocationScale, const FName Tag)
+bool ULearningAgentsObservationObject::GetLocationObservation(FVector& OutLocation, const FLearningAgentsObservationObjectElement Element, const FTransform RelativeTransform, const float LocationScale, const FName Name) const
 {
 	TStaticArray<float, 3> OutValues;
-	if (!GetContinuousObservationToArrayView(OutValues, Object, Element, Tag))
+	if (!GetContinuousObservationToArrayView(OutValues, Element, Name))
 	{
 		OutLocation = FVector::ZeroVector;
 		return false;
@@ -3456,10 +3246,10 @@ bool ULearningAgentsObservations::GetLocationObservation(FVector& OutLocation, c
 	return true;
 }
 
-bool ULearningAgentsObservations::GetRotationObservation(FRotator& OutRotation, const ULearningAgentsObservationObject* Object, const FLearningAgentsObservationObjectElement Element, const FRotator RelativeRotation, const FName Tag)
+bool ULearningAgentsObservationObject::GetRotationObservation(FRotator& OutRotation, const FLearningAgentsObservationObjectElement Element, const FRotator RelativeRotation, const FName Name) const
 {
 	FQuat OutRotationQuat;
-	if (!GetRotationObservationAsQuat(OutRotationQuat, Object, Element, FQuat::MakeFromRotator(RelativeRotation), Tag))
+	if (!GetRotationObservationAsQuat(OutRotationQuat, Element, FQuat::MakeFromRotator(RelativeRotation), Name))
 	{
 		OutRotation = FRotator::ZeroRotator;
 		return false;
@@ -3469,10 +3259,10 @@ bool ULearningAgentsObservations::GetRotationObservation(FRotator& OutRotation,
 	return true;
 }
 
-bool ULearningAgentsObservations::GetRotationObservationAsQuat(FQuat& OutRotation, const ULearningAgentsObservationObject* Object, const FLearningAgentsObservationObjectElement Element, const FQuat RelativeRotation, const FName Tag)
+bool ULearningAgentsObservationObject::GetRotationObservationAsQuat(FQuat& OutRotation, const FLearningAgentsObservationObjectElement Element, const FQuat RelativeRotation, const FName Name) const
 {
 	TStaticArray<float, 6> OutValues;
-	if (!GetContinuousObservationToArrayView(OutValues, Object, Element, Tag))
+	if (!GetContinuousObservationToArrayView(OutValues, Element, Name))
 	{
 		OutRotation = FQuat::Identity;
 		return false;
@@ -3493,10 +3283,10 @@ bool ULearningAgentsObservations::GetRotationObservationAsQuat(FQuat& OutRotatio
 	return true;
 }
 
-bool ULearningAgentsObservations::GetScaleObservation(FVector& OutScale, const ULearningAgentsObservationObject* Object, const FLearningAgentsObservationObjectElement Element, const FVector RelativeScale, const FName Tag)
+bool ULearningAgentsObservationObject::GetScaleObservation(FVector& OutScale, const FLearningAgentsObservationObjectElement Element, const FVector RelativeScale, const FName Name) const
 {
 	TStaticArray<float, 3> OutValues;
-	if (!GetContinuousObservationToArrayView(OutValues, Object, Element, Tag))
+	if (!GetContinuousObservationToArrayView(OutValues, Element, Name))
 	{
 		OutScale = FVector::OneVector;
 		return false;
@@ -3506,11 +3296,11 @@ bool ULearningAgentsObservations::GetScaleObservation(FVector& OutScale, const U
 	return true;
 }
 
-bool ULearningAgentsObservations::GetTransformObservation(FTransform& OutTransform, const ULearningAgentsObservationObject* Object, const FLearningAgentsObservationObjectElement Element, const FTransform RelativeTransform, const float LocationScale, const FName Tag)
+bool ULearningAgentsObservationObject::GetTransformObservation(FTransform& OutTransform, const FLearningAgentsObservationObjectElement Element, const FTransform RelativeTransform, const float LocationScale, const FName Name) const
 {
 	TStaticArray<FName, 3> OutElementNames;
 	TStaticArray<FLearningAgentsObservationObjectElement, 3> OutElements;
-	if (!GetStructObservationToArrayViews(OutElementNames, OutElements, Object, Element, Tag))
+	if (!GetStructObservationToArrayViews(OutElementNames, OutElements, Element, Name))
 	{
 		OutTransform = FTransform::Identity;
 		return false;
@@ -3518,7 +3308,7 @@ bool ULearningAgentsObservations::GetTransformObservation(FTransform& OutTransfo
 
 	const int32 LocationElement = MakeArrayView(OutElementNames).Find(TEXT("Location"));
 	FVector OutLocation;
-	if (LocationElement == INDEX_NONE || !GetLocationObservation(OutLocation, Object, OutElements[LocationElement], RelativeTransform, LocationScale))
+	if (LocationElement == INDEX_NONE || !GetLocationObservation(OutLocation, OutElements[LocationElement], RelativeTransform, LocationScale))
 	{
 		OutTransform = FTransform::Identity;
 		return false;
@@ -3526,7 +3316,7 @@ bool ULearningAgentsObservations::GetTransformObservation(FTransform& OutTransfo
 
 	const int32 RotationElement = MakeArrayView(OutElementNames).Find(TEXT("Rotation"));
 	FQuat OutRotation;
-	if (RotationElement == INDEX_NONE || !GetRotationObservationAsQuat(OutRotation, Object, OutElements[RotationElement], RelativeTransform.GetRotation()))
+	if (RotationElement == INDEX_NONE || !GetRotationObservationAsQuat(OutRotation, OutElements[RotationElement], RelativeTransform.GetRotation()))
 	{
 		OutTransform = FTransform::Identity;
 		return false;
@@ -3534,7 +3324,7 @@ bool ULearningAgentsObservations::GetTransformObservation(FTransform& OutTransfo
 
 	const int32 ScaleElement = MakeArrayView(OutElementNames).Find(TEXT("Scale"));
 	FVector OutScale;
-	if (ScaleElement == INDEX_NONE || !GetScaleObservation(OutScale, Object, OutElements[ScaleElement], RelativeTransform.GetScale3D()))
+	if (ScaleElement == INDEX_NONE || !GetScaleObservation(OutScale, OutElements[ScaleElement], RelativeTransform.GetScale3D()))
 	{
 		OutTransform = FTransform::Identity;
 		return false;
@@ -3544,10 +3334,10 @@ bool ULearningAgentsObservations::GetTransformObservation(FTransform& OutTransfo
 	return true;
 }
 
-bool ULearningAgentsObservations::GetAngleObservationRadians(float& OutAngle, const ULearningAgentsObservationObject* Object, const FLearningAgentsObservationObjectElement Element, const float RelativeAngle, const FName Tag)
+bool ULearningAgentsObservationObject::GetAngleObservationRadians(float& OutAngle, const FLearningAgentsObservationObjectElement Element, const float RelativeAngle, const FName Name) const
 {
 	TStaticArray<float, 2> OutValues;
-	if (!GetContinuousObservationToArrayView(OutValues, Object, Element, Tag))
+	if (!GetContinuousObservationToArrayView(OutValues, Element, Name))
 	{
 		OutAngle = 0.0f;
 		return false;
@@ -3558,9 +3348,9 @@ bool ULearningAgentsObservations::GetAngleObservationRadians(float& OutAngle, co
 }
 
 
-bool ULearningAgentsObservations::GetAngleObservation(float& OutAngle, const ULearningAgentsObservationObject* Object, const FLearningAgentsObservationObjectElement Element, const float RelativeAngle, const FName Tag)
+bool ULearningAgentsObservationObject::GetAngleObservation(float& OutAngle, const FLearningAgentsObservationObjectElement Element, const float RelativeAngle, const FName Name) const
 {
-	if (!GetAngleObservationRadians(OutAngle, Object, Element, FMath::DegreesToRadians(RelativeAngle), Tag))
+	if (!GetAngleObservationRadians(OutAngle, Element, FMath::DegreesToRadians(RelativeAngle), Name))
 	{
 		return false;
 	}
@@ -3569,10 +3359,10 @@ bool ULearningAgentsObservations::GetAngleObservation(float& OutAngle, const ULe
 	return true;
 }
 
-bool ULearningAgentsObservations::GetVelocityObservation(FVector& OutVelocity, const ULearningAgentsObservationObject* Object, const FLearningAgentsObservationObjectElement Element, const FTransform RelativeTransform, const float VelocityScale, const FName Tag)
+bool ULearningAgentsObservationObject::GetVelocityObservation(FVector& OutVelocity, const FLearningAgentsObservationObjectElement Element, const FTransform RelativeTransform, const float VelocityScale, const FName Name) const
 {
 	TStaticArray<float, 3> OutValues;
-	if (!GetContinuousObservationToArrayView(OutValues, Object, Element, Tag))
+	if (!GetContinuousObservationToArrayView(OutValues, Element, Name))
 	{
 		OutVelocity = FVector::ZeroVector;
 		return false;
@@ -3582,10 +3372,10 @@ bool ULearningAgentsObservations::GetVelocityObservation(FVector& OutVelocity, c
 	return true;
 }
 
-bool ULearningAgentsObservations::GetDirectionObservation(FVector& OutDirection, const ULearningAgentsObservationObject* Object, const FLearningAgentsObservationObjectElement Element, const FTransform RelativeTransform, const FName Tag)
+bool ULearningAgentsObservationObject::GetDirectionObservation(FVector& OutDirection, const FLearningAgentsObservationObjectElement Element, const FTransform RelativeTransform, const FName Name) const
 {
 	TStaticArray<float, 3> OutValues;
-	if (!GetContinuousObservationToArrayView(OutValues, Object, Element, Tag))
+	if (!GetContinuousObservationToArrayView(OutValues, Element, Name))
 	{
 		OutDirection = FVector::ForwardVector;
 		return false;
@@ -3595,16 +3385,16 @@ bool ULearningAgentsObservations::GetDirectionObservation(FVector& OutDirection,
 	return true;
 }
 
-bool ULearningAgentsObservations::GetLocationAlongSplineObservation(FVector& OutLocation, const ULearningAgentsObservationObject* Object, const FLearningAgentsObservationObjectElement Element, const FTransform RelativeTransform, const float LocationScale, const FName Tag)
+bool ULearningAgentsObservationObject::GetLocationAlongSplineObservation(FVector& OutLocation, const FLearningAgentsObservationObjectElement Element, const FTransform RelativeTransform, const float LocationScale, const FName Name)
 {
-	return GetLocationObservation(OutLocation, Object, Element, RelativeTransform, LocationScale, Tag);
+	return GetLocationObservation(OutLocation, Element, RelativeTransform, LocationScale, Name);
 }
 
-bool ULearningAgentsObservations::GetProportionAlongSplineObservation(bool& bOutIsClosedLoop, float& OutAngle, float& OutPropotion, const ULearningAgentsObservationObject* Object, const FLearningAgentsObservationObjectElement Element, const FName Tag)
+bool ULearningAgentsObservationObject::GetProportionAlongSplineObservation(bool& bOutIsClosedLoop, float& OutAngle, float& OutPropotion, const FLearningAgentsObservationObjectElement Element, const FName Name)
 {
 	FName SubName;
 	FLearningAgentsObservationObjectElement SubElement;
-	if (!GetExclusiveUnionObservation(SubName, SubElement, Object, Element, Tag))
+	if (!GetExclusiveUnionObservation(SubName, SubElement, Element, Name))
 	{
 		bOutIsClosedLoop = false;
 		OutAngle = 0.0f;
@@ -3616,26 +3406,26 @@ bool ULearningAgentsObservations::GetProportionAlongSplineObservation(bool& bOut
 	{
 		bOutIsClosedLoop = true;
 		OutPropotion = 0.0f;
-		return GetAngleObservation(OutAngle, Object, SubElement);
+		return GetAngleObservation(OutAngle, SubElement);
 	}
 	else
 	{
 		bOutIsClosedLoop = false;
 		OutAngle = 0.0f;
-		return GetFloatObservation(OutPropotion, Object, SubElement);
+		return GetFloatObservation(OutPropotion, SubElement);
 	}
 }
 
-bool ULearningAgentsObservations::GetDirectionAlongSplineObservation(FVector& OutDirection, const ULearningAgentsObservationObject* Object, const FLearningAgentsObservationObjectElement Element, const FTransform RelativeTransform, const FName Tag)
+bool ULearningAgentsObservationObject::GetDirectionAlongSplineObservation(FVector& OutDirection, const FLearningAgentsObservationObjectElement Element, const FTransform RelativeTransform, const FName Name)
 {
-	return GetDirectionObservation(OutDirection, Object, Element, RelativeTransform, Tag);
+	return GetDirectionObservation(OutDirection, Element, RelativeTransform, Name);
 }
 
-bool ULearningAgentsObservations::GetPropertiesAlongSplineObservation(FVector& OutLocation, bool& bOutIsClosedLoop, float& OutAngle, float& OutPropotion, FVector& OutDirection, const ULearningAgentsObservationObject* Object, const FLearningAgentsObservationObjectElement Element, const FTransform RelativeTransform, const float LocationScale, const FName Tag)
+bool ULearningAgentsObservationObject::GetPropertiesAlongSplineObservation(FVector& OutLocation, bool& bOutIsClosedLoop, float& OutAngle, float& OutPropotion, FVector& OutDirection, const FLearningAgentsObservationObjectElement Element, const FTransform RelativeTransform, const float LocationScale, const FName Name)
 {
 	TStaticArray<FName, 3> OutElementNames;
 	TStaticArray<FLearningAgentsObservationObjectElement, 3> OutElements;
-	if (!GetStructObservationToArrayViews(OutElementNames, OutElements, Object, Element, Tag))
+	if (!GetStructObservationToArrayViews(OutElementNames, OutElements, Element, Name))
 	{
 		OutLocation = FVector::ZeroVector;
 		bOutIsClosedLoop = false;
@@ -3646,7 +3436,7 @@ bool ULearningAgentsObservations::GetPropertiesAlongSplineObservation(FVector& O
 	}
 
 	const int32 LocationElement = MakeArrayView(OutElementNames).Find(TEXT("Location"));
-	if (LocationElement == INDEX_NONE || !GetLocationAlongSplineObservation(OutLocation, Object, OutElements[LocationElement], RelativeTransform, LocationScale))
+	if (LocationElement == INDEX_NONE || !GetLocationAlongSplineObservation(OutLocation, OutElements[LocationElement], RelativeTransform, LocationScale))
 	{
 		OutLocation = FVector::ZeroVector;
 		bOutIsClosedLoop = false;
@@ -3657,7 +3447,7 @@ bool ULearningAgentsObservations::GetPropertiesAlongSplineObservation(FVector& O
 	}
 
 	const int32 ProportionElement = MakeArrayView(OutElementNames).Find(TEXT("Proportion"));
-	if (ProportionElement == INDEX_NONE || !GetProportionAlongSplineObservation(bOutIsClosedLoop, OutAngle, OutPropotion, Object, OutElements[ProportionElement]))
+	if (ProportionElement == INDEX_NONE || !GetProportionAlongSplineObservation(bOutIsClosedLoop, OutAngle, OutPropotion, OutElements[ProportionElement]))
 	{
 		OutLocation = FVector::ZeroVector;
 		bOutIsClosedLoop = false;
@@ -3668,7 +3458,7 @@ bool ULearningAgentsObservations::GetPropertiesAlongSplineObservation(FVector& O
 	}
 
 	const int32 DirectionElement = MakeArrayView(OutElementNames).Find(TEXT("Direction"));
-	if (DirectionElement == INDEX_NONE || !GetDirectionAlongSplineObservation(OutDirection, Object, OutElements[ProportionElement], RelativeTransform))
+	if (DirectionElement == INDEX_NONE || !GetDirectionAlongSplineObservation(OutDirection, OutElements[ProportionElement], RelativeTransform))
 	{
 		OutLocation = FVector::ZeroVector;
 		bOutIsClosedLoop = false;
@@ -3681,9 +3471,9 @@ bool ULearningAgentsObservations::GetPropertiesAlongSplineObservation(FVector& O
 	return true;
 }
 
-bool ULearningAgentsObservations::GetProportionAlongRayObservation(float& OutProportion, const ULearningAgentsObservationObject* Object, const FLearningAgentsObservationObjectElement Element, const FName Tag)
+bool ULearningAgentsObservationObject::GetProportionAlongRayObservation(float& OutProportion, const FLearningAgentsObservationObjectElement Element, const FName Name)
 {
-	if (!GetFloatObservation(OutProportion, Object, Element, 1.0f, Tag))
+	if (!GetFloatObservation(OutProportion, Element, 1.0f, Name))
 	{
 		OutProportion = 0.0f;
 		return false;
@@ -3693,22 +3483,22 @@ bool ULearningAgentsObservations::GetProportionAlongRayObservation(float& OutPro
 	return true;
 }
 
-bool ULearningAgentsObservations::GetProportionAlongRaysObservationNum(int32& OutProportionNum, const ULearningAgentsObservationObject* Object, const FLearningAgentsObservationObjectElement Element, const FName Tag)
+bool ULearningAgentsObservationObject::GetProportionAlongRaysObservationNum(int32& OutProportionNum, const FLearningAgentsObservationObjectElement Element, const FName Name)
 {
-	return GetStaticArrayObservationNum(OutProportionNum, Object, Element, Tag);
+	return GetStaticArrayObservationNum(OutProportionNum, Element, Name);
 }
 
-bool ULearningAgentsObservations::GetProportionAlongRaysObservation(TArray<float>& OutProportions, const ULearningAgentsObservationObject* Object, const FLearningAgentsObservationObjectElement Element, const FName Tag)
+bool ULearningAgentsObservationObject::GetProportionAlongRaysObservation(TArray<float>& OutProportions, const FLearningAgentsObservationObjectElement Element, const FName Name)
 {
 	int32 ProportionNum;
-	if (!GetProportionAlongRaysObservationNum(ProportionNum, Object, Element, Tag))
+	if (!GetProportionAlongRaysObservationNum(ProportionNum, Element, Name))
 	{
 		OutProportions.Empty();
 		return false;
 	}
 
 	OutProportions.SetNumUninitialized(ProportionNum);
-	if (!GetProportionAlongRaysObservationToArrayView(OutProportions, Object, Element, Tag))
+	if (!GetProportionAlongRaysObservationToArrayView(OutProportions, Element, Name))
 	{
 		OutProportions.Empty();
 		return false;
@@ -3717,26 +3507,19 @@ bool ULearningAgentsObservations::GetProportionAlongRaysObservation(TArray<float
 	return true;
 }
 
-bool ULearningAgentsObservations::GetProportionAlongRaysObservationToArrayView(TArrayView<float> OutProportions, const ULearningAgentsObservationObject* Object, const FLearningAgentsObservationObjectElement Element, const FName Tag)
+bool ULearningAgentsObservationObject::GetProportionAlongRaysObservationToArrayView(TArrayView<float> OutProportions, const FLearningAgentsObservationObjectElement Element, const FName Name)
 {
 	int32 ProportionNum;
-	if (!GetStaticArrayObservationNum(ProportionNum, Object, Element, Tag))
-	{
-		UE::Learning::Array::Zero<1, float>(OutProportions);
-		return false;
-	}
-
-	if (!Object)
+	if (!GetStaticArrayObservationNum(ProportionNum, Element, Name))
 	{
-		UE_LOG(LogLearning, Error, TEXT("GetProportionAlongRaysObservationToArrayView: Object is nullptr."));
 		UE::Learning::Array::Zero<1, float>(OutProportions);
 		return false;
 	}
 
 	if (ProportionNum != OutProportions.Num())
 	{
-		UE_LOG(LogLearning, Error, TEXT("GetProportionAlongRaysObservationToArrayView: Observation '%s' size does not match. Observation is '%i' elements but asked for '%i'."),
-			*Object->ObservationObject.GetTag(Element.ObjectElement).ToString(),
+		UE_LOG(LogLearning, Error, TEXT("%s: Observation '%s' size does not match. Observation is '%i' elements but asked for '%i'."),
+			*GetName(), *ObservationObject.GetName(Element.ObjectElement).ToString(),
 			ProportionNum, OutProportions.Num());
 		UE::Learning::Array::Zero<1, float>(OutProportions);
 		return false;
@@ -3744,7 +3527,7 @@ bool ULearningAgentsObservations::GetProportionAlongRaysObservationToArrayView(T
 
 	TArray<FLearningAgentsObservationObjectElement, TInlineAllocator<32>> SubElements;
 	SubElements.SetNumUninitialized(ProportionNum);
-	if (!GetStaticArrayObservationToArrayView(SubElements, Object, Element, Tag))
+	if (!GetStaticArrayObservationToArrayView(SubElements, Element, Name))
 	{
 		UE::Learning::Array::Zero<1, float>(OutProportions);
 		return false;
@@ -3752,7 +3535,7 @@ bool ULearningAgentsObservations::GetProportionAlongRaysObservationToArrayView(T
 
 	for (int32 SubElementIdx = 0; SubElementIdx < ProportionNum; SubElementIdx++)
 	{
-		if (!GetProportionAlongRayObservation(OutProportions[SubElementIdx], Object, SubElements[SubElementIdx]))
+		if (!GetProportionAlongRayObservation(OutProportions[SubElementIdx], SubElements[SubElementIdx]))
 		{
 			UE::Learning::Array::Zero<1, float>(OutProportions);
 			return false;
