LearningAgents: Fixed Set-based Attention Observation

#rnx

#virtualized

diff --git a/Engine/Plugins/Experimental/LearningAgents/Source/LearningAgents/Public/LearningAgentsActions.h b/Engine/Plugins/Experimental/LearningAgents/Source/LearningAgents/Public/LearningAgentsActions.h
--- a/Engine/Plugins/Experimental/LearningAgents/Source/LearningAgents/Public/LearningAgentsActions.h
+++ b/Engine/Plugins/Experimental/LearningAgents/Source/LearningAgents/Public/LearningAgentsActions.h
@@ -8,6 +8,7 @@
 #include "LearningAgentsNeuralNetwork.h" // Included for ELearningAgentsActivationFunction
 
 #include "Engine/EngineTypes.h"
+#include "GameFramework/OnlineReplStructs.h"
 #include "Kismet/BlueprintFunctionLibrary.h"
 
 #include "LearningAgentsActions.generated.h"
@@ -41,6 +42,15 @@ bool operator==(const FLearningAgentsActionObjectElement& Lhs, const FLearningAg
 /** Hashing operator for Action Object Elements */
 uint32 GetTypeHash(const FLearningAgentsActionObjectElement& Element);
 
+template<>
+struct TStructOpsTypeTraits<FLearningAgentsActionObjectElement> : public TStructOpsTypeTraitsBase2<FLearningAgentsActionObjectElement>
+{
+	enum
+	{
+		WithIdenticalViaEquality = true,
+	};
+};
+
 /**
  * Action Schema
  *
diff --git a/Engine/Plugins/Experimental/LearningAgents/Source/LearningAgents/Public/LearningAgentsObservations.h b/Engine/Plugins/Experimental/LearningAgents/Source/LearningAgents/Public/LearningAgentsObservations.h
--- a/Engine/Plugins/Experimental/LearningAgents/Source/LearningAgents/Public/LearningAgentsObservations.h
+++ b/Engine/Plugins/Experimental/LearningAgents/Source/LearningAgents/Public/LearningAgentsObservations.h
@@ -8,6 +8,7 @@
 #include "LearningAgentsNeuralNetwork.h" // Included for ELearningAgentsActivationFunction
 
 #include "Engine/EngineTypes.h"
+#include "GameFramework/OnlineReplStructs.h"
 #include "Kismet/BlueprintFunctionLibrary.h"
 
 #include "LearningAgentsObservations.generated.h"
@@ -44,6 +45,15 @@ bool operator==(const FLearningAgentsObservationObjectElement& Lhs, const FLearn
 /** Hashing operator for Observation Object Elements */
 uint32 GetTypeHash(const FLearningAgentsObservationObjectElement& Element);
 
+template<>
+struct TStructOpsTypeTraits<FLearningAgentsObservationObjectElement> : public TStructOpsTypeTraitsBase2<FLearningAgentsObservationObjectElement>
+{
+	enum
+	{
+		WithIdenticalViaEquality = true,
+	};
+};
+
 /**
  * Observation Schema
  *
