LearningAgents: Fix for crash when providing the same network asset in multiple slots for the policy.

diff --git a/Engine/Plugins/Experimental/LearningAgents/Source/Learning/Private/LearningCritic.cpp b/Engine/Plugins/Experimental/LearningAgents/Source/Learning/Private/LearningCritic.cpp
--- a/Engine/Plugins/Experimental/LearningAgents/Source/Learning/Private/LearningCritic.cpp
+++ b/Engine/Plugins/Experimental/LearningAgents/Source/Learning/Private/LearningCritic.cpp
@@ -55,4 +55,9 @@ namespace UE::Learning
 	{
 		NeuralNetworkFunction->UpdateNeuralNetwork(NewNeuralNetwork);
 	}
+
+	const TSharedPtr<FNeuralNetwork>& FNeuralNetworkCritic::GetNeuralNetwork() const
+	{
+		return NeuralNetworkFunction->GetNeuralNetwork();
+	}
 }
diff --git a/Engine/Plugins/Experimental/LearningAgents/Source/Learning/Private/LearningNeuralNetwork.cpp b/Engine/Plugins/Experimental/LearningAgents/Source/Learning/Private/LearningNeuralNetwork.cpp
--- a/Engine/Plugins/Experimental/LearningAgents/Source/Learning/Private/LearningNeuralNetwork.cpp
+++ b/Engine/Plugins/Experimental/LearningAgents/Source/Learning/Private/LearningNeuralNetwork.cpp
@@ -396,4 +396,9 @@ namespace UE::Learning
 		}
 	}
 
+	const TSharedPtr<FNeuralNetwork>& FNeuralNetworkFunction::GetNeuralNetwork() const
+	{
+		return NeuralNetwork;
+	}
+
 }
diff --git a/Engine/Plugins/Experimental/LearningAgents/Source/Learning/Private/LearningPolicy.cpp b/Engine/Plugins/Experimental/LearningAgents/Source/Learning/Private/LearningPolicy.cpp
--- a/Engine/Plugins/Experimental/LearningAgents/Source/Learning/Private/LearningPolicy.cpp
+++ b/Engine/Plugins/Experimental/LearningAgents/Source/Learning/Private/LearningPolicy.cpp
@@ -67,4 +67,9 @@ namespace UE::Learning
 	{
 		NeuralNetworkFunction->UpdateNeuralNetwork(NewNeuralNetwork);
 	}
+
+	const TSharedPtr<FNeuralNetwork>& FNeuralNetworkPolicy::GetNeuralNetwork() const
+	{
+		return NeuralNetworkFunction->GetNeuralNetwork();
+	}
 }
diff --git a/Engine/Plugins/Experimental/LearningAgents/Source/Learning/Public/LearningCritic.h b/Engine/Plugins/Experimental/LearningAgents/Source/Learning/Public/LearningCritic.h
--- a/Engine/Plugins/Experimental/LearningAgents/Source/Learning/Public/LearningCritic.h
+++ b/Engine/Plugins/Experimental/LearningAgents/Source/Learning/Public/LearningCritic.h
@@ -47,6 +47,9 @@ namespace UE::Learning
 		/** Sets the NeuralNetwork and re-creates the NeuralNetworkInference object */
 		void UpdateNeuralNetwork(const TSharedPtr<FNeuralNetwork>& NewNeuralNetwork);
 
+		/** Gets the NeuralNetwork associated with this Critic */
+		const TSharedPtr<FNeuralNetwork>& GetNeuralNetwork() const;
+
 	private:
 
 		int32 ObservationEncodedNum = 0;
diff --git a/Engine/Plugins/Experimental/LearningAgents/Source/Learning/Public/LearningNeuralNetwork.h b/Engine/Plugins/Experimental/LearningAgents/Source/Learning/Public/LearningNeuralNetwork.h
--- a/Engine/Plugins/Experimental/LearningAgents/Source/Learning/Public/LearningNeuralNetwork.h
+++ b/Engine/Plugins/Experimental/LearningAgents/Source/Learning/Public/LearningNeuralNetwork.h
@@ -232,6 +232,9 @@ namespace UE::Learning
 		/** Sets the NeuralNetwork and re-creates the NeuralNetworkInference object. */
 		void UpdateNeuralNetwork(const TSharedPtr<FNeuralNetwork>& NewNeuralNetwork);
 
+		/** Gets the NeuralNetwork associated with this function */
+		const TSharedPtr<FNeuralNetwork>& GetNeuralNetwork() const;
+
 	private:
 
 		int32 MaxInstanceNum = 0;
diff --git a/Engine/Plugins/Experimental/LearningAgents/Source/Learning/Public/LearningPolicy.h b/Engine/Plugins/Experimental/LearningAgents/Source/Learning/Public/LearningPolicy.h
--- a/Engine/Plugins/Experimental/LearningAgents/Source/Learning/Public/LearningPolicy.h
+++ b/Engine/Plugins/Experimental/LearningAgents/Source/Learning/Public/LearningPolicy.h
@@ -51,6 +51,9 @@ namespace UE::Learning
 		/** Sets the NeuralNetwork and re-creates the NeuralNetworkInference object */
 		void UpdateNeuralNetwork(const TSharedPtr<FNeuralNetwork>& NewNeuralNetwork);
 
+		/** Gets the NeuralNetwork associated with this Policy */
+		const TSharedPtr<FNeuralNetwork>& GetNeuralNetwork() const;
+
 	private:
 
 		int32 ObservationEncodedNum = 0;
diff --git a/Engine/Plugins/Experimental/LearningAgents/Source/LearningAgents/Private/LearningAgentsCritic.cpp b/Engine/Plugins/Experimental/LearningAgents/Source/LearningAgents/Private/LearningAgentsCritic.cpp
--- a/Engine/Plugins/Experimental/LearningAgents/Source/LearningAgents/Private/LearningAgentsCritic.cpp
+++ b/Engine/Plugins/Experimental/LearningAgents/Source/LearningAgents/Private/LearningAgentsCritic.cpp
@@ -300,6 +300,20 @@ void ULearningAgentsCritic::EvaluateCritic()
 
 	// Evaluate Critic
 
+	if (CriticObject->GetNeuralNetwork()->GetInputSize() != Policy->ObservationVectorsEncoded.Num<1>() + Policy->MemoryState.Num<1>())
+	{
+		UE_LOG(LogLearning, Error, TEXT("%s: Critic Network Input size doesn't match. Network input size is %i but Critic expects %i."), *GetName(),
+			CriticObject->GetNeuralNetwork()->GetInputSize(), Policy->ObservationVectorsEncoded.Num<1>() + Policy->MemoryState.Num<1>());
+		return;
+	}
+
+	if (CriticObject->GetNeuralNetwork()->GetOutputSize() != 1)
+	{
+		UE_LOG(LogLearning, Error, TEXT("%s: Critic Network Output size don't match. Network output size is %i but Critic expects %i."), *GetName(),
+			CriticObject->GetNeuralNetwork()->GetOutputSize(), 1);
+		return;
+	}
+
 	CriticObject->Evaluate(
 		Returns,
 		Policy->ObservationVectorsEncoded,
diff --git a/Engine/Plugins/Experimental/LearningAgents/Source/LearningAgents/Private/LearningAgentsPolicy.cpp b/Engine/Plugins/Experimental/LearningAgents/Source/LearningAgents/Private/LearningAgentsPolicy.cpp
--- a/Engine/Plugins/Experimental/LearningAgents/Source/LearningAgents/Private/LearningAgentsPolicy.cpp
+++ b/Engine/Plugins/Experimental/LearningAgents/Source/LearningAgents/Private/LearningAgentsPolicy.cpp
@@ -118,6 +118,26 @@ void ULearningAgentsPolicy::SetupPolicy(
 		return;
 	}
 
+	// Warn if Network Assets are not unique
+
+	if (PolicyNeuralNetworkAsset && PolicyNeuralNetworkAsset == EncoderNeuralNetworkAsset)
+	{
+		UE_LOG(LogLearning, Warning, TEXT("%s: Identical Network Assets given as both Policy and Encoder: %s."), *GetName(), *PolicyNeuralNetworkAsset->GetName());
+	}
+
+
+	if (PolicyNeuralNetworkAsset && PolicyNeuralNetworkAsset == DecoderNeuralNetworkAsset)
+	{
+		UE_LOG(LogLearning, Warning, TEXT("%s: Identical Network Assets given as both Policy and Decoder: %s."), *GetName(), *PolicyNeuralNetworkAsset->GetName());
+	}
+
+	if (EncoderNeuralNetworkAsset && EncoderNeuralNetworkAsset == DecoderNeuralNetworkAsset)
+	{
+		UE_LOG(LogLearning, Warning, TEXT("%s: Identical Network Assets given as both Encoder and Decoder: %s."), *GetName(), *EncoderNeuralNetworkAsset->GetName());
+	}
+
+	// Begin Setup
+
 	Manager = InManager;
 	Interactor = InInteractor;
 
@@ -523,6 +543,20 @@ void ULearningAgentsPolicy::EncodeObservations()
 
 	// Encode Observations
 
+	if (EncoderObject->GetNeuralNetwork()->GetInputSize() != Interactor->ObservationVectors.Num<1>())
+	{
+		UE_LOG(LogLearning, Error, TEXT("%s: Encoder Network Input size doesn't match. Network input size is %i but Encoder expects %i."), *GetName(),
+			EncoderObject->GetNeuralNetwork()->GetInputSize(), Interactor->ObservationVectors.Num<1>());
+		return;
+	}
+
+	if (EncoderObject->GetNeuralNetwork()->GetOutputSize() != ObservationVectorsEncoded.Num<1>())
+	{
+		UE_LOG(LogLearning, Error, TEXT("%s: Encoder Network Output size don't match. Network output size is %i but Encoder expects %i."), *GetName(),
+			EncoderObject->GetNeuralNetwork()->GetOutputSize(), ObservationVectorsEncoded.Num<1>());
+		return;
+	}
+
 	EncoderObject->Evaluate(ObservationVectorsEncoded, Interactor->ObservationVectors, ValidAgentSet);
 
 	for (const int32 AgentId : ValidAgentSet)
@@ -570,6 +604,20 @@ void ULearningAgentsPolicy::EvaluatePolicy()
 
 	// Evaluate policy
 
+	if (PolicyObject->GetNeuralNetwork()->GetInputSize() != ObservationVectorsEncoded.Num<1>() + MemoryState.Num<1>())
+	{
+		UE_LOG(LogLearning, Error, TEXT("%s: Policy Network Input size doesn't match. Network input size is %i but Policy expects %i."), *GetName(),
+			PolicyObject->GetNeuralNetwork()->GetInputSize(), ObservationVectorsEncoded.Num<1>() + MemoryState.Num<1>());
+		return;
+	}
+
+	if (PolicyObject->GetNeuralNetwork()->GetOutputSize() != ActionDistributionVectors.Num<1>() + MemoryState.Num<1>())
+	{
+		UE_LOG(LogLearning, Error, TEXT("%s: Policy Network Output size don't match. Network output size is %i but Policy expects %i."), *GetName(),
+			PolicyObject->GetNeuralNetwork()->GetOutputSize(), ActionDistributionVectors.Num<1>() + MemoryState.Num<1>());
+		return;
+	}
+
 	PolicyObject->Evaluate(
 		ActionVectorsEncoded,
 		MemoryState,
@@ -621,6 +669,20 @@ void ULearningAgentsPolicy::DecodeAndSampleActions(const float ActionNoiseScale)
 
 	// Decode to produce action distribution vectors
 
+	if (DecoderObject->GetNeuralNetwork()->GetInputSize() != ActionVectorsEncoded.Num<1>())
+	{
+		UE_LOG(LogLearning, Error, TEXT("%s: Decoder Network Input size doesn't match. Network input size is %i but Decoder expects %i."), *GetName(),
+			DecoderObject->GetNeuralNetwork()->GetInputSize(), ActionVectorsEncoded.Num<1>());
+		return;
+	}
+
+	if (DecoderObject->GetNeuralNetwork()->GetOutputSize() != ActionDistributionVectors.Num<1>())
+	{
+		UE_LOG(LogLearning, Error, TEXT("%s: Decoder Network Output size don't match. Network output size is %i but Decoder expects %i."), *GetName(),
+			DecoderObject->GetNeuralNetwork()->GetOutputSize(), ActionDistributionVectors.Num<1>());
+		return;
+	}
+
 	DecoderObject->Evaluate(ActionDistributionVectors, ActionVectorsEncoded, ValidAgentSet);
 
 	// Sample actions from action distribution vectors
