LearningAgents: Small fixes and improvements.

diff --git a/Engine/Plugins/Experimental/LearningAgents/Source/LearningAgents/Private/LearningAgentsActions.cpp b/Engine/Plugins/Experimental/LearningAgents/Source/LearningAgents/Private/LearningAgentsActions.cpp
--- a/Engine/Plugins/Experimental/LearningAgents/Source/LearningAgents/Private/LearningAgentsActions.cpp
+++ b/Engine/Plugins/Experimental/LearningAgents/Source/LearningAgents/Private/LearningAgentsActions.cpp
@@ -2,6 +2,9 @@
 
 #include "LearningAgentsActions.h"
 
+#include "LearningAgentsManagerListener.h"
+#include "LearningAgentsDebug.h"
+
 #include "LearningArray.h"
 #include "LearningLog.h"
 
@@ -2547,7 +2550,17 @@ bool ULearningAgentsActions::GetBoolAction(bool& bOutValue, const ULearningAgent
 	return true;
 }
 
-bool ULearningAgentsActions::GetFloatAction(float& OutValue, const ULearningAgentsActionObject* Object, const FLearningAgentsActionObjectElement Element, const float FloatScale, const FName Tag)
+bool ULearningAgentsActions::GetFloatAction(
+	float& OutValue, 
+	const ULearningAgentsActionObject* Object, 
+	const FLearningAgentsActionObjectElement Element, 
+	const float FloatScale, 
+	const FName Tag,
+	const bool bVisualLoggerEnabled,
+	ULearningAgentsManagerListener* VisualLoggerListener,
+	const int32 VisualLoggerAgentId,
+	const FVector VisualLoggerLocation,
+	const FLinearColor VisualLoggerColor)
 {
 	float OutValuesData;
 	if (!GetContinuousActionToArrayView(MakeArrayView(&OutValuesData, 1), Object, Element, Tag))
@@ -2557,6 +2570,24 @@ bool ULearningAgentsActions::GetFloatAction(float& OutValue, const ULearningAgen
 	}
 
 	OutValue = OutValuesData * FloatScale;
+
+#if UE_LEARNING_AGENTS_ENABLE_VISUAL_LOG
+	if (bVisualLoggerEnabled && VisualLoggerListener)
+	{
+		const ULearningAgentsVisualLoggerObject* VisualLoggerObject = VisualLoggerListener->GetOrAddVisualLoggerObject(Tag);
+
+		UE_LEARNING_AGENTS_VLOG_STRING(VisualLoggerObject, LogLearning, Display, VisualLoggerLocation,
+			VisualLoggerColor.ToFColor(true),
+			TEXT("Listener: %s\nTag: %s\nAgent Id: % 3i\nEncoded: [% 6.2f]\nScale: [% 6.2f]\nValue: [% 6.1f]"),
+			*VisualLoggerListener->GetName(),
+			*Tag.ToString(),
+			VisualLoggerAgentId,
+			OutValuesData,
+			FloatScale,
+			OutValue);
+	}
+#endif
+
 	return true;
 }
 
diff --git a/Engine/Plugins/Experimental/LearningAgents/Source/LearningAgents/Private/LearningAgentsObservations.cpp b/Engine/Plugins/Experimental/LearningAgents/Source/LearningAgents/Private/LearningAgentsObservations.cpp
--- a/Engine/Plugins/Experimental/LearningAgents/Source/LearningAgents/Private/LearningAgentsObservations.cpp
+++ b/Engine/Plugins/Experimental/LearningAgents/Source/LearningAgents/Private/LearningAgentsObservations.cpp
@@ -1058,32 +1058,11 @@ FLearningAgentsObservationSchemaElement ULearningAgentsObservations::SpecifyDire
 	return SpecifyDirectionObservation(Schema, Tag);
 }
 
-FLearningAgentsObservationSchemaElement ULearningAgentsObservations::SpecifyPropertiesAlongSplineObservation(ULearningAgentsObservationSchema* Schema, const FName Tag)
-{
-	return SpecifyStructObservationFromArrayViews(Schema,
-		{
-			TEXT("Location"),
-			TEXT("Proportion"),
-			TEXT("Direction")
-		},
-		{
-			SpecifyLocationAlongSplineObservation(Schema),
-			SpecifyProportionAlongSplineObservation(Schema),
-			SpecifyDirectionAlongSplineObservation(Schema),
-		},
-		Tag);
-}
-
 FLearningAgentsObservationSchemaElement ULearningAgentsObservations::SpecifyProportionAlongRayObservation(ULearningAgentsObservationSchema* Schema, const FName Tag)
 {
 	return SpecifyFloatObservation(Schema, Tag);
 }
 
-FLearningAgentsObservationSchemaElement ULearningAgentsObservations::SpecifyProportionAlongRaysObservation(ULearningAgentsObservationSchema* Schema, const int32 Num, const FName Tag)
-{
-	return SpecifyStaticArrayObservation(Schema, SpecifyProportionAlongRayObservation(Schema), Num, Tag);
-}
-
 void ULearningAgentsObservations::LogObservation(const ULearningAgentsObservationObject* Object, const FLearningAgentsObservationObjectElement Element)
 {
 	if (!Object)
@@ -1951,15 +1930,55 @@ FLearningAgentsObservationObjectElement ULearningAgentsObservations::MakeAngleOb
 		}, Tag);
 }
 
-FLearningAgentsObservationObjectElement ULearningAgentsObservations::MakeVelocityObservation(ULearningAgentsObservationObject* Object, const FVector Velocity, const FTransform RelativeTransform, const float VelocityScale, const FName Tag)
+FLearningAgentsObservationObjectElement ULearningAgentsObservations::MakeVelocityObservation(
+	ULearningAgentsObservationObject* Object, 
+	const FVector Velocity, 
+	const FTransform RelativeTransform, 
+	const float VelocityScale, 
+	const FName Tag,
+	const bool bVisualLoggerEnabled,
+	ULearningAgentsManagerListener* VisualLoggerListener,
+	const int32 VisualLoggerAgentId,
+	const FVector VisualLoggerVelocityLocation,
+	const FVector VisualLoggerLocation,
+	const FLinearColor VisualLoggerColor)
 {
 	const FVector LocalVelocity = RelativeTransform.InverseTransformVectorNoScale(Velocity);
+	const FVector EncodedVelocity = FVector(
+		LocalVelocity.X / FMath::Max(VelocityScale, UE_SMALL_NUMBER),
+		LocalVelocity.Y / FMath::Max(VelocityScale, UE_SMALL_NUMBER),
+		LocalVelocity.Z / FMath::Max(VelocityScale, UE_SMALL_NUMBER));
 
-	return MakeContinuousObservationFromArrayView(Object, {
-		(float)LocalVelocity.X / FMath::Max(VelocityScale, UE_SMALL_NUMBER),
-		(float)LocalVelocity.Y / FMath::Max(VelocityScale, UE_SMALL_NUMBER),
-		(float)LocalVelocity.Z / FMath::Max(VelocityScale, UE_SMALL_NUMBER),
-		}, Tag);
+#if UE_LEARNING_AGENTS_ENABLE_VISUAL_LOG
+	if (bVisualLoggerEnabled && VisualLoggerListener)
+	{
+		const ULearningAgentsVisualLoggerObject* VisualLoggerObject = VisualLoggerListener->GetOrAddVisualLoggerObject(Tag);
+
+		UE_LEARNING_AGENTS_VLOG_ARROW(VisualLoggerObject, LogLearning, Display,
+			VisualLoggerVelocityLocation,
+			VisualLoggerVelocityLocation + Velocity,
+			VisualLoggerColor.ToFColor(true),
+			TEXT(""));
+
+		UE_LEARNING_AGENTS_VLOG_TRANSFORM(VisualLoggerObject, LogLearning, Display,
+			RelativeTransform.GetTranslation(),
+			RelativeTransform.GetRotation(),
+			VisualLoggerColor.ToFColor(true),
+			TEXT(""));
+
+		UE_LEARNING_AGENTS_VLOG_STRING(VisualLoggerObject, LogLearning, Display, VisualLoggerLocation,
+			VisualLoggerColor.ToFColor(true),
+			TEXT("Listener: %s\nTag: %s\nAgent Id: % 3i\nVelocity: [% 6.1f % 6.1f % 6.1f]\nLocal Velocity: [% 6.1f % 6.1f % 6.1f]\nEncoded: [% 6.2f % 6.2f % 6.2f]"),
+			*VisualLoggerListener->GetName(),
+			*Tag.ToString(),
+			VisualLoggerAgentId,
+			Velocity.X, Velocity.Y, Velocity.Z,
+			LocalVelocity.X, LocalVelocity.Y, LocalVelocity.Z,
+			EncodedVelocity.X, EncodedVelocity.Y, EncodedVelocity.Z);
+	}
+#endif
+
+	return MakeContinuousObservationFromArrayView(Object, { (float)EncodedVelocity.X, (float)EncodedVelocity.Y, (float)EncodedVelocity.Z }, Tag);
 }
 
 FLearningAgentsObservationObjectElement ULearningAgentsObservations::MakeDirectionObservation(
@@ -2032,7 +2051,20 @@ FLearningAgentsObservationObjectElement ULearningAgentsObservations::MakeLocatio
 		return FLearningAgentsObservationObjectElement();
 	}
 
-	return MakeLocationObservation(Object, SplineComponent->GetLocationAtDistanceAlongSpline(DistanceAlongSpline, ESplineCoordinateSpace::World), RelativeTransform, LocationScale, Tag, bVisualLoggerEnabled, VisualLoggerListener, VisualLoggerAgentId, VisualLoggerLocation, VisualLoggerColor);
+	const float TotalDistance = SplineComponent->GetSplineLength();
+	const float LoopedDistance = SplineComponent->IsClosedLoop() ? FMath::Wrap(DistanceAlongSpline, 0.0f, TotalDistance) : FMath::Clamp(DistanceAlongSpline, 0.0f, TotalDistance);
+
+	return MakeLocationObservation(
+		Object, 
+		SplineComponent->GetLocationAtDistanceAlongSpline(LoopedDistance, ESplineCoordinateSpace::World), 
+		RelativeTransform, 
+		LocationScale, 
+		Tag, 
+		bVisualLoggerEnabled, 
+		VisualLoggerListener, 
+		VisualLoggerAgentId, 
+		VisualLoggerLocation, 
+		VisualLoggerColor);
 }
 
 FLearningAgentsObservationObjectElement ULearningAgentsObservations::MakeProportionAlongSplineObservation(ULearningAgentsObservationObject* Object, const USplineComponent* SplineComponent, const float DistanceAlongSpline, const FName Tag)
@@ -2077,37 +2109,30 @@ FLearningAgentsObservationObjectElement ULearningAgentsObservations::MakeDirecti
 		return FLearningAgentsObservationObjectElement();
 	}
 
+	const float TotalDistance = SplineComponent->GetSplineLength();
+	const float LoopedDistance = SplineComponent->IsClosedLoop() ? FMath::Wrap(DistanceAlongSpline, 0.0f, TotalDistance) : FMath::Clamp(DistanceAlongSpline, 0.0f, TotalDistance);
+
 #if UE_LEARNING_AGENTS_ENABLE_VISUAL_LOG
 	const FVector VisualLoggerDirectionLocation = 
 		(bVisualLoggerEnabled && VisualLoggerListener) ?
-		SplineComponent->GetLocationAtDistanceAlongSpline(DistanceAlongSpline, ESplineCoordinateSpace::World) : 
+		SplineComponent->GetLocationAtDistanceAlongSpline(LoopedDistance, ESplineCoordinateSpace::World) :
 		FVector::ZeroVector;
 #else
 	const FVector VisualLoggerDirectionLocation = FVector::ZeroVector;
 #endif
 
-	return MakeDirectionObservation(Object, SplineComponent->GetDirectionAtDistanceAlongSpline(DistanceAlongSpline, ESplineCoordinateSpace::World), RelativeTransform, Tag, bVisualLoggerEnabled, VisualLoggerListener, VisualLoggerAgentId, VisualLoggerDirectionLocation, VisualLoggerLocation, VisualLoggerArrowLength, VisualLoggerColor);
-}
-
-FLearningAgentsObservationObjectElement ULearningAgentsObservations::MakePropertiesAlongSplineObservation(ULearningAgentsObservationObject* Object, const USplineComponent* SplineComponent, const float DistanceAlongSpline, const FTransform RelativeTransform, const float LocationScale, const FName Tag)
-{
-	if (!SplineComponent)
-	{
-		UE_LOG(LogLearning, Error, TEXT("MakePropertiesAlongSplineObservation: SplineComponent was nullptr."));
-		return FLearningAgentsObservationObjectElement();
-	}
-
-	return MakeStructObservationFromArrayViews(Object,
-		{
-			TEXT("Location"),
-			TEXT("Proportion"),
-			TEXT("Direction")
-		},
-		{
-			MakeLocationAlongSplineObservation(Object, SplineComponent, DistanceAlongSpline, RelativeTransform, LocationScale),
-			MakeProportionAlongSplineObservation(Object, SplineComponent, DistanceAlongSpline),
-			MakeDirectionAlongSplineObservation(Object, SplineComponent, DistanceAlongSpline, RelativeTransform),
-		}, Tag);
+	return MakeDirectionObservation(
+		Object, 
+		SplineComponent->GetDirectionAtDistanceAlongSpline(LoopedDistance, ESplineCoordinateSpace::World), 
+		RelativeTransform, 
+		Tag, 
+		bVisualLoggerEnabled, 
+		VisualLoggerListener, 
+		VisualLoggerAgentId, 
+		VisualLoggerDirectionLocation, 
+		VisualLoggerLocation, 
+		VisualLoggerArrowLength, 
+		VisualLoggerColor);
 }
 
 FLearningAgentsObservationObjectElement ULearningAgentsObservations::MakeProportionAlongRayObservation(ULearningAgentsObservationObject* Object, const FVector RayStart, const FVector RayEnd, const FTransform RayTransform, const ECollisionChannel CollisionChannel, const FName Tag)
@@ -2133,27 +2158,6 @@ FLearningAgentsObservationObjectElement ULearningAgentsObservations::MakeProport
 	}
 }
 
-FLearningAgentsObservationObjectElement ULearningAgentsObservations::MakeProportionAlongRaysObservation(ULearningAgentsObservationObject* Object, const TArray<FVector>& RayStarts, const TArray<FVector>& RayEnds, const FTransform RayTransform, const ECollisionChannel CollisionChannel, const FName Tag)
-{
-	if (RayStarts.Num() != RayEnds.Num())
-	{
-		UE_LOG(LogLearning, Error, TEXT("MakeProportionAlongRaysObservation: Different number of RayStarts and RayEnds (%i vs %i)."), RayStarts.Num(), RayEnds.Num());
-		return FLearningAgentsObservationObjectElement();
-	}
-
-	const int32 RayNum = RayStarts.Num();
-
-	TArray<FLearningAgentsObservationObjectElement, TInlineAllocator<32>> RayElements;
-	RayElements.Reserve(RayNum);
-
-	for (int32 RayIdx = 0; RayIdx < RayNum; RayIdx++)
-	{
-		RayElements.Add(MakeProportionAlongRayObservation(Object, RayStarts[RayIdx], RayEnds[RayIdx], RayTransform, CollisionChannel));
-	}
-
-	return MakeStaticArrayObservationFromArrayView(Object, RayElements, Tag);
-}
-
 
 bool ULearningAgentsObservations::GetNullObservation(const ULearningAgentsObservationObject* Object, const FLearningAgentsObservationObjectElement Element, const FName Tag)
 {
@@ -3631,56 +3635,6 @@ bool ULearningAgentsObservations::GetDirectionAlongSplineObservation(FVector& Ou
 	return GetDirectionObservation(OutDirection, Object, Element, RelativeTransform, Tag);
 }
 
-bool ULearningAgentsObservations::GetPropertiesAlongSplineObservation(FVector& OutLocation, bool& bOutIsClosedLoop, float& OutAngle, float& OutPropotion, FVector& OutDirection, const ULearningAgentsObservationObject* Object, const FLearningAgentsObservationObjectElement Element, const FTransform RelativeTransform, const float LocationScale, const FName Tag)
-{
-	TStaticArray<FName, 3> OutElementNames;
-	TStaticArray<FLearningAgentsObservationObjectElement, 3> OutElements;
-	if (!GetStructObservationToArrayViews(OutElementNames, OutElements, Object, Element, Tag))
-	{
-		OutLocation = FVector::ZeroVector;
-		bOutIsClosedLoop = false;
-		OutAngle = 0.0f;
-		OutPropotion = 0.0f;
-		OutDirection = FVector::ForwardVector;
-		return false;
-	}
-
-	const int32 LocationElement = MakeArrayView(OutElementNames).Find(TEXT("Location"));
-	if (LocationElement == INDEX_NONE || !GetLocationAlongSplineObservation(OutLocation, Object, OutElements[LocationElement], RelativeTransform, LocationScale))
-	{
-		OutLocation = FVector::ZeroVector;
-		bOutIsClosedLoop = false;
-		OutAngle = 0.0f;
-		OutPropotion = 0.0f;
-		OutDirection = FVector::ForwardVector;
-		return false;
-	}
-
-	const int32 ProportionElement = MakeArrayView(OutElementNames).Find(TEXT("Proportion"));
-	if (ProportionElement == INDEX_NONE || !GetProportionAlongSplineObservation(bOutIsClosedLoop, OutAngle, OutPropotion, Object, OutElements[ProportionElement]))
-	{
-		OutLocation = FVector::ZeroVector;
-		bOutIsClosedLoop = false;
-		OutAngle = 0.0f;
-		OutPropotion = 0.0f;
-		OutDirection = FVector::ForwardVector;
-		return false;
-	}
-
-	const int32 DirectionElement = MakeArrayView(OutElementNames).Find(TEXT("Direction"));
-	if (DirectionElement == INDEX_NONE || !GetDirectionAlongSplineObservation(OutDirection, Object, OutElements[ProportionElement], RelativeTransform))
-	{
-		OutLocation = FVector::ZeroVector;
-		bOutIsClosedLoop = false;
-		OutAngle = 0.0f;
-		OutPropotion = 0.0f;
-		OutDirection = FVector::ForwardVector;
-		return false;
-	}
-
-	return true;
-}
-
 bool ULearningAgentsObservations::GetProportionAlongRayObservation(float& OutProportion, const ULearningAgentsObservationObject* Object, const FLearningAgentsObservationObjectElement Element, const FName Tag)
 {
 	if (!GetFloatObservation(OutProportion, Object, Element, 1.0f, Tag))
@@ -3693,72 +3647,3 @@ bool ULearningAgentsObservations::GetProportionAlongRayObservation(float& OutPro
 	return true;
 }
 
-bool ULearningAgentsObservations::GetProportionAlongRaysObservationNum(int32& OutProportionNum, const ULearningAgentsObservationObject* Object, const FLearningAgentsObservationObjectElement Element, const FName Tag)
-{
-	return GetStaticArrayObservationNum(OutProportionNum, Object, Element, Tag);
-}
-
-bool ULearningAgentsObservations::GetProportionAlongRaysObservation(TArray<float>& OutProportions, const ULearningAgentsObservationObject* Object, const FLearningAgentsObservationObjectElement Element, const FName Tag)
-{
-	int32 ProportionNum;
-	if (!GetProportionAlongRaysObservationNum(ProportionNum, Object, Element, Tag))
-	{
-		OutProportions.Empty();
-		return false;
-	}
-
-	OutProportions.SetNumUninitialized(ProportionNum);
-	if (!GetProportionAlongRaysObservationToArrayView(OutProportions, Object, Element, Tag))
-	{
-		OutProportions.Empty();
-		return false;
-	}
-
-	return true;
-}
-
-bool ULearningAgentsObservations::GetProportionAlongRaysObservationToArrayView(TArrayView<float> OutProportions, const ULearningAgentsObservationObject* Object, const FLearningAgentsObservationObjectElement Element, const FName Tag)
-{
-	int32 ProportionNum;
-	if (!GetStaticArrayObservationNum(ProportionNum, Object, Element, Tag))
-	{
-		UE::Learning::Array::Zero<1, float>(OutProportions);
-		return false;
-	}
-
-	if (!Object)
-	{
-		UE_LOG(LogLearning, Error, TEXT("GetProportionAlongRaysObservationToArrayView: Object is nullptr."));
-		UE::Learning::Array::Zero<1, float>(OutProportions);
-		return false;
-	}
-
-	if (ProportionNum != OutProportions.Num())
-	{
-		UE_LOG(LogLearning, Error, TEXT("GetProportionAlongRaysObservationToArrayView: Observation '%s' size does not match. Observation is '%i' elements but asked for '%i'."),
-			*Object->ObservationObject.GetTag(Element.ObjectElement).ToString(),
-			ProportionNum, OutProportions.Num());
-		UE::Learning::Array::Zero<1, float>(OutProportions);
-		return false;
-	}
-
-	TArray<FLearningAgentsObservationObjectElement, TInlineAllocator<32>> SubElements;
-	SubElements.SetNumUninitialized(ProportionNum);
-	if (!GetStaticArrayObservationToArrayView(SubElements, Object, Element, Tag))
-	{
-		UE::Learning::Array::Zero<1, float>(OutProportions);
-		return false;
-	}
-
-	for (int32 SubElementIdx = 0; SubElementIdx < ProportionNum; SubElementIdx++)
-	{
-		if (!GetProportionAlongRayObservation(OutProportions[SubElementIdx], Object, SubElements[SubElementIdx]))
-		{
-			UE::Learning::Array::Zero<1, float>(OutProportions);
-			return false;
-		}
-	}
-
-	return true;
-}
-
diff --git a/Engine/Plugins/Experimental/LearningAgents/Source/LearningAgents/Public/LearningAgentsActions.h b/Engine/Plugins/Experimental/LearningAgents/Source/LearningAgents/Public/LearningAgentsActions.h
--- a/Engine/Plugins/Experimental/LearningAgents/Source/LearningAgents/Public/LearningAgentsActions.h
+++ b/Engine/Plugins/Experimental/LearningAgents/Source/LearningAgents/Public/LearningAgentsActions.h
@@ -397,8 +397,32 @@ public:
 	UFUNCTION(BlueprintPure = false, Category = "LearningAgents", meta = (AdvancedDisplay = 3, ReturnDisplayName = "Success"))
 	static bool GetBoolAction(bool& bOutValue, const ULearningAgentsActionObject* Object, const FLearningAgentsActionObjectElement Element, const FName Tag = TEXT("Bool"));
 
-	UFUNCTION(BlueprintPure = false, Category = "LearningAgents", meta = (AdvancedDisplay = 4, ReturnDisplayName = "Success"))
-	static bool GetFloatAction(float& OutValue, const ULearningAgentsActionObject* Object, const FLearningAgentsActionObjectElement Element, const float FloatScale = 1.0f, const FName Tag = TEXT("Float"));
+	/**
+	 * Get the value for a float action.
+	 *
+	 * @param OutValue The output float value.
+	 * @param Object The Action Object
+	 * @param Element The Action Object Element
+	 * @param FloatScale The scale used to control the overall magnitude of the outputted float action.
+	 * @param Tag The tag of the corresponding action. Must match the tag given during Specify.
+	 * @param bVisualLoggerEnabled When true, debug data will be sent to the visual logger.
+	 * @param VisualLoggerListener The listener object which is making this action. This must be set to use logging.
+	 * @param VisualLoggerAgentId The agent id associated with this action.
+	 * @param VisualLoggerLocation A location for the visual logger information in the world.
+	 * @param VisualLoggerColor The color for the visual logger display.
+	 */
+	UFUNCTION(BlueprintPure = false, Category = "LearningAgents", meta = (AdvancedDisplay = 4, ReturnDisplayName = "Success", DefaultToSelf = "VisualLoggerListener"))
+	static bool GetFloatAction(
+		float& OutValue, 
+		const ULearningAgentsActionObject* Object, 
+		const FLearningAgentsActionObjectElement Element, 
+		const float FloatScale = 1.0f, 
+		const FName Tag = TEXT("Float"),
+		const bool bVisualLoggerEnabled = false,
+		ULearningAgentsManagerListener* VisualLoggerListener = nullptr,
+		const int32 VisualLoggerAgentId = -1,
+		const FVector VisualLoggerLocation = FVector::ZeroVector,
+		const FLinearColor VisualLoggerColor = FLinearColor::Blue);
 
 	UFUNCTION(BlueprintPure = false, Category = "LearningAgents", meta = (AdvancedDisplay = 5, ReturnDisplayName = "Success"))
 	static bool GetLocationAction(FVector& OutLocation, const ULearningAgentsActionObject* Object, const FLearningAgentsActionObjectElement Element, const FTransform RelativeTransform = FTransform(), const float LocationScale = 100.0f, const FName Tag = TEXT("Location"));
diff --git a/Engine/Plugins/Experimental/LearningAgents/Source/LearningAgents/Public/LearningAgentsObservations.h b/Engine/Plugins/Experimental/LearningAgents/Source/LearningAgents/Public/LearningAgentsObservations.h
--- a/Engine/Plugins/Experimental/LearningAgents/Source/LearningAgents/Public/LearningAgentsObservations.h
+++ b/Engine/Plugins/Experimental/LearningAgents/Source/LearningAgents/Public/LearningAgentsObservations.h
@@ -211,6 +211,7 @@ public:
 	/**
 	 * Specifies a new bool observation. A true or false observation.
 	 *
+	 * @param Schema The Observation Schema
 	 * @param Tag The tag of this new observation. Used during observation object validation and debugging.
 	 * @return The newly created observation schema element.
 	 */
@@ -221,6 +222,7 @@ public:
 	 * Specifies a new float observation. A simple observation which can be used as a catch-all for situations where a
 	 * type-specific observation does not exist.
 	 *
+	 * @param Schema The Observation Schema
 	 * @param Tag The tag of this new observation. Used during observation object validation and debugging.
 	 * @return The newly created observation schema element.
 	 */
@@ -230,6 +232,7 @@ public:
 	/**
 	 * Specifies a new location observation. Allows an agent to observe the location of some entity.
 	 *
+	 * @param Schema The Observation Schema
 	 * @param Tag The tag of this new observation. Used during observation object validation and debugging.
 	 * @return The newly created observation schema element.
 	 */
@@ -265,17 +268,11 @@ public:
 	UFUNCTION(BlueprintPure, Category = "LearningAgents", meta = (AdvancedDisplay = 1))
 	static FLearningAgentsObservationSchemaElement SpecifyDirectionAlongSplineObservation(ULearningAgentsObservationSchema* Schema, const FName Tag = TEXT("DirectionAlongSpline"));
 
-	UFUNCTION(BlueprintPure, Category = "LearningAgents", meta = (AdvancedDisplay = 1))
-	static FLearningAgentsObservationSchemaElement SpecifyPropertiesAlongSplineObservation(ULearningAgentsObservationSchema* Schema, const FName Tag = TEXT("PropertiesAlongSpline"));
-
 	// Specify Ray Cast Observations
 
 	UFUNCTION(BlueprintPure, Category = "LearningAgents", meta = (AdvancedDisplay = 1))
 	static FLearningAgentsObservationSchemaElement SpecifyProportionAlongRayObservation(ULearningAgentsObservationSchema* Schema, const FName Tag = TEXT("ProportionAlongRay"));
 
-	UFUNCTION(BlueprintPure, Category = "LearningAgents", meta = (AdvancedDisplay = 2))
-	static FLearningAgentsObservationSchemaElement SpecifyProportionAlongRaysObservation(ULearningAgentsObservationSchema* Schema, const int32 Num, const FName Tag = TEXT("ProportionAlongRays"));
-
 public:
 
 	// Make Basic Observations
@@ -372,6 +369,7 @@ public:
 	/**
 	 * Make a new bool observation.
 	 *
+	 * @param Object The Observation Object
 	 * @param Value The new value of this observation.
 	 * @param Tag The tag of the corresponding observation. Must match the tag given during Specify.
 	 * @param bVisualLoggerEnabled When true, debug data will be sent to the visual logger.
@@ -395,6 +393,7 @@ public:
 	/**
 	 * Make a new float observation.
 	 *
+	 * @param Object The Observation Object
 	 * @param Value The new value of this observation.
 	 * @param FloatScale Used to normalize the data for this observation.
 	 * @param Tag The tag of the corresponding observation. Must match the tag given during Specify.
@@ -420,6 +419,7 @@ public:
 	/**
 	 * Make a new location observation.
 	 *
+	 * @param Object The Observation Object
 	 * @param Location The location of interest to the agent.
 	 * @param RelativeTransform The transform the provided location should be encoded relative to.
 	 * @param LocationScale Used to normalize the data for this observation.
@@ -447,6 +447,7 @@ public:
 	/**
 	 * Make a new rotation observation.
 	 *
+	 * @param Object The Observation Object
 	 * @param Rotation The rotation of interest to the agent.
 	 * @param RelativeRotation The rotation the provided rotation should be encoded relative to.
 	 * @param Tag The tag of the corresponding observation. Must match the tag given during Specify.
@@ -474,6 +475,7 @@ public:
 	/**
 	 * Make a new rotation observation from a quaternion.
 	 *
+	 * @param Object The Observation Object
 	 * @param Rotation The rotation of interest to the agent.
 	 * @param RelativeRotation The rotation the provided rotation should be encoded relative to.
 	 * @param Tag The tag of the corresponding observation. Must match the tag given during Specify.
@@ -501,6 +503,7 @@ public:
 	/**
 	 * Make a new scale observation. Negative scales are not supported by this observation type.
 	 *
+	 * @param Object The Observation Object
 	 * @param Scale The scale of interest to the agent.
 	 * @param RelativeScale The scale the provided scale should be encoded relative to.
 	 * @param Tag The tag of the corresponding observation. Must match the tag given during Specify.
@@ -528,6 +531,7 @@ public:
 	/**
 	 * Make a new transform observation.
 	 *
+	 * @param Object The Observation Object
 	 * @param Transform The transform of interest to the agent.
 	 * @param RelativeTransform The transform the provided transform should be encoded relative to.
 	 * @param LocationScale Used to normalize the transform's location for this observation.
@@ -558,9 +562,52 @@ public:
 	UFUNCTION(BlueprintPure, Category = "LearningAgents", meta = (AdvancedDisplay = 3))
 	static FLearningAgentsObservationObjectElement MakeAngleObservationRadians(ULearningAgentsObservationObject* Object, const float Angle, const float RelativeAngle = 0.0f, const FName Tag = TEXT("Angle"));
 
-	UFUNCTION(BlueprintPure, Category = "LearningAgents", meta = (AdvancedDisplay = 4))
-	static FLearningAgentsObservationObjectElement MakeVelocityObservation(ULearningAgentsObservationObject* Object, const FVector Velocity, const FTransform RelativeTransform = FTransform(), const float VelocityScale = 200.0f, const FName Tag = TEXT("Velocity"));
+	/**
+	 * Make a new velocity observation.
+	 *
+	 * @param Object The Observation Object
+	 * @param Velocity The velocity of interest to the agent.
+	 * @param RelativeTransform The transform the provided velocity should be encoded relative to.
+	 * @param VelocityScale Used to normalize the data for this observation.
+	 * @param Tag The tag of the corresponding observation. Must match the tag given during Specify.
+	 * @param bVisualLoggerEnabled When true, debug data will be sent to the visual logger.
+	 * @param VisualLoggerListener The listener object which is making this observation. This must be set to use logging.
+	 * @param VisualLoggerAgentId The agent id associated with this observation.
+	 * @param VisualLoggerVelocityLocation A location for the visual logger to display the velocity in the world.
+	 * @param VisualLoggerLocation A location for the visual logger information in the world.
+	 * @param VisualLoggerColor The color for the visual logger display.
+	 * @return The newly created observation object element.
+	 */
+	UFUNCTION(BlueprintPure, Category = "LearningAgents", meta = (AdvancedDisplay = 4, DefaultToSelf = "VisualLoggerListener"))
+	static FLearningAgentsObservationObjectElement MakeVelocityObservation(
+		ULearningAgentsObservationObject* Object, 
+		const FVector Velocity, 
+		const FTransform RelativeTransform = FTransform(), 
+		const float VelocityScale = 200.0f, 
+		const FName Tag = TEXT("Velocity"),
+		const bool bVisualLoggerEnabled = false,
+		ULearningAgentsManagerListener* VisualLoggerListener = nullptr,
+		const int32 VisualLoggerAgentId = -1,
+		const FVector VisualLoggerVelocityLocation = FVector::ZeroVector,
+		const FVector VisualLoggerLocation = FVector::ZeroVector,
+		const FLinearColor VisualLoggerColor = FLinearColor::Red);
 
+	/**
+	 * Make a new direction observation.
+	 *
+	 * @param Object The Observation Object
+	 * @param Direction The direction of interest to the agent.
+	 * @param RelativeTransform The transform the provided direction should be encoded relative to.
+	 * @param Tag The tag of the corresponding observation. Must match the tag given during Specify.
+	 * @param bVisualLoggerEnabled When true, debug data will be sent to the visual logger.
+	 * @param VisualLoggerListener The listener object which is making this observation. This must be set to use logging.
+	 * @param VisualLoggerAgentId The agent id associated with this observation.
+	 * @param VisualLoggerDirectionLocation A location for the visual logger to display the direction in the world.
+	 * @param VisualLoggerLocation A location for the visual logger information in the world.
+	 * @param VisualLoggerArrowLength The length of the arrow to display to represent the direction.
+	 * @param VisualLoggerColor The color for the visual logger display.
+	 * @return The newly created observation object element.
+	 */
 	UFUNCTION(BlueprintPure, Category = "LearningAgents", meta = (AdvancedDisplay = 3, DefaultToSelf = "VisualLoggerListener"))
 	static FLearningAgentsObservationObjectElement MakeDirectionObservation(
 		ULearningAgentsObservationObject* Object,
@@ -592,7 +639,11 @@ public:
 		const FLinearColor VisualLoggerColor = FLinearColor::Red);
 
 	UFUNCTION(BlueprintPure, Category = "LearningAgents", meta = (AdvancedDisplay = 3))
-	static FLearningAgentsObservationObjectElement MakeProportionAlongSplineObservation(ULearningAgentsObservationObject* Object, const USplineComponent* SplineComponent, const float DistanceAlongSpline, const FName Tag = TEXT("ProportionAlongSpline"));
+	static FLearningAgentsObservationObjectElement MakeProportionAlongSplineObservation(
+		ULearningAgentsObservationObject* Object, 
+		const USplineComponent* SplineComponent, 
+		const float DistanceAlongSpline, 
+		const FName Tag = TEXT("ProportionAlongSpline"));
 
 	UFUNCTION(BlueprintPure, Category = "LearningAgents", meta = (AdvancedDisplay = 4, DefaultToSelf = "VisualLoggerListener"))
 	static FLearningAgentsObservationObjectElement MakeDirectionAlongSplineObservation(
@@ -608,16 +659,16 @@ public:
 		const float VisualLoggerArrowLength = 100.0f,
 		const FLinearColor VisualLoggerColor = FLinearColor::Red);
 
-	UFUNCTION(BlueprintPure, Category = "LearningAgents", meta = (AdvancedDisplay = 5))
-	static FLearningAgentsObservationObjectElement MakePropertiesAlongSplineObservation(ULearningAgentsObservationObject* Object, const USplineComponent* SplineComponent, const float DistanceAlongSpline, const FTransform RelativeTransform = FTransform(), const float LocationScale = 100.0f, const FName Tag = TEXT("PropertiesAlongSpline"));
-
 	// Make Ray Cast Observations
 
 	UFUNCTION(BlueprintPure, Category = "LearningAgents", meta = (AdvancedDisplay = 5))
-	static FLearningAgentsObservationObjectElement MakeProportionAlongRayObservation(ULearningAgentsObservationObject* Object, const FVector RayStart, const FVector RayEnd, const FTransform RayTransform = FTransform(), const ECollisionChannel CollisionChannel = ECollisionChannel::ECC_WorldStatic, const FName Tag = TEXT("ProportionAlongRay"));
-
-	UFUNCTION(BlueprintPure, Category = "LearningAgents", meta = (AdvancedDisplay = 5))
-	static FLearningAgentsObservationObjectElement MakeProportionAlongRaysObservation(ULearningAgentsObservationObject* Object, const TArray<FVector>& RayStarts, const TArray<FVector>& RayEnds, const FTransform RayTransform, const ECollisionChannel CollisionChannel = ECollisionChannel::ECC_WorldStatic, const FName Tag = TEXT("ProportionAlongRays"));
+	static FLearningAgentsObservationObjectElement MakeProportionAlongRayObservation(
+		ULearningAgentsObservationObject* Object, 
+		const FVector RayStart, 
+		const FVector RayEnd, 
+		const FTransform RayTransform = FTransform(), 
+		const ECollisionChannel CollisionChannel = ECollisionChannel::ECC_WorldStatic, 
+		const FName Tag = TEXT("ProportionAlongRay"));
 
 public:
 
@@ -768,21 +819,11 @@ public:
 	UFUNCTION(BlueprintPure = false, Category = "LearningAgents", meta = (AdvancedDisplay = 4, ReturnDisplayName = "Success"))
 	static bool GetDirectionAlongSplineObservation(FVector& OutDirection, const ULearningAgentsObservationObject* Object, const FLearningAgentsObservationObjectElement Element, const FTransform RelativeTransform = FTransform(), const FName Tag = TEXT("DirectionAlongSpline"));
 
-	UFUNCTION(BlueprintPure = false, Category = "LearningAgents", meta = (AdvancedDisplay = 9, ReturnDisplayName = "Success"))
-	static bool GetPropertiesAlongSplineObservation(FVector& OutLocation, bool& bOutIsClosedLoop, float& OutAngle, float& OutPropotion, FVector& OutDirection, const ULearningAgentsObservationObject* Object, const FLearningAgentsObservationObjectElement Element, const FTransform RelativeTransform = FTransform(), const float LocationScale = 100.0f, const FName Tag = TEXT("PropertiesAlongSpline"));
-
 	// Get Ray Cast Observations
 
 	UFUNCTION(BlueprintPure, Category = "LearningAgents", meta = (AdvancedDisplay = 3, ReturnDisplayName = "Success"))
 	static bool GetProportionAlongRayObservation(float& OutProportion, const ULearningAgentsObservationObject* Object, const FLearningAgentsObservationObjectElement Element, const FName Tag = TEXT("ProportionAlongRay"));
 
-	UFUNCTION(BlueprintPure, Category = "LearningAgents", meta = (AdvancedDisplay = 3, ReturnDisplayName = "Success"))
-	static bool GetProportionAlongRaysObservationNum(int32& OutProportionNum, const ULearningAgentsObservationObject* Object, const FLearningAgentsObservationObjectElement Element, const FName Tag = TEXT("ProportionAlongRays"));
-
-	UFUNCTION(BlueprintPure, Category = "LearningAgents", meta = (AdvancedDisplay = 3, ReturnDisplayName = "Success"))
-	static bool GetProportionAlongRaysObservation(TArray<float>& OutProportions, const ULearningAgentsObservationObject* Object, const FLearningAgentsObservationObjectElement Element, const FName Tag = TEXT("ProportionAlongRays"));
-	static bool GetProportionAlongRaysObservationToArrayView(TArrayView<float> OutProportions, const ULearningAgentsObservationObject* Object, const FLearningAgentsObservationObjectElement Element, const FName Tag = TEXT("ProportionAlongRays"));
-
 };
 
 
diff --git a/Engine/Plugins/Experimental/LearningAgents/Source/LearningAgentsTraining/Private/LearningAgentsCompletions.cpp b/Engine/Plugins/Experimental/LearningAgents/Source/LearningAgentsTraining/Private/LearningAgentsCompletions.cpp
--- a/Engine/Plugins/Experimental/LearningAgents/Source/LearningAgentsTraining/Private/LearningAgentsCompletions.cpp
+++ b/Engine/Plugins/Experimental/LearningAgents/Source/LearningAgentsTraining/Private/LearningAgentsCompletions.cpp
@@ -376,6 +376,18 @@ ELearningAgentsCompletion ULearningAgentsCompletions::MakeCompletionOnLocationOu
 	{
 		const ULearningAgentsVisualLoggerObject* VisualLoggerObject = VisualLoggerListener->GetOrAddVisualLoggerObject(Tag);
 
+		UE_LEARNING_AGENTS_VLOG_LOCATION(VisualLoggerObject, LogLearning, Display,
+			Location,
+			10,
+			VisualLoggerColor.ToFColor(true),
+			TEXT(""));
+
+		UE_LEARNING_AGENTS_VLOG_SEGMENT(VisualLoggerObject, LogLearning, Display,
+			Location,
+			BoundsTransform.GetLocation(),
+			VisualLoggerColor.ToFColor(true),
+			TEXT(""));
+
 		UE_LEARNING_AGENTS_VLOG_OBOX(VisualLoggerObject, LogLearning, Display,
 			FBox(BoundsMins, BoundsMaxs),
 			BoundsTransform.ToMatrixWithScale(),
diff --git a/Engine/Plugins/Experimental/LearningAgents/Source/LearningAgentsTraining/Private/LearningAgentsRewards.cpp b/Engine/Plugins/Experimental/LearningAgents/Source/LearningAgentsTraining/Private/LearningAgentsRewards.cpp
--- a/Engine/Plugins/Experimental/LearningAgents/Source/LearningAgentsTraining/Private/LearningAgentsRewards.cpp
+++ b/Engine/Plugins/Experimental/LearningAgents/Source/LearningAgentsTraining/Private/LearningAgentsRewards.cpp
@@ -504,8 +504,8 @@ float ULearningAgentsRewards::MakeRewardFromVelocityAlongSpline(
 		}
 	}
 
-	const float DistanceVelocity = ((Distance1 - Distance0) / FiniteDiff) * (Velocity.Length() / FMath::Max(VelocityScale, UE_SMALL_NUMBER));
-	const float Reward = RewardScale * DistanceVelocity;
+	const float SplineVelocity = ((Distance1 - Distance0) / FiniteDiff) * Velocity.Length();
+	const float Reward = RewardScale * (SplineVelocity / FMath::Max(VelocityScale, UE_SMALL_NUMBER));
 
 #if UE_LEARNING_AGENTS_ENABLE_VISUAL_LOG
 	if (bVisualLoggerEnabled && VisualLoggerListener)
@@ -521,19 +521,20 @@ float ULearningAgentsRewards::MakeRewardFromVelocityAlongSpline(
 			LogLearning,
 			Display,
 			SplineLocation0,
-			SplineLocation0 + DistanceVelocity * SplineVelocityDirection,
+			SplineLocation0 + FMath::Abs(SplineVelocity) * SplineVelocityDirection,
 			VisualLoggerColor.ToFColor(true),
 			TEXT(""));
 
 		UE_LEARNING_AGENTS_VLOG_STRING(VisualLoggerObject, LogLearning, Display, VisualLoggerLocation,
 			VisualLoggerColor.ToFColor(true),
-			TEXT("Listener: %s\nTag: %s\nAgent Id: % 3i\nLocation [% 6.2f % 6.2f % 6.2f]\nVelocity [% 6.2f % 6.2f % 6.2f]\nDistanceVelocity: [% 6.2f]\nScale: [% 6.2f]\nReward: [% 6.2f]"),
+			TEXT("Listener: %s\nTag: %s\nAgent Id: % 3i\nLocation [% 6.2f % 6.2f % 6.2f]\nVelocity [% 6.2f % 6.2f % 6.2f]\nVelocity Along Spline: [% 6.2f]\nVelocity Scale: [% 6.2f]\nScale: [% 6.2f]\nReward: [% 6.2f]"),
 			*VisualLoggerListener->GetName(),
 			*Tag.ToString(),
 			VisualLoggerAgentId,
 			Location.X, Location.Y, Location.Z,
 			Velocity.X, Velocity.Y, Velocity.Z,
-			DistanceVelocity,
+			SplineVelocity,
+			VelocityScale,
 			RewardScale,
 			Reward);
 	}
