LearningAgents: Automatically disable editor cpu throttling when training.

diff --git a/Engine/Plugins/Experimental/LearningAgents/Source/LearningAgentsTraining/LearningAgentsTraining.Build.cs b/Engine/Plugins/Experimental/LearningAgents/Source/LearningAgentsTraining/LearningAgentsTraining.Build.cs
--- a/Engine/Plugins/Experimental/LearningAgents/Source/LearningAgentsTraining/LearningAgentsTraining.Build.cs
+++ b/Engine/Plugins/Experimental/LearningAgents/Source/LearningAgentsTraining/LearningAgentsTraining.Build.cs
@@ -35,7 +35,12 @@ public class LearningAgentsTraining : ModuleRules
 				"LearningTraining",
 				// ... add private dependencies that you statically link with here ...	
 			});
-		
+
+		if (Target.bBuildEditor)
+		{
+			PrivateDependencyModuleNames.Add("UnrealEd");
+		}
+
 		DynamicallyLoadedModuleNames.AddRange(
 			new string[]
 			{
diff --git a/Engine/Plugins/Experimental/LearningAgents/Source/LearningAgentsTraining/Private/LearningAgentsTrainer.cpp b/Engine/Plugins/Experimental/LearningAgents/Source/LearningAgentsTraining/Private/LearningAgentsTrainer.cpp
--- a/Engine/Plugins/Experimental/LearningAgents/Source/LearningAgentsTraining/Private/LearningAgentsTrainer.cpp
+++ b/Engine/Plugins/Experimental/LearningAgents/Source/LearningAgentsTraining/Private/LearningAgentsTrainer.cpp
@@ -25,6 +25,10 @@
 #include "Engine/GameViewportClient.h"
 #include "EngineDefines.h"
 
+#if WITH_EDITOR
+#include "Editor/EditorPerformanceSettings.h"
+#endif
+
 namespace UE::Learning::Agents
 {
 	ELearningAgentsTrainerDevice GetLearningAgentsTrainerDevice(const ETrainerDevice Device)
@@ -431,6 +435,14 @@ void ULearningAgentsTrainer::BeginTraining(
 		ViewModeIndex = ViewportClient->ViewModeIndex;
 	}
 
+#if WITH_EDITOR
+	UEditorPerformanceSettings* EditorPerformanceSettings = GetMutableDefault<UEditorPerformanceSettings>();
+	if (EditorPerformanceSettings)
+	{
+		bUseLessCPUInTheBackground = EditorPerformanceSettings->bThrottleCPUWhenNotForeground;
+	}
+#endif
+
 	// Apply Training GameState Settings
 
 	FApp::SetUseFixedTimeStep(TrainerGameSettings.bUseFixedTimeStep);
@@ -459,6 +471,17 @@ void ULearningAgentsTrainer::BeginTraining(
 		ViewportClient->ViewModeIndex = EViewModeIndex::VMI_Unlit;
 	}
 
+#if WITH_EDITOR
+	if (TrainerGameSettings.bDisableUseLessCPUInTheBackground && EditorPerformanceSettings)
+	{
+		if (TrainerGameSettings.bDisableUseLessCPUInTheBackground)
+		{
+			EditorPerformanceSettings->bThrottleCPUWhenNotForeground = false;
+			EditorPerformanceSettings->PostEditChange();
+		}
+	}
+#endif
+
 	// Start Trainer
 
 	UE::Learning::FPPOTrainerTrainingSettings PPOTrainingSettings;
@@ -592,6 +615,15 @@ void ULearningAgentsTrainer::DoneTraining()
 			ViewportClient->ViewModeIndex = ViewModeIndex;
 		}
 
+#if WITH_EDITOR
+		UEditorPerformanceSettings* EditorPerformanceSettings = GetMutableDefault<UEditorPerformanceSettings>();
+		if (EditorPerformanceSettings)
+		{
+			EditorPerformanceSettings->bThrottleCPUWhenNotForeground = bUseLessCPUInTheBackground;
+			EditorPerformanceSettings->PostEditChange();
+		}
+#endif
+
 		bIsTraining = false;
 	}
 }
diff --git a/Engine/Plugins/Experimental/LearningAgents/Source/LearningAgentsTraining/Public/LearningAgentsTrainer.h b/Engine/Plugins/Experimental/LearningAgents/Source/LearningAgentsTraining/Public/LearningAgentsTrainer.h
--- a/Engine/Plugins/Experimental/LearningAgents/Source/LearningAgentsTraining/Public/LearningAgentsTrainer.h
+++ b/Engine/Plugins/Experimental/LearningAgents/Source/LearningAgentsTraining/Public/LearningAgentsTrainer.h
@@ -99,6 +99,14 @@ public:
 	/** If true, the viewport rendering will be unlit; Otherwise, it will not. Disabling lighting can speed up the game simulation. */
 	UPROPERTY(EditAnywhere, Category = "LearningAgents")
 	bool bUseUnlitViewportRendering = false;
+
+#if WITH_EDITORONLY_DATA
+
+	/** If true, the Use Less CPU In The Background editor setting will be disabled. This prevents the editor from running slowly when minimized. */
+	UPROPERTY(EditAnywhere, Category = "LearningAgents")
+	bool bDisableUseLessCPUInTheBackground = true;
+
+#endif
 };
 
 /** Enumeration of the training devices. */
@@ -681,6 +689,7 @@ private:
 // ----- Private Recording of GameSettings ----- 
 private:
 
+	bool bUseLessCPUInTheBackground = true;
 	bool bFixedTimestepUsed = false;
 	float FixedTimeStepDeltaTime = -1.0f;
 	bool bVSyncEnabled = true;
