LearningAgents: Added some more docs and visual logging.

diff --git a/Engine/Plugins/Experimental/LearningAgents/Source/LearningAgents/Private/LearningAgentsObservations.cpp b/Engine/Plugins/Experimental/LearningAgents/Source/LearningAgents/Private/LearningAgentsObservations.cpp
--- a/Engine/Plugins/Experimental/LearningAgents/Source/LearningAgents/Private/LearningAgentsObservations.cpp
+++ b/Engine/Plugins/Experimental/LearningAgents/Source/LearningAgents/Private/LearningAgentsObservations.cpp
@@ -1915,19 +1915,71 @@ FLearningAgentsObservationObjectElement ULearningAgentsObservations::MakeTransfo
 		Tag);
 }
 
-FLearningAgentsObservationObjectElement ULearningAgentsObservations::MakeAngleObservation(ULearningAgentsObservationObject* Object, const float Angle, const float RelativeAngle, const FName Tag)
+FLearningAgentsObservationObjectElement ULearningAgentsObservations::MakeAngleObservation(
+	ULearningAgentsObservationObject* Object, 
+	const float Angle, 
+	const float RelativeAngle, 
+	const FName Tag,
+	const bool bVisualLoggerEnabled,
+	ULearningAgentsManagerListener* VisualLoggerListener,
+	const int32 VisualLoggerAgentId,
+	const FVector VisualLoggerLocation,
+	const FLinearColor VisualLoggerColor)
 {
-	return MakeAngleObservationRadians(Object, FMath::DegreesToRadians(Angle), FMath::DegreesToRadians(RelativeAngle), Tag);
+	const float LocalAngle = FMath::FindDeltaAngleDegrees(RelativeAngle, Angle);
+	const float EncodedX = FMath::Sin(FMath::DegreesToRadians(Angle));
+	const float EncodedY = FMath::Cos(FMath::DegreesToRadians(Angle));
+
+#if UE_LEARNING_AGENTS_ENABLE_VISUAL_LOG
+	if (bVisualLoggerEnabled && VisualLoggerListener)
+	{
+		const ULearningAgentsVisualLoggerObject* VisualLoggerObject = VisualLoggerListener->GetOrAddVisualLoggerObject(Tag);
+
+		UE_LEARNING_AGENTS_VLOG_ANGLE_DEGREES(VisualLoggerObject, LogLearning, Display,
+			Angle,
+			RelativeAngle,
+			VisualLoggerLocation,
+			10.0f,
+			VisualLoggerColor.ToFColor(true),
+			TEXT(""));
+
+		UE_LEARNING_AGENTS_VLOG_STRING(VisualLoggerObject, LogLearning, Display, VisualLoggerLocation,
+			VisualLoggerColor.ToFColor(true),
+			TEXT("Listener: %s\nTag: %s\nAgent Id: % 3i\nAngle: [% 6.1f]\nRelative Angle: [% 6.1f]\nLocal Angle: [% 6.1f]\nEncoded: [% 6.2f % 6.2f]"),
+			*VisualLoggerListener->GetName(),
+			*Tag.ToString(),
+			VisualLoggerAgentId,
+			Angle,
+			RelativeAngle,
+			LocalAngle,
+			EncodedX, EncodedY);
+	}
+#endif
+
+	return MakeContinuousObservationFromArrayView(Object, { EncodedX, EncodedY }, Tag);
 }
 
-FLearningAgentsObservationObjectElement ULearningAgentsObservations::MakeAngleObservationRadians(ULearningAgentsObservationObject* Object, const float Angle, const float RelativeAngle, const FName Tag)
+FLearningAgentsObservationObjectElement ULearningAgentsObservations::MakeAngleObservationRadians(
+	ULearningAgentsObservationObject* Object, 
+	const float Angle, 
+	const float RelativeAngle, 
+	const FName Tag,
+	const bool bVisualLoggerEnabled,
+	ULearningAgentsManagerListener* VisualLoggerListener,
+	const int32 VisualLoggerAgentId,
+	const FVector VisualLoggerLocation,
+	const FLinearColor VisualLoggerColor)
 {
-	const float LocalAngle = FMath::FindDeltaAngleRadians(RelativeAngle, Angle);
-
-	return MakeContinuousObservationFromArrayView(Object, {
-		(float)FMath::Sin(Angle),
-		(float)FMath::Cos(Angle),
-		}, Tag);
+	return MakeAngleObservation(
+		Object, 
+		FMath::RadiansToDegrees(Angle), 
+		FMath::RadiansToDegrees(RelativeAngle), 
+		Tag, 
+		bVisualLoggerEnabled, 
+		VisualLoggerListener, 
+		VisualLoggerAgentId, 
+		VisualLoggerLocation, 
+		VisualLoggerColor);
 }
 
 FLearningAgentsObservationObjectElement ULearningAgentsObservations::MakeVelocityObservation(
@@ -2067,7 +2119,16 @@ FLearningAgentsObservationObjectElement ULearningAgentsObservations::MakeLocatio
 		VisualLoggerColor);
 }
 
-FLearningAgentsObservationObjectElement ULearningAgentsObservations::MakeProportionAlongSplineObservation(ULearningAgentsObservationObject* Object, const USplineComponent* SplineComponent, const float DistanceAlongSpline, const FName Tag)
+FLearningAgentsObservationObjectElement ULearningAgentsObservations::MakeProportionAlongSplineObservation(
+	ULearningAgentsObservationObject* Object, 
+	const USplineComponent* SplineComponent, 
+	const float DistanceAlongSpline, 
+	const FName Tag,
+	const bool bVisualLoggerEnabled,
+	ULearningAgentsManagerListener* VisualLoggerListener,
+	const int32 VisualLoggerAgentId,
+	const FVector VisualLoggerLocation,
+	const FLinearColor VisualLoggerColor)
 {
 	if (!SplineComponent)
 	{
@@ -2081,12 +2142,32 @@ FLearningAgentsObservationObjectElement ULearningAgentsObservations::MakeProport
 		const float WrapDistance = FMath::Wrap(DistanceAlongSpline, 0.0f, TotalDistance);
 		const float Proportion = WrapDistance / FMath::Max(TotalDistance, UE_SMALL_NUMBER);
 		const float Angle = FMath::Wrap(UE_TWO_PI * Proportion, -UE_PI, UE_PI);
-		return MakeExclusiveUnionObservation(Object, TEXT("Angle"), MakeAngleObservation(Object, Angle), Tag);
+		
+		return MakeExclusiveUnionObservation(Object, TEXT("Angle"), MakeAngleObservationRadians(
+			Object, 
+			Angle,
+			0.0f,
+			TEXT("Angle"),
+			bVisualLoggerEnabled,
+			VisualLoggerListener,
+			VisualLoggerAgentId,
+			VisualLoggerLocation,
+			VisualLoggerColor), Tag);
 	}
 	else
 	{
 		const float Proportion = FMath::Clamp(DistanceAlongSpline / FMath::Max(SplineComponent->GetSplineLength(), UE_SMALL_NUMBER), 0.0f, 1.0f);
-		return MakeExclusiveUnionObservation(Object, TEXT("Proportion"), MakeFloatObservation(Object, Proportion), Tag);
+
+		return MakeExclusiveUnionObservation(Object, TEXT("Proportion"), MakeFloatObservation(
+			Object, 
+			Proportion, 
+			1.0f, 
+			TEXT("Float"),
+			bVisualLoggerEnabled, 
+			VisualLoggerListener, 
+			VisualLoggerAgentId, 
+			VisualLoggerLocation, 
+			VisualLoggerColor), Tag);
 	}
 }
 
@@ -2135,27 +2216,74 @@ FLearningAgentsObservationObjectElement ULearningAgentsObservations::MakeDirecti
 		VisualLoggerColor);
 }
 
-FLearningAgentsObservationObjectElement ULearningAgentsObservations::MakeProportionAlongRayObservation(ULearningAgentsObservationObject* Object, const FVector RayStart, const FVector RayEnd, const FTransform RayTransform, const ECollisionChannel CollisionChannel, const FName Tag)
+FLearningAgentsObservationObjectElement ULearningAgentsObservations::MakeProportionAlongRayObservation(
+	ULearningAgentsObservationObject* Object, 
+	const FVector RayStart, 
+	const FVector RayEnd, 
+	const FTransform RayTransform, 
+	const ECollisionChannel CollisionChannel, 
+	const FName Tag,
+	const bool bVisualLoggerEnabled,
+	ULearningAgentsManagerListener* VisualLoggerListener,
+	const int32 VisualLoggerAgentId,
+	const FVector VisualLoggerLocation,
+	const FLinearColor VisualLoggerColor)
 {
 	const FVector RayStartWorld = RayTransform.TransformPosition(RayStart);
 	const FVector RayEndWorld = RayTransform.TransformPosition(RayEnd);
 	const float RayDistance = FVector::Distance(RayStartWorld, RayEndWorld);
 
-
 	FCollisionObjectQueryParams ObjectQueryParams;
 	ObjectQueryParams.AddObjectTypesToQuery(CollisionChannel);
 
 	FHitResult TraceHit;
 	const bool bHit = Object->GetWorld()->LineTraceSingleByObjectType(TraceHit, RayStartWorld, RayEndWorld, ObjectQueryParams);
 
-	if (bHit)
-	{
-		return MakeFloatObservation(Object, 1.0f - TraceHit.Time, 1.0f, Tag);
-	}
-	else
+	const float Encoded = bHit ? 1.0f - TraceHit.Time : 0.0f;
+
+#if UE_LEARNING_AGENTS_ENABLE_VISUAL_LOG
+	if (bVisualLoggerEnabled && VisualLoggerListener)
 	{
-		return MakeFloatObservation(Object, 0.0f, 1.0f, Tag);
+		const ULearningAgentsVisualLoggerObject* VisualLoggerObject = VisualLoggerListener->GetOrAddVisualLoggerObject(Tag);
+
+		UE_LEARNING_AGENTS_VLOG_SEGMENT(VisualLoggerObject, LogLearning, Display,
+			RayStartWorld,
+			RayEndWorld,
+			VisualLoggerColor.ToFColor(true),
+			TEXT(""));
+
+		if (bHit)
+		{
+			UE_LEARNING_AGENTS_VLOG_LOCATION(VisualLoggerObject, LogLearning, Display,
+				RayStartWorld + TraceHit.Time * (RayEndWorld - RayStartWorld),
+				10.f,
+				VisualLoggerColor.ToFColor(true),
+				TEXT(""));
+		}
+
+		UE_LEARNING_AGENTS_VLOG_TRANSFORM(VisualLoggerObject, LogLearning, Display,
+			RayTransform.GetTranslation(),
+			RayTransform.GetRotation(),
+			VisualLoggerColor.ToFColor(true),
+			TEXT(""));
+
+		UE_LEARNING_AGENTS_VLOG_STRING(VisualLoggerObject, LogLearning, Display, VisualLoggerLocation,
+			VisualLoggerColor.ToFColor(true),
+			TEXT("Listener: %s\nTag: %s\nAgent Id: % 3i\nRay Start: [% 6.1f % 6.1f % 6.1f]\nRay End: [% 6.1f % 6.1f % 6.1f]\nRay Start World: [% 6.1f % 6.1f % 6.1f]\nRay End World: [% 6.1f % 6.1f % 6.1f]\nCollision Channel: [%s]\nHit: [%s]\nEncoded: [% 6.2f]"),
+			*VisualLoggerListener->GetName(),
+			*Tag.ToString(),
+			VisualLoggerAgentId,
+			RayStart.X, RayStart.Y, RayStart.Z,
+			RayEnd.X, RayEnd.Y, RayEnd.Z,
+			RayStartWorld.X, RayStartWorld.Y, RayStartWorld.Z,
+			RayEndWorld.X, RayEndWorld.Y, RayEndWorld.Z,
+			*StaticEnum<ECollisionChannel>()->GetNameStringByValue(CollisionChannel),
+			bHit ? TEXT("true") : TEXT("false"),
+			Encoded);
 	}
+#endif
+
+	return MakeFloatObservation(Object, Encoded, 1.0f, Tag);
 }
 
 
diff --git a/Engine/Plugins/Experimental/LearningAgents/Source/LearningAgents/Public/LearningAgentsDebug.h b/Engine/Plugins/Experimental/LearningAgents/Source/LearningAgents/Public/LearningAgentsDebug.h
--- a/Engine/Plugins/Experimental/LearningAgents/Source/LearningAgents/Public/LearningAgentsDebug.h
+++ b/Engine/Plugins/Experimental/LearningAgents/Source/LearningAgents/Public/LearningAgentsDebug.h
@@ -45,7 +45,7 @@
 		UE_LEARNING_AGENTS_VLOG_STRING(Owner, Category, Verbosity, Location + FVector(0.0f, 0.0f, 20.0f), Color, Format, ##__VA_ARGS__); \
 	}
 
-#define UE_LEARNING_AGENTS_VLOG_ANGLE(Owner, Category, Verbosity, Angle, RelativeAngle, Location, Radius, Color, Format, ...) \
+#define UE_LEARNING_AGENTS_VLOG_ANGLE_RADIANS(Owner, Category, Verbosity, Angle, RelativeAngle, Location, Radius, Color, Format, ...) \
 	{ \
 		UE_LEARNING_AGENTS_VLOG_CIRCLE(Owner, Category, Verbosity, Location, FVector::UpVector, Radius, Color, TEXT("")); \
 		UE_LEARNING_AGENTS_VLOG_SEGMENT(Owner, Category, Verbosity, Location, Location + Radius * FVector(FMath::Sin(RelativeAngle), FMath::Cos(RelativeAngle), 0.0f), Color, TEXT("")); \
@@ -53,4 +53,7 @@
 		UE_LEARNING_AGENTS_VLOG_LOCATION(Owner, Category, Verbosity, Location + Radius * FVector(FMath::Sin(Angle), FMath::Cos(Angle), 0.0f), static_cast<uint16>(Radius / 20.0f), Color, Format, ##__VA_ARGS__); \
 	}
 
+#define UE_LEARNING_AGENTS_VLOG_ANGLE_DEGREES(Owner, Category, Verbosity, Angle, RelativeAngle, Location, Radius, Color, Format, ...) \
+	UE_LEARNING_AGENTS_VLOG_ANGLE_RADIANS(Owner, Category, Verbosity, FMath::DegreesToRadians(Angle), FMath::DegreesToRadians(RelativeAngle), Location, Radius, Color, Format)
+
 #endif
diff --git a/Engine/Plugins/Experimental/LearningAgents/Source/LearningAgents/Public/LearningAgentsObservations.h b/Engine/Plugins/Experimental/LearningAgents/Source/LearningAgents/Public/LearningAgentsObservations.h
--- a/Engine/Plugins/Experimental/LearningAgents/Source/LearningAgents/Public/LearningAgentsObservations.h
+++ b/Engine/Plugins/Experimental/LearningAgents/Source/LearningAgents/Public/LearningAgentsObservations.h
@@ -132,79 +132,339 @@ public:
 		const ULearningAgentsObservationObject* Object,
 		const FLearningAgentsObservationObjectElement ObjectElement);
 
+	/**
+	 * Logs an Observation Object Element. Useful for debugging.
+	 *
+	 * @param Object				Observation Object
+	 * @param ObjectElement			Observation Object Element
+	 */
 	UFUNCTION(BlueprintPure = false, Category = "LearningAgents")
 	static void LogObservation(const ULearningAgentsObservationObject* Object, const FLearningAgentsObservationObjectElement Element);
 
 public:
 
-	// Specify Basic Observations
-
+	/**
+	 * Specifies a new null observation. This represents an empty observation and can be useful when an observation is needed which has no value.
+	 *
+	 * @param Schema The Observation Schema
+	 * @param Tag The tag of this new observation. Used during observation object validation and debugging.
+	 * @return The newly created observation schema element.
+	 */
 	UFUNCTION(BlueprintPure, Category = "LearningAgents", meta = (AdvancedDisplay = 1))
 	static FLearningAgentsObservationSchemaElement SpecifyNullObservation(ULearningAgentsObservationSchema* Schema, const FName Tag = TEXT("Null"));
 
+	/**
+	 * Specifies a new continuous observation. This represents an observation made up of several float values.
+	 *
+	 * @param Schema The Observation Schema
+	 * @param Size The number of float values in the observation.
+	 * @param Tag The tag of this new observation. Used during observation object validation and debugging.
+	 * @return The newly created observation schema element.
+	 */
 	UFUNCTION(BlueprintPure, Category = "LearningAgents", meta = (AdvancedDisplay = 2))
 	static FLearningAgentsObservationSchemaElement SpecifyContinuousObservation(ULearningAgentsObservationSchema* Schema, const int32 Size, const FName Tag = TEXT("Continuous"));
 
+	/**
+	 * Specifies a new exclusive discrete observation. This represents a discrete observation which is an exclusive selection from multiple choices.
+	 *
+	 * @param Schema The Observation Schema
+	 * @param Size The number of discrete options in the observation.
+	 * @param Tag The tag of this new observation. Used during observation object validation and debugging.
+	 * @return The newly created observation schema element.
+	 */
 	UFUNCTION(BlueprintPure, Category = "LearningAgents", meta = (AdvancedDisplay = 2))
 	static FLearningAgentsObservationSchemaElement SpecifyExclusiveDiscreteObservation(ULearningAgentsObservationSchema* Schema, const int32 Size, const FName Tag = TEXT("ExclusiveDiscrete"));
 
+	/**
+	 * Specifies a new inclusive discrete observation. This represents a discrete observation which is an inclusive selection from multiple choices.
+	 *
+	 * @param Schema The Observation Schema
+	 * @param Size The number of discrete options in the observation.
+	 * @param Tag The tag of this new observation. Used during observation object validation and debugging.
+	 * @return The newly created observation schema element.
+	 */
 	UFUNCTION(BlueprintPure, Category = "LearningAgents", meta = (AdvancedDisplay = 2))
 	static FLearningAgentsObservationSchemaElement SpecifyInclusiveDiscreteObservation(ULearningAgentsObservationSchema* Schema, const int32 Size, const FName Tag = TEXT("InclusiveDiscrete"));
 
+	/**
+	 * Specifies a new index observation. This represents an index into something such as an array.
+	 *
+	 * @param Schema The Observation Schema
+	 * @param Size The size of the object being indexed.
+	 * @param Tag The tag of this new observation. Used during observation object validation and debugging.
+	 * @return The newly created observation schema element.
+	 */
 	UFUNCTION(BlueprintPure, Category = "LearningAgents", meta = (AdvancedDisplay = 2))
 	static FLearningAgentsObservationSchemaElement SpecifyIndexObservation(ULearningAgentsObservationSchema* Schema, const int32 Size, const FName Tag = TEXT("Index"));
 
+	/**
+	 * Specifies a new count observation. This represents a count of something such as the size of an array.
+	 *
+	 * @param Schema The Observation Schema
+	 * @param Tag The tag of this new observation. Used during observation object validation and debugging.
+	 * @return The newly created observation schema element.
+	 */
 	UFUNCTION(BlueprintPure, Category = "LearningAgents", meta = (AdvancedDisplay = 1))
 	static FLearningAgentsObservationSchemaElement SpecifyCountObservation(ULearningAgentsObservationSchema* Schema, const FName Tag = TEXT("Count"));
 
+	/**
+	 * Specifies a new struct observation. This represents a group of named sub-observations.
+	 *
+	 * @param Schema The Observation Schema
+	 * @param Elements The sub-observations that make up this struct.
+	 * @param Tag The tag of this new observation. Used during observation object validation and debugging.
+	 * @return The newly created observation schema element.
+	 */
 	UFUNCTION(BlueprintPure, Category = "LearningAgents", meta = (AdvancedDisplay = 2))
 	static FLearningAgentsObservationSchemaElement SpecifyStructObservation(ULearningAgentsObservationSchema* Schema, const TMap<FName, FLearningAgentsObservationSchemaElement>& Elements, const FName Tag = TEXT("Struct"));
 
+	/**
+	 * Specifies a new struct observation. This represents a group of named sub-observations.
+	 *
+	 * @param Schema The Observation Schema
+	 * @param ElementNames The names of the sub-observations that make up this struct.
+	 * @param Elements The corresponding sub-observations that make up this struct. Must be the same size as ElementNames.
+	 * @param Tag The tag of this new observation. Used during observation object validation and debugging.
+	 * @return The newly created observation schema element.
+	 */
 	UFUNCTION(BlueprintPure, Category = "LearningAgents", meta = (AdvancedDisplay = 3))
 	static FLearningAgentsObservationSchemaElement SpecifyStructObservationFromArrays(ULearningAgentsObservationSchema* Schema, const TArray<FName>& ElementNames, const TArray<FLearningAgentsObservationSchemaElement>& Elements, const FName Tag = TEXT("Struct"));
+	
+	/**
+	 * Specifies a new struct observation. This represents a group of named sub-observations.
+	 *
+	 * @param Schema The Observation Schema
+	 * @param ElementNames The names of the sub-observations that make up this struct.
+	 * @param Elements The corresponding sub-observations that make up this struct. Must be the same size as ElementNames.
+	 * @param Tag The tag of this new observation. Used during observation object validation and debugging.
+	 * @return The newly created observation schema element.
+	 */
 	static FLearningAgentsObservationSchemaElement SpecifyStructObservationFromArrayViews(ULearningAgentsObservationSchema* Schema, const TArrayView<const FName> ElementNames, const TArrayView<const FLearningAgentsObservationSchemaElement> Elements, const FName Tag = TEXT("Struct"));
 
+	/**
+	 * Specifies a new exclusive union observation. This represents an observation which is exclusively chosen from a set of named sub-observations.
+	 * In other words, when this observation is created, you only need to provide one observation from the given sub-observations.
+	 *
+	 * @param Schema The Observation Schema
+	 * @param Elements The sub-observations that make up this union.
+	 * @param EncodingSize The encoding size used to encode each sub-observation.
+	 * @param Tag The tag of this new observation. Used during observation object validation and debugging.
+	 * @return The newly created observation schema element.
+	 */
 	UFUNCTION(BlueprintPure, Category = "LearningAgents", meta = (AdvancedDisplay = 2))
 	static FLearningAgentsObservationSchemaElement SpecifyExclusiveUnionObservation(ULearningAgentsObservationSchema* Schema, const TMap<FName, FLearningAgentsObservationSchemaElement>& Elements, const int32 EncodingSize = 128, const FName Tag = TEXT("ExclusiveUnion"));
 
+	/**
+	 * Specifies a new exclusive union observation. This represents an observation which is exclusively chosen from a set of named sub-observations.
+	 * In other words, when this observation is created, you only need to provide one observation from the given sub-observations.
+	 *
+	 * @param Schema The Observation Schema
+	 * @param ElementNames The names of the sub-observations that make up this union.
+	 * @param Elements The corresponding sub-observations that make up this union. Must be the same size as ElementNames.
+	 * @param EncodingSize The encoding size used to encode each sub-observation.
+	 * @param Tag The tag of this new observation. Used during observation object validation and debugging.
+	 * @return The newly created observation schema element.
+	 */
 	UFUNCTION(BlueprintPure, Category = "LearningAgents", meta = (AdvancedDisplay = 3))
 	static FLearningAgentsObservationSchemaElement SpecifyExclusiveUnionObservationFromArrays(ULearningAgentsObservationSchema* Schema, const TArray<FName>& ElementNames, const TArray<FLearningAgentsObservationSchemaElement>& Elements, const int32 EncodingSize = 128, const FName Tag = TEXT("ExclusiveUnion"));
+	
+	/**
+	 * Specifies a new exclusive union observation. This represents an observation which is exclusively chosen from a set of named sub-observations.
+	 * In other words, when this observation is created, you only need to provide one observation from the given sub-observations.
+	 *
+	 * @param Schema The Observation Schema
+	 * @param ElementNames The names of the sub-observations that make up this union.
+	 * @param Elements The corresponding sub-observations that make up this union. Must be the same size as ElementNames.
+	 * @param EncodingSize The encoding size used to encode each sub-observation.
+	 * @param Tag The tag of this new observation. Used during observation object validation and debugging.
+	 * @return The newly created observation schema element.
+	 */
 	static FLearningAgentsObservationSchemaElement SpecifyExclusiveUnionObservationFromArrayViews(ULearningAgentsObservationSchema* Schema, const TArrayView<const FName> ElementNames, const TArrayView<const FLearningAgentsObservationSchemaElement> Elements, const int32 EncodingSize = 128, const FName Tag = TEXT("ExclusiveUnion"));
 
+	/**
+	 * Specifies a new inclusive union observation. This represents an observation which is inclusively chosen from a set of named sub-observations.
+	 * In other words, when this observation is created, you can provide any combination of observations from the given sub-observations. Internally
+	 * this observation uses Attention so can be slower to evaluate and more difficult to train than other observation types. For this reason it 
+	 * should be used sparingly.
+	 *
+	 * @param Schema The Observation Schema
+	 * @param Elements The sub-observations that make up this union.
+	 * @param AttentionEncodingSize The encoding size used by the attention mechanism.
+	 * @param AttentionHeadNum The number of heads used by the attention mechanism.
+	 * @param ValueEncodingSize The output encoding size used by the attention mechanism.
+	 * @param Tag The tag of this new observation. Used during observation object validation and debugging.
+	 * @return The newly created observation schema element.
+	 */
 	UFUNCTION(BlueprintPure, Category = "LearningAgents", meta = (AdvancedDisplay = 2))
 	static FLearningAgentsObservationSchemaElement SpecifyInclusiveUnionObservation(ULearningAgentsObservationSchema* Schema, const TMap<FName, FLearningAgentsObservationSchemaElement>& Elements, const int32 AttentionEncodingSize = 32, const int32 AttentionHeadNum = 4, const int32 ValueEncodingSize = 32, const FName Tag = TEXT("InclusiveUnion"));
 
+	/**
+	 * Specifies a new inclusive union observation. This represents an observation which is inclusively chosen from a set of named sub-observations.
+	 * In other words, when this observation is created, you can provide any combination of observations from the given sub-observations. Internally
+	 * this observation uses Attention so can be slower to evaluate and more difficult to train than other observation types. For this reason it
+	 * should be used sparingly.
+	 *
+	 * @param Schema The Observation Schema
+	 * @param ElementNames The names of the sub-observations that make up this union.
+	 * @param Elements The corresponding sub-observations that make up this union. Must be the same size as ElementNames.
+	 * @param AttentionEncodingSize The encoding size used by the attention mechanism.
+	 * @param AttentionHeadNum The number of heads used by the attention mechanism.
+	 * @param ValueEncodingSize The output encoding size used by the attention mechanism.
+	 * @param Tag The tag of this new observation. Used during observation object validation and debugging.
+	 * @return The newly created observation schema element.
+	 */
 	UFUNCTION(BlueprintPure, Category = "LearningAgents", meta = (AdvancedDisplay = 3))
 	static FLearningAgentsObservationSchemaElement SpecifyInclusiveUnionObservationFromArrays(ULearningAgentsObservationSchema* Schema, const TArray<FName>& ElementNames, const TArray<FLearningAgentsObservationSchemaElement>& Elements, const int32 AttentionEncodingSize = 32, const int32 AttentionHeadNum = 4, const int32 ValueEncodingSize = 32, const FName Tag = TEXT("InclusiveUnion"));
+	
+	/**
+	 * Specifies a new inclusive union observation. This represents an observation which is inclusively chosen from a set of named sub-observations.
+	 * In other words, when this observation is created, you can provide any combination of observations from the given sub-observations. Internally
+	 * this observation uses Attention so can be slower to evaluate and more difficult to train than other observation types. For this reason it
+	 * should be used sparingly.
+	 *
+	 * @param Schema The Observation Schema
+	 * @param ElementNames The names of the sub-observations that make up this union.
+	 * @param Elements The corresponding sub-observations that make up this union. Must be the same size as ElementNames.
+	 * @param AttentionEncodingSize The encoding size used by the attention mechanism.
+	 * @param AttentionHeadNum The number of heads used by the attention mechanism.
+	 * @param ValueEncodingSize The output encoding size used by the attention mechanism.
+	 * @param Tag The tag of this new observation. Used during observation object validation and debugging.
+	 * @return The newly created observation schema element.
+	 */
 	static FLearningAgentsObservationSchemaElement SpecifyInclusiveUnionObservationFromArrayViews(ULearningAgentsObservationSchema* Schema, const TArrayView<const FName> ElementNames, const TArrayView<const FLearningAgentsObservationSchemaElement> Elements, const int32 AttentionEncodingSize = 32, const int32 AttentionHeadNum = 4, const int32 ValueEncodingSize = 32, const FName Tag = TEXT("InclusiveUnion"));
 
+	/**
+	 * Specifies a new static array observation. This represents an observation made up of a fixed-size array of some other observation.
+	 *
+	 * @param Schema The Observation Schema
+	 * @param Element The sub-observation that represents elements of this array.
+	 * @param Num The number of elements in the fixed size array.
+	 * @param Tag The tag of this new observation. Used during observation object validation and debugging.
+	 * @return The newly created observation schema element.
+	 */
 	UFUNCTION(BlueprintPure, Category = "LearningAgents", meta = (AdvancedDisplay = 3))
 	static FLearningAgentsObservationSchemaElement SpecifyStaticArrayObservation(ULearningAgentsObservationSchema* Schema, const FLearningAgentsObservationSchemaElement Element, const int32 Num, const FName Tag = TEXT("StaticArray"));
 
+	/**
+	 * Specifies a new set observation. This represents an observation made up of a Set of some other observation. This Set can be variable in size 
+	 * (up to some fixed maximum size) and elements are considered unordered. Internally this observation uses Attention so can be slower to evaluate 
+	 * and more difficult to train than other observation types. For this reason it should be used sparingly.
+	 *
+	 * @param Schema The Observation Schema
+	 * @param Element The sub-observation that represents elements of this array.
+	 * @param MaxNum The maximum number of elements that can be included in the set.
+	 * @param AttentionEncodingSize The encoding size used by the attention mechanism.
+	 * @param AttentionHeadNum The number of heads used by the attention mechanism.
+	 * @param ValueEncodingSize The output encoding size used by the attention mechanism.
+	 * @param Tag The tag of this new observation. Used during observation object validation and debugging.
+	 * @return The newly created observation schema element.
+	 */
 	UFUNCTION(BlueprintPure, Category = "LearningAgents", meta = (AdvancedDisplay = 3))
 	static FLearningAgentsObservationSchemaElement SpecifySetObservation(ULearningAgentsObservationSchema* Schema, const FLearningAgentsObservationSchemaElement Element, const int32 MaxNum, const int32 AttentionEncodingSize = 32, const int32 AttentionHeadNum = 4, const int32 ValueEncodingSize = 32, const FName Tag = TEXT("Set"));
 
+	/**
+	 * Specifies a new pair observation. This represents an observation made up of two sub-observations.
+	 *
+	 * @param Schema The Observation Schema
+	 * @param Key The first sub-observation.
+	 * @param Value The second sub-observation.
+	 * @param Tag The tag of this new observation. Used during observation object validation and debugging.
+	 * @return The newly created observation schema element.
+	 */
 	UFUNCTION(BlueprintPure, Category = "LearningAgents", meta = (AdvancedDisplay = 3))
 	static FLearningAgentsObservationSchemaElement SpecifyPairObservation(ULearningAgentsObservationSchema* Schema, const FLearningAgentsObservationSchemaElement Key, const FLearningAgentsObservationSchemaElement Value, const FName Tag = TEXT("Pair"));
 
+	/**
+	 * Specifies a new array observation. This represents an observation made up of an Array of some other observation. This Array can be variable in 
+	 * size (up to some fixed maximum size) and the order of elements is taken into consideration. Internally this observation uses Attention so can 
+	 * be slower to evaluate and more difficult to train than other observation types. For this reason it should be used sparingly.
+	 *
+	 * @param Schema The Observation Schema
+	 * @param Element The sub-observation that represents elements of this array.
+	 * @param MaxNum The maximum number of elements that can be included in the array.
+	 * @param AttentionEncodingSize The encoding size used by the attention mechanism.
+	 * @param AttentionHeadNum The number of heads used by the attention mechanism.
+	 * @param ValueEncodingSize The output encoding size used by the attention mechanism.
+	 * @param Tag The tag of this new observation. Used during observation object validation and debugging.
+	 * @return The newly created observation schema element.
+	 */
 	UFUNCTION(BlueprintPure, Category = "LearningAgents", meta = (AdvancedDisplay = 3))
 	static FLearningAgentsObservationSchemaElement SpecifyArrayObservation(ULearningAgentsObservationSchema* Schema, const FLearningAgentsObservationSchemaElement Element, const int32 MaxNum, const int32 AttentionEncodingSize = 32, const int32 AttentionHeadNum = 4, const int32 ValueEncodingSize = 32, const FName Tag = TEXT("Array"));
 
+	/**
+	 * Specifies a new map observation. This represents an observation made up of a Map of some other key and pair observations. This Map can be 
+	 * variable in size (up to some fixed maximum size) and elements are considered unordered. Internally this observation uses Attention so can
+	 * be slower to evaluate and more difficult to train than other observation types. For this reason it should be used sparingly.
+	 *
+	 * @param Schema The Observation Schema
+	 * @param KeyElement The sub-observation that represents keys in this map.
+	 * @param ValueElement The sub-observation that represents values in this map.
+	 * @param MaxNum The maximum number of elements that can be included in the map.
+	 * @param AttentionEncodingSize The encoding size used by the attention mechanism.
+	 * @param AttentionHeadNum The number of heads used by the attention mechanism.
+	 * @param ValueEncodingSize The output encoding size used by the attention mechanism.
+	 * @param Tag The tag of this new observation. Used during observation object validation and debugging.
+	 * @return The newly created observation schema element.
+	 */
 	UFUNCTION(BlueprintPure, Category = "LearningAgents", meta = (AdvancedDisplay = 4))
 	static FLearningAgentsObservationSchemaElement SpecifyMapObservation(ULearningAgentsObservationSchema* Schema, const FLearningAgentsObservationSchemaElement KeyElement, const FLearningAgentsObservationSchemaElement ValueElement, const int32 MaxNum, const int32 AttentionEncodingSize = 32, const int32 AttentionHeadNum = 4, const int32 ValueEncodingSize = 32, const FName Tag = TEXT("Map"));
 
+	/**
+	 * Specifies a new enum observation. This represents an exclusive choice from elements of the given Enum. To use this with an Enum defined in C++ 
+	 * use the FindEnumByName convenience function.
+	 *
+	 * @param Schema The Observation Schema
+	 * @param Enum The enum type to use.
+	 * @return The newly created observation schema element.
+	 */
 	UFUNCTION(BlueprintPure, Category = "LearningAgents", meta = (AdvancedDisplay = 2))
 	static FLearningAgentsObservationSchemaElement SpecifyEnumObservation(ULearningAgentsObservationSchema* Schema, const UEnum* Enum, const FName Tag = TEXT("Enum"));
 
+	/**
+	 * Specifies a new bitmask observation. This represents an inclusive choice from elements of the given Enum. To use this with an Enum defined in 
+	 * C++ use the FindEnumByName convenience function.
+	 *
+	 * @param Schema The Observation Schema
+	 * @param Enum The enum type to use.
+	 * @return The newly created observation schema element.
+	 */
 	UFUNCTION(BlueprintPure, Category = "LearningAgents", meta = (AdvancedDisplay = 2))
 	static FLearningAgentsObservationSchemaElement SpecifyBitmaskObservation(ULearningAgentsObservationSchema* Schema, const UEnum* Enum, const FName Tag = TEXT("Bitmask"));
 
+	/**
+	 * Specifies a new optional observation. This represents an observation which may or may not be provided.
+	 *
+	 * @param Schema The Observation Schema
+	 * @param Elements The sub-observation that may or may not be provided.
+	 * @param EncodingSize The encoding size used to encode this sub-observation.
+	 * @return The newly created observation schema element.
+	 */
 	UFUNCTION(BlueprintPure, Category = "LearningAgents", meta = (AdvancedDisplay = 2))
 	static FLearningAgentsObservationSchemaElement SpecifyOptionalObservation(ULearningAgentsObservationSchema* Schema, const FLearningAgentsObservationSchemaElement Element, const int32 EncodingSize = 128, const FName Tag = TEXT("Optional"));
 
+	/**
+	 * Specifies a new either observation. This represents an observation which will be either sub-observation A or sub-observation B.
+	 *
+	 * @param Schema The Observation Schema
+	 * @param A The first sub-observation.
+	 * @param A The second sub-observation.
+	 * @param EncodingSize The encoding size used to encode each sub-observation.
+	 * @return The newly created observation schema element.
+	 */
 	UFUNCTION(BlueprintPure, Category = "LearningAgents", meta = (AdvancedDisplay = 3))
 	static FLearningAgentsObservationSchemaElement SpecifyEitherObservation(ULearningAgentsObservationSchema* Schema, const FLearningAgentsObservationSchemaElement A, const FLearningAgentsObservationSchemaElement B, const int32 EncodingSize = 128, const FName Tag = TEXT("Either"));
 
+	/**
+	 * Specifies a new encoding observation. This represents an observation which will be an encoding of another sub-observation using a small neural
+	 * network.
+	 *
+	 * @param Schema The Observation Schema
+	 * @param Element The sub-observation to be encoded.
+	 * @param EncodingSize The encoding size used to encode this sub-observation.
+	 * @param LayerNum The number of layers used to encode this sub-observation.
+	 * @param ActivationFunction The activation function used to encode this sub-observation.
+	 * @return The newly created observation schema element.
+	 */
 	UFUNCTION(BlueprintPure, Category = "LearningAgents", meta = (AdvancedDisplay = 2))
 	static FLearningAgentsObservationSchemaElement SpecifyEncodingObservation(ULearningAgentsObservationSchema* Schema, const FLearningAgentsObservationSchemaElement Element, const int32 EncodingSize = 128, const int32 LayerNum = 1, const ELearningAgentsActivationFunction ActivationFunction = ELearningAgentsActivationFunction::ELU, const FName Tag = TEXT("Encoding"));
 
@@ -239,130 +499,500 @@ public:
 	UFUNCTION(BlueprintPure, Category = "LearningAgents", meta = (AdvancedDisplay = 1))
 	static FLearningAgentsObservationSchemaElement SpecifyLocationObservation(ULearningAgentsObservationSchema* Schema, const FName Tag = TEXT("Location"));
 
+	/**
+	 * Specifies a new rotation observation. Allows an agent to observe the rotation of some entity. Rotations are encoded as two columns of the 
+	 * rotation matrix to ensure there is no discontinuity in the encoding.
+	 *
+	 * @param Schema The Observation Schema
+	 * @param Tag The tag of this new observation. Used during observation object validation and debugging.
+	 * @return The newly created observation schema element.
+	 */
 	UFUNCTION(BlueprintPure, Category = "LearningAgents", meta = (AdvancedDisplay = 1))
 	static FLearningAgentsObservationSchemaElement SpecifyRotationObservation(ULearningAgentsObservationSchema* Schema, const FName Tag = TEXT("Rotation"));
 
+	/**
+	 * Specifies a new scale observation. Allows an agent to observe the scale of some entity. Negative scales are not supported by this observation 
+	 * type.
+	 *
+	 * @param Schema The Observation Schema
+	 * @param Tag The tag of this new observation. Used during observation object validation and debugging.
+	 * @return The newly created observation schema element.
+	 */
 	UFUNCTION(BlueprintPure, Category = "LearningAgents", meta = (AdvancedDisplay = 1))
 	static FLearningAgentsObservationSchemaElement SpecifyScaleObservation(ULearningAgentsObservationSchema* Schema, const FName Tag = TEXT("Scale"));
 
+	/**
+	 * Specifies a new transform observation. Allows an agent to observe the transform of some entity.
+	 *
+	 * @param Schema The Observation Schema
+	 * @param Tag The tag of this new observation. Used during observation object validation and debugging.
+	 * @return The newly created observation schema element.
+	 */
 	UFUNCTION(BlueprintPure, Category = "LearningAgents", meta = (AdvancedDisplay = 1))
 	static FLearningAgentsObservationSchemaElement SpecifyTransformObservation(ULearningAgentsObservationSchema* Schema, const FName Tag = TEXT("Transform"));
 
+	/**
+	 * Specifies a new angle observation. This will be encoded as a 2-dimension Cartesian vector so that 0 and 350 are close to each other in the 
+	 * encoded space.
+	 *
+	 * @param Schema The Observation Schema
+	 * @param Tag The tag of this new observation. Used during observation object validation and debugging.
+	 * @return The newly created observation schema element.
+	 */
 	UFUNCTION(BlueprintPure, Category = "LearningAgents", meta = (AdvancedDisplay = 1))
 	static FLearningAgentsObservationSchemaElement SpecifyAngleObservation(ULearningAgentsObservationSchema* Schema, const FName Tag = TEXT("Angle"));
 
+	/**
+	 * Specifies a new velocity observation. Allows an agent to observe the velocity of some entity.
+	 *
+	 * @param Schema The Observation Schema
+	 * @param Tag The tag of this new observation. Used during observation object validation and debugging.
+	 * @return The newly created observation schema element.
+	 */
 	UFUNCTION(BlueprintPure, Category = "LearningAgents", meta = (AdvancedDisplay = 1))
 	static FLearningAgentsObservationSchemaElement SpecifyVelocityObservation(ULearningAgentsObservationSchema* Schema, const FName Tag = TEXT("Velocity"));
 
+	/**
+	 * Specifies a new direction observation. Allows an agent to observe the direction of some entity.
+	 *
+	 * @param Schema The Observation Schema
+	 * @param Tag The tag of this new observation. Used during observation object validation and debugging.
+	 * @return The newly created observation schema element.
+	 */
 	UFUNCTION(BlueprintPure, Category = "LearningAgents", meta = (AdvancedDisplay = 1))
 	static FLearningAgentsObservationSchemaElement SpecifyDirectionObservation(ULearningAgentsObservationSchema* Schema, const FName Tag = TEXT("Direction"));
 
-	// Specify Spline Observations
-
+	/**
+	 * Specifies a new location along spline observation. This observes the location of the spline at the given distance along that spline.
+	 *
+	 * @param Schema The Observation Schema
+	 * @param Tag The tag of this new observation. Used during observation object validation and debugging.
+	 * @return The newly created observation schema element.
+	 */
 	UFUNCTION(BlueprintPure, Category = "LearningAgents", meta = (AdvancedDisplay = 1))
 	static FLearningAgentsObservationSchemaElement SpecifyLocationAlongSplineObservation(ULearningAgentsObservationSchema* Schema, const FName Tag = TEXT("LocationAlongSpline"));
 
+	/**
+	 * Specifies a new proportion along spline observation. This observes the proportion along a spline at the given distance. For looped splines 
+	 * this will be treated effectively like an angle between 0 and 360 degrees and encoded appropriately so that 0 and 350 are close to each other in the 
+	 * encoded space, while for non-looped splines this will be treated as a value between 0 and 1.
+	 *
+	 * @param Schema The Observation Schema
+	 * @param Tag The tag of this new observation. Used during observation object validation and debugging.
+	 * @return The newly created observation schema element.
+	 */
 	UFUNCTION(BlueprintPure, Category = "LearningAgents", meta = (AdvancedDisplay = 1))
 	static FLearningAgentsObservationSchemaElement SpecifyProportionAlongSplineObservation(ULearningAgentsObservationSchema* Schema, const FName Tag = TEXT("ProportionAlongSpline"));
 
+	/**
+	 * Specifies a new direction along spline observation. This observes the direction of the spline at the given distance along that spline.
+	 *
+	 * @param Schema The Observation Schema
+	 * @param Tag The tag of this new observation. Used during observation object validation and debugging.
+	 * @return The newly created observation schema element.
+	 */
 	UFUNCTION(BlueprintPure, Category = "LearningAgents", meta = (AdvancedDisplay = 1))
 	static FLearningAgentsObservationSchemaElement SpecifyDirectionAlongSplineObservation(ULearningAgentsObservationSchema* Schema, const FName Tag = TEXT("DirectionAlongSpline"));
 
-	// Specify Ray Cast Observations
-
+	/**
+	 * Specifies a new proportion along ray observation. This observes how far a you can travel along a ray before collision. Rays that can travel 
+	 * the full distance are encoded as zero, while rays that collide instantly are encoded as one.
+	 *
+	 * @param Schema The Observation Schema
+	 * @param Tag The tag of this new observation. Used during observation object validation and debugging.
+	 * @return The newly created observation schema element.
+	 */
 	UFUNCTION(BlueprintPure, Category = "LearningAgents", meta = (AdvancedDisplay = 1))
 	static FLearningAgentsObservationSchemaElement SpecifyProportionAlongRayObservation(ULearningAgentsObservationSchema* Schema, const FName Tag = TEXT("ProportionAlongRay"));
 
 public:
 
-	// Make Basic Observations
-
+	/**
+	 * Make a new null observation.
+	 *
+	 * @param Object The Observation Object
+	 * @param Tag The tag of the corresponding observation. Must match the tag given during Specify.
+	 * @return The newly created observation object element.
+	 */
 	UFUNCTION(BlueprintPure, Category = "LearningAgents", meta = (AdvancedDisplay = 1))
 	static FLearningAgentsObservationObjectElement MakeNullObservation(ULearningAgentsObservationObject* Object, const FName Tag = TEXT("Null"));
 
+	/**
+	 * Make a new continuous observation. The size of Values must match the Size given during Specify.
+	 *
+	 * @param Object The Observation Object
+	 * @param Values The observation values.
+	 * @param Tag The tag of the corresponding observation. Must match the tag given during Specify.
+	 * @return The newly created observation object element.
+	 */
 	UFUNCTION(BlueprintPure, Category = "LearningAgents", meta = (AdvancedDisplay = 2))
 	static FLearningAgentsObservationObjectElement MakeContinuousObservation(ULearningAgentsObservationObject* Object, const TArray<float>& Values, const FName Tag = TEXT("Continuous"));
+
+	/**
+	 * Make a new continuous observation. The size of Values must match the Size given during Specify.
+	 *
+	 * @param Object The Observation Object
+	 * @param Values The observation values.
+	 * @param Tag The tag of the corresponding observation. Must match the tag given during Specify.
+	 * @return The newly created observation object element.
+	 */
 	static FLearningAgentsObservationObjectElement MakeContinuousObservationFromArrayView(ULearningAgentsObservationObject* Object, const TArrayView<const float> Values, const FName Tag = TEXT("Continuous"));
 
+	/**
+	 * Make a new exclusive discrete observation.
+	 *
+	 * @param Object The Observation Object
+	 * @param DiscreteIndex The index of the discrete observation. Values must be smaller than the given Size.
+	 * @param Size The size of the discrete observation. Must be equal to the size given during Specify.
+	 * @param Tag The tag of the corresponding observation. Must match the tag given during Specify.
+	 * @return The newly created observation object element.
+	 */
 	UFUNCTION(BlueprintPure, Category = "LearningAgents", meta = (AdvancedDisplay = 3))
 	static FLearningAgentsObservationObjectElement MakeExclusiveDiscreteObservation(ULearningAgentsObservationObject* Object, const int32 DiscreteIndex, const int32 Size, const FName Tag = TEXT("ExclusiveDiscrete"));
 
+	/**
+	 * Make a new inclusive discrete observation.
+	 *
+	 * @param Object The Observation Object
+	 * @param DiscreteIndices The indices of the discrete observations. All values must be smaller than the given Size.
+	 * @param Size The size of the discrete observation. Must be equal to the size given during Specify.
+	 * @param Tag The tag of the corresponding observation. Must match the tag given during Specify.
+	 * @return The newly created observation object element.
+	 */
 	UFUNCTION(BlueprintPure, Category = "LearningAgents", meta = (AdvancedDisplay = 3))
 	static FLearningAgentsObservationObjectElement MakeInclusiveDiscreteObservation(ULearningAgentsObservationObject* Object, const TArray<int32>& DiscreteIndices, const int32 Size, const FName Tag = TEXT("InclusiveDiscrete"));
+	
+	/**
+	 * Make a new inclusive discrete observation.
+	 *
+	 * @param Object The Observation Object
+	 * @param DiscreteIndices The indices of the discrete observations. All values must be smaller than the given Size.
+	 * @param Size The size of the discrete observation. Must be equal to the size given during Specify.
+	 * @param Tag The tag of the corresponding observation. Must match the tag given during Specify.
+	 * @return The newly created observation object element.
+	 */
 	static FLearningAgentsObservationObjectElement MakeInclusiveDiscreteObservationFromArrayView(ULearningAgentsObservationObject* Object, const TArrayView<const int32> DiscreteIndices, const int32 Size, const FName Tag = TEXT("InclusiveDiscrete"));
 
+	/**
+	 * Make a new index observation.
+	 *
+	 * @param Object The Observation Object
+	 * @param Index The index. Value must be smaller than the given Size.
+	 * @param Size The size of the discrete observation. Must be equal to the size given during Specify.
+	 * @param Tag The tag of the corresponding observation. Must match the tag given during Specify.
+	 * @return The newly created observation object element.
+	 */
 	UFUNCTION(BlueprintPure, Category = "LearningAgents", meta = (AdvancedDisplay = 3))
 	static FLearningAgentsObservationObjectElement MakeIndexObservation(ULearningAgentsObservationObject* Object, const int32 Index, const int32 Size, const FName Tag = TEXT("Index"));
 
+	/**
+	 * Make a new count observation.
+	 *
+	 * @param Object The Observation Object
+	 * @param Num The number of items. Must be less than or equal to MaxNum.
+	 * @param MaxNum The maximum number of items possible.
+	 * @param Tag The tag of the corresponding observation. Must match the tag given during Specify.
+	 * @return The newly created observation object element.
+	 */
 	UFUNCTION(BlueprintPure, Category = "LearningAgents", meta = (AdvancedDisplay = 3))
 	static FLearningAgentsObservationObjectElement MakeCountObservation(ULearningAgentsObservationObject* Object, const int32 Num, const int32 MaxNum, const FName Tag = TEXT("Count"));
 
+	/**
+	 * Make a new struct observation.
+	 *
+	 * @param Object The Observation Object
+	 * @param Elements The named sub-observations. Must match what was given during Specify.
+	 * @param Tag The tag of the corresponding observation. Must match the tag given during Specify.
+	 * @return The newly created observation object element.
+	 */
 	UFUNCTION(BlueprintPure, Category = "LearningAgents", meta = (AdvancedDisplay = 2))
 	static FLearningAgentsObservationObjectElement MakeStructObservation(ULearningAgentsObservationObject* Object, const TMap<FName, FLearningAgentsObservationObjectElement>& Elements, const FName Tag = TEXT("Struct"));
 
+	/**
+	 * Make a new struct observation.
+	 *
+	 * @param Object The Observation Object
+	 * @param ElementNames The names of the sub-observations. Must match what was given during Specify.
+	 * @param Elements The corresponding sub-observations. Must match what was given during Specify.
+	 * @param Tag The tag of the corresponding observation. Must match the tag given during Specify.
+	 * @return The newly created observation object element.
+	 */
 	UFUNCTION(BlueprintPure, Category = "LearningAgents", meta = (AdvancedDisplay = 3))
 	static FLearningAgentsObservationObjectElement MakeStructObservationFromArrays(ULearningAgentsObservationObject* Object, const TArray<FName>& ElementNames, const TArray<FLearningAgentsObservationObjectElement>& Elements, const FName Tag = TEXT("Struct"));
+	
+	/**
+	 * Make a new struct observation.
+	 *
+	 * @param Object The Observation Object
+	 * @param ElementNames The names of the sub-observations. Must match what was given during Specify.
+	 * @param Elements The corresponding sub-observations. Must match what was given during Specify.
+	 * @param Tag The tag of the corresponding observation. Must match the tag given during Specify.
+	 * @return The newly created observation object element.
+	 */
 	static FLearningAgentsObservationObjectElement MakeStructObservationFromArrayViews(ULearningAgentsObservationObject* Object, const TArrayView<const FName> ElementNames, const TArrayView<const FLearningAgentsObservationObjectElement> Elements, const FName Tag = TEXT("Struct"));
 
+	/**
+	 * Make a new exclusive union observation.
+	 *
+	 * @param Object The Observation Object
+	 * @param ElementName The name of the chosen sub-observation.
+	 * @param Element The corresponding chosen sub-observation.
+	 * @param Tag The tag of the corresponding observation. Must match the tag given during Specify.
+	 * @return The newly created observation object element.
+	 */
 	UFUNCTION(BlueprintPure, Category = "LearningAgents", meta = (AdvancedDisplay = 3))
 	static FLearningAgentsObservationObjectElement MakeExclusiveUnionObservation(ULearningAgentsObservationObject* Object, const FName ElementName, const FLearningAgentsObservationObjectElement Element, const FName Tag = TEXT("ExclusiveUnion"));
 
+	/**
+	 * Make a new inclusive union observation.
+	 *
+	 * @param Object The Observation Object
+	 * @param Elements The chosen sub-observations.
+	 * @param Tag The tag of the corresponding observation. Must match the tag given during Specify.
+	 * @return The newly created observation object element.
+	 */
 	UFUNCTION(BlueprintPure, Category = "LearningAgents", meta = (AdvancedDisplay = 2))
 	static FLearningAgentsObservationObjectElement MakeInclusiveUnionObservation(ULearningAgentsObservationObject* Object, const TMap<FName, FLearningAgentsObservationObjectElement>& Elements, const FName Tag = TEXT("InclusiveUnion"));
 
+	/**
+	 * Make a new inclusive union observation.
+	 *
+	 * @param Object The Observation Object
+	 * @param ElementNames The names of the chosen sub-observations.
+	 * @param Elements The corresponding chosen sub-observations.
+	 * @param Tag The tag of the corresponding observation. Must match the tag given during Specify.
+	 * @return The newly created observation object element.
+	 */
 	UFUNCTION(BlueprintPure, Category = "LearningAgents", meta = (AdvancedDisplay = 3))
 	static FLearningAgentsObservationObjectElement MakeInclusiveUnionObservationFromArrays(ULearningAgentsObservationObject* Object, const TArray<FName>& ElementNames, const TArray<FLearningAgentsObservationObjectElement>& Elements, const FName Tag = TEXT("InclusiveUnion"));
+	
+	/**
+	 * Make a new inclusive union observation.
+	 *
+	 * @param Object The Observation Object
+	 * @param ElementNames The names of the chosen sub-observations.
+	 * @param Elements The corresponding chosen sub-observations.
+	 * @param Tag The tag of the corresponding observation. Must match the tag given during Specify.
+	 * @return The newly created observation object element.
+	 */
 	static FLearningAgentsObservationObjectElement MakeInclusiveUnionObservationFromArrayViews(ULearningAgentsObservationObject* Object, const TArrayView<const FName> ElementNames, const TArrayView<const FLearningAgentsObservationObjectElement> Elements, const FName Tag = TEXT("InclusiveUnion"));
 
+	/**
+	 * Make a new static array observation.
+	 *
+	 * @param Object The Observation Object
+	 * @param Elements The sub-observations. The number of elements here must match what was given during Specify.
+	 * @param Tag The tag of the corresponding observation. Must match the tag given during Specify.
+	 * @return The newly created observation object element.
+	 */
 	UFUNCTION(BlueprintPure, Category = "LearningAgents", meta = (AdvancedDisplay = 2))
 	static FLearningAgentsObservationObjectElement MakeStaticArrayObservation(ULearningAgentsObservationObject* Object, const TArray<FLearningAgentsObservationObjectElement>& Elements, const FName Tag = TEXT("StaticArray"));
+
+	/**
+	 * Make a new static array observation.
+	 *
+	 * @param Object The Observation Object
+	 * @param Elements The sub-observations. The number of elements here must match what was given during Specify.
+	 * @param Tag The tag of the corresponding observation. Must match the tag given during Specify.
+	 * @return The newly created observation object element.
+	 */
 	static FLearningAgentsObservationObjectElement MakeStaticArrayObservationFromArrayView(ULearningAgentsObservationObject* Object, const TArrayView<const FLearningAgentsObservationObjectElement> Elements, const FName Tag = TEXT("StaticArray"));
 
+	/**
+	 * Make a new set observation.
+	 *
+	 * @param Object The Observation Object
+	 * @param Elements The sub-observations. The number of elements here must be less than or equal to the maximum that was given during Specify.
+	 * @param Tag The tag of the corresponding observation. Must match the tag given during Specify.
+	 * @return The newly created observation object element.
+	 */
 	UFUNCTION(BlueprintPure, Category = "LearningAgents", meta = (AdvancedDisplay = 2))
 	static FLearningAgentsObservationObjectElement MakeSetObservation(ULearningAgentsObservationObject* Object, const TSet<FLearningAgentsObservationObjectElement>& Elements, const FName Tag = TEXT("Set"));
 
+	/**
+	 * Make a new set observation.
+	 *
+	 * @param Object The Observation Object
+	 * @param Elements The sub-observations. The number of elements here must be less than or equal to the maximum that was given during Specify.
+	 * @param Tag The tag of the corresponding observation. Must match the tag given during Specify.
+	 * @return The newly created observation object element.
+	 */
 	UFUNCTION(BlueprintPure, Category = "LearningAgents", meta = (AdvancedDisplay = 2))
 	static FLearningAgentsObservationObjectElement MakeSetObservationFromArray(ULearningAgentsObservationObject* Object, const TArray<FLearningAgentsObservationObjectElement>& Elements, const FName Tag = TEXT("Set"));
+
+	/**
+	 * Make a new set observation.
+	 *
+	 * @param Object The Observation Object
+	 * @param Elements The sub-observations. The number of elements here must be less than or equal to the maximum that was given during Specify.
+	 * @param Tag The tag of the corresponding observation. Must match the tag given during Specify.
+	 * @return The newly created observation object element.
+	 */
 	static FLearningAgentsObservationObjectElement MakeSetObservationFromArrayView(ULearningAgentsObservationObject* Object, const TArrayView<const FLearningAgentsObservationObjectElement> Elements, const FName Tag = TEXT("Set"));
 
+	/**
+	 * Make a new pair observation.
+	 *
+	 * @param Object The Observation Object
+	 * @param Key The key sub-observation.
+	 * @param Key The value sub-observation.
+	 * @param Tag The tag of the corresponding observation. Must match the tag given during Specify.
+	 * @return The newly created observation object element.
+	 */
 	UFUNCTION(BlueprintPure, Category = "LearningAgents", meta = (AdvancedDisplay = 3))
 	static FLearningAgentsObservationObjectElement MakePairObservation(ULearningAgentsObservationObject* Object, const FLearningAgentsObservationObjectElement Key, const FLearningAgentsObservationObjectElement Value, const FName Tag = TEXT("Pair"));
 
+	/**
+	 * Make a new array observation.
+	 *
+	 * @param Object The Observation Object
+	 * @param Elements The sub-observations. The number of elements here must be less than or equal to the maximum that was given during Specify.
+	 * @param Tag The tag of the corresponding observation. Must match the tag given during Specify.
+	 * @return The newly created observation object element.
+	 */
 	UFUNCTION(BlueprintPure, Category = "LearningAgents", meta = (AdvancedDisplay = 2))
 	static FLearningAgentsObservationObjectElement MakeArrayObservation(ULearningAgentsObservationObject* Object, const TArray<FLearningAgentsObservationObjectElement>& Elements, const FName Tag = TEXT("Array"));
+
+	/**
+	 * Make a new array observation.
+	 *
+	 * @param Object The Observation Object
+	 * @param Elements The sub-observations. The number of elements here must be less than or equal to the maximum that was given during Specify.
+	 * @param Tag The tag of the corresponding observation. Must match the tag given during Specify.
+	 * @return The newly created observation object element.
+	 */
 	static FLearningAgentsObservationObjectElement MakeArrayObservationFromArrayView(ULearningAgentsObservationObject* Object, const TArrayView<const FLearningAgentsObservationObjectElement> Elements, const FName Tag = TEXT("Array"));
 
+	/**
+	 * Make a new map observation.
+	 *
+	 * @param Object The Observation Object
+	 * @param Elements The sub-observations. The number of elements here must be less than or equal to the maximum that was given during Specify.
+	 * @param Tag The tag of the corresponding observation. Must match the tag given during Specify.
+	 * @return The newly created observation object element.
+	 */
 	UFUNCTION(BlueprintPure, Category = "LearningAgents", meta = (AdvancedDisplay = 2))
 	static FLearningAgentsObservationObjectElement MakeMapObservation(ULearningAgentsObservationObject* Object, const TMap<FLearningAgentsObservationObjectElement, FLearningAgentsObservationObjectElement>& Map, const FName Tag = TEXT("Map"));
 
+	/**
+	 * Make a new map observation.
+	 *
+	 * @param Object The Observation Object
+	 * @param Elements The sub-observations. The number of elements here must be less than or equal to the maximum that was given during Specify.
+	 * @param Tag The tag of the corresponding observation. Must match the tag given during Specify.
+	 * @return The newly created observation object element.
+	 */
 	UFUNCTION(BlueprintPure, Category = "LearningAgents", meta = (AdvancedDisplay = 3))
 	static FLearningAgentsObservationObjectElement MakeMapObservationFromArrays(ULearningAgentsObservationObject* Object, const TArray<FLearningAgentsObservationObjectElement>& Keys, const TArray<FLearningAgentsObservationObjectElement>& Values, const FName Tag = TEXT("Map"));
+	
+	/**
+	 * Make a new map observation.
+	 *
+	 * @param Object The Observation Object
+	 * @param Elements The sub-observations. The number of elements here must be less than or equal to the maximum that was given during Specify.
+	 * @param Tag The tag of the corresponding observation. Must match the tag given during Specify.
+	 * @return The newly created observation object element.
+	 */
 	static FLearningAgentsObservationObjectElement MakeMapObservationFromArrayViews(ULearningAgentsObservationObject* Object, const TArrayView<const FLearningAgentsObservationObjectElement> Keys, const TArrayView<const FLearningAgentsObservationObjectElement> Values, const FName Tag = TEXT("Map"));
 
+	/**
+	 * Make a new enum observation.
+	 *
+	 * @param Object The Observation Object
+	 * @param Enum The enum type for this observation. Must match what was given during Specify.
+	 * @param EnumValue The enum value.
+	 * @param Tag The tag of the corresponding observation. Must match the tag given during Specify.
+	 * @return The newly created observation object element.
+	 */
 	UFUNCTION(BlueprintPure, Category = "LearningAgents", meta = (AdvancedDisplay = 3))
 	static FLearningAgentsObservationObjectElement MakeEnumObservation(ULearningAgentsObservationObject* Object, const UEnum* Enum, const uint8 EnumValue, const FName Tag = TEXT("Enum"));
 
+	/**
+	 * Make a new bitmask observation.
+	 *
+	 * @param Object The Observation Object
+	 * @param Enum The enum type for this observation. Must match what was given during Specify.
+	 * @param BitmaskValue The bitmask value.
+	 * @param Tag The tag of the corresponding observation. Must match the tag given during Specify.
+	 * @return The newly created observation object element.
+	 */
 	UFUNCTION(BlueprintPure, Category = "LearningAgents", meta = (AdvancedDisplay = 3))
 	static FLearningAgentsObservationObjectElement MakeBitmaskObservation(ULearningAgentsObservationObject* Object, const UEnum* Enum, const int32 BitmaskValue, const FName Tag = TEXT("Bitmask"));
 
+	/**
+	 * Make a new optional observation.
+	 *
+	 * @param Object The Observation Object
+	 * @param Element The sub-observation given.
+	 * @param Option The indicator as to if this is observation should be used.
+	 * @param Tag The tag of the corresponding observation. Must match the tag given during Specify.
+	 * @return The newly created observation object element.
+	 */
 	UFUNCTION(BlueprintPure, Category = "LearningAgents", meta = (AdvancedDisplay = 3))
 	static FLearningAgentsObservationObjectElement MakeOptionalObservation(ULearningAgentsObservationObject* Object, const FLearningAgentsObservationObjectElement Element, const ELearningAgentsOptionalObservation Option, const FName Tag = TEXT("Optional"));
 
+	/**
+	 * Make a new null optional observation. Use this to provide a null optional observation.
+	 *
+	 * @param Object The Observation Object
+	 * @param Tag The tag of the corresponding observation. Must match the tag given during Specify.
+	 * @return The newly created observation object element.
+	 */
 	UFUNCTION(BlueprintPure, Category = "LearningAgents", meta = (AdvancedDisplay = 1))
 	static FLearningAgentsObservationObjectElement MakeOptionalNullObservation(ULearningAgentsObservationObject* Object, const FName Tag = TEXT("Optional"));
 
+	/**
+	 * Make a new valid optional observation. Use this to provide a valid optional observation.
+	 *
+	 * @param Object The Observation Object
+	 * @param Tag The tag of the corresponding observation. Must match the tag given during Specify.
+	 * @return The newly created observation object element.
+	 */
 	UFUNCTION(BlueprintPure, Category = "LearningAgents", meta = (AdvancedDisplay = 2))
 	static FLearningAgentsObservationObjectElement MakeOptionalValidObservation(ULearningAgentsObservationObject* Object, const FLearningAgentsObservationObjectElement Element, const FName Tag = TEXT("Optional"));
 
+	/**
+	 * Make a new either observation.
+	 *
+	 * @param Object The Observation Object
+	 * @param Element The sub-observation given.
+	 * @param Option The indicator as to if this is observation A or B.
+	 * @param Tag The tag of the corresponding observation. Must match the tag given during Specify.
+	 * @return The newly created observation object element.
+	 */
 	UFUNCTION(BlueprintPure, Category = "LearningAgents", meta = (AdvancedDisplay = 3))
 	static FLearningAgentsObservationObjectElement MakeEitherObservation(ULearningAgentsObservationObject* Object, const FLearningAgentsObservationObjectElement Element, const ELearningAgentsEitherObservation Either, const FName Tag = TEXT("Either"));
 
+	/**
+	 * Make a new either A observation. Use this to provide option A.
+	 *
+	 * @param Object The Observation Object
+	 * @param Element The sub-observation given.
+	 * @param Tag The tag of the corresponding observation. Must match the tag given during Specify.
+	 * @return The newly created observation object element.
+	 */
 	UFUNCTION(BlueprintPure, Category = "LearningAgents", meta = (AdvancedDisplay = 2, DisplayName = "Make Either A Observation"))
 	static FLearningAgentsObservationObjectElement MakeEitherAObservation(ULearningAgentsObservationObject* Object, const FLearningAgentsObservationObjectElement A, const FName Tag = TEXT("Either"));
 
+	/**
+	 * Make a new either B observation. Use this to provide option B.
+	 *
+	 * @param Object The Observation Object
+	 * @param Element The sub-observation given.
+	 * @param Tag The tag of the corresponding observation. Must match the tag given during Specify.
+	 * @return The newly created observation object element.
+	 */
 	UFUNCTION(BlueprintPure, Category = "LearningAgents", meta = (AdvancedDisplay = 2, DisplayName = "Make Either B Observation"))
 	static FLearningAgentsObservationObjectElement MakeEitherBObservation(ULearningAgentsObservationObject* Object, const FLearningAgentsObservationObjectElement B, const FName Tag = TEXT("Either"));
 
+	/**
+	 * Make a new encoding observation. This must be used in conjunction with SpecifyEncodingObservation.
+	 * 
+	 * @param Object The Observation Object
+	 * @param Element The Observation Element to be encoded.
+	 * @param Tag The tag of the corresponding observation. Must match the tag given during Specify.
+	 * @return The newly created observation object element.
+	 */
 	UFUNCTION(BlueprintPure, Category = "LearningAgents", meta = (AdvancedDisplay = 2))
 	static FLearningAgentsObservationObjectElement MakeEncodingObservation(ULearningAgentsObservationObject* Object, const FLearningAgentsObservationObjectElement Element, const FName Tag = TEXT("Encoding"));
 
@@ -501,7 +1131,7 @@ public:
 		const FLinearColor VisualLoggerColor = FLinearColor::Red);
 
 	/**
-	 * Make a new scale observation. Negative scales are not supported by this observation type.
+	 * Make a new scale observation.
 	 *
 	 * @param Object The Observation Object
 	 * @param Scale The scale of interest to the agent.
@@ -556,11 +1186,57 @@ public:
 		const FVector VisualLoggerLocation = FVector::ZeroVector,
 		const FLinearColor VisualLoggerColor = FLinearColor::Red);
 
-	UFUNCTION(BlueprintPure, Category = "LearningAgents", meta = (AdvancedDisplay = 3))
-	static FLearningAgentsObservationObjectElement MakeAngleObservation(ULearningAgentsObservationObject* Object, const float Angle, const float RelativeAngle = 0.0f, const FName Tag = TEXT("Angle"));
+	/**
+	 * Make a new angle observation. Angles should be given in degrees.
+	 *
+	 * @param Object The Observation Object
+	 * @param Angle The angle of interest to the agent.
+	 * @param RelativeAngle The angle the provided angle should be encoded relative to.
+	 * @param Tag The tag of the corresponding observation. Must match the tag given during Specify.
+	 * @param bVisualLoggerEnabled When true, debug data will be sent to the visual logger.
+	 * @param VisualLoggerListener The listener object which is making this observation. This must be set to use logging.
+	 * @param VisualLoggerAgentId The agent id associated with this observation.
+	 * @param VisualLoggerLocation A location for the visual logger information in the world.
+	 * @param VisualLoggerColor The color for the visual logger display.
+	 * @return The newly created observation object element.
+	 */
+	UFUNCTION(BlueprintPure, Category = "LearningAgents", meta = (AdvancedDisplay = 3, DefaultToSelf = "VisualLoggerListener"))
+	static FLearningAgentsObservationObjectElement MakeAngleObservation(
+		ULearningAgentsObservationObject* Object, 
+		const float Angle, 
+		const float RelativeAngle = 0.0f, 
+		const FName Tag = TEXT("Angle"),
+		const bool bVisualLoggerEnabled = false,
+		ULearningAgentsManagerListener* VisualLoggerListener = nullptr,
+		const int32 VisualLoggerAgentId = -1,
+		const FVector VisualLoggerLocation = FVector::ZeroVector,
+		const FLinearColor VisualLoggerColor = FLinearColor::Red);
 
-	UFUNCTION(BlueprintPure, Category = "LearningAgents", meta = (AdvancedDisplay = 3))
-	static FLearningAgentsObservationObjectElement MakeAngleObservationRadians(ULearningAgentsObservationObject* Object, const float Angle, const float RelativeAngle = 0.0f, const FName Tag = TEXT("Angle"));
+	/**
+	 * Make a new angle observation. Angles should be given in radians.
+	 * 
+	 * @param Object The Observation Object
+	 * @param Angle The angle of interest to the agent.
+	 * @param RelativeAngle The angle the provided angle should be encoded relative to.
+	 * @param Tag The tag of the corresponding observation. Must match the tag given during Specify.
+	 * @param bVisualLoggerEnabled When true, debug data will be sent to the visual logger.
+	 * @param VisualLoggerListener The listener object which is making this observation. This must be set to use logging.
+	 * @param VisualLoggerAgentId The agent id associated with this observation.
+	 * @param VisualLoggerLocation A location for the visual logger information in the world.
+	 * @param VisualLoggerColor The color for the visual logger display.
+	 * @return The newly created observation object element.
+	 */
+	UFUNCTION(BlueprintPure, Category = "LearningAgents", meta = (AdvancedDisplay = 3, DefaultToSelf = "VisualLoggerListener"))
+	static FLearningAgentsObservationObjectElement MakeAngleObservationRadians(
+		ULearningAgentsObservationObject* Object, 
+		const float Angle, 
+		const float RelativeAngle = 0.0f, 
+		const FName Tag = TEXT("Angle"),
+		const bool bVisualLoggerEnabled = false,
+		ULearningAgentsManagerListener* VisualLoggerListener = nullptr,
+		const int32 VisualLoggerAgentId = -1,
+		const FVector VisualLoggerLocation = FVector::ZeroVector,
+		const FLinearColor VisualLoggerColor = FLinearColor::Red);
 
 	/**
 	 * Make a new velocity observation.
@@ -622,8 +1298,22 @@ public:
 		const float VisualLoggerArrowLength = 100.0f,
 		const FLinearColor VisualLoggerColor = FLinearColor::Red);
 
-	// Make Spline Observations
-
+	/**
+	 * Make a new location along spline observation.
+	 *
+	 * @param Object The Observation Object
+	 * @param SplineComponent The spline to observe.
+	 * @param DistanceAlongSpline The distance along that spline.
+	 * @param RelativeTransform The transform the provided location should be encoded relative to.
+	 * @param LocationScale Used to normalize the transform's location for this observation.
+	 * @param Tag The tag of the corresponding observation. Must match the tag given during Specify.
+	 * @param bVisualLoggerEnabled When true, debug data will be sent to the visual logger.
+	 * @param VisualLoggerListener The listener object which is making this observation. This must be set to use logging.
+	 * @param VisualLoggerAgentId The agent id associated with this observation.
+	 * @param VisualLoggerLocation A location for the visual logger information in the world.
+	 * @param VisualLoggerColor The color for the visual logger display.
+	 * @return The newly created observation object element.
+	 */
 	UFUNCTION(BlueprintPure, Category = "LearningAgents", meta = (AdvancedDisplay = 5, DefaultToSelf = "VisualLoggerListener"))
 	static FLearningAgentsObservationObjectElement MakeLocationAlongSplineObservation(
 		ULearningAgentsObservationObject* Object,
@@ -638,13 +1328,48 @@ public:
 		const FVector VisualLoggerLocation = FVector::ZeroVector,
 		const FLinearColor VisualLoggerColor = FLinearColor::Red);
 
-	UFUNCTION(BlueprintPure, Category = "LearningAgents", meta = (AdvancedDisplay = 3))
+	/**
+	 * Make a new proportion along spline observation.
+	 *
+	 * @param Object The Observation Object
+	 * @param SplineComponent The spline to observe.
+	 * @param DistanceAlongSpline The distance along that spline.
+	 * @param Tag The tag of the corresponding observation. Must match the tag given during Specify.
+	 * @param bVisualLoggerEnabled When true, debug data will be sent to the visual logger.
+	 * @param VisualLoggerListener The listener object which is making this observation. This must be set to use logging.
+	 * @param VisualLoggerAgentId The agent id associated with this observation.
+	 * @param VisualLoggerLocation A location for the visual logger information in the world.
+	 * @param VisualLoggerColor The color for the visual logger display.
+	 * @return The newly created observation object element.
+	 */
+	UFUNCTION(BlueprintPure, Category = "LearningAgents", meta = (AdvancedDisplay = 3, DefaultToSelf = "VisualLoggerListener"))
 	static FLearningAgentsObservationObjectElement MakeProportionAlongSplineObservation(
 		ULearningAgentsObservationObject* Object, 
 		const USplineComponent* SplineComponent, 
 		const float DistanceAlongSpline, 
-		const FName Tag = TEXT("ProportionAlongSpline"));
+		const FName Tag = TEXT("ProportionAlongSpline"),
+		const bool bVisualLoggerEnabled = false,
+		ULearningAgentsManagerListener* VisualLoggerListener = nullptr,
+		const int32 VisualLoggerAgentId = -1,
+		const FVector VisualLoggerLocation = FVector::ZeroVector,
+		const FLinearColor VisualLoggerColor = FLinearColor::Red);
 
+	/**
+	 * Make a new direction along spline observation.
+	 *
+	 * @param Object The Observation Object
+	 * @param SplineComponent The spline to observe.
+	 * @param DistanceAlongSpline The distance along that spline.
+	 * @param RelativeTransform The transform the provided direction should be encoded relative to.
+	 * @param Tag The tag of the corresponding observation. Must match the tag given during Specify.
+	 * @param bVisualLoggerEnabled When true, debug data will be sent to the visual logger.
+	 * @param VisualLoggerListener The listener object which is making this observation. This must be set to use logging.
+	 * @param VisualLoggerAgentId The agent id associated with this observation.
+	 * @param VisualLoggerLocation A location for the visual logger information in the world.
+	 * @param VisualLoggerArrowLength The length of the arrow to display to represent the direction.
+	 * @param VisualLoggerColor The color for the visual logger display.
+	 * @return The newly created observation object element.
+	 */
 	UFUNCTION(BlueprintPure, Category = "LearningAgents", meta = (AdvancedDisplay = 4, DefaultToSelf = "VisualLoggerListener"))
 	static FLearningAgentsObservationObjectElement MakeDirectionAlongSplineObservation(
 		ULearningAgentsObservationObject* Object,
@@ -659,16 +1384,35 @@ public:
 		const float VisualLoggerArrowLength = 100.0f,
 		const FLinearColor VisualLoggerColor = FLinearColor::Red);
 
-	// Make Ray Cast Observations
-
-	UFUNCTION(BlueprintPure, Category = "LearningAgents", meta = (AdvancedDisplay = 5))
+	/**
+	 * Make a new proportion along ray observation. 
+	 *
+	 * @param Object The Observation Object
+	 * @param RayStart The local ray start location.
+	 * @param RayEnd The local ray end location.
+	 * @param RayTransform The transform to use to transform the local ray starts and ends into the world space.
+	 * @param CollisionChannel The collision channel to collide against.
+	 * @param Tag The tag of the corresponding observation. Must match the tag given during Specify.
+	 * @param bVisualLoggerEnabled When true, debug data will be sent to the visual logger.
+	 * @param VisualLoggerListener The listener object which is making this observation. This must be set to use logging.
+	 * @param VisualLoggerAgentId The agent id associated with this observation.
+	 * @param VisualLoggerLocation A location for the visual logger information in the world.
+	 * @param VisualLoggerColor The color for the visual logger display.
+	 * @return The newly created observation object element.
+	 */
+	UFUNCTION(BlueprintPure, Category = "LearningAgents", meta = (AdvancedDisplay = 5, DefaultToSelf = "VisualLoggerListener"))
 	static FLearningAgentsObservationObjectElement MakeProportionAlongRayObservation(
 		ULearningAgentsObservationObject* Object, 
 		const FVector RayStart, 
 		const FVector RayEnd, 
 		const FTransform RayTransform = FTransform(), 
 		const ECollisionChannel CollisionChannel = ECollisionChannel::ECC_WorldStatic, 
-		const FName Tag = TEXT("ProportionAlongRay"));
+		const FName Tag = TEXT("ProportionAlongRay"),
+		const bool bVisualLoggerEnabled = false,
+		ULearningAgentsManagerListener* VisualLoggerListener = nullptr,
+		const int32 VisualLoggerAgentId = -1,
+		const FVector VisualLoggerLocation = FVector::ZeroVector,
+		const FLinearColor VisualLoggerColor = FLinearColor::Red);
 
 public:
 
diff --git a/Engine/Plugins/Experimental/LearningAgents/Source/LearningAgentsTraining/Private/LearningAgentsRewards.cpp b/Engine/Plugins/Experimental/LearningAgents/Source/LearningAgentsTraining/Private/LearningAgentsRewards.cpp
--- a/Engine/Plugins/Experimental/LearningAgents/Source/LearningAgentsTraining/Private/LearningAgentsRewards.cpp
+++ b/Engine/Plugins/Experimental/LearningAgents/Source/LearningAgentsTraining/Private/LearningAgentsRewards.cpp
@@ -269,7 +269,7 @@ float ULearningAgentsRewards::MakeRewardFromAngleSimilarity(
 	{
 		const ULearningAgentsVisualLoggerObject* VisualLoggerObject = VisualLoggerListener->GetOrAddVisualLoggerObject(Tag);
 
-		UE_LEARNING_AGENTS_VLOG_ANGLE(
+		UE_LEARNING_AGENTS_VLOG_ANGLE_DEGREES(
 			VisualLoggerObject,
 			LogLearning,
 			Display,
@@ -280,7 +280,7 @@ float ULearningAgentsRewards::MakeRewardFromAngleSimilarity(
 			VisualLoggerColor.ToFColor(true),
 			TEXT(""));
 
-		UE_LEARNING_AGENTS_VLOG_ANGLE(
+		UE_LEARNING_AGENTS_VLOG_ANGLE_DEGREES(
 			VisualLoggerObject,
 			LogLearning,
 			Display,
diff --git a/Engine/Plugins/Experimental/LearningAgents/Source/LearningAgentsTraining/Public/LearningAgentsRewards.h b/Engine/Plugins/Experimental/LearningAgents/Source/LearningAgentsTraining/Public/LearningAgentsRewards.h
--- a/Engine/Plugins/Experimental/LearningAgents/Source/LearningAgentsTraining/Public/LearningAgentsRewards.h
+++ b/Engine/Plugins/Experimental/LearningAgents/Source/LearningAgentsTraining/Public/LearningAgentsRewards.h
@@ -209,7 +209,7 @@ public:
 		const FLinearColor VisualLoggerColor = FLinearColor::Green);
 
 	/**
-	 * Make a reward based on how similar two angles are.
+	 * Make a reward based on how similar two angles are. Angles should be given in degrees.
 	 *
 	 * @param AngleA The first angle.
 	 * @param AngleB The second angle.
@@ -241,8 +241,8 @@ public:
 	/**
 	 * Make a reward based on how similar two directions are.
 	 *
-	 * @param DirectionA The first angle.
-	 * @param DirectionB The second angle.
+	 * @param DirectionA The first direction.
+	 * @param DirectionB The second direction.
 	 * @param RewardScale The scale of the reward. Use a negative scale to create a penalty.
 	 * @param Tag The tag for the reward. Used for debugging.
 	 * @param bVisualLoggerEnabled When true, debug data will be sent to the visual logger.
