LearningAgents: Added visual logging to some more observations.

#virtualized

diff --git a/Engine/Plugins/Experimental/LearningAgents/Source/LearningAgents/Private/LearningAgentsObservations.cpp b/Engine/Plugins/Experimental/LearningAgents/Source/LearningAgents/Private/LearningAgentsObservations.cpp
--- a/Engine/Plugins/Experimental/LearningAgents/Source/LearningAgents/Private/LearningAgentsObservations.cpp
+++ b/Engine/Plugins/Experimental/LearningAgents/Source/LearningAgents/Private/LearningAgentsObservations.cpp
@@ -1085,12 +1085,28 @@ FLearningAgentsObservationObjectElement ULearningAgentsObservations::MakeNullObs
 	return { Object->ObservationObject.CreateNull(Tag) };
 }
 
-FLearningAgentsObservationObjectElement ULearningAgentsObservations::MakeContinuousObservation(ULearningAgentsObservationObject* Object, const TArray<float>& Values, const FName Tag)
+FLearningAgentsObservationObjectElement ULearningAgentsObservations::MakeContinuousObservation(
+	ULearningAgentsObservationObject* Object, 
+	const TArray<float>& Values, 
+	const FName Tag,
+	const bool bVisualLoggerEnabled,
+	ULearningAgentsManagerListener* VisualLoggerListener,
+	const int32 VisualLoggerAgentId,
+	const FVector VisualLoggerLocation,
+	const FLinearColor VisualLoggerColor)
 {
-	return MakeContinuousObservationFromArrayView(Object, Values, Tag);
+	return MakeContinuousObservationFromArrayView(Object, Values, Tag, bVisualLoggerEnabled, VisualLoggerListener, VisualLoggerAgentId, VisualLoggerLocation, VisualLoggerColor);
 }
 
-FLearningAgentsObservationObjectElement ULearningAgentsObservations::MakeContinuousObservationFromArrayView(ULearningAgentsObservationObject* Object, const TArrayView<const float> Values, const FName Tag)
+FLearningAgentsObservationObjectElement ULearningAgentsObservations::MakeContinuousObservationFromArrayView(
+	ULearningAgentsObservationObject* Object, 
+	const TArrayView<const float> Values, 
+	const FName Tag,
+	const bool bVisualLoggerEnabled,
+	ULearningAgentsManagerListener* VisualLoggerListener,
+	const int32 VisualLoggerAgentId,
+	const FVector VisualLoggerLocation,
+	const FLinearColor VisualLoggerColor)
 {
 	if (!Object)
 	{
@@ -1103,10 +1119,35 @@ FLearningAgentsObservationObjectElement ULearningAgentsObservations::MakeContinu
 		UE_LOG(LogLearning, Warning, TEXT("MakeContinuousObservationFromArrayView: Creating zero-sized Continuous Observation."));
 	}
 
+#if UE_LEARNING_AGENTS_ENABLE_VISUAL_LOG
+	if (bVisualLoggerEnabled && VisualLoggerListener)
+	{
+		const ULearningAgentsVisualLoggerObject* VisualLoggerObject = VisualLoggerListener->GetOrAddVisualLoggerObject(Tag);
+
+		UE_LEARNING_AGENTS_VLOG_STRING(VisualLoggerObject, LogLearning, Display, VisualLoggerLocation,
+			VisualLoggerColor.ToFColor(true),
+			TEXT("Listener: %s\nTag: %s\nAgent Id: % 3i\nValues: %s\nEncoded: %s"),
+			*VisualLoggerListener->GetName(),
+			*Tag.ToString(),
+			VisualLoggerAgentId,
+			*UE::Learning::Array::FormatFloat(Values),
+			*UE::Learning::Array::FormatFloat(Values)); // Encoded is identical to provided values
+	}
+#endif
+
 	return { Object->ObservationObject.CreateContinuous({ Values }, Tag) };
 }
 
-FLearningAgentsObservationObjectElement ULearningAgentsObservations::MakeExclusiveDiscreteObservation(ULearningAgentsObservationObject* Object, const int32 DiscreteIndex, const int32 Size, const FName Tag)
+FLearningAgentsObservationObjectElement ULearningAgentsObservations::MakeExclusiveDiscreteObservation(
+	ULearningAgentsObservationObject* Object, 
+	const int32 DiscreteIndex, 
+	const int32 Size, 
+	const FName Tag,
+	const bool bVisualLoggerEnabled,
+	ULearningAgentsManagerListener* VisualLoggerListener,
+	const int32 VisualLoggerAgentId,
+	const FVector VisualLoggerLocation,
+	const FLinearColor VisualLoggerColor)
 {
 	if (DiscreteIndex < 0 || DiscreteIndex >= Size)
 	{
@@ -1117,15 +1158,51 @@ FLearningAgentsObservationObjectElement ULearningAgentsObservations::MakeExclusi
 	TArray<float, TInlineAllocator<32>> Values;
 	Values.Init(0.0f, Size);
 	Values[DiscreteIndex] = 1.0f;
+
+#if UE_LEARNING_AGENTS_ENABLE_VISUAL_LOG
+	if (bVisualLoggerEnabled && VisualLoggerListener)
+	{
+		const ULearningAgentsVisualLoggerObject* VisualLoggerObject = VisualLoggerListener->GetOrAddVisualLoggerObject(Tag);
+
+		UE_LEARNING_AGENTS_VLOG_STRING(VisualLoggerObject, LogLearning, Display, VisualLoggerLocation,
+			VisualLoggerColor.ToFColor(true),
+			TEXT("Listener: %s\nTag: %s\nAgent Id: % 3i\nSize: [%i]\nIndex: [%i]\nEncoded: %s"),
+			*VisualLoggerListener->GetName(),
+			*Tag.ToString(),
+			VisualLoggerAgentId,
+			Size,
+			DiscreteIndex,
+			*UE::Learning::Array::FormatFloat(Values, Size));
+	}
+#endif
+
 	return MakeContinuousObservationFromArrayView(Object, Values, Tag);
 }
 
-FLearningAgentsObservationObjectElement ULearningAgentsObservations::MakeInclusiveDiscreteObservation(ULearningAgentsObservationObject* Object, const TArray<int32>& DiscreteIndices, const int32 Size, const FName Tag)
+FLearningAgentsObservationObjectElement ULearningAgentsObservations::MakeInclusiveDiscreteObservation(
+	ULearningAgentsObservationObject* Object, 
+	const TArray<int32>& DiscreteIndices, 
+	const int32 Size, 
+	const FName Tag,
+	const bool bVisualLoggerEnabled,
+	ULearningAgentsManagerListener* VisualLoggerListener,
+	const int32 VisualLoggerAgentId,
+	const FVector VisualLoggerLocation,
+	const FLinearColor VisualLoggerColor)
 {
-	return MakeInclusiveDiscreteObservationFromArrayView(Object, DiscreteIndices, Size, Tag);
+	return MakeInclusiveDiscreteObservationFromArrayView(Object, DiscreteIndices, Size, Tag, bVisualLoggerEnabled, VisualLoggerListener, VisualLoggerAgentId, VisualLoggerLocation, VisualLoggerColor);
 }
 
-FLearningAgentsObservationObjectElement ULearningAgentsObservations::MakeInclusiveDiscreteObservationFromArrayView(ULearningAgentsObservationObject* Object, const TArrayView<const int32> DiscreteIndices, const int32 Size, const FName Tag)
+FLearningAgentsObservationObjectElement ULearningAgentsObservations::MakeInclusiveDiscreteObservationFromArrayView(
+	ULearningAgentsObservationObject* Object, 
+	const TArrayView<const int32> DiscreteIndices, 
+	const int32 Size, 
+	const FName Tag,
+	const bool bVisualLoggerEnabled,
+	ULearningAgentsManagerListener* VisualLoggerListener,
+	const int32 VisualLoggerAgentId,
+	const FVector VisualLoggerLocation,
+	const FLinearColor VisualLoggerColor)
 {
 	if (UE::Learning::Agents::Observation::Private::ContainsDuplicates(DiscreteIndices))
 	{
@@ -1147,10 +1224,36 @@ FLearningAgentsObservationObjectElement ULearningAgentsObservations::MakeInclusi
 		Values[DiscreteIndices[Idx]] = 1.0f;
 	}
 
+#if UE_LEARNING_AGENTS_ENABLE_VISUAL_LOG
+	if (bVisualLoggerEnabled && VisualLoggerListener)
+	{
+		const ULearningAgentsVisualLoggerObject* VisualLoggerObject = VisualLoggerListener->GetOrAddVisualLoggerObject(Tag);
+
+		UE_LEARNING_AGENTS_VLOG_STRING(VisualLoggerObject, LogLearning, Display, VisualLoggerLocation,
+			VisualLoggerColor.ToFColor(true),
+			TEXT("Listener: %s\nTag: %s\nAgent Id: % 3i\nSize: [%i]\nIndices: %s\nEncoded: %s"),
+			*VisualLoggerListener->GetName(),
+			*Tag.ToString(),
+			VisualLoggerAgentId,
+			Size,
+			*UE::Learning::Array::FormatInt32(DiscreteIndices, Size),
+			*UE::Learning::Array::FormatFloat(Values, Size));
+	}
+#endif
+
 	return MakeContinuousObservationFromArrayView(Object, Values, Tag);
 }
 
-FLearningAgentsObservationObjectElement ULearningAgentsObservations::MakeIndexObservation(ULearningAgentsObservationObject* Object, const int32 Index, const int32 Size, const FName Tag)
+FLearningAgentsObservationObjectElement ULearningAgentsObservations::MakeIndexObservation(
+	ULearningAgentsObservationObject* Object, 
+	const int32 Index, 
+	const int32 Size, 
+	const FName Tag,
+	const bool bVisualLoggerEnabled,
+	ULearningAgentsManagerListener* VisualLoggerListener,
+	const int32 VisualLoggerAgentId,
+	const FVector VisualLoggerLocation,
+	const FLinearColor VisualLoggerColor)
 {
 	if (Index < 0 || Index >= Size)
 	{
@@ -1166,10 +1269,36 @@ FLearningAgentsObservationObjectElement ULearningAgentsObservations::MakeIndexOb
 		Values[Idx] = 1.0f;
 	}
 
+#if UE_LEARNING_AGENTS_ENABLE_VISUAL_LOG
+	if (bVisualLoggerEnabled && VisualLoggerListener)
+	{
+		const ULearningAgentsVisualLoggerObject* VisualLoggerObject = VisualLoggerListener->GetOrAddVisualLoggerObject(Tag);
+
+		UE_LEARNING_AGENTS_VLOG_STRING(VisualLoggerObject, LogLearning, Display, VisualLoggerLocation,
+			VisualLoggerColor.ToFColor(true),
+			TEXT("Listener: %s\nTag: %s\nAgent Id: % 3i\nIndex: [%i]\nSize: [%i]\nEncoded: %s"),
+			*VisualLoggerListener->GetName(),
+			*Tag.ToString(),
+			VisualLoggerAgentId,
+			Index,
+			Size,
+			*UE::Learning::Array::FormatFloat(Values, Size));
+	}
+#endif
+
 	return MakeContinuousObservationFromArrayView(Object, Values, Tag);
 }
 
-FLearningAgentsObservationObjectElement ULearningAgentsObservations::MakeCountObservation(ULearningAgentsObservationObject* Object, const int32 Num, const int32 MaxNum, const FName Tag)
+FLearningAgentsObservationObjectElement ULearningAgentsObservations::MakeCountObservation(
+	ULearningAgentsObservationObject* Object, 
+	const int32 Num, 
+	const int32 MaxNum, 
+	const FName Tag,
+	const bool bVisualLoggerEnabled,
+	ULearningAgentsManagerListener* VisualLoggerListener,
+	const int32 VisualLoggerAgentId,
+	const FVector VisualLoggerLocation,
+	const FLinearColor VisualLoggerColor)
 {
 	if (MaxNum == 0)
 	{
@@ -1177,7 +1306,26 @@ FLearningAgentsObservationObjectElement ULearningAgentsObservations::MakeCountOb
 		return FLearningAgentsObservationObjectElement();
 	}
 
-	return MakeContinuousObservationFromArrayView(Object, { (float)Num / (float)MaxNum }, Tag);
+	const float Encoded = (float)Num / (float)MaxNum;
+
+#if UE_LEARNING_AGENTS_ENABLE_VISUAL_LOG
+	if (bVisualLoggerEnabled && VisualLoggerListener)
+	{
+		const ULearningAgentsVisualLoggerObject* VisualLoggerObject = VisualLoggerListener->GetOrAddVisualLoggerObject(Tag);
+
+		UE_LEARNING_AGENTS_VLOG_STRING(VisualLoggerObject, LogLearning, Display, VisualLoggerLocation,
+			VisualLoggerColor.ToFColor(true),
+			TEXT("Listener: %s\nTag: %s\nAgent Id: % 3i\nNum: [%i]\nMax Num: [%i]\nEncoded: [%6.2f]"),
+			*VisualLoggerListener->GetName(),
+			*Tag.ToString(),
+			VisualLoggerAgentId,
+			Num,
+			MaxNum,
+			Encoded);
+	}
+#endif
+
+	return MakeContinuousObservationFromArrayView(Object, { Encoded }, Tag);
 }
 
 FLearningAgentsObservationObjectElement ULearningAgentsObservations::MakeStructObservation(ULearningAgentsObservationObject* Object, const TMap<FName, FLearningAgentsObservationObjectElement>& Elements, const FName Tag)
@@ -1516,7 +1664,16 @@ FLearningAgentsObservationObjectElement ULearningAgentsObservations::MakeMapObse
 	return MakeSetObservationFromArrayView(Object, SubElements);
 }
 
-FLearningAgentsObservationObjectElement ULearningAgentsObservations::MakeEnumObservation(ULearningAgentsObservationObject* Object, const UEnum* Enum, const uint8 EnumValue, const FName Tag)
+FLearningAgentsObservationObjectElement ULearningAgentsObservations::MakeEnumObservation(
+	ULearningAgentsObservationObject* Object, 
+	const UEnum* Enum, 
+	const uint8 EnumValue, 
+	const FName Tag,
+	const bool bVisualLoggerEnabled,
+	ULearningAgentsManagerListener* VisualLoggerListener,
+	const int32 VisualLoggerAgentId,
+	const FVector VisualLoggerLocation,
+	const FLinearColor VisualLoggerColor)
 {
 	if (!Enum)
 	{
@@ -1536,10 +1693,38 @@ FLearningAgentsObservationObjectElement ULearningAgentsObservations::MakeEnumObs
 	OneHot.Init(0.0f, Enum->NumEnums() - 1);
 	OneHot[EnumValueIndex] = 1.0f;
 
+#if UE_LEARNING_AGENTS_ENABLE_VISUAL_LOG
+	if (bVisualLoggerEnabled && VisualLoggerListener)
+	{
+		const ULearningAgentsVisualLoggerObject* VisualLoggerObject = VisualLoggerListener->GetOrAddVisualLoggerObject(Tag);
+
+		UE_LEARNING_AGENTS_VLOG_STRING(VisualLoggerObject, LogLearning, Display, VisualLoggerLocation,
+			VisualLoggerColor.ToFColor(true),
+			TEXT("Listener: %s\nTag: %s\nAgent Id: % 3i\nEnum: %s\nSize: [%i]\nValue: [%s]\nIndex: [%i]\nEncoded: %s"),
+			*VisualLoggerListener->GetName(),
+			*Tag.ToString(),
+			VisualLoggerAgentId,
+			*Enum->GetName(),
+			Enum->NumEnums() - 1,
+			*Enum->GetDisplayNameTextByValue(EnumValue).ToString(),
+			EnumValueIndex,
+			*UE::Learning::Array::FormatFloat(OneHot, Enum->NumEnums() - 1));
+	}
+#endif
+
 	return MakeContinuousObservationFromArrayView(Object, OneHot, Tag);
 }
 
-FLearningAgentsObservationObjectElement ULearningAgentsObservations::MakeBitmaskObservation(ULearningAgentsObservationObject* Object, const UEnum* Enum, const int32 BitmaskValue, const FName Tag)
+FLearningAgentsObservationObjectElement ULearningAgentsObservations::MakeBitmaskObservation(
+	ULearningAgentsObservationObject* Object, 
+	const UEnum* Enum, 
+	const int32 BitmaskValue, 
+	const FName Tag,
+	const bool bVisualLoggerEnabled,
+	ULearningAgentsManagerListener* VisualLoggerListener,
+	const int32 VisualLoggerAgentId,
+	const FVector VisualLoggerLocation,
+	const FLinearColor VisualLoggerColor)
 {
 	if (!Enum)
 	{
@@ -1564,6 +1749,40 @@ FLearningAgentsObservationObjectElement ULearningAgentsObservations::MakeBitmask
 		}
 	}
 
+#if UE_LEARNING_AGENTS_ENABLE_VISUAL_LOG
+	if (bVisualLoggerEnabled && VisualLoggerListener)
+	{
+		const ULearningAgentsVisualLoggerObject* VisualLoggerObject = VisualLoggerListener->GetOrAddVisualLoggerObject(Tag);
+
+		FString ValuesString;
+		FString IndicesString;
+
+		for (int32 EnumIdx = 0; EnumIdx < Enum->NumEnums() - 1; EnumIdx++)
+		{
+			if (BitmaskValue & (1 << EnumIdx))
+			{
+				ValuesString += Enum->GetDisplayNameTextByIndex(EnumIdx).ToString() + TEXT(" ");
+				IndicesString += FString::FromInt(EnumIdx) + TEXT(" ");
+			}
+		}
+
+		ValuesString = ValuesString.TrimEnd();
+		IndicesString = IndicesString.TrimEnd();
+
+		UE_LEARNING_AGENTS_VLOG_STRING(VisualLoggerObject, LogLearning, Display, VisualLoggerLocation,
+			VisualLoggerColor.ToFColor(true),
+			TEXT("Listener: %s\nTag: %s\nAgent Id: % 3i\nEnum: %s\nSize: [%i]\nValues: [%s]\nIndices: [%s]\nEncoded: %s"),
+			*VisualLoggerListener->GetName(),
+			*Tag.ToString(),
+			VisualLoggerAgentId,
+			*Enum->GetName(),
+			Enum->NumEnums() - 1,
+			*ValuesString,
+			*IndicesString,
+			*UE::Learning::Array::FormatFloat(OneHot, Enum->NumEnums() - 1));
+	}
+#endif
+
 	return MakeContinuousObservationFromArrayView(Object, OneHot, Tag);
 }
 
diff --git a/Engine/Plugins/Experimental/LearningAgents/Source/LearningAgents/Public/LearningAgentsObservations.h b/Engine/Plugins/Experimental/LearningAgents/Source/LearningAgents/Public/LearningAgentsObservations.h
--- a/Engine/Plugins/Experimental/LearningAgents/Source/LearningAgents/Public/LearningAgentsObservations.h
+++ b/Engine/Plugins/Experimental/LearningAgents/Source/LearningAgents/Public/LearningAgentsObservations.h
@@ -623,10 +623,23 @@ public:
 	 * @param Object The Observation Object
 	 * @param Values The observation values.
 	 * @param Tag The tag of the corresponding observation. Must match the tag given during Specify.
+	 * @param bVisualLoggerEnabled When true, debug data will be sent to the visual logger.
+	 * @param VisualLoggerListener The listener object which is making this observation. This must be set to use logging.
+	 * @param VisualLoggerAgentId The agent id associated with this observation.
+	 * @param VisualLoggerLocation A location for the visual logger information in the world.
+	 * @param VisualLoggerColor The color for the visual logger display.
 	 * @return The newly created observation object element.
 	 */
-	UFUNCTION(BlueprintPure, Category = "LearningAgents", meta = (AdvancedDisplay = 2))
-	static FLearningAgentsObservationObjectElement MakeContinuousObservation(ULearningAgentsObservationObject* Object, const TArray<float>& Values, const FName Tag = TEXT("Continuous"));
+	UFUNCTION(BlueprintPure, Category = "LearningAgents", meta = (AdvancedDisplay = 2, DefaultToSelf = "VisualLoggerListener"))
+	static FLearningAgentsObservationObjectElement MakeContinuousObservation(
+		ULearningAgentsObservationObject* Object, 
+		const TArray<float>& Values, 
+		const FName Tag = TEXT("Continuous"),
+		const bool bVisualLoggerEnabled = false,
+		ULearningAgentsManagerListener* VisualLoggerListener = nullptr,
+		const int32 VisualLoggerAgentId = -1,
+		const FVector VisualLoggerLocation = FVector::ZeroVector,
+		const FLinearColor VisualLoggerColor = FLinearColor::Red);
 
 	/**
 	 * Make a new continuous observation. The size of Values must match the Size given during Specify.
@@ -634,9 +647,22 @@ public:
 	 * @param Object The Observation Object
 	 * @param Values The observation values.
 	 * @param Tag The tag of the corresponding observation. Must match the tag given during Specify.
+	 * @param bVisualLoggerEnabled When true, debug data will be sent to the visual logger.
+	 * @param VisualLoggerListener The listener object which is making this observation. This must be set to use logging.
+	 * @param VisualLoggerAgentId The agent id associated with this observation.
+	 * @param VisualLoggerLocation A location for the visual logger information in the world.
+	 * @param VisualLoggerColor The color for the visual logger display.
 	 * @return The newly created observation object element.
 	 */
-	static FLearningAgentsObservationObjectElement MakeContinuousObservationFromArrayView(ULearningAgentsObservationObject* Object, const TArrayView<const float> Values, const FName Tag = TEXT("Continuous"));
+	static FLearningAgentsObservationObjectElement MakeContinuousObservationFromArrayView(
+		ULearningAgentsObservationObject* Object, 
+		const TArrayView<const float> Values, 
+		const FName Tag = TEXT("Continuous"),
+		const bool bVisualLoggerEnabled = false,
+		ULearningAgentsManagerListener* VisualLoggerListener = nullptr,
+		const int32 VisualLoggerAgentId = -1,
+		const FVector VisualLoggerLocation = FVector::ZeroVector,
+		const FLinearColor VisualLoggerColor = FLinearColor::Red);
 
 	/**
 	 * Make a new exclusive discrete observation.
@@ -645,10 +671,24 @@ public:
 	 * @param DiscreteIndex The index of the discrete observation. Values must be smaller than the given Size.
 	 * @param Size The size of the discrete observation. Must be equal to the size given during Specify.
 	 * @param Tag The tag of the corresponding observation. Must match the tag given during Specify.
+	 * @param bVisualLoggerEnabled When true, debug data will be sent to the visual logger.
+	 * @param VisualLoggerListener The listener object which is making this observation. This must be set to use logging.
+	 * @param VisualLoggerAgentId The agent id associated with this observation.
+	 * @param VisualLoggerLocation A location for the visual logger information in the world.
+	 * @param VisualLoggerColor The color for the visual logger display.
 	 * @return The newly created observation object element.
 	 */
-	UFUNCTION(BlueprintPure, Category = "LearningAgents", meta = (AdvancedDisplay = 3))
-	static FLearningAgentsObservationObjectElement MakeExclusiveDiscreteObservation(ULearningAgentsObservationObject* Object, const int32 DiscreteIndex, const int32 Size, const FName Tag = TEXT("ExclusiveDiscrete"));
+	UFUNCTION(BlueprintPure, Category = "LearningAgents", meta = (AdvancedDisplay = 3, DefaultToSelf = "VisualLoggerListener"))
+	static FLearningAgentsObservationObjectElement MakeExclusiveDiscreteObservation(
+		ULearningAgentsObservationObject* Object, 
+		const int32 DiscreteIndex, 
+		const int32 Size,
+		const FName Tag = TEXT("ExclusiveDiscrete"),
+		const bool bVisualLoggerEnabled = false,
+		ULearningAgentsManagerListener* VisualLoggerListener = nullptr,
+		const int32 VisualLoggerAgentId = -1,
+		const FVector VisualLoggerLocation = FVector::ZeroVector,
+		const FLinearColor VisualLoggerColor = FLinearColor::Red);
 
 	/**
 	 * Make a new inclusive discrete observation.
@@ -657,10 +697,24 @@ public:
 	 * @param DiscreteIndices The indices of the discrete observations. All values must be smaller than the given Size.
 	 * @param Size The size of the discrete observation. Must be equal to the size given during Specify.
 	 * @param Tag The tag of the corresponding observation. Must match the tag given during Specify.
+	 * @param bVisualLoggerEnabled When true, debug data will be sent to the visual logger.
+	 * @param VisualLoggerListener The listener object which is making this observation. This must be set to use logging.
+	 * @param VisualLoggerAgentId The agent id associated with this observation.
+	 * @param VisualLoggerLocation A location for the visual logger information in the world.
+	 * @param VisualLoggerColor The color for the visual logger display.
 	 * @return The newly created observation object element.
 	 */
-	UFUNCTION(BlueprintPure, Category = "LearningAgents", meta = (AdvancedDisplay = 3))
-	static FLearningAgentsObservationObjectElement MakeInclusiveDiscreteObservation(ULearningAgentsObservationObject* Object, const TArray<int32>& DiscreteIndices, const int32 Size, const FName Tag = TEXT("InclusiveDiscrete"));
+	UFUNCTION(BlueprintPure, Category = "LearningAgents", meta = (AdvancedDisplay = 3, DefaultToSelf = "VisualLoggerListener"))
+	static FLearningAgentsObservationObjectElement MakeInclusiveDiscreteObservation(
+		ULearningAgentsObservationObject* Object, 
+		const TArray<int32>& DiscreteIndices, 
+		const int32 Size, 
+		const FName Tag = TEXT("InclusiveDiscrete"),
+		const bool bVisualLoggerEnabled = false,
+		ULearningAgentsManagerListener* VisualLoggerListener = nullptr,
+		const int32 VisualLoggerAgentId = -1,
+		const FVector VisualLoggerLocation = FVector::ZeroVector,
+		const FLinearColor VisualLoggerColor = FLinearColor::Red);
 	
 	/**
 	 * Make a new inclusive discrete observation.
@@ -669,9 +723,23 @@ public:
 	 * @param DiscreteIndices The indices of the discrete observations. All values must be smaller than the given Size.
 	 * @param Size The size of the discrete observation. Must be equal to the size given during Specify.
 	 * @param Tag The tag of the corresponding observation. Must match the tag given during Specify.
+	 * @param bVisualLoggerEnabled When true, debug data will be sent to the visual logger.
+	 * @param VisualLoggerListener The listener object which is making this observation. This must be set to use logging.
+	 * @param VisualLoggerAgentId The agent id associated with this observation.
+	 * @param VisualLoggerLocation A location for the visual logger information in the world.
+	 * @param VisualLoggerColor The color for the visual logger display.
 	 * @return The newly created observation object element.
 	 */
-	static FLearningAgentsObservationObjectElement MakeInclusiveDiscreteObservationFromArrayView(ULearningAgentsObservationObject* Object, const TArrayView<const int32> DiscreteIndices, const int32 Size, const FName Tag = TEXT("InclusiveDiscrete"));
+	static FLearningAgentsObservationObjectElement MakeInclusiveDiscreteObservationFromArrayView(
+		ULearningAgentsObservationObject* Object, 
+		const TArrayView<const int32> DiscreteIndices, 
+		const int32 Size, 
+		const FName Tag = TEXT("InclusiveDiscrete"),
+		const bool bVisualLoggerEnabled = false,
+		ULearningAgentsManagerListener* VisualLoggerListener = nullptr,
+		const int32 VisualLoggerAgentId = -1,
+		const FVector VisualLoggerLocation = FVector::ZeroVector,
+		const FLinearColor VisualLoggerColor = FLinearColor::Red);
 
 	/**
 	 * Make a new index observation.
@@ -680,10 +748,24 @@ public:
 	 * @param Index The index. Value must be smaller than the given Size.
 	 * @param Size The size of the discrete observation. Must be equal to the size given during Specify.
 	 * @param Tag The tag of the corresponding observation. Must match the tag given during Specify.
+	 * @param bVisualLoggerEnabled When true, debug data will be sent to the visual logger.
+	 * @param VisualLoggerListener The listener object which is making this observation. This must be set to use logging.
+	 * @param VisualLoggerAgentId The agent id associated with this observation.
+	 * @param VisualLoggerLocation A location for the visual logger information in the world.
+	 * @param VisualLoggerColor The color for the visual logger display.
 	 * @return The newly created observation object element.
 	 */
-	UFUNCTION(BlueprintPure, Category = "LearningAgents", meta = (AdvancedDisplay = 3))
-	static FLearningAgentsObservationObjectElement MakeIndexObservation(ULearningAgentsObservationObject* Object, const int32 Index, const int32 Size, const FName Tag = TEXT("Index"));
+	UFUNCTION(BlueprintPure, Category = "LearningAgents", meta = (AdvancedDisplay = 3, DefaultToSelf = "VisualLoggerListener"))
+	static FLearningAgentsObservationObjectElement MakeIndexObservation(
+		ULearningAgentsObservationObject* Object, 
+		const int32 Index, 
+		const int32 Size, 
+		const FName Tag = TEXT("Index"),
+		const bool bVisualLoggerEnabled = false,
+		ULearningAgentsManagerListener* VisualLoggerListener = nullptr,
+		const int32 VisualLoggerAgentId = -1,
+		const FVector VisualLoggerLocation = FVector::ZeroVector,
+		const FLinearColor VisualLoggerColor = FLinearColor::Red);
 
 	/**
 	 * Make a new count observation.
@@ -692,10 +774,24 @@ public:
 	 * @param Num The number of items. Must be less than or equal to MaxNum.
 	 * @param MaxNum The maximum number of items possible.
 	 * @param Tag The tag of the corresponding observation. Must match the tag given during Specify.
+	 * @param bVisualLoggerEnabled When true, debug data will be sent to the visual logger.
+	 * @param VisualLoggerListener The listener object which is making this observation. This must be set to use logging.
+	 * @param VisualLoggerAgentId The agent id associated with this observation.
+	 * @param VisualLoggerLocation A location for the visual logger information in the world.
+	 * @param VisualLoggerColor The color for the visual logger display.
 	 * @return The newly created observation object element.
 	 */
-	UFUNCTION(BlueprintPure, Category = "LearningAgents", meta = (AdvancedDisplay = 3))
-	static FLearningAgentsObservationObjectElement MakeCountObservation(ULearningAgentsObservationObject* Object, const int32 Num, const int32 MaxNum, const FName Tag = TEXT("Count"));
+	UFUNCTION(BlueprintPure, Category = "LearningAgents", meta = (AdvancedDisplay = 3, DefaultToSelf = "VisualLoggerListener"))
+	static FLearningAgentsObservationObjectElement MakeCountObservation(
+		ULearningAgentsObservationObject* Object, 
+		const int32 Num, 
+		const int32 MaxNum, 
+		const FName Tag = TEXT("Count"),
+		const bool bVisualLoggerEnabled = false,
+		ULearningAgentsManagerListener* VisualLoggerListener = nullptr,
+		const int32 VisualLoggerAgentId = -1,
+		const FVector VisualLoggerLocation = FVector::ZeroVector,
+		const FLinearColor VisualLoggerColor = FLinearColor::Red);
 
 	/**
 	 * Make a new struct observation.
@@ -902,10 +998,24 @@ public:
 	 * @param Enum The enum type for this observation. Must match what was given during Specify.
 	 * @param EnumValue The enum value.
 	 * @param Tag The tag of the corresponding observation. Must match the tag given during Specify.
+	 * @param bVisualLoggerEnabled When true, debug data will be sent to the visual logger.
+	 * @param VisualLoggerListener The listener object which is making this observation. This must be set to use logging.
+	 * @param VisualLoggerAgentId The agent id associated with this observation.
+	 * @param VisualLoggerLocation A location for the visual logger information in the world.
+	 * @param VisualLoggerColor The color for the visual logger display.
 	 * @return The newly created observation object element.
 	 */
-	UFUNCTION(BlueprintPure, Category = "LearningAgents", meta = (AdvancedDisplay = 3))
-	static FLearningAgentsObservationObjectElement MakeEnumObservation(ULearningAgentsObservationObject* Object, const UEnum* Enum, const uint8 EnumValue, const FName Tag = TEXT("Enum"));
+	UFUNCTION(BlueprintPure, Category = "LearningAgents", meta = (AdvancedDisplay = 3, DefaultToSelf = "VisualLoggerListener"))
+	static FLearningAgentsObservationObjectElement MakeEnumObservation(
+		ULearningAgentsObservationObject* Object, 
+		const UEnum* Enum, 
+		const uint8 EnumValue, 
+		const FName Tag = TEXT("Enum"),
+		const bool bVisualLoggerEnabled = false,
+		ULearningAgentsManagerListener* VisualLoggerListener = nullptr,
+		const int32 VisualLoggerAgentId = -1,
+		const FVector VisualLoggerLocation = FVector::ZeroVector,
+		const FLinearColor VisualLoggerColor = FLinearColor::Red);
 
 	/**
 	 * Make a new bitmask observation.
@@ -914,10 +1024,24 @@ public:
 	 * @param Enum The enum type for this observation. Must match what was given during Specify.
 	 * @param BitmaskValue The bitmask value.
 	 * @param Tag The tag of the corresponding observation. Must match the tag given during Specify.
+	 * @param bVisualLoggerEnabled When true, debug data will be sent to the visual logger.
+	 * @param VisualLoggerListener The listener object which is making this observation. This must be set to use logging.
+	 * @param VisualLoggerAgentId The agent id associated with this observation.
+	 * @param VisualLoggerLocation A location for the visual logger information in the world.
+	 * @param VisualLoggerColor The color for the visual logger display.
 	 * @return The newly created observation object element.
 	 */
-	UFUNCTION(BlueprintPure, Category = "LearningAgents", meta = (AdvancedDisplay = 3))
-	static FLearningAgentsObservationObjectElement MakeBitmaskObservation(ULearningAgentsObservationObject* Object, const UEnum* Enum, const int32 BitmaskValue, const FName Tag = TEXT("Bitmask"));
+	UFUNCTION(BlueprintPure, Category = "LearningAgents", meta = (AdvancedDisplay = 3, DefaultToSelf = "VisualLoggerListener"))
+	static FLearningAgentsObservationObjectElement MakeBitmaskObservation(
+		ULearningAgentsObservationObject* Object, 
+		const UEnum* Enum, 
+		const int32 BitmaskValue, 
+		const FName Tag = TEXT("Bitmask"),
+		const bool bVisualLoggerEnabled = false,
+		ULearningAgentsManagerListener* VisualLoggerListener = nullptr,
+		const int32 VisualLoggerAgentId = -1,
+		const FVector VisualLoggerLocation = FVector::ZeroVector,
+		const FLinearColor VisualLoggerColor = FLinearColor::Red);
 
 	/**
 	 * Make a new optional observation.
