LearningAgents: Fixed some issues on Mac and Linux and added error message for training on these platforms.

diff --git a/Engine/Plugins/Experimental/LearningAgents/Source/LearningAgentsTraining/Private/LearningAgentsImitationTrainer.cpp b/Engine/Plugins/Experimental/LearningAgents/Source/LearningAgentsTraining/Private/LearningAgentsImitationTrainer.cpp
--- a/Engine/Plugins/Experimental/LearningAgents/Source/LearningAgentsTraining/Private/LearningAgentsImitationTrainer.cpp
+++ b/Engine/Plugins/Experimental/LearningAgents/Source/LearningAgentsTraining/Private/LearningAgentsImitationTrainer.cpp
@@ -115,6 +115,12 @@ void ULearningAgentsImitationTrainer::BeginTraining(
 	const FLearningAgentsImitationTrainerTrainingSettings& ImitationTrainerTrainingSettings,
 	const FLearningAgentsTrainerPathSettings& ImitationTrainerPathSettings)
 {
+	if (!PLATFORM_WINDOWS)
+	{
+		UE_LOG(LogLearning, Error, TEXT("%s: Training currently only supported on Windows."), *GetName());
+		return;
+	}
+
 	if (!IsSetup())
 	{
 		UE_LOG(LogLearning, Error, TEXT("%s: Setup not complete."), *GetName());
diff --git a/Engine/Plugins/Experimental/LearningAgents/Source/LearningAgentsTraining/Private/LearningAgentsTrainer.cpp b/Engine/Plugins/Experimental/LearningAgents/Source/LearningAgentsTraining/Private/LearningAgentsTrainer.cpp
--- a/Engine/Plugins/Experimental/LearningAgents/Source/LearningAgentsTraining/Private/LearningAgentsTrainer.cpp
+++ b/Engine/Plugins/Experimental/LearningAgents/Source/LearningAgentsTraining/Private/LearningAgentsTrainer.cpp
@@ -372,6 +372,12 @@ void ULearningAgentsTrainer::BeginTraining(
 	const FLearningAgentsTrainerPathSettings& TrainerPathSettings,
 	const bool bResetAgentsOnBegin)
 {
+	if (!PLATFORM_WINDOWS)
+	{
+		UE_LOG(LogLearning, Error, TEXT("%s: Training currently only supported on Windows."), *GetName());
+		return;
+	}
+
 	if (!IsSetup())
 	{
 		UE_LOG(LogLearning, Error, TEXT("%s: Setup not complete."), *GetName());
diff --git a/Engine/Plugins/Experimental/LearningAgents/Source/LearningTraining/Private/LearningTrainer.cpp b/Engine/Plugins/Experimental/LearningAgents/Source/LearningTraining/Private/LearningTrainer.cpp
--- a/Engine/Plugins/Experimental/LearningAgents/Source/LearningTraining/Private/LearningTrainer.cpp
+++ b/Engine/Plugins/Experimental/LearningAgents/Source/LearningTraining/Private/LearningTrainer.cpp
@@ -296,12 +296,14 @@ namespace UE::Learning::Trainer
 	{
 		UE_LEARNING_CHECKF(PLATFORM_WINDOWS || PLATFORM_MAC || PLATFORM_LINUX, TEXT("Python only supported on Windows, Mac, and Linux."));
 
-		return EngineDir / TEXT("Binaries/ThirdParty/Python3") / FPlatformMisc::GetUBTPlatform() / (PLATFORM_WINDOWS ? TEXT("python.exe") : TEXT("bin/python"));
+		return EngineDir / TEXT("Binaries/ThirdParty/Python3") / FPlatformMisc::GetUBTPlatform() / (PLATFORM_WINDOWS ? TEXT("python.exe") : TEXT("bin/python3"));
 	}
 
 	FString GetSitePackagesPath(const FString& IntermediateDir)
 	{
-		return IntermediateDir / TEXT("PipInstall/Lib/site-packages");
+		UE_LEARNING_CHECKF(PLATFORM_WINDOWS || PLATFORM_MAC || PLATFORM_LINUX, TEXT("Python only supported on Windows, Mac, and Linux."));
+
+		return PLATFORM_WINDOWS ? IntermediateDir / TEXT("PipInstall/Lib/site-packages") : IntermediateDir / TEXT("PipInstall/lib/python3.9/site-packages");
 	}
 
 	FString GetPythonContentPath(const FString& EngineDir)
