LearningAgents: Added comments and visual logging to actions. Changed tags to include Observation/Action/Reward/Completion to avoid naming conflicts in the visual logger.

#virtualized

diff --git a/Engine/Plugins/Experimental/LearningAgents/Source/LearningAgents/Private/LearningAgentsActions.cpp b/Engine/Plugins/Experimental/LearningAgents/Source/LearningAgents/Private/LearningAgentsActions.cpp
--- a/Engine/Plugins/Experimental/LearningAgents/Source/LearningAgents/Private/LearningAgentsActions.cpp
+++ b/Engine/Plugins/Experimental/LearningAgents/Source/LearningAgents/Private/LearningAgentsActions.cpp
@@ -55,7 +55,7 @@ namespace UE::Learning::Agents::Action::Private
 		}
 	}
 
-	static bool ValidateObjectMatchesSchema(
+	static bool ValidateActionObjectMatchesSchema(
 		const Learning::Action::FSchema& Schema,
 		const Learning::Action::FSchemaElement SchemaElement,
 		const Learning::Action::FObject& Object,
@@ -65,13 +65,13 @@ namespace UE::Learning::Agents::Action::Private
 
 		if (!Schema.IsValid(SchemaElement))
 		{
-			UE_LOG(LogLearning, Error, TEXT("ValidateObjectMatchesSchema: Invalid Action Schema Element."));
+			UE_LOG(LogLearning, Error, TEXT("ValidateActionObjectMatchesSchema: Invalid Action Schema Element."));
 			return false;
 		}
 
 		if (!Object.IsValid(ObjectElement))
 		{
-			UE_LOG(LogLearning, Error, TEXT("ValidateObjectMatchesSchema: Invalid Action Object Element."));
+			UE_LOG(LogLearning, Error, TEXT("ValidateActionObjectMatchesSchema: Invalid Action Object Element."));
 			return false;
 		}
 
@@ -82,7 +82,7 @@ namespace UE::Learning::Agents::Action::Private
 
 		if (ActionSchemaElementTag != ActionObjectElementTag)
 		{
-			UE_LOG(LogLearning, Warning, TEXT("ValidateObjectMatchesSchema: Action tag does not match Schema. Expected '%s', got '%s'."),
+			UE_LOG(LogLearning, Warning, TEXT("ValidateActionObjectMatchesSchema: Action tag does not match Schema. Expected '%s', got '%s'."),
 				*ActionSchemaElementTag.ToString(), *ActionObjectElementTag.ToString());
 		}
 
@@ -93,7 +93,7 @@ namespace UE::Learning::Agents::Action::Private
 
 		if (ActionSchemaElementType != ActionObjectElementType)
 		{
-			UE_LOG(LogLearning, Error, TEXT("ValidateObjectMatchesSchema: Action '%s' type does not match Schema. Expected type '%s', got type '%s'."),
+			UE_LOG(LogLearning, Error, TEXT("ValidateActionObjectMatchesSchema: Action '%s' type does not match Schema. Expected type '%s', got type '%s'."),
 				*ActionSchemaElementTag.ToString(),
 				GetActionTypeString(ActionSchemaElementType),
 				GetActionTypeString(ActionObjectElementType));
@@ -113,7 +113,7 @@ namespace UE::Learning::Agents::Action::Private
 
 			if (SchemaElementSize != ObjectElementSize)
 			{
-				UE_LOG(LogLearning, Error, TEXT("ValidateObjectMatchesSchema: Action '%s' size does not match Schema. Expected '%i', got '%i'."),
+				UE_LOG(LogLearning, Error, TEXT("ValidateActionObjectMatchesSchema: Action '%s' size does not match Schema. Expected '%i', got '%i'."),
 					*ActionSchemaElementTag.ToString(),
 					SchemaElementSize,
 					ObjectElementSize);
@@ -130,7 +130,7 @@ namespace UE::Learning::Agents::Action::Private
 
 			if (ObjectElementIndex < 0 || ObjectElementIndex >= SchemaElementSize)
 			{
-				UE_LOG(LogLearning, Error, TEXT("ValidateObjectMatchesSchema: Action '%s' index out of range for Schema. Expected '<%i', got '%i'."),
+				UE_LOG(LogLearning, Error, TEXT("ValidateActionObjectMatchesSchema: Action '%s' index out of range for Schema. Expected '<%i', got '%i'."),
 					*ActionSchemaElementTag.ToString(),
 					SchemaElementSize,
 					ObjectElementIndex);
@@ -147,7 +147,7 @@ namespace UE::Learning::Agents::Action::Private
 
 			if (ObjectElementIndices.Num() > SchemaElementSize)
 			{
-				UE_LOG(LogLearning, Error, TEXT("ValidateObjectMatchesSchema: Action '%s' too many indices provided. Expected at most '%i', got '%i'."),
+				UE_LOG(LogLearning, Error, TEXT("ValidateActionObjectMatchesSchema: Action '%s' too many indices provided. Expected at most '%i', got '%i'."),
 					*ActionSchemaElementTag.ToString(),
 					SchemaElementSize,
 					ObjectElementIndices.Num());
@@ -158,7 +158,7 @@ namespace UE::Learning::Agents::Action::Private
 			{
 				if (ObjectElementIndices[SubElementIdx] < 0 || ObjectElementIndices[SubElementIdx] >= SchemaElementSize)
 				{
-					UE_LOG(LogLearning, Error, TEXT("ValidateObjectMatchesSchema: Action '%s' index out of range for Schema. Expected '<%i', got '%i'."),
+					UE_LOG(LogLearning, Error, TEXT("ValidateActionObjectMatchesSchema: Action '%s' index out of range for Schema. Expected '<%i', got '%i'."),
 						*ActionSchemaElementTag.ToString(),
 						SchemaElementSize,
 						ObjectElementIndices[SubElementIdx]);
@@ -178,7 +178,7 @@ namespace UE::Learning::Agents::Action::Private
 
 			if (SchemaParameters.Elements.Num() != ObjectParameters.Elements.Num())
 			{
-				UE_LOG(LogLearning, Error, TEXT("ValidateObjectMatchesSchema: Action '%s' number of sub-elements does not match Schema. Expected '%i', got '%i'."),
+				UE_LOG(LogLearning, Error, TEXT("ValidateActionObjectMatchesSchema: Action '%s' number of sub-elements does not match Schema. Expected '%i', got '%i'."),
 					*ActionSchemaElementTag.ToString(),
 					SchemaParameters.Elements.Num(),
 					ObjectParameters.Elements.Num());
@@ -191,13 +191,13 @@ namespace UE::Learning::Agents::Action::Private
 
 				if (ObjectElementIdx == INDEX_NONE)
 				{
-					UE_LOG(LogLearning, Error, TEXT("ValidateObjectMatchesSchema: Action '%s' does not include '%s' action required by Schema."),
+					UE_LOG(LogLearning, Error, TEXT("ValidateActionObjectMatchesSchema: Action '%s' does not include '%s' action required by Schema."),
 						*ActionSchemaElementTag.ToString(),
 						*SchemaParameters.ElementNames[SchemaElementIdx].ToString());
 					return false;
 				}
 
-				if (!ValidateObjectMatchesSchema(
+				if (!ValidateActionObjectMatchesSchema(
 					Schema,
 					SchemaParameters.Elements[SchemaElementIdx],
 					Object,
@@ -220,13 +220,13 @@ namespace UE::Learning::Agents::Action::Private
 
 			if (SchemaSubElementIdx == INDEX_NONE)
 			{
-				UE_LOG(LogLearning, Error, TEXT("ValidateObjectMatchesSchema: Action '%s' Schema does not include '%s' action."),
+				UE_LOG(LogLearning, Error, TEXT("ValidateActionObjectMatchesSchema: Action '%s' Schema does not include '%s' action."),
 					*ActionSchemaElementTag.ToString(),
 					*ObjectParameters.ElementName.ToString());
 				return false;
 			}
 
-			return ValidateObjectMatchesSchema(
+			return ValidateActionObjectMatchesSchema(
 				Schema,
 				SchemaParameters.Elements[SchemaSubElementIdx],
 				Object,
@@ -240,7 +240,7 @@ namespace UE::Learning::Agents::Action::Private
 
 			if (ObjectParameters.Elements.Num() > SchemaParameters.Elements.Num())
 			{
-				UE_LOG(LogLearning, Error, TEXT("ValidateObjectMatchesSchema: Action '%s' too many sub-actions provided. Expected at most '%i', got '%i'."),
+				UE_LOG(LogLearning, Error, TEXT("ValidateActionObjectMatchesSchema: Action '%s' too many sub-actions provided. Expected at most '%i', got '%i'."),
 					*ActionSchemaElementTag.ToString(),
 					SchemaParameters.Elements.Num(),
 					ObjectParameters.Elements.Num());
@@ -253,13 +253,13 @@ namespace UE::Learning::Agents::Action::Private
 
 				if (SchemaSubElementIdx == INDEX_NONE)
 				{
-					UE_LOG(LogLearning, Error, TEXT("ValidateObjectMatchesSchema: Action '%s' Schema does not include '%s' action."),
+					UE_LOG(LogLearning, Error, TEXT("ValidateActionObjectMatchesSchema: Action '%s' Schema does not include '%s' action."),
 						*ActionSchemaElementTag.ToString(),
 						*ObjectParameters.ElementNames[ObjectSubElementIdx].ToString());
 					return false;
 				}
 
-				if (!ValidateObjectMatchesSchema(
+				if (!ValidateActionObjectMatchesSchema(
 					Schema,
 					SchemaParameters.Elements[SchemaSubElementIdx],
 					Object,
@@ -279,7 +279,7 @@ namespace UE::Learning::Agents::Action::Private
 
 			if (ObjectParameters.Elements.Num() != SchemaParameters.Num)
 			{
-				UE_LOG(LogLearning, Error, TEXT("ValidateObjectMatchesSchema: Action '%s' array incorrect size. Expected '%i' elements, got '%i'."),
+				UE_LOG(LogLearning, Error, TEXT("ValidateActionObjectMatchesSchema: Action '%s' array incorrect size. Expected '%i' elements, got '%i'."),
 					*ActionSchemaElementTag.ToString(),
 					SchemaParameters.Num,
 					ObjectParameters.Elements.Num());
@@ -288,7 +288,7 @@ namespace UE::Learning::Agents::Action::Private
 
 			for (int32 ElementIdx = 0; ElementIdx < ObjectParameters.Elements.Num(); ElementIdx++)
 			{
-				if (!ValidateObjectMatchesSchema(
+				if (!ValidateActionObjectMatchesSchema(
 					Schema,
 					SchemaParameters.Element,
 					Object,
@@ -306,7 +306,7 @@ namespace UE::Learning::Agents::Action::Private
 			const Learning::Action::FSchemaEncodingParameters SchemaParameters = Schema.GetEncoding(SchemaElement);
 			const Learning::Action::FObjectEncodingParameters ObjectParameters = Object.GetEncoding(ObjectElement);
 
-			return ValidateObjectMatchesSchema(
+			return ValidateActionObjectMatchesSchema(
 					Schema,
 					SchemaParameters.Element,
 					Object,
@@ -502,7 +502,7 @@ namespace UE::Learning::Agents::Action::Private
 	}
 }
 
-bool ULearningAgentsActions::ValidateObjectMatchesSchema(
+bool ULearningAgentsActions::ValidateActionObjectMatchesSchema(
 	const ULearningAgentsActionSchema* Schema,
 	const FLearningAgentsActionSchemaElement SchemaElement,
 	const ULearningAgentsActionObject* Object,
@@ -510,17 +510,17 @@ bool ULearningAgentsActions::ValidateObjectMatchesSchema(
 {
 	if (!Schema)
 	{
-		UE_LOG(LogLearning, Error, TEXT("ValidateObjectMatchesSchema: Schema is nullptr."));
+		UE_LOG(LogLearning, Error, TEXT("ValidateActionObjectMatchesSchema: Schema is nullptr."));
 		return false;
 	}
 
 	if (!Object)
 	{
-		UE_LOG(LogLearning, Error, TEXT("ValidateObjectMatchesSchema: Object is nullptr."));
+		UE_LOG(LogLearning, Error, TEXT("ValidateActionObjectMatchesSchema: Object is nullptr."));
 		return false;
 	}
 
-	return UE::Learning::Agents::Action::Private::ValidateObjectMatchesSchema(
+	return UE::Learning::Agents::Action::Private::ValidateActionObjectMatchesSchema(
 		Schema->ActionSchema,
 		SchemaElement.SchemaElement,
 		Object->ActionObject,
@@ -1152,12 +1152,6 @@ FLearningAgentsActionSchemaElement ULearningAgentsActions::SpecifyDirectionActio
 	return SpecifyContinuousAction(Schema, 3, Tag);
 }
 
-FLearningAgentsActionSchemaElement ULearningAgentsActions::SpecifySpeedAction(ULearningAgentsActionSchema* Schema, const FName Tag)
-{
-	return SpecifyContinuousAction(Schema, 1, Tag);
-}
-
-
 void ULearningAgentsActions::LogAction(const ULearningAgentsActionObject* Object, const FLearningAgentsActionObjectElement Element)
 {
 	if (!Object)
@@ -1637,11 +1631,6 @@ FLearningAgentsActionObjectElement ULearningAgentsActions::MakeDirectionAction(U
 		}, Tag);
 }
 
-FLearningAgentsActionObjectElement ULearningAgentsActions::MakeSpeedAction(ULearningAgentsActionObject* Object, const float Speed, const float SpeedScale, const FName Tag)
-{
-	return MakeContinuousActionFromArrayView(Object, { Speed / FMath::Max(SpeedScale, UE_SMALL_NUMBER) });
-}
-
 bool ULearningAgentsActions::GetNullAction(const ULearningAgentsActionObject* Object, const FLearningAgentsActionObjectElement Element, const FName Tag)
 {
 	if (!Object)
@@ -1708,7 +1697,16 @@ bool ULearningAgentsActions::GetContinuousActionNum(int32& OutNum, const ULearni
 	return true;
 }
 
-bool ULearningAgentsActions::GetContinuousAction(TArray<float>& OutValues, const ULearningAgentsActionObject* Object, const FLearningAgentsActionObjectElement Element, const FName Tag)
+bool ULearningAgentsActions::GetContinuousAction(
+	TArray<float>& OutValues, 
+	const ULearningAgentsActionObject* Object, 
+	const FLearningAgentsActionObjectElement Element, 
+	const FName Tag,
+	const bool bVisualLoggerEnabled,
+	ULearningAgentsManagerListener* VisualLoggerListener,
+	const int32 VisualLoggerAgentId,
+	const FVector VisualLoggerLocation,
+	const FLinearColor VisualLoggerColor)
 {
 	int32 OutValueNum = 0;
 	if (!GetContinuousActionNum(OutValueNum, Object, Element, Tag))
@@ -1719,7 +1717,7 @@ bool ULearningAgentsActions::GetContinuousAction(TArray<float>& OutValues, const
 
 	OutValues.SetNumUninitialized(OutValueNum);
 
-	if (!GetContinuousActionToArrayView(OutValues, Object, Element, Tag))
+	if (!GetContinuousActionToArrayView(OutValues, Object, Element, Tag, bVisualLoggerEnabled, VisualLoggerListener, VisualLoggerAgentId, VisualLoggerLocation, VisualLoggerColor))
 	{
 		OutValues.Empty();
 		return false;
@@ -1728,7 +1726,16 @@ bool ULearningAgentsActions::GetContinuousAction(TArray<float>& OutValues, const
 	return true;
 }
 
-bool ULearningAgentsActions::GetContinuousActionToArrayView(TArrayView<float> OutValues, const ULearningAgentsActionObject* Object, const FLearningAgentsActionObjectElement Element, const FName Tag)
+bool ULearningAgentsActions::GetContinuousActionToArrayView(
+	TArrayView<float> OutValues, 
+	const ULearningAgentsActionObject* Object, 
+	const FLearningAgentsActionObjectElement Element, 
+	const FName Tag,
+	const bool bVisualLoggerEnabled,
+	ULearningAgentsManagerListener* VisualLoggerListener,
+	const int32 VisualLoggerAgentId,
+	const FVector VisualLoggerLocation,
+	const FLinearColor VisualLoggerColor)
 {
 	if (!Object)
 	{
@@ -1771,10 +1778,36 @@ bool ULearningAgentsActions::GetContinuousActionToArrayView(TArrayView<float> Ou
 	}
 
 	UE::Learning::Array::Copy<1, float>(OutValues, Values);
+
+#if UE_LEARNING_AGENTS_ENABLE_VISUAL_LOG
+	if (bVisualLoggerEnabled && VisualLoggerListener)
+	{
+		const ULearningAgentsVisualLoggerObject* VisualLoggerObject = VisualLoggerListener->GetOrAddVisualLoggerObject(Tag);
+
+		UE_LEARNING_AGENTS_VLOG_STRING(VisualLoggerObject, LogLearning, Display, VisualLoggerLocation,
+			VisualLoggerColor.ToFColor(true),
+			TEXT("Listener: %s\nTag: %s\nAgent Id: % 3i\nEncoded: %s\nValues: %s"),
+			*VisualLoggerListener->GetName(),
+			*Tag.ToString(),
+			VisualLoggerAgentId,
+			*UE::Learning::Array::FormatFloat(OutValues),
+			*UE::Learning::Array::FormatFloat(OutValues)); // Encoded is identical to provided values
+	}
+#endif
+
 	return true;
 }
 
-bool ULearningAgentsActions::GetExclusiveDiscreteAction(int32& OutIndex, const ULearningAgentsActionObject* Object, const FLearningAgentsActionObjectElement Element, const FName Tag)
+bool ULearningAgentsActions::GetExclusiveDiscreteAction(
+	int32& OutIndex, 
+	const ULearningAgentsActionObject* Object, 
+	const FLearningAgentsActionObjectElement Element, 
+	const FName Tag,
+	const bool bVisualLoggerEnabled,
+	ULearningAgentsManagerListener* VisualLoggerListener,
+	const int32 VisualLoggerAgentId,
+	const FVector VisualLoggerLocation,
+	const FLinearColor VisualLoggerColor)
 {
 	if (!Object)
 	{
@@ -1806,6 +1839,22 @@ bool ULearningAgentsActions::GetExclusiveDiscreteAction(int32& OutIndex, const U
 	}
 
 	OutIndex = Object->ActionObject.GetDiscreteExclusive(Element.ObjectElement).DiscreteIndex;
+
+#if UE_LEARNING_AGENTS_ENABLE_VISUAL_LOG
+	if (bVisualLoggerEnabled && VisualLoggerListener)
+	{
+		const ULearningAgentsVisualLoggerObject* VisualLoggerObject = VisualLoggerListener->GetOrAddVisualLoggerObject(Tag);
+
+		UE_LEARNING_AGENTS_VLOG_STRING(VisualLoggerObject, LogLearning, Display, VisualLoggerLocation,
+			VisualLoggerColor.ToFColor(true),
+			TEXT("Listener: %s\nTag: %s\nAgent Id: % 3i\nIndex: [%i]"),
+			*VisualLoggerListener->GetName(),
+			*Tag.ToString(),
+			VisualLoggerAgentId,
+			OutIndex);
+	}
+#endif
+
 	return true;
 }
 
@@ -1844,7 +1893,16 @@ bool ULearningAgentsActions::GetInclusiveDiscreteActionNum(int32& OutNum, const
 	return true;
 }
 
-bool ULearningAgentsActions::GetInclusiveDiscreteAction(TArray<int32>& OutIndices, const ULearningAgentsActionObject* Object, const FLearningAgentsActionObjectElement Element, const FName Tag)
+bool ULearningAgentsActions::GetInclusiveDiscreteAction(
+	TArray<int32>& OutIndices, 
+	const ULearningAgentsActionObject* Object, 
+	const FLearningAgentsActionObjectElement Element, 
+	const FName Tag,
+	const bool bVisualLoggerEnabled,
+	ULearningAgentsManagerListener* VisualLoggerListener,
+	const int32 VisualLoggerAgentId,
+	const FVector VisualLoggerLocation,
+	const FLinearColor VisualLoggerColor)
 {
 	int32 OutIndexNum = 0;
 	if (!GetInclusiveDiscreteActionNum(OutIndexNum, Object, Element, Tag))
@@ -1855,7 +1913,7 @@ bool ULearningAgentsActions::GetInclusiveDiscreteAction(TArray<int32>& OutIndice
 
 	OutIndices.SetNumUninitialized(OutIndexNum);
 
-	if (!GetInclusiveDiscreteActionToArrayView(OutIndices, Object, Element, Tag))
+	if (!GetInclusiveDiscreteActionToArrayView(OutIndices, Object, Element, Tag, bVisualLoggerEnabled, VisualLoggerListener, VisualLoggerAgentId, VisualLoggerLocation, VisualLoggerColor))
 	{
 		OutIndices.Empty();
 		return false;
@@ -1864,7 +1922,16 @@ bool ULearningAgentsActions::GetInclusiveDiscreteAction(TArray<int32>& OutIndice
 	return true;
 }
 
-bool ULearningAgentsActions::GetInclusiveDiscreteActionToArrayView(TArrayView<int32> OutIndices, const ULearningAgentsActionObject* Object, const FLearningAgentsActionObjectElement Element, const FName Tag)
+bool ULearningAgentsActions::GetInclusiveDiscreteActionToArrayView(
+	TArrayView<int32> OutIndices,
+	const ULearningAgentsActionObject* Object, 
+	const FLearningAgentsActionObjectElement Element, 
+	const FName Tag,
+	const bool bVisualLoggerEnabled,
+	ULearningAgentsManagerListener* VisualLoggerListener,
+	const int32 VisualLoggerAgentId,
+	const FVector VisualLoggerLocation,
+	const FLinearColor VisualLoggerColor)
 {
 	if (!Object)
 	{
@@ -1907,6 +1974,22 @@ bool ULearningAgentsActions::GetInclusiveDiscreteActionToArrayView(TArrayView<in
 	}
 
 	UE::Learning::Array::Copy<1, int32>(OutIndices, Indices);
+
+#if UE_LEARNING_AGENTS_ENABLE_VISUAL_LOG
+	if (bVisualLoggerEnabled && VisualLoggerListener)
+	{
+		const ULearningAgentsVisualLoggerObject* VisualLoggerObject = VisualLoggerListener->GetOrAddVisualLoggerObject(Tag);
+
+		UE_LEARNING_AGENTS_VLOG_STRING(VisualLoggerObject, LogLearning, Display, VisualLoggerLocation,
+			VisualLoggerColor.ToFColor(true),
+			TEXT("Listener: %s\nTag: %s\nAgent Id: % 3i\nIndices: %s"),
+			*VisualLoggerListener->GetName(),
+			*Tag.ToString(),
+			VisualLoggerAgentId,
+			*UE::Learning::Array::FormatInt32(OutIndices, 256));
+	}
+#endif
+
 	return true;
 }
 
@@ -2393,7 +2476,17 @@ bool ULearningAgentsActions::GetPairAction(FLearningAgentsActionObjectElement& O
 	return true;
 }
 
-bool ULearningAgentsActions::GetEnumAction(uint8& OutEnumValue, const ULearningAgentsActionObject* Object, const FLearningAgentsActionObjectElement Element, const UEnum* Enum, const FName Tag)
+bool ULearningAgentsActions::GetEnumAction(
+	uint8& OutEnumValue, 
+	const ULearningAgentsActionObject* Object, 
+	const FLearningAgentsActionObjectElement Element, 
+	const UEnum* Enum, 
+	const FName Tag,
+	const bool bVisualLoggerEnabled,
+	ULearningAgentsManagerListener* VisualLoggerListener,
+	const int32 VisualLoggerAgentId,
+	const FVector VisualLoggerLocation,
+	const FLinearColor VisualLoggerColor)
 {
 	if (!Enum)
 	{
@@ -2426,10 +2519,40 @@ bool ULearningAgentsActions::GetEnumAction(uint8& OutEnumValue, const ULearningA
 	}
 
 	OutEnumValue = (uint8)EnumValue;
+
+
+#if UE_LEARNING_AGENTS_ENABLE_VISUAL_LOG
+	if (bVisualLoggerEnabled && VisualLoggerListener)
+	{
+		const ULearningAgentsVisualLoggerObject* VisualLoggerObject = VisualLoggerListener->GetOrAddVisualLoggerObject(Tag);
+
+		UE_LEARNING_AGENTS_VLOG_STRING(VisualLoggerObject, LogLearning, Display, VisualLoggerLocation,
+			VisualLoggerColor.ToFColor(true),
+			TEXT("Listener: %s\nTag: %s\nAgent Id: % 3i\nEnum: %s\nSize: [%i]\nValue: [%s]\nIndex: [%i]"),
+			*VisualLoggerListener->GetName(),
+			*Tag.ToString(),
+			VisualLoggerAgentId,
+			*Enum->GetName(),
+			Enum->NumEnums() - 1,
+			*Enum->GetDisplayNameTextByValue(OutEnumValue).ToString(),
+			OutIndex);
+	}
+#endif
+
 	return true;
 }
 
-bool ULearningAgentsActions::GetBitmaskAction(int32& OutBitmaskValue, const ULearningAgentsActionObject* Object, const FLearningAgentsActionObjectElement Element, const UEnum* Enum, const FName Tag)
+bool ULearningAgentsActions::GetBitmaskAction(
+	int32& OutBitmaskValue, 
+	const ULearningAgentsActionObject* Object, 
+	const FLearningAgentsActionObjectElement Element, 
+	const UEnum* Enum, 
+	const FName Tag,
+	const bool bVisualLoggerEnabled,
+	ULearningAgentsManagerListener* VisualLoggerListener,
+	const int32 VisualLoggerAgentId,
+	const FVector VisualLoggerLocation,
+	const FLinearColor VisualLoggerColor)
 {
 	if (!Enum)
 	{
@@ -2472,6 +2595,40 @@ bool ULearningAgentsActions::GetBitmaskAction(int32& OutBitmaskValue, const ULea
 	{
 		OutBitmaskValue |= (1 << OutIndex);
 	}
+
+#if UE_LEARNING_AGENTS_ENABLE_VISUAL_LOG
+	if (bVisualLoggerEnabled && VisualLoggerListener)
+	{
+		const ULearningAgentsVisualLoggerObject* VisualLoggerObject = VisualLoggerListener->GetOrAddVisualLoggerObject(Tag);
+
+		FString ValuesString;
+		FString IndicesString;
+
+		for (int32 EnumIdx = 0; EnumIdx < Enum->NumEnums() - 1; EnumIdx++)
+		{
+			if (OutBitmaskValue & (1 << EnumIdx))
+			{
+				ValuesString += Enum->GetDisplayNameTextByIndex(EnumIdx).ToString() + TEXT(" ");
+				IndicesString += FString::FromInt(EnumIdx) + TEXT(" ");
+			}
+		}
+
+		ValuesString = ValuesString.TrimEnd();
+		IndicesString = IndicesString.TrimEnd();
+
+		UE_LEARNING_AGENTS_VLOG_STRING(VisualLoggerObject, LogLearning, Display, VisualLoggerLocation,
+			VisualLoggerColor.ToFColor(true),
+			TEXT("Listener: %s\nTag: %s\nAgent Id: % 3i\nEnum: %s\nSize: [%i]\nValues: [%s]\nIndices: [%s]"),
+			*VisualLoggerListener->GetName(),
+			*Tag.ToString(),
+			VisualLoggerAgentId,
+			*Enum->GetName(),
+			Enum->NumEnums() - 1,
+			*ValuesString,
+			*IndicesString);
+	}
+#endif
+
 	return true;
 }
 
@@ -2537,7 +2694,16 @@ bool ULearningAgentsActions::GetEncodingAction(FLearningAgentsActionObjectElemen
 	return true;
 }
 
-bool ULearningAgentsActions::GetBoolAction(bool& bOutValue, const ULearningAgentsActionObject* Object, const FLearningAgentsActionObjectElement Element, const FName Tag)
+bool ULearningAgentsActions::GetBoolAction(
+	bool& bOutValue, 
+	const ULearningAgentsActionObject* Object, 
+	const FLearningAgentsActionObjectElement Element, 
+	const FName Tag,
+	const bool bVisualLoggerEnabled,
+	ULearningAgentsManagerListener* VisualLoggerListener,
+	const int32 VisualLoggerAgentId,
+	const FVector VisualLoggerLocation,
+	const FLinearColor VisualLoggerColor)
 {
 	int32 OutIndex = 0;
 	if (!GetExclusiveDiscreteAction(OutIndex, Object, Element, Tag))
@@ -2547,6 +2713,22 @@ bool ULearningAgentsActions::GetBoolAction(bool& bOutValue, const ULearningAgent
 	}
 
 	bOutValue = OutIndex == 1;
+	
+#if UE_LEARNING_AGENTS_ENABLE_VISUAL_LOG
+	if (bVisualLoggerEnabled && VisualLoggerListener)
+	{
+		const ULearningAgentsVisualLoggerObject* VisualLoggerObject = VisualLoggerListener->GetOrAddVisualLoggerObject(Tag);
+
+		UE_LEARNING_AGENTS_VLOG_STRING(VisualLoggerObject, LogLearning, Display, VisualLoggerLocation,
+			VisualLoggerColor.ToFColor(true),
+			TEXT("Listener: %s\nTag: %s\nAgent Id: % 3i\nValue: [%s]"),
+			*VisualLoggerListener->GetName(),
+			*Tag.ToString(),
+			VisualLoggerAgentId,
+			bOutValue ? TEXT("true") : TEXT("false"));
+	}
+#endif
+	
 	return true;
 }
 
@@ -2591,7 +2773,18 @@ bool ULearningAgentsActions::GetFloatAction(
 	return true;
 }
 
-bool ULearningAgentsActions::GetLocationAction(FVector& OutLocation, const ULearningAgentsActionObject* Object, const FLearningAgentsActionObjectElement Element, const FTransform RelativeTransform, const float LocationScale, const FName Tag)
+bool ULearningAgentsActions::GetLocationAction(
+	FVector& OutLocation, 
+	const ULearningAgentsActionObject* Object, 
+	const FLearningAgentsActionObjectElement Element, 
+	const FTransform RelativeTransform, 
+	const float LocationScale, 
+	const FName Tag,
+	const bool bVisualLoggerEnabled,
+	ULearningAgentsManagerListener* VisualLoggerListener,
+	const int32 VisualLoggerAgentId,
+	const FVector VisualLoggerLocation,
+	const FLinearColor VisualLoggerColor)
 {
 	TStaticArray<float, 3> OutValues;
 	if (!GetContinuousActionToArrayView(OutValues, Object, Element, Tag))
@@ -2602,13 +2795,62 @@ bool ULearningAgentsActions::GetLocationAction(FVector& OutLocation, const ULear
 
 	const FVector LocalLocation = LocationScale * FVector(OutValues[0], OutValues[1], OutValues[2]);
 	OutLocation = RelativeTransform.TransformPosition(LocalLocation);
+	
+#if UE_LEARNING_AGENTS_ENABLE_VISUAL_LOG
+	if (bVisualLoggerEnabled && VisualLoggerListener)
+	{
+		const ULearningAgentsVisualLoggerObject* VisualLoggerObject = VisualLoggerListener->GetOrAddVisualLoggerObject(Tag);
+
+		UE_LEARNING_AGENTS_VLOG_LOCATION(VisualLoggerObject, LogLearning, Display,
+			OutLocation,
+			10,
+			VisualLoggerColor.ToFColor(true),
+			TEXT(""));
+
+		UE_LEARNING_AGENTS_VLOG_SEGMENT(VisualLoggerObject, LogLearning, Display,
+			RelativeTransform.GetTranslation(),
+			OutLocation,
+			VisualLoggerColor.ToFColor(true),
+			TEXT(""));
+
+		UE_LEARNING_AGENTS_VLOG_TRANSFORM(VisualLoggerObject, LogLearning, Display,
+			RelativeTransform.GetTranslation(),
+			RelativeTransform.GetRotation(),
+			VisualLoggerColor.ToFColor(true),
+			TEXT(""));
+
+		UE_LEARNING_AGENTS_VLOG_STRING(VisualLoggerObject, LogLearning, Display, VisualLoggerLocation,
+			VisualLoggerColor.ToFColor(true),
+			TEXT("Listener: %s\nTag: %s\nAgent Id: % 3i\nEncoded: [% 6.2f % 6.2f % 6.2f]\nScale: [% 6.2f]\nLocal Location: [% 6.1f % 6.1f % 6.1f]\nLocation: [% 6.1f % 6.1f % 6.1f]"),
+			*VisualLoggerListener->GetName(),
+			*Tag.ToString(),
+			VisualLoggerAgentId,
+			OutValues[0], OutValues[1], OutValues[2],
+			LocationScale,
+			LocalLocation.X, LocalLocation.Y, LocalLocation.Z,
+			OutLocation.X, OutLocation.Y, OutLocation.Z);
+	}
+#endif
+	
 	return true;
 }
 
-bool ULearningAgentsActions::GetRotationAction(FRotator& OutRotation, const ULearningAgentsActionObject* Object, const FLearningAgentsActionObjectElement Element, const FRotator RelativeRotation, const float RotationScale, const FName Tag)
+bool ULearningAgentsActions::GetRotationAction(
+	FRotator& OutRotation, 
+	const ULearningAgentsActionObject* Object, 
+	const FLearningAgentsActionObjectElement Element, 
+	const FRotator RelativeRotation, 
+	const float RotationScale, 
+	const FName Tag,
+	const bool bVisualLoggerEnabled,
+	ULearningAgentsManagerListener* VisualLoggerListener,
+	const int32 VisualLoggerAgentId,
+	const FVector VisualLoggerRotationLocation,
+	const FVector VisualLoggerLocation,
+	const FLinearColor VisualLoggerColor)
 {
 	FQuat OutRotationQuat;
-	if (!GetRotationActionAsQuat(OutRotationQuat, Object, Element, FQuat::MakeFromRotator(RelativeRotation), RotationScale, Tag))
+	if (!GetRotationActionAsQuat(OutRotationQuat, Object, Element, FQuat::MakeFromRotator(RelativeRotation), RotationScale, Tag, bVisualLoggerEnabled, VisualLoggerListener, VisualLoggerAgentId, VisualLoggerRotationLocation, VisualLoggerLocation, VisualLoggerColor))
 	{
 		OutRotation = FRotator::ZeroRotator;
 		return false;
@@ -2618,7 +2860,19 @@ bool ULearningAgentsActions::GetRotationAction(FRotator& OutRotation, const ULea
 	return true;
 }
 
-bool ULearningAgentsActions::GetRotationActionAsQuat(FQuat& OutRotation, const ULearningAgentsActionObject* Object, const FLearningAgentsActionObjectElement Element, const FQuat RelativeRotation, const float RotationScale, const FName Tag)
+bool ULearningAgentsActions::GetRotationActionAsQuat(
+	FQuat& OutRotation, 
+	const ULearningAgentsActionObject* Object, 
+	const FLearningAgentsActionObjectElement Element, 
+	const FQuat RelativeRotation, 
+	const float RotationScale, 
+	const FName Tag,
+	const bool bVisualLoggerEnabled,
+	ULearningAgentsManagerListener* VisualLoggerListener,
+	const int32 VisualLoggerAgentId,
+	const FVector VisualLoggerRotationLocation,
+	const FVector VisualLoggerLocation,
+	const FLinearColor VisualLoggerColor)
 {
 	TStaticArray<float, 3> OutValues;
 	if (!GetContinuousActionToArrayView(OutValues, Object, Element, Tag))
@@ -2628,11 +2882,49 @@ bool ULearningAgentsActions::GetRotationActionAsQuat(FQuat& OutRotation, const U
 	}
 
 	const FVector LocalRotationVector = FMath::DegreesToRadians(RotationScale) * FVector(OutValues[0], OutValues[1], OutValues[2]);
-	OutRotation = RelativeRotation * FQuat::MakeFromRotationVector(LocalRotationVector);
+	const FQuat LocalRotation = FQuat::MakeFromRotationVector(LocalRotationVector);
+	OutRotation = RelativeRotation * LocalRotation;
+
+#if UE_LEARNING_AGENTS_ENABLE_VISUAL_LOG
+	if (bVisualLoggerEnabled && VisualLoggerListener)
+	{
+		const ULearningAgentsVisualLoggerObject* VisualLoggerObject = VisualLoggerListener->GetOrAddVisualLoggerObject(Tag);
+
+		UE_LEARNING_AGENTS_VLOG_TRANSFORM(VisualLoggerObject, LogLearning, Display,
+			VisualLoggerRotationLocation,
+			LocalRotation.Rotator(),
+			VisualLoggerColor.ToFColor(true),
+			TEXT(""));
+
+		UE_LEARNING_AGENTS_VLOG_STRING(VisualLoggerObject, LogLearning, Display, VisualLoggerLocation,
+			VisualLoggerColor.ToFColor(true),
+			TEXT("Listener: %s\nTag: %s\nAgent Id: % 3i\nEncoded: [% 6.2f % 6.2f % 6.2f]\nScale: [% 6.2f]\nLocal Rotation Vector: [% 6.1f % 6.1f % 6.1f]\nLocal Rotation: [% 6.1f % 6.1f % 6.1f % 6.1f]\nRotation: [% 6.1f % 6.1f % 6.1f % 6.1f]"),
+			*VisualLoggerListener->GetName(),
+			*Tag.ToString(),
+			VisualLoggerAgentId,
+			OutValues[0], OutValues[1], OutValues[2],
+			RotationScale,
+			LocalRotationVector.X, LocalRotationVector.Y, LocalRotationVector.Z,
+			LocalRotation.X, LocalRotation.Y, LocalRotation.Z, LocalRotation.W,
+			OutRotation.X, OutRotation.Y, OutRotation.Z, OutRotation.W);
+	}
+#endif
+
 	return true;
 }
 
-bool ULearningAgentsActions::GetScaleAction(FVector& OutScale, const ULearningAgentsActionObject* Object, const FLearningAgentsActionObjectElement Element, const FVector RelativeScale, const float Scale, const FName Tag)
+bool ULearningAgentsActions::GetScaleAction(
+	FVector& OutScale, 
+	const ULearningAgentsActionObject* Object, 
+	const FLearningAgentsActionObjectElement Element, 
+	const FVector RelativeScale, 
+	const float Scale, 
+	const FName Tag,
+	const bool bVisualLoggerEnabled,
+	ULearningAgentsManagerListener* VisualLoggerListener,
+	const int32 VisualLoggerAgentId,
+	const FVector VisualLoggerLocation,
+	const FLinearColor VisualLoggerColor)
 {
 	TStaticArray<float, 3> OutValues;
 	if (!GetContinuousActionToArrayView(OutValues, Object, Element, Tag))
@@ -2643,10 +2935,42 @@ bool ULearningAgentsActions::GetScaleAction(FVector& OutScale, const ULearningAg
 
 	const FVector LocalScaleVector = UE::Learning::Agents::Action::Private::VectorExp(Scale * FVector(OutValues[0], OutValues[1], OutValues[2]));
 	OutScale = RelativeScale * LocalScaleVector;
+
+#if UE_LEARNING_AGENTS_ENABLE_VISUAL_LOG
+	if (bVisualLoggerEnabled && VisualLoggerListener)
+	{
+		const ULearningAgentsVisualLoggerObject* VisualLoggerObject = VisualLoggerListener->GetOrAddVisualLoggerObject(Tag);
+
+		UE_LEARNING_AGENTS_VLOG_STRING(VisualLoggerObject, LogLearning, Display, VisualLoggerLocation,
+			VisualLoggerColor.ToFColor(true),
+			TEXT("Listener: %s\nTag: %s\nAgent Id: % 3i\nEncoded: [% 6.2f % 6.2f % 6.2f]\nScale: [% 6.2f]\nLocal Scale: [% 6.1f % 6.1f % 6.1f]\nScale: [% 6.1f % 6.1f % 6.1f]"),
+			*VisualLoggerListener->GetName(),
+			*Tag.ToString(),
+			VisualLoggerAgentId,
+			OutValues[0], OutValues[1], OutValues[2],
+			Scale,
+			LocalScaleVector.X, LocalScaleVector.Y, LocalScaleVector.Z,
+			OutScale.X, OutScale.Y, OutScale.Z);
+	}
+#endif
+
 	return true;
 }
 
-bool ULearningAgentsActions::GetTransformAction(FTransform& OutTransform, const ULearningAgentsActionObject* Object, const FLearningAgentsActionObjectElement Element, const FTransform RelativeTransform, const float LocationScale, const float RotationScale, const float ScaleScale, const FName Tag)
+bool ULearningAgentsActions::GetTransformAction(
+	FTransform& OutTransform, 
+	const ULearningAgentsActionObject* Object, 
+	const FLearningAgentsActionObjectElement Element, 
+	const FTransform RelativeTransform, 
+	const float LocationScale, 
+	const float RotationScale, 
+	const float ScaleScale, 
+	const FName Tag,
+	const bool bVisualLoggerEnabled,
+	ULearningAgentsManagerListener* VisualLoggerListener,
+	const int32 VisualLoggerAgentId,
+	const FVector VisualLoggerLocation,
+	const FLinearColor VisualLoggerColor)
 {
 	TStaticArray<FName, 3> OutElementNames;
 	TStaticArray<FLearningAgentsActionObjectElement, 3> OutElements;
@@ -2678,36 +3002,126 @@ bool ULearningAgentsActions::GetTransformAction(FTransform& OutTransform, const
 	}
 
 	OutTransform = FTransform(OutRotation, OutLocation, OutScale);
+
+#if UE_LEARNING_AGENTS_ENABLE_VISUAL_LOG
+	if (bVisualLoggerEnabled && VisualLoggerListener)
+	{
+		const ULearningAgentsVisualLoggerObject* VisualLoggerObject = VisualLoggerListener->GetOrAddVisualLoggerObject(Tag);
+
+		UE_LEARNING_AGENTS_VLOG_TRANSFORM(VisualLoggerObject, LogLearning, Display,
+			OutLocation,
+			OutRotation,
+			VisualLoggerColor.ToFColor(true),
+			TEXT(""));
+
+		UE_LEARNING_AGENTS_VLOG_STRING(VisualLoggerObject, LogLearning, Display, VisualLoggerLocation,
+			VisualLoggerColor.ToFColor(true),
+			TEXT("Listener: %s\nTag: %s\nAgent Id: % 3i\nLocation Scale: [% 6.1f]\nRotation Scale: [% 6.1f]\nScale Scale: [% 6.1f]\nLocation: [% 6.1f % 6.1f % 6.1f]\nRotation: [% 6.1f % 6.1f % 6.1f % 6.1f]\nScale: [% 6.1f % 6.1f % 6.1f]"),
+			*VisualLoggerListener->GetName(),
+			*Tag.ToString(),
+			VisualLoggerAgentId,
+			LocationScale,
+			RotationScale,
+			ScaleScale,
+			OutLocation.X, OutLocation.Y, OutLocation.Z,
+			OutRotation.X, OutRotation.Y, OutRotation.Z, OutRotation.W,
+			OutScale.X, OutScale.Y, OutScale.Z);
+	}
+#endif
+
 	return true;
 }
 
-bool ULearningAgentsActions::GetAngleAction(float& OutAngle, const ULearningAgentsActionObject* Object, const FLearningAgentsActionObjectElement Element, const float RelativeAngle, const float AngleScale, const FName Tag)
+bool ULearningAgentsActions::GetAngleAction(
+	float& OutAngle, 
+	const ULearningAgentsActionObject* Object, 
+	const FLearningAgentsActionObjectElement Element, 
+	const float RelativeAngle, 
+	const float AngleScale, 
+	const FName Tag,
+	const bool bVisualLoggerEnabled,
+	ULearningAgentsManagerListener* VisualLoggerListener,
+	const int32 VisualLoggerAgentId,
+	const FVector VisualLoggerAngleLocation,
+	const FVector VisualLoggerLocation,
+	const FLinearColor VisualLoggerColor)
 {
-	if (GetAngleActionRadians(OutAngle, Object, Element, FMath::DegreesToRadians(RelativeAngle), FMath::DegreesToRadians(AngleScale), Tag))
-	{
-		OutAngle = FMath::RadiansToDegrees(OutAngle);
-		return true;
-	}
-	else
+	float OutValue = 0.0f;
+	if (!GetContinuousActionToArrayView(MakeArrayView(&OutValue, 1), Object, Element, Tag))
 	{
 		OutAngle = 0.0f;
 		return false;
 	}
+
+	const float LocalAngle = AngleScale * OutValue;
+	OutAngle = RelativeAngle + LocalAngle;
+
+#if UE_LEARNING_AGENTS_ENABLE_VISUAL_LOG
+	if (bVisualLoggerEnabled && VisualLoggerListener)
+	{
+		const ULearningAgentsVisualLoggerObject* VisualLoggerObject = VisualLoggerListener->GetOrAddVisualLoggerObject(Tag);
+
+		UE_LEARNING_AGENTS_VLOG_ANGLE_DEGREES(VisualLoggerObject, LogLearning, Display,
+			OutAngle,
+			0.0f,
+			VisualLoggerAngleLocation,
+			10.0f,
+			VisualLoggerColor.ToFColor(true),
+			TEXT(""));
+
+		UE_LEARNING_AGENTS_VLOG_STRING(VisualLoggerObject, LogLearning, Display, VisualLoggerLocation,
+			VisualLoggerColor.ToFColor(true),
+			TEXT("Listener: %s\nTag: %s\nAgent Id: % 3i\nEncoded: [% 6.2f]\nScale: [% 6.2f]\nLocal Angle: [% 6.1f]\nAngle: [% 6.1f]"),
+			*VisualLoggerListener->GetName(),
+			*Tag.ToString(),
+			VisualLoggerAgentId,
+			OutValue,
+			AngleScale,
+			LocalAngle,
+			OutAngle);
+	}
+#endif
+
+	return true;
 }
 
-bool ULearningAgentsActions::GetAngleActionRadians(float& OutAngle, const ULearningAgentsActionObject* Object, const FLearningAgentsActionObjectElement Element, const float RelativeAngle, const float AngleScale, const FName Tag)
+bool ULearningAgentsActions::GetAngleActionRadians(
+	float& OutAngle, 
+	const ULearningAgentsActionObject* Object, 
+	const FLearningAgentsActionObjectElement Element, 
+	const float RelativeAngle, 
+	const float AngleScale, 
+	const FName Tag,
+	const bool bVisualLoggerEnabled,
+	ULearningAgentsManagerListener* VisualLoggerListener,
+	const int32 VisualLoggerAgentId,
+	const FVector VisualLoggerAngleLocation,
+	const FVector VisualLoggerLocation,
+	const FLinearColor VisualLoggerColor)
 {
-	if (!GetContinuousActionToArrayView(MakeArrayView(&OutAngle, 1), Object, Element, Tag))
+	if (!GetAngleAction(OutAngle, Object, Element, FMath::RadiansToDegrees(RelativeAngle), FMath::RadiansToDegrees(AngleScale), Tag, bVisualLoggerEnabled, VisualLoggerListener, VisualLoggerAgentId, VisualLoggerAngleLocation, VisualLoggerLocation, VisualLoggerColor))
 	{
 		OutAngle = 0.0f;
 		return false;
 	}
 
-	OutAngle = RelativeAngle + AngleScale * OutAngle;
+	OutAngle = FMath::DegreesToRadians(OutAngle);
 	return true;
 }
 
-bool ULearningAgentsActions::GetVelocityAction(FVector& OutVelocity, const ULearningAgentsActionObject* Object, const FLearningAgentsActionObjectElement Element, const FTransform RelativeTransform, const float VelocityScale, const FName Tag)
+bool ULearningAgentsActions::GetVelocityAction(
+	FVector& OutVelocity, 
+	const ULearningAgentsActionObject* Object, 
+	const FLearningAgentsActionObjectElement Element, 
+	const FTransform RelativeTransform, 
+	const float VelocityScale, 
+	const FName Tag,
+	const bool bVisualLoggerEnabled,
+	ULearningAgentsManagerListener* VisualLoggerListener,
+	const int32 VisualLoggerAgentId,
+	const FVector VisualLoggerVelocityLocation,
+	const FVector VisualLoggerLocation,
+	const FLinearColor VisualLoggerColor)
 {
 	TStaticArray<float, 3> OutValues;
 	if (!GetContinuousActionToArrayView(OutValues, Object, Element, Tag))
@@ -2716,11 +3130,55 @@ bool ULearningAgentsActions::GetVelocityAction(FVector& OutVelocity, const ULear
 		return false;
 	}
 
-	OutVelocity = RelativeTransform.TransformVector(VelocityScale * FVector(OutValues[0], OutValues[1], OutValues[2]));
+	const FVector LocalVelocity = VelocityScale * FVector(OutValues[0], OutValues[1], OutValues[2]);
+	OutVelocity = RelativeTransform.TransformVector(LocalVelocity);
+
+#if UE_LEARNING_AGENTS_ENABLE_VISUAL_LOG
+	if (bVisualLoggerEnabled && VisualLoggerListener)
+	{
+		const ULearningAgentsVisualLoggerObject* VisualLoggerObject = VisualLoggerListener->GetOrAddVisualLoggerObject(Tag);
+
+		UE_LEARNING_AGENTS_VLOG_ARROW(VisualLoggerObject, LogLearning, Display,
+			VisualLoggerVelocityLocation,
+			VisualLoggerVelocityLocation + OutVelocity,
+			VisualLoggerColor.ToFColor(true),
+			TEXT(""));
+
+		UE_LEARNING_AGENTS_VLOG_TRANSFORM(VisualLoggerObject, LogLearning, Display,
+			RelativeTransform.GetTranslation(),
+			RelativeTransform.GetRotation(),
+			VisualLoggerColor.ToFColor(true),
+			TEXT(""));
+
+		UE_LEARNING_AGENTS_VLOG_STRING(VisualLoggerObject, LogLearning, Display, VisualLoggerLocation,
+			VisualLoggerColor.ToFColor(true),
+			TEXT("Listener: %s\nTag: %s\nAgent Id: % 3i\nEncoded: [% 6.2f % 6.2f % 6.2f]\nScale: [% 6.2f]\nLocal Velocity: [% 6.1f % 6.1f % 6.1f]\nVelocity: [% 6.1f % 6.1f % 6.1f]"),
+			*VisualLoggerListener->GetName(),
+			*Tag.ToString(),
+			VisualLoggerAgentId,
+			OutValues[0], OutValues[1], OutValues[2],
+			VelocityScale,
+			LocalVelocity.X, LocalVelocity.Y, LocalVelocity.Z,
+			OutVelocity.X, OutVelocity.Y, OutVelocity.Z);
+	}
+#endif
+
 	return true;
 }
 
-bool ULearningAgentsActions::GetDirectionAction(FVector& OutDirection, const ULearningAgentsActionObject* Object, const FLearningAgentsActionObjectElement Element, const FTransform RelativeTransform, const FName Tag)
+bool ULearningAgentsActions::GetDirectionAction(
+	FVector& OutDirection, 
+	const ULearningAgentsActionObject* Object, 
+	const FLearningAgentsActionObjectElement Element, 
+	const FTransform RelativeTransform, 
+	const FName Tag,
+	const bool bVisualLoggerEnabled,
+	ULearningAgentsManagerListener* VisualLoggerListener,
+	const int32 VisualLoggerAgentId,
+	const FVector VisualLoggerDirectionLocation,
+	const FVector VisualLoggerLocation,
+	const float VisualLoggerArrowLength,
+	const FLinearColor VisualLoggerColor)
 {
 	TStaticArray<float, 3> OutValues;
 	if (!GetContinuousActionToArrayView(OutValues, Object, Element, Tag))
@@ -2729,18 +3187,37 @@ bool ULearningAgentsActions::GetDirectionAction(FVector& OutDirection, const ULe
 		return false;
 	}
 
-	OutDirection = RelativeTransform.TransformVectorNoScale(FVector(OutValues[0], OutValues[1], OutValues[2]).GetSafeNormal(UE_SMALL_NUMBER, FVector::ForwardVector));
-	return true;
-}
+	const FVector LocalDirection = FVector(OutValues[0], OutValues[1], OutValues[2]).GetSafeNormal(UE_SMALL_NUMBER, FVector::ForwardVector);
+	OutDirection = RelativeTransform.TransformVectorNoScale(LocalDirection);
 
-bool ULearningAgentsActions::GetSpeedAction(float& OutSpeed, const ULearningAgentsActionObject* Object, const FLearningAgentsActionObjectElement Element, const float SpeedScale, const FName Tag)
-{
-	if (!GetContinuousActionToArrayView(MakeArrayView(&OutSpeed, 1), Object, Element, Tag))
+#if UE_LEARNING_AGENTS_ENABLE_VISUAL_LOG
+	if (bVisualLoggerEnabled && VisualLoggerListener)
 	{
-		OutSpeed = 0.0f;
-		return false;
+		const ULearningAgentsVisualLoggerObject* VisualLoggerObject = VisualLoggerListener->GetOrAddVisualLoggerObject(Tag);
+
+		UE_LEARNING_AGENTS_VLOG_ARROW(VisualLoggerObject, LogLearning, Display,
+			VisualLoggerDirectionLocation,
+			VisualLoggerDirectionLocation + VisualLoggerArrowLength * OutDirection,
+			VisualLoggerColor.ToFColor(true),
+			TEXT(""));
+
+		UE_LEARNING_AGENTS_VLOG_TRANSFORM(VisualLoggerObject, LogLearning, Display,
+			RelativeTransform.GetTranslation(),
+			RelativeTransform.GetRotation(),
+			VisualLoggerColor.ToFColor(true),
+			TEXT(""));
+
+		UE_LEARNING_AGENTS_VLOG_STRING(VisualLoggerObject, LogLearning, Display, VisualLoggerLocation,
+			VisualLoggerColor.ToFColor(true),
+			TEXT("Listener: %s\nTag: %s\nAgent Id: % 3i\nEncoded: [% 6.2f % 6.2f % 6.2f]\nLocal Direction: [% 6.1f % 6.1f % 6.1f]\nDirection: [% 6.1f % 6.1f % 6.1f]"),
+			*VisualLoggerListener->GetName(),
+			*Tag.ToString(),
+			VisualLoggerAgentId,
+			OutValues[0], OutValues[1], OutValues[2],
+			LocalDirection.X, LocalDirection.Y, LocalDirection.Z,
+			OutDirection.X, OutDirection.Y, OutDirection.Z);
 	}
+#endif
 
-	OutSpeed = SpeedScale * OutSpeed;
 	return true;
-}
\ No newline at end of file
+}
diff --git a/Engine/Plugins/Experimental/LearningAgents/Source/LearningAgents/Private/LearningAgentsController.cpp b/Engine/Plugins/Experimental/LearningAgents/Source/LearningAgents/Private/LearningAgentsController.cpp
--- a/Engine/Plugins/Experimental/LearningAgents/Source/LearningAgents/Private/LearningAgentsController.cpp
+++ b/Engine/Plugins/Experimental/LearningAgents/Source/LearningAgents/Private/LearningAgentsController.cpp
@@ -149,7 +149,7 @@ void ULearningAgentsController::EvaluateController()
 
 	for (int32 AgentIdx = 0; AgentIdx < Manager->GetAgentNum(); AgentIdx++)
 	{
-		if (ULearningAgentsActions::ValidateObjectMatchesSchema(
+		if (ULearningAgentsActions::ValidateActionObjectMatchesSchema(
 			Interactor->ActionSchema,
 			Interactor->ActionSchemaElement, 
 			Interactor->ActionObject, 
diff --git a/Engine/Plugins/Experimental/LearningAgents/Source/LearningAgents/Private/LearningAgentsInteractor.cpp b/Engine/Plugins/Experimental/LearningAgents/Source/LearningAgents/Private/LearningAgentsInteractor.cpp
--- a/Engine/Plugins/Experimental/LearningAgents/Source/LearningAgents/Private/LearningAgentsInteractor.cpp
+++ b/Engine/Plugins/Experimental/LearningAgents/Source/LearningAgents/Private/LearningAgentsInteractor.cpp
@@ -267,7 +267,11 @@ void ULearningAgentsInteractor::GatherObservations(const UE::Learning::FIndexSet
 
 	for (int32 AgentIdx = 0; AgentIdx < AgentSet.Num(); AgentIdx++)
 	{
-		if (ULearningAgentsObservations::ValidateObjectMatchesSchema(ObservationSchema, ObservationSchemaElement, ObservationObject, ObservationObjectElements[AgentIdx]))
+		if (ULearningAgentsObservations::ValidateObservationObjectMatchesSchema(
+			ObservationSchema, 
+			ObservationSchemaElement, 
+			ObservationObject, 
+			ObservationObjectElements[AgentIdx]))
 		{
 			UE::Learning::Observation::SetVectorFromObject(
 				ObservationVectors[AgentSet[AgentIdx]],
diff --git a/Engine/Plugins/Experimental/LearningAgents/Source/LearningAgents/Private/LearningAgentsObservations.cpp b/Engine/Plugins/Experimental/LearningAgents/Source/LearningAgents/Private/LearningAgentsObservations.cpp
--- a/Engine/Plugins/Experimental/LearningAgents/Source/LearningAgents/Private/LearningAgentsObservations.cpp
+++ b/Engine/Plugins/Experimental/LearningAgents/Source/LearningAgents/Private/LearningAgentsObservations.cpp
@@ -58,7 +58,7 @@ namespace UE::Learning::Agents::Observation::Private
 		}
 	}
 
-	static bool ValidateObjectMatchesSchema(
+	static bool ValidateObservationObjectMatchesSchema(
 		const Learning::Observation::FSchema& Schema,
 		const Learning::Observation::FSchemaElement SchemaElement,
 		const Learning::Observation::FObject& Object,
@@ -68,13 +68,13 @@ namespace UE::Learning::Agents::Observation::Private
 
 		if (!Schema.IsValid(SchemaElement))
 		{
-			UE_LOG(LogLearning, Error, TEXT("ValidateObjectMatchesSchema: Invalid Observation Schema Element."));
+			UE_LOG(LogLearning, Error, TEXT("ValidateObservationObjectMatchesSchema: Invalid Observation Schema Element."));
 			return false;
 		}
 
 		if (!Object.IsValid(ObjectElement))
 		{
-			UE_LOG(LogLearning, Error, TEXT("ValidateObjectMatchesSchema: Invalid Observation Object Element."));
+			UE_LOG(LogLearning, Error, TEXT("ValidateObservationObjectMatchesSchema: Invalid Observation Object Element."));
 			return false;
 		}
 
@@ -85,7 +85,7 @@ namespace UE::Learning::Agents::Observation::Private
 
 		if (ObservationSchemaElementTag != ObservationObjectElementTag)
 		{
-			UE_LOG(LogLearning, Warning, TEXT("ValidateObjectMatchesSchema: Observation tag does not match Schema. Expected '%s', got '%s'."),
+			UE_LOG(LogLearning, Warning, TEXT("ValidateObservationObjectMatchesSchema: Observation tag does not match Schema. Expected '%s', got '%s'."),
 				*ObservationSchemaElementTag.ToString(), *ObservationObjectElementTag.ToString());
 		}
 
@@ -96,7 +96,7 @@ namespace UE::Learning::Agents::Observation::Private
 
 		if (ObservationSchemaElementType != ObservationObjectElementType)
 		{
-			UE_LOG(LogLearning, Error, TEXT("ValidateObjectMatchesSchema: Observation '%s' type does not match Schema. Expected type '%s', got type '%s'."),
+			UE_LOG(LogLearning, Error, TEXT("ValidateObservationObjectMatchesSchema: Observation '%s' type does not match Schema. Expected type '%s', got type '%s'."),
 				*ObservationSchemaElementTag.ToString(),
 				GetObservationTypeString(ObservationSchemaElementType),
 				GetObservationTypeString(ObservationObjectElementType));
@@ -116,7 +116,7 @@ namespace UE::Learning::Agents::Observation::Private
 
 			if (SchemaElementSize != ObjectElementSize)
 			{
-				UE_LOG(LogLearning, Error, TEXT("ValidateObjectMatchesSchema: Observation '%s' size does not match Schema. Expected '%i', got '%i'."),
+				UE_LOG(LogLearning, Error, TEXT("ValidateObservationObjectMatchesSchema: Observation '%s' size does not match Schema. Expected '%i', got '%i'."),
 					*ObservationSchemaElementTag.ToString(),
 					SchemaElementSize,
 					ObjectElementSize);
@@ -135,7 +135,7 @@ namespace UE::Learning::Agents::Observation::Private
 
 			if (SchemaParameters.Elements.Num() != ObjectParameters.Elements.Num())
 			{
-				UE_LOG(LogLearning, Error, TEXT("ValidateObjectMatchesSchema: Observation '%s' number of sub-elements does not match Schema. Expected '%i', got '%i'."),
+				UE_LOG(LogLearning, Error, TEXT("ValidateObservationObjectMatchesSchema: Observation '%s' number of sub-elements does not match Schema. Expected '%i', got '%i'."),
 					*ObservationSchemaElementTag.ToString(),
 					SchemaParameters.Elements.Num(),
 					ObjectParameters.Elements.Num());
@@ -148,13 +148,13 @@ namespace UE::Learning::Agents::Observation::Private
 
 				if (ObjectElementIdx == INDEX_NONE)
 				{
-					UE_LOG(LogLearning, Error, TEXT("ValidateObjectMatchesSchema: Observation '%s' does not include '%s' observation required by Schema."),
+					UE_LOG(LogLearning, Error, TEXT("ValidateObservationObjectMatchesSchema: Observation '%s' does not include '%s' observation required by Schema."),
 						*ObservationSchemaElementTag.ToString(),
 						*SchemaParameters.ElementNames[SchemaElementIdx].ToString());
 					return false;
 				}
 
-				if (!ValidateObjectMatchesSchema(
+				if (!ValidateObservationObjectMatchesSchema(
 					Schema,
 					SchemaParameters.Elements[SchemaElementIdx],
 					Object,
@@ -177,13 +177,13 @@ namespace UE::Learning::Agents::Observation::Private
 
 			if (SchemaSubElementIdx == INDEX_NONE)
 			{
-				UE_LOG(LogLearning, Error, TEXT("ValidateObjectMatchesSchema: Observation '%s' Schema does not include '%s' observation."),
+				UE_LOG(LogLearning, Error, TEXT("ValidateObservationObjectMatchesSchema: Observation '%s' Schema does not include '%s' observation."),
 					*ObservationSchemaElementTag.ToString(),
 					*ObjectParameters.ElementName.ToString());
 				return false;
 			}
 
-			return ValidateObjectMatchesSchema(
+			return ValidateObservationObjectMatchesSchema(
 				Schema,
 				SchemaParameters.Elements[SchemaSubElementIdx],
 				Object,
@@ -197,7 +197,7 @@ namespace UE::Learning::Agents::Observation::Private
 
 			if (ObjectParameters.Elements.Num() > SchemaParameters.Elements.Num())
 			{
-				UE_LOG(LogLearning, Error, TEXT("ValidateObjectMatchesSchema: Observation '%s' too many sub-observations provided. Expected at most '%i', got '%i'."),
+				UE_LOG(LogLearning, Error, TEXT("ValidateObservationObjectMatchesSchema: Observation '%s' too many sub-observations provided. Expected at most '%i', got '%i'."),
 					*ObservationSchemaElementTag.ToString(),
 					SchemaParameters.Elements.Num(),
 					ObjectParameters.Elements.Num());
@@ -210,13 +210,13 @@ namespace UE::Learning::Agents::Observation::Private
 
 				if (SchemaSubElementIdx == INDEX_NONE)
 				{
-					UE_LOG(LogLearning, Error, TEXT("ValidateObjectMatchesSchema: Observation '%s' Schema does not include '%s' observation."),
+					UE_LOG(LogLearning, Error, TEXT("ValidateObservationObjectMatchesSchema: Observation '%s' Schema does not include '%s' observation."),
 						*ObservationSchemaElementTag.ToString(),
 						*ObjectParameters.ElementNames[ObjectSubElementIdx].ToString());
 					return false;
 				}
 
-				if (!ValidateObjectMatchesSchema(
+				if (!ValidateObservationObjectMatchesSchema(
 					Schema,
 					SchemaParameters.Elements[SchemaSubElementIdx],
 					Object,
@@ -236,7 +236,7 @@ namespace UE::Learning::Agents::Observation::Private
 
 			if (ObjectParameters.Elements.Num() != SchemaParameters.Num)
 			{
-				UE_LOG(LogLearning, Error, TEXT("ValidateObjectMatchesSchema: Observation '%s' array incorrect size. Expected '%i' elements, got '%i'."),
+				UE_LOG(LogLearning, Error, TEXT("ValidateObservationObjectMatchesSchema: Observation '%s' array incorrect size. Expected '%i' elements, got '%i'."),
 					*ObservationSchemaElementTag.ToString(),
 					SchemaParameters.Num,
 					ObjectParameters.Elements.Num());
@@ -245,7 +245,7 @@ namespace UE::Learning::Agents::Observation::Private
 
 			for (int32 ElementIdx = 0; ElementIdx < ObjectParameters.Elements.Num(); ElementIdx++)
 			{
-				if (!ValidateObjectMatchesSchema(
+				if (!ValidateObservationObjectMatchesSchema(
 					Schema,
 					SchemaParameters.Element,
 					Object,
@@ -265,7 +265,7 @@ namespace UE::Learning::Agents::Observation::Private
 
 			if (ObjectParameters.Elements.Num() > SchemaParameters.MaxNum)
 			{
-				UE_LOG(LogLearning, Error, TEXT("ValidateObjectMatchesSchema: Observation '%s' set too large. Expected at most '%i' elements, got '%i'."),
+				UE_LOG(LogLearning, Error, TEXT("ValidateObservationObjectMatchesSchema: Observation '%s' set too large. Expected at most '%i' elements, got '%i'."),
 					*ObservationSchemaElementTag.ToString(),
 					SchemaParameters.MaxNum,
 					ObjectParameters.Elements.Num());
@@ -274,7 +274,7 @@ namespace UE::Learning::Agents::Observation::Private
 
 			for (int32 ElementIdx = 0; ElementIdx < ObjectParameters.Elements.Num(); ElementIdx++)
 			{
-				if (!ValidateObjectMatchesSchema(
+				if (!ValidateObservationObjectMatchesSchema(
 					Schema,
 					SchemaParameters.Element,
 					Object,
@@ -292,7 +292,7 @@ namespace UE::Learning::Agents::Observation::Private
 			const Learning::Observation::FSchemaEncodingParameters SchemaParameters = Schema.GetEncoding(SchemaElement);
 			const Learning::Observation::FObjectEncodingParameters ObjectParameters = Object.GetEncoding(ObjectElement);
 
-			return ValidateObjectMatchesSchema(
+			return ValidateObservationObjectMatchesSchema(
 				Schema,
 				SchemaParameters.Element,
 				Object,
@@ -455,7 +455,7 @@ UEnum* ULearningAgentsObservations::FindEnumByName(const FString& Name)
 	return FindObject<UEnum>(nullptr, *Name);
 }
 
-bool ULearningAgentsObservations::ValidateObjectMatchesSchema(
+bool ULearningAgentsObservations::ValidateObservationObjectMatchesSchema(
 	const ULearningAgentsObservationSchema* Schema,
 	const FLearningAgentsObservationSchemaElement SchemaElement,
 	const ULearningAgentsObservationObject* Object,
@@ -463,17 +463,17 @@ bool ULearningAgentsObservations::ValidateObjectMatchesSchema(
 {
 	if (!Schema)
 	{
-		UE_LOG(LogLearning, Error, TEXT("ValidateObjectMatchesSchema: Schema is nullptr."));
+		UE_LOG(LogLearning, Error, TEXT("ValidateObservationObjectMatchesSchema: Schema is nullptr."));
 		return false;
 	}
 
 	if (!Object)
 	{
-		UE_LOG(LogLearning, Error, TEXT("ValidateObjectMatchesSchema: Object is nullptr."));
+		UE_LOG(LogLearning, Error, TEXT("ValidateObservationObjectMatchesSchema: Object is nullptr."));
 		return false;
 	}
 
-	return UE::Learning::Agents::Observation::Private::ValidateObjectMatchesSchema(
+	return UE::Learning::Agents::Observation::Private::ValidateObservationObjectMatchesSchema(
 		Schema->ObservationSchema,
 		SchemaElement.SchemaElement,
 		Object->ObservationObject,
@@ -2094,16 +2094,16 @@ FLearningAgentsObservationObjectElement ULearningAgentsObservations::MakeTransfo
 		const FRotator LocalRotation = LocalTransform.Rotator();
 		const FVector LocalScale = LocalTransform.GetScale3D();
 
-		UE_LEARNING_AGENTS_VLOG_TRANSFORM(VisualLoggerObject, LogLearning, Display,
-			LocalLocation,
-			LocalRotation,
-			VisualLoggerColor.ToFColor(true),
-			TEXT(""));
-
 		const FVector Location = Transform.GetLocation();
 		const FRotator Rotation = Transform.Rotator();
 		const FVector Scale = Transform.GetScale3D();
 
+		UE_LEARNING_AGENTS_VLOG_TRANSFORM(VisualLoggerObject, LogLearning, Display,
+			Location,
+			Rotation,
+			VisualLoggerColor.ToFColor(true),
+			TEXT(""));
+
 		UE_LEARNING_AGENTS_VLOG_STRING(VisualLoggerObject, LogLearning, Display, VisualLoggerLocation,
 			VisualLoggerColor.ToFColor(true),
 			TEXT("Listener: %s\nTag: %s\nAgent Id: % 3i\nLocation: [% 6.1f % 6.1f % 6.1f]\nLocal Location: [% 6.1f % 6.1f % 6.1f]\nRotation: [% 6.1f % 6.1f % 6.1f]\nLocal Rotation: [% 6.1f % 6.1f % 6.1f]\nScale: [% 6.1f % 6.1f % 6.1f]\nLocal Scale: [% 6.1f % 6.1f % 6.1f]\nLocation Scale: [% 6.1f]"),
diff --git a/Engine/Plugins/Experimental/LearningAgents/Source/LearningAgents/Public/LearningAgentsActions.h b/Engine/Plugins/Experimental/LearningAgents/Source/LearningAgents/Public/LearningAgentsActions.h
--- a/Engine/Plugins/Experimental/LearningAgents/Source/LearningAgents/Public/LearningAgentsActions.h
+++ b/Engine/Plugins/Experimental/LearningAgents/Source/LearningAgents/Public/LearningAgentsActions.h
@@ -105,7 +105,7 @@ class LEARNINGAGENTS_API ULearningAgentsActions : public UBlueprintFunctionLibra
 public:
 
 	/**
-	 * Validates that the given object matches the schema. Will log errors on objects that don't match.
+	 * Validates that the given action object matches the schema. Will log errors on objects that don't match.
 	 *
 	 * @param Schema				Action Schema
 	 * @param SchemaElement			Action Schema Element
@@ -114,345 +114,1292 @@ public:
 	 * @returns						true if the object matches the schema
 	 */
 	UFUNCTION(BlueprintPure = false, Category = "LearningAgents")
-	static bool ValidateObjectMatchesSchema(
+	static bool ValidateActionObjectMatchesSchema(
 		const ULearningAgentsActionSchema* Schema,
 		const FLearningAgentsActionSchemaElement SchemaElement,
 		const ULearningAgentsActionObject* Object,
 		const FLearningAgentsActionObjectElement ObjectElement);
 
+	/**
+	 * Logs an Action Object Element. Useful for debugging.
+	 *
+	 * @param Object				Action Object
+	 * @param ObjectElement			Action Object Element
+	 */
 	UFUNCTION(BlueprintPure = false, Category = "LearningAgents")
-	static void LogAction(
-		const ULearningAgentsActionObject* Object, 
-		const FLearningAgentsActionObjectElement Element);
+	static void LogAction(const ULearningAgentsActionObject* Object, const FLearningAgentsActionObjectElement Element);
 
 public:
 
+	/**
+	 * Specifies a new null action. This represents an empty action and can be useful when an action is needed which does nothing.
+	 *
+	 * @param Schema The Action Schema
+	 * @param Tag The tag of this new action. Used during action object validation and debugging.
+	 * @return The newly created action schema element.
+	 */
 	UFUNCTION(BlueprintPure, Category = "LearningAgents", meta = (AdvancedDisplay = 1))
-	static FLearningAgentsActionSchemaElement SpecifyNullAction(ULearningAgentsActionSchema* Schema, const FName Tag = TEXT("Null"));
+	static FLearningAgentsActionSchemaElement SpecifyNullAction(ULearningAgentsActionSchema* Schema, const FName Tag = TEXT("NullAction"));
 
+	/**
+	 * Specifies a new continuous action. This represents an action made up of several float values sampled from a Gaussian distribution.
+	 *
+	 * @param Schema The Action Schema
+	 * @param Size The number of float values in the action.
+	 * @param Tag The tag of this new action. Used during action object validation and debugging.
+	 * @return The newly created action schema element.
+	 */
 	UFUNCTION(BlueprintPure, Category = "LearningAgents", meta = (AdvancedDisplay = 2))
-	static FLearningAgentsActionSchemaElement SpecifyContinuousAction(ULearningAgentsActionSchema* Schema, const int32 Size, const FName Tag = TEXT("Continuous"));
+	static FLearningAgentsActionSchemaElement SpecifyContinuousAction(ULearningAgentsActionSchema* Schema, const int32 Size, const FName Tag = TEXT("ContinuousAction"));
 
+	/**
+	 * Specifies a new exclusive discrete action. This represents an action which is an exclusive choice from a number of discrete options, sampled 
+	 * from a Categorical distribution.
+	 *
+	 * @param Schema The Action Schema
+	 * @param Size The number of discrete options in the action.
+	 * @param PriorProbabilities The prior probabilities of each option. Can be left empty to use a uniform distribution over options. Should sum to one.
+	 * @param Tag The tag of this new action. Used during action object validation and debugging.
+	 * @return The newly created action schema element.
+	 */
 	UFUNCTION(BlueprintPure, Category = "LearningAgents", meta = (AdvancedDisplay = 2, AutoCreateRefTerm = "PriorProbabilities"))
-	static FLearningAgentsActionSchemaElement SpecifyExclusiveDiscreteAction(ULearningAgentsActionSchema* Schema, const int32 Size, const TArray<float>& PriorProbabilities, const FName Tag = TEXT("DiscreteExclusive"));
-	static FLearningAgentsActionSchemaElement SpecifyExclusiveDiscreteActionFromArrayView(ULearningAgentsActionSchema* Schema, const int32 Size, const TArrayView<const float> PriorProbabilities = {}, const FName Tag = TEXT("DiscreteExclusive"));
+	static FLearningAgentsActionSchemaElement SpecifyExclusiveDiscreteAction(ULearningAgentsActionSchema* Schema, const int32 Size, const TArray<float>& PriorProbabilities, const FName Tag = TEXT("DiscreteExclusiveAction"));
+	
+	/**
+	 * Specifies a new exclusive discrete action. This represents an action which is an exclusive choice from a number of discrete options, sampled
+	 * from a Categorical distribution.
+	 *
+	 * @param Schema The Action Schema
+	 * @param Size The number of discrete options in the action.
+	 * @param PriorProbabilities The prior probabilities of each option. Can be left empty to use a uniform distribution over options. Should sum to one.
+	 * @param Tag The tag of this new action. Used during action object validation and debugging.
+	 * @return The newly created action schema element.
+	 */
+	static FLearningAgentsActionSchemaElement SpecifyExclusiveDiscreteActionFromArrayView(ULearningAgentsActionSchema* Schema, const int32 Size, const TArrayView<const float> PriorProbabilities = {}, const FName Tag = TEXT("DiscreteExclusiveAction"));
 
+	/**
+	 * Specifies a new inclusive discrete action. This represents an action which is an inclusive choice from a number of discrete options, sampled
+	 * from a Bernoulli distribution.
+	 *
+	 * @param Schema The Action Schema
+	 * @param Size The number of discrete options in the action.
+	 * @param PriorProbabilities The prior probabilities of each option. Can be left empty to use a probability of 0.5 for each option.
+	 * @param Tag The tag of this new action. Used during action object validation and debugging.
+	 * @return The newly created action schema element.
+	 */
 	UFUNCTION(BlueprintPure, Category = "LearningAgents", meta = (AdvancedDisplay = 2, AutoCreateRefTerm = "PriorProbabilities"))
-	static FLearningAgentsActionSchemaElement SpecifyInclusiveDiscreteAction(ULearningAgentsActionSchema* Schema, const int32 Size, const TArray<float>& PriorProbabilities, const FName Tag = TEXT("DiscreteInclusive"));
-	static FLearningAgentsActionSchemaElement SpecifyInclusiveDiscreteActionFromArrayView(ULearningAgentsActionSchema* Schema, const int32 Size, const TArrayView<const float> PriorProbabilities = {}, const FName Tag = TEXT("DiscreteInclusive"));
+	static FLearningAgentsActionSchemaElement SpecifyInclusiveDiscreteAction(ULearningAgentsActionSchema* Schema, const int32 Size, const TArray<float>& PriorProbabilities, const FName Tag = TEXT("DiscreteInclusiveAction"));
+	
+	/**
+	 * Specifies a new inclusive discrete action. This represents an action which is an inclusive choice from a number of discrete options, sampled
+	 * from a Bernoulli distribution.
+	 *
+	 * @param Schema The Action Schema
+	 * @param Size The number of discrete options in the action.
+	 * @param PriorProbabilities The prior probabilities of each option. Can be left empty to use a probability of 0.5 for each option.
+	 * @param Tag The tag of this new action. Used during action object validation and debugging.
+	 * @return The newly created action schema element.
+	 */
+	static FLearningAgentsActionSchemaElement SpecifyInclusiveDiscreteActionFromArrayView(ULearningAgentsActionSchema* Schema, const int32 Size, const TArrayView<const float> PriorProbabilities = {}, const FName Tag = TEXT("DiscreteInclusiveAction"));
 
+	/**
+	 * Specifies a new struct action. This represents an action which is made up of a number of named sub-actions.
+	 *
+	 * @param Schema The Action Schema
+	 * @param Elements The sub-actions.
+	 * @param Tag The tag of this new action. Used during action object validation and debugging.
+	 * @return The newly created action schema element.
+	 */
 	UFUNCTION(BlueprintPure, Category = "LearningAgents", meta = (AdvancedDisplay = 2))
-	static FLearningAgentsActionSchemaElement SpecifyStructAction(ULearningAgentsActionSchema* Schema, const TMap<FName, FLearningAgentsActionSchemaElement>& Elements, const FName Tag = TEXT("Struct"));
+	static FLearningAgentsActionSchemaElement SpecifyStructAction(ULearningAgentsActionSchema* Schema, const TMap<FName, FLearningAgentsActionSchemaElement>& Elements, const FName Tag = TEXT("StructAction"));
 
+	/**
+	 * Specifies a new struct action. This represents an action which is made up of a number of named sub-actions.
+	 *
+	 * @param Schema The Action Schema
+	 * @param ElementNames The names of the sub-actions.
+	 * @param Elements The corresponding sub-actions. Must be the same size as ElementNames.
+	 * @param Tag The tag of this new action. Used during action object validation and debugging.
+	 * @return The newly created action schema element.
+	 */
 	UFUNCTION(BlueprintPure, Category = "LearningAgents", meta = (AdvancedDisplay = 3))
-	static FLearningAgentsActionSchemaElement SpecifyStructActionFromArrays(ULearningAgentsActionSchema* Schema, const TArray<FName>& ElementNames, const TArray<FLearningAgentsActionSchemaElement>& Elements, const FName Tag = TEXT("Struct"));
-	static FLearningAgentsActionSchemaElement SpecifyStructActionFromArrayViews(ULearningAgentsActionSchema* Schema, const TArrayView<const FName> ElementNames, const TArrayView<const FLearningAgentsActionSchemaElement> Elements, const FName Tag = TEXT("Struct"));
+	static FLearningAgentsActionSchemaElement SpecifyStructActionFromArrays(ULearningAgentsActionSchema* Schema, const TArray<FName>& ElementNames, const TArray<FLearningAgentsActionSchemaElement>& Elements, const FName Tag = TEXT("StructAction"));
+	
+	/**
+	 * Specifies a new struct action. This represents an action which is made up of a number of named sub-actions.
+	 *
+	 * @param Schema The Action Schema
+	 * @param ElementNames The names of the sub-actions.
+	 * @param Elements The corresponding sub-actions. Must be the same size as ElementNames.
+	 * @param Tag The tag of this new action. Used during action object validation and debugging.
+	 * @return The newly created action schema element.
+	 */
+	static FLearningAgentsActionSchemaElement SpecifyStructActionFromArrayViews(ULearningAgentsActionSchema* Schema, const TArrayView<const FName> ElementNames, const TArrayView<const FLearningAgentsActionSchemaElement> Elements, const FName Tag = TEXT("StructAction"));
 
+	/**
+	 * Specifies a new exclusive union action. This represents an action which is an exclusive choice from a number of named sub-actions, sampled
+	 * from a Categorical distribution.
+	 *
+	 * @param Schema The Action Schema
+	 * @param Elements The sub-actions.
+	 * @param PriorProbabilities The prior probabilities of each option. Can be left empty to use a uniform distribution over options. Should sum to one.
+	 * @param Tag The tag of this new action. Used during action object validation and debugging.
+	 * @return The newly created action schema element.
+	 */
 	UFUNCTION(BlueprintPure, Category = "LearningAgents", meta = (AdvancedDisplay = 2, AutoCreateRefTerm = "PriorProbabilities"))
-	static FLearningAgentsActionSchemaElement SpecifyExclusiveUnionAction(ULearningAgentsActionSchema* Schema, const TMap<FName, FLearningAgentsActionSchemaElement>& Elements, const TMap<FName, float>& PriorProbabilities, const FName Tag = TEXT("ExclusiveUnion"));
+	static FLearningAgentsActionSchemaElement SpecifyExclusiveUnionAction(ULearningAgentsActionSchema* Schema, const TMap<FName, FLearningAgentsActionSchemaElement>& Elements, const TMap<FName, float>& PriorProbabilities, const FName Tag = TEXT("ExclusiveUnionAction"));
 
+	/**
+	 * Specifies a new exclusive union action. This represents an action which is an exclusive choice from a number of named sub-actions, sampled
+	 * from a Categorical distribution.
+	 *
+	 * @param Schema The Action Schema
+	 * @param ElementNames The names of the sub-actions.
+	 * @param Elements The corresponding sub-actions. Must be the same size as ElementNames.
+	 * @param PriorProbabilities The prior probabilities of each option. Can be left empty to use a uniform distribution over options. Should sum to one.
+	 * @param Tag The tag of this new action. Used during action object validation and debugging.
+	 * @return The newly created action schema element.
+	 */
 	UFUNCTION(BlueprintPure, Category = "LearningAgents", meta = (AdvancedDisplay = 3, AutoCreateRefTerm = "PriorProbabilities"))
-	static FLearningAgentsActionSchemaElement SpecifyExclusiveUnionActionFromArrays(ULearningAgentsActionSchema* Schema, const TArray<FName>& ElementNames, const TArray<FLearningAgentsActionSchemaElement>& Elements, const TArray<float>& PriorProbabilities, const FName Tag = TEXT("ExclusiveUnion"));
-	static FLearningAgentsActionSchemaElement SpecifyExclusiveUnionActionFromArrayViews(ULearningAgentsActionSchema* Schema, const TArrayView<const FName> ElementNames, const TArrayView<const FLearningAgentsActionSchemaElement> Elements, const TArrayView<const float> PriorProbabilities = {}, const FName Tag = TEXT("ExclusiveUnion"));
+	static FLearningAgentsActionSchemaElement SpecifyExclusiveUnionActionFromArrays(ULearningAgentsActionSchema* Schema, const TArray<FName>& ElementNames, const TArray<FLearningAgentsActionSchemaElement>& Elements, const TArray<float>& PriorProbabilities, const FName Tag = TEXT("ExclusiveUnionAction"));
+	
+	/**
+	 * Specifies a new exclusive union action. This represents an action which is an exclusive choice from a number of named sub-actions, sampled
+	 * from a Categorical distribution.
+	 *
+	 * @param Schema The Action Schema
+	 * @param ElementNames The names of the sub-actions.
+	 * @param Elements The corresponding sub-actions. Must be the same size as ElementNames.
+	 * @param PriorProbabilities The prior probabilities of each option. Can be left empty to use a uniform distribution over options. Should sum to one.
+	 * @param Tag The tag of this new action. Used during action object validation and debugging.
+	 * @return The newly created action schema element.
+	 */
+	static FLearningAgentsActionSchemaElement SpecifyExclusiveUnionActionFromArrayViews(ULearningAgentsActionSchema* Schema, const TArrayView<const FName> ElementNames, const TArrayView<const FLearningAgentsActionSchemaElement> Elements, const TArrayView<const float> PriorProbabilities = {}, const FName Tag = TEXT("ExclusiveUnionAction"));
 
+	/**
+	 * Specifies a new inclusive union action. This represents an action which is an inclusive choice from a number of named sub-actions, sampled
+	 * from a Bernoulli distribution.
+	 *
+	 * @param Schema The Action Schema
+	 * @param Elements The sub-actions.
+	 * @param PriorProbabilities The prior probabilities of each option. Can be left empty to use a probability of 0.5 for each option.
+	 * @param Tag The tag of this new action. Used during action object validation and debugging.
+	 * @return The newly created action schema element.
+	 */
 	UFUNCTION(BlueprintPure, Category = "LearningAgents", meta = (AdvancedDisplay = 2, AutoCreateRefTerm = "PriorProbabilities"))
-	static FLearningAgentsActionSchemaElement SpecifyInclusiveUnionAction(ULearningAgentsActionSchema* Schema, const TMap<FName, FLearningAgentsActionSchemaElement>& Elements, const TMap<FName, float>& PriorProbabilities, const FName Tag = TEXT("InclusiveUnion"));
+	static FLearningAgentsActionSchemaElement SpecifyInclusiveUnionAction(ULearningAgentsActionSchema* Schema, const TMap<FName, FLearningAgentsActionSchemaElement>& Elements, const TMap<FName, float>& PriorProbabilities, const FName Tag = TEXT("InclusiveUnionAction"));
 
+	/**
+	 * Specifies a new inclusive union action. This represents an action which is an inclusive choice from a number of named sub-actions, sampled
+	 * from a Bernoulli distribution.
+	 *
+	 * @param Schema The Action Schema
+	 * @param ElementNames The names of the sub-actions.
+	 * @param Elements The corresponding sub-actions. Must be the same size as ElementNames.
+	 * @param PriorProbabilities The prior probabilities of each option. Can be left empty to use a probability of 0.5 for each option.
+	 * @param Tag The tag of this new action. Used during action object validation and debugging.
+	 * @return The newly created action schema element.
+	 */
 	UFUNCTION(BlueprintPure, Category = "LearningAgents", meta = (AdvancedDisplay = 3, AutoCreateRefTerm = "PriorProbabilities"))
-	static FLearningAgentsActionSchemaElement SpecifyInclusiveUnionActionFromArrays(ULearningAgentsActionSchema* Schema, const TArray<FName> ElementNames, const TArray<FLearningAgentsActionSchemaElement>& Elements, const TArray<float>& PriorProbabilities, const FName Tag = TEXT("InclusiveUnion"));
-	static FLearningAgentsActionSchemaElement SpecifyInclusiveUnionActionFromArrayViews(ULearningAgentsActionSchema* Schema, const TArrayView<const FName> ElementNames, const TArrayView<const FLearningAgentsActionSchemaElement> Elements, const TArrayView<const float> PriorProbabilities = {}, const FName Tag = TEXT("InclusiveUnion"));
+	static FLearningAgentsActionSchemaElement SpecifyInclusiveUnionActionFromArrays(ULearningAgentsActionSchema* Schema, const TArray<FName> ElementNames, const TArray<FLearningAgentsActionSchemaElement>& Elements, const TArray<float>& PriorProbabilities, const FName Tag = TEXT("InclusiveUnionAction"));
+	
+	/**
+	 * Specifies a new inclusive union action. This represents an action which is an inclusive choice from a number of named sub-actions, sampled
+	 * from a Bernoulli distribution.
+	 *
+	 * @param Schema The Action Schema
+	 * @param ElementNames The names of the sub-actions.
+	 * @param Elements The corresponding sub-actions. Must be the same size as ElementNames.
+	 * @param PriorProbabilities The prior probabilities of each option. Can be left empty to use a probability of 0.5 for each option.
+	 * @param Tag The tag of this new action. Used during action object validation and debugging.
+	 * @return The newly created action schema element.
+	 */
+	static FLearningAgentsActionSchemaElement SpecifyInclusiveUnionActionFromArrayViews(ULearningAgentsActionSchema* Schema, const TArrayView<const FName> ElementNames, const TArrayView<const FLearningAgentsActionSchemaElement> Elements, const TArrayView<const float> PriorProbabilities = {}, const FName Tag = TEXT("InclusiveUnionAction"));
 
+	/**
+	 * Specifies a new static array action. This represents an action which is a fixed sized array of some sub-action.
+	 *
+	 * @param Schema The Action Schema
+	 * @param Element The sub-action.
+	 * @param Num The number of elements in the array.
+	 * @param Tag The tag of this new action. Used during action object validation and debugging.
+	 * @return The newly created action schema element.
+	 */
 	UFUNCTION(BlueprintPure, Category = "LearningAgents", meta = (AdvancedDisplay = 3))
-	static FLearningAgentsActionSchemaElement SpecifyStaticArrayAction(ULearningAgentsActionSchema* Schema, const FLearningAgentsActionSchemaElement Element, const int32 Num, const FName Tag = TEXT("StaticArray"));
+	static FLearningAgentsActionSchemaElement SpecifyStaticArrayAction(ULearningAgentsActionSchema* Schema, const FLearningAgentsActionSchemaElement Element, const int32 Num, const FName Tag = TEXT("StaticArrayAction"));
 
+	/**
+	 * Specifies a new pair action. This represents an action which is made up of a key and value sub-actions.
+	 *
+	 * @param Schema The Action Schema
+	 * @param Key The key sub-action.
+	 * @param Value The value sub-action.
+	 * @param Tag The tag of this new action. Used during action object validation and debugging.
+	 * @return The newly created action schema element.
+	 */
 	UFUNCTION(BlueprintPure, Category = "LearningAgents", meta = (AdvancedDisplay = 3))
-	static FLearningAgentsActionSchemaElement SpecifyPairAction(ULearningAgentsActionSchema* Schema, const FLearningAgentsActionSchemaElement Key, const FLearningAgentsActionSchemaElement Value, const FName Tag = TEXT("Pair"));
+	static FLearningAgentsActionSchemaElement SpecifyPairAction(ULearningAgentsActionSchema* Schema, const FLearningAgentsActionSchemaElement Key, const FLearningAgentsActionSchemaElement Value, const FName Tag = TEXT("PairAction"));
 
+	/**
+	 * Specifies a new enum action. This represents an action which is an exclusive choice from entries of an Enum, sampled from a Categorical distribution.
+	 *
+	 * @param Schema The Action Schema
+	 * @param Enum The Enum type.
+	 * @param PriorProbabilities The prior probabilities of each enum element. Can be left empty to use a uniform distribution over elements. Should sum to one.
+	 * @param Tag The tag of this new action. Used during action object validation and debugging.
+	 * @return The newly created action schema element.
+	 */
 	UFUNCTION(BlueprintPure, Category = "LearningAgents", meta = (AdvancedDisplay = 2, AutoCreateRefTerm = "PriorProbabilities"))
-	static FLearningAgentsActionSchemaElement SpecifyEnumAction(ULearningAgentsActionSchema* Schema, const UEnum* Enum, const TMap<uint8, float>& PriorProbabilities, const FName Tag = TEXT("Enum"));
+	static FLearningAgentsActionSchemaElement SpecifyEnumAction(ULearningAgentsActionSchema* Schema, const UEnum* Enum, const TMap<uint8, float>& PriorProbabilities, const FName Tag = TEXT("EnumAction"));
 
+	/**
+	 * Specifies a new enum action. This represents an action which is an exclusive choice from entries of an Enum, sampled from a Categorical distribution.
+	 *
+	 * @param Schema The Action Schema
+	 * @param Enum The Enum type.
+	 * @param PriorProbabilities The prior probabilities of each enum element. Can be left empty to use a uniform distribution over elements. Should sum to one.
+	 * @param Tag The tag of this new action. Used during action object validation and debugging.
+	 * @return The newly created action schema element.
+	 */
 	UFUNCTION(BlueprintPure, Category = "LearningAgents", meta = (AdvancedDisplay = 2, AutoCreateRefTerm = "PriorProbabilities"))
-	static FLearningAgentsActionSchemaElement SpecifyEnumActionFromArray(ULearningAgentsActionSchema* Schema, const UEnum* Enum, const TArray<float>& PriorProbabilities, const FName Tag = TEXT("Enum"));
-	static FLearningAgentsActionSchemaElement SpecifyEnumActionFromArrayView(ULearningAgentsActionSchema* Schema, const UEnum* Enum, const TArrayView<const float> PriorProbabilities = {}, const FName Tag = TEXT("Enum"));
+	static FLearningAgentsActionSchemaElement SpecifyEnumActionFromArray(ULearningAgentsActionSchema* Schema, const UEnum* Enum, const TArray<float>& PriorProbabilities, const FName Tag = TEXT("EnumAction"));
+	
+	/**
+	 * Specifies a new enum action. This represents an action which is an exclusive choice from entries of an Enum, sampled from a Categorical distribution.
+	 *
+	 * @param Schema The Action Schema
+	 * @param Enum The Enum type.
+	 * @param PriorProbabilities The prior probabilities of each enum element. Can be left empty to use a uniform distribution over elements. Should sum to one.
+	 * @param Tag The tag of this new action. Used during action object validation and debugging.
+	 * @return The newly created action schema element.
+	 */
+	static FLearningAgentsActionSchemaElement SpecifyEnumActionFromArrayView(ULearningAgentsActionSchema* Schema, const UEnum* Enum, const TArrayView<const float> PriorProbabilities = {}, const FName Tag = TEXT("EnumAction"));
 
+	/**
+	 * Specifies a new bitmask action. This represents an action which is an inclusive choice from entries of an Enum, sampled from a Bernoulli 
+	 * distribution.
+	 *
+	 * @param Schema The Action Schema
+	 * @param Enum The Enum type.
+	 * @param PriorProbabilities The prior probabilities of each enum element. Can be left empty to use a probability of 0.5 for each element.
+	 * @param Tag The tag of this new action. Used during action object validation and debugging.
+	 * @return The newly created action schema element.
+	 */
 	UFUNCTION(BlueprintPure, Category = "LearningAgents", meta = (AdvancedDisplay = 2, AutoCreateRefTerm = "PriorProbabilities"))
-	static FLearningAgentsActionSchemaElement SpecifyBitmaskAction(ULearningAgentsActionSchema* Schema, const UEnum* Enum, const TMap<uint8, float>& PriorProbabilities, const FName Tag = TEXT("Bitmask"));
+	static FLearningAgentsActionSchemaElement SpecifyBitmaskAction(ULearningAgentsActionSchema* Schema, const UEnum* Enum, const TMap<uint8, float>& PriorProbabilities, const FName Tag = TEXT("BitmaskAction"));
 
+	/**
+	 * Specifies a new bitmask action. This represents an action which is an inclusive choice from entries of an Enum, sampled from a Bernoulli
+	 * distribution.
+	 *
+	 * @param Schema The Action Schema
+	 * @param Enum The Enum type.
+	 * @param PriorProbabilities The prior probabilities of each enum element. Can be left empty to use a probability of 0.5 for each element.
+	 * @param Tag The tag of this new action. Used during action object validation and debugging.
+	 * @return The newly created action schema element.
+	 */
 	UFUNCTION(BlueprintPure, Category = "LearningAgents", meta = (AdvancedDisplay = 2, AutoCreateRefTerm = "PriorProbabilities"))
-	static FLearningAgentsActionSchemaElement SpecifyBitmaskActionFromArray(ULearningAgentsActionSchema* Schema, const UEnum* Enum, const TArray<float>& PriorProbabilities, const FName Tag = TEXT("Bitmask"));
-	static FLearningAgentsActionSchemaElement SpecifyBitmaskActionFromArrayView(ULearningAgentsActionSchema* Schema, const UEnum* Enum, const TArrayView<const float> PriorProbabilities = {}, const FName Tag = TEXT("Bitmask"));
+	static FLearningAgentsActionSchemaElement SpecifyBitmaskActionFromArray(ULearningAgentsActionSchema* Schema, const UEnum* Enum, const TArray<float>& PriorProbabilities, const FName Tag = TEXT("BitmaskAction"));
+	
+	/**
+	 * Specifies a new bitmask action. This represents an action which is an inclusive choice from entries of an Enum, sampled from a Bernoulli
+	 * distribution.
+	 *
+	 * @param Schema The Action Schema
+	 * @param Enum The Enum type.
+	 * @param PriorProbabilities The prior probabilities of each enum element. Can be left empty to use a probability of 0.5 for each element.
+	 * @param Tag The tag of this new action. Used during action object validation and debugging.
+	 * @return The newly created action schema element.
+	 */
+	static FLearningAgentsActionSchemaElement SpecifyBitmaskActionFromArrayView(ULearningAgentsActionSchema* Schema, const UEnum* Enum, const TArrayView<const float> PriorProbabilities = {}, const FName Tag = TEXT("BitmaskAction"));
 
+	/**
+	 * Specifies a new optional action. This represents an action which may or may not be generated.
+	 *
+	 * @param Schema The Action Schema
+	 * @param Element The sub-action.
+	 * @param PriorProbabilities The prior probability of sampling this action.
+	 * @param Tag The tag of this new action. Used during action object validation and debugging.
+	 * @return The newly created action schema element.
+	 */
 	UFUNCTION(BlueprintPure, Category = "LearningAgents", meta = (AdvancedDisplay = 2))
-	static FLearningAgentsActionSchemaElement SpecifyOptionalAction(ULearningAgentsActionSchema* Schema, const FLearningAgentsActionSchemaElement Element, const float PriorProbability = 0.5f, const FName Tag = TEXT("Optional"));
+	static FLearningAgentsActionSchemaElement SpecifyOptionalAction(ULearningAgentsActionSchema* Schema, const FLearningAgentsActionSchemaElement Element, const float PriorProbability = 0.5f, const FName Tag = TEXT("OptionalAction"));
 
+	/**
+	 * Specifies a new either action. This represents an action which is either action A or action B.
+	 *
+	 * @param Schema The Action Schema
+	 * @param A The sub-action A.
+	 * @param B The sub-action B.
+	 * @param PriorProbabilityOfA The prior probability of sampling action A over action B.
+	 * @param Tag The tag of this new action. Used during action object validation and debugging.
+	 * @return The newly created action schema element.
+	 */
 	UFUNCTION(BlueprintPure, Category = "LearningAgents", meta = (AdvancedDisplay = 3))
-	static FLearningAgentsActionSchemaElement SpecifyEitherAction(ULearningAgentsActionSchema* Schema, const FLearningAgentsActionSchemaElement A, const FLearningAgentsActionSchemaElement B, const float PriorProbabilityOfA = 0.5f, const FName Tag = TEXT("Either"));
+	static FLearningAgentsActionSchemaElement SpecifyEitherAction(ULearningAgentsActionSchema* Schema, const FLearningAgentsActionSchemaElement A, const FLearningAgentsActionSchemaElement B, const float PriorProbabilityOfA = 0.5f, const FName Tag = TEXT("EitherAction"));
 
+	/**
+	 * Specifies a new encoding action. This represents an action which will be a decoding of another sub-action using a small neural network.
+	 *
+	 * @param Schema The Action Schema
+	 * @param Element The sub-action.
+	 * @param EncodingSize The encoding size used to decode this sub-action.
+	 * @param LayerNum The number of layers used to decode this sub-action.
+	 * @param ActivationFunction The activation function used to decode this sub-action.
+	 * @param Tag The tag of this new action. Used during action object validation and debugging.
+	 * @return The newly created action schema element.
+	 */
 	UFUNCTION(BlueprintPure, Category = "LearningAgents", meta = (AdvancedDisplay = 2))
-	static FLearningAgentsActionSchemaElement SpecifyEncodingAction(ULearningAgentsActionSchema* Schema, const FLearningAgentsActionSchemaElement Element, const int32 EncodingSize = 128, const int32 LayerNum = 1, const ELearningAgentsActivationFunction ActivationFunction = ELearningAgentsActivationFunction::ELU, const FName Tag = TEXT("Encoding"));
-
-	UFUNCTION(BlueprintPure, Category = "LearningAgents", meta = (AdvancedDisplay = 1))
-	static FLearningAgentsActionSchemaElement SpecifyBoolAction(ULearningAgentsActionSchema* Schema, const float PriorProbability = 0.5f, const FName Tag = TEXT("Bool"));
+	static FLearningAgentsActionSchemaElement SpecifyEncodingAction(ULearningAgentsActionSchema* Schema, const FLearningAgentsActionSchemaElement Element, const int32 EncodingSize = 128, const int32 LayerNum = 1, const ELearningAgentsActivationFunction ActivationFunction = ELearningAgentsActivationFunction::ELU, const FName Tag = TEXT("EncodingAction"));
 
+	/**
+	 * Specifies a new bool action. This represents an action which is either true or false.
+	 *
+	 * @param Schema The Action Schema
+	 * @param PriorProbability The prior probability of this action being true.
+	 * @param Tag The tag of this new action. Used during action object validation and debugging.
+	 * @return The newly created action schema element.
+	 */
 	UFUNCTION(BlueprintPure, Category = "LearningAgents", meta = (AdvancedDisplay = 1))
-	static FLearningAgentsActionSchemaElement SpecifyFloatAction(ULearningAgentsActionSchema* Schema, const FName Tag = TEXT("Float"));
+	static FLearningAgentsActionSchemaElement SpecifyBoolAction(ULearningAgentsActionSchema* Schema, const float PriorProbability = 0.5f, const FName Tag = TEXT("BoolAction"));
 
+	/**
+	 * Specifies a new float action. This represents an action which is a single float sampled from a Gaussian distribution. It can be used as a 
+	 * catch-all for situations where a type-specific action does not exist.
+	 *
+	 * @param Schema The Action Schema
+	 * @param Tag The tag of this new action. Used during action object validation and debugging.
+	 * @return The newly created action schema element.
+	 */
 	UFUNCTION(BlueprintPure, Category = "LearningAgents", meta = (AdvancedDisplay = 1))
-	static FLearningAgentsActionSchemaElement SpecifyLocationAction(ULearningAgentsActionSchema* Schema, const FName Tag = TEXT("Location"));
+	static FLearningAgentsActionSchemaElement SpecifyFloatAction(ULearningAgentsActionSchema* Schema, const FName Tag = TEXT("FloatAction"));
 
+	/**
+	 * Specifies a new location action. This represents an action which is a location sampled from a Gaussian distribution.
+	 *
+	 * @param Schema The Action Schema
+	 * @param Tag The tag of this new action. Used during action object validation and debugging.
+	 * @return The newly created action schema element.
+	 */
 	UFUNCTION(BlueprintPure, Category = "LearningAgents", meta = (AdvancedDisplay = 1))
-	static FLearningAgentsActionSchemaElement SpecifyRotationAction(ULearningAgentsActionSchema* Schema, const FName Tag = TEXT("Rotation"));
+	static FLearningAgentsActionSchemaElement SpecifyLocationAction(ULearningAgentsActionSchema* Schema, const FName Tag = TEXT("LocationAction"));
 
+	/**
+	 * Specifies a new rotation action. This represents an action which is a rotation sampled from a Gaussian distribution in the angle-axis space.
+	 *
+	 * @param Schema The Action Schema
+	 * @param Tag The tag of this new action. Used during action object validation and debugging.
+	 * @return The newly created action schema element.
+	 */
 	UFUNCTION(BlueprintPure, Category = "LearningAgents", meta = (AdvancedDisplay = 1))
-	static FLearningAgentsActionSchemaElement SpecifyScaleAction(ULearningAgentsActionSchema* Schema, const FName Tag = TEXT("Scale"));
+	static FLearningAgentsActionSchemaElement SpecifyRotationAction(ULearningAgentsActionSchema* Schema, const FName Tag = TEXT("RotationAction"));
 
+	/**
+	 * Specifies a new scale action. This represents an action which is a scale sampled from a Gaussian distribution in the log space.
+	 *
+	 * @param Schema The Action Schema
+	 * @param Tag The tag of this new action. Used during action object validation and debugging.
+	 * @return The newly created action schema element.
+	 */
 	UFUNCTION(BlueprintPure, Category = "LearningAgents", meta = (AdvancedDisplay = 1))
-	static FLearningAgentsActionSchemaElement SpecifyTransformAction(ULearningAgentsActionSchema* Schema, const FName Tag = TEXT("Transform"));
+	static FLearningAgentsActionSchemaElement SpecifyScaleAction(ULearningAgentsActionSchema* Schema, const FName Tag = TEXT("ScaleAction"));
 
+	/**
+	 * Specifies a new transform action.
+	 *
+	 * @param Schema The Action Schema
+	 * @param Tag The tag of this new action. Used during action object validation and debugging.
+	 * @return The newly created action schema element.
+	 */
 	UFUNCTION(BlueprintPure, Category = "LearningAgents", meta = (AdvancedDisplay = 1))
-	static FLearningAgentsActionSchemaElement SpecifyAngleAction(ULearningAgentsActionSchema* Schema, const FName Tag = TEXT("Angle"));
+	static FLearningAgentsActionSchemaElement SpecifyTransformAction(ULearningAgentsActionSchema* Schema, const FName Tag = TEXT("TransformAction"));
 
+	/**
+	 * Specifies a new angle action. This represents an action which is an angle sampled from a Gaussian distribution centered around zero.
+	 *
+	 * @param Schema The Action Schema
+	 * @param Tag The tag of this new action. Used during action object validation and debugging.
+	 * @return The newly created action schema element.
+	 */
 	UFUNCTION(BlueprintPure, Category = "LearningAgents", meta = (AdvancedDisplay = 1))
-	static FLearningAgentsActionSchemaElement SpecifyVelocityAction(ULearningAgentsActionSchema* Schema, const FName Tag = TEXT("Velocity"));
+	static FLearningAgentsActionSchemaElement SpecifyAngleAction(ULearningAgentsActionSchema* Schema, const FName Tag = TEXT("AngleAction"));
 
+	/**
+	 * Specifies a new velocity action. This represents an action which is a velocity sampled from a Gaussian distribution.
+	 *
+	 * @param Schema The Action Schema
+	 * @param Tag The tag of this new action. Used during action object validation and debugging.
+	 * @return The newly created action schema element.
+	 */
 	UFUNCTION(BlueprintPure, Category = "LearningAgents", meta = (AdvancedDisplay = 1))
-	static FLearningAgentsActionSchemaElement SpecifyDirectionAction(ULearningAgentsActionSchema* Schema, const FName Tag = TEXT("Direction"));
+	static FLearningAgentsActionSchemaElement SpecifyVelocityAction(ULearningAgentsActionSchema* Schema, const FName Tag = TEXT("VelocityAction"));
 
+	/**
+	 * Specifies a new direction action. This represents an action which is a direction sampled from a Gaussian distribution and normalized.
+	 *
+	 * @param Schema The Action Schema
+	 * @param Tag The tag of this new action. Used during action object validation and debugging.
+	 * @return The newly created action schema element.
+	 */
 	UFUNCTION(BlueprintPure, Category = "LearningAgents", meta = (AdvancedDisplay = 1))
-	static FLearningAgentsActionSchemaElement SpecifySpeedAction(ULearningAgentsActionSchema* Schema, const FName Tag = TEXT("Speed"));
+	static FLearningAgentsActionSchemaElement SpecifyDirectionAction(ULearningAgentsActionSchema* Schema, const FName Tag = TEXT("DirectionAction"));
 
 public:
 
 	UFUNCTION(BlueprintPure, Category = "LearningAgents", meta = (AdvancedDisplay = 1))
-	static FLearningAgentsActionObjectElement MakeNullAction(ULearningAgentsActionObject* Object, const FName Tag = TEXT("Null"));
+	static FLearningAgentsActionObjectElement MakeNullAction(ULearningAgentsActionObject* Object, const FName Tag = TEXT("NullAction"));
 
 	UFUNCTION(BlueprintPure, Category = "LearningAgents", meta = (AdvancedDisplay = 2))
-	static FLearningAgentsActionObjectElement MakeContinuousAction(ULearningAgentsActionObject* Object, const TArray<float>& Values, const FName Tag = TEXT("Continuous"));
-	static FLearningAgentsActionObjectElement MakeContinuousActionFromArrayView(ULearningAgentsActionObject* Object, const TArrayView<const float> Values, const FName Tag = TEXT("Continuous"));
+	static FLearningAgentsActionObjectElement MakeContinuousAction(ULearningAgentsActionObject* Object, const TArray<float>& Values, const FName Tag = TEXT("ContinuousAction"));
+	static FLearningAgentsActionObjectElement MakeContinuousActionFromArrayView(ULearningAgentsActionObject* Object, const TArrayView<const float> Values, const FName Tag = TEXT("ContinuousAction"));
 
 	UFUNCTION(BlueprintPure, Category = "LearningAgents", meta = (AdvancedDisplay = 2))
-	static FLearningAgentsActionObjectElement MakeExclusiveDiscreteAction(ULearningAgentsActionObject* Object, const int32 Index, const FName Tag = TEXT("DiscreteExclusive"));
+	static FLearningAgentsActionObjectElement MakeExclusiveDiscreteAction(ULearningAgentsActionObject* Object, const int32 Index, const FName Tag = TEXT("DiscreteExclusiveAction"));
 
 	UFUNCTION(BlueprintPure, Category = "LearningAgents", meta = (AdvancedDisplay = 2))
-	static FLearningAgentsActionObjectElement MakeInclusiveDiscreteAction(ULearningAgentsActionObject* Object, const TArray<int32>& Indices, const FName Tag = TEXT("DiscreteInclusive"));
-	static FLearningAgentsActionObjectElement MakeInclusiveDiscreteActionFromArrayView(ULearningAgentsActionObject* Object, const TArrayView<const int32> Indices, const FName Tag = TEXT("DiscreteInclusive"));
+	static FLearningAgentsActionObjectElement MakeInclusiveDiscreteAction(ULearningAgentsActionObject* Object, const TArray<int32>& Indices, const FName Tag = TEXT("DiscreteInclusiveAction"));
+	static FLearningAgentsActionObjectElement MakeInclusiveDiscreteActionFromArrayView(ULearningAgentsActionObject* Object, const TArrayView<const int32> Indices, const FName Tag = TEXT("DiscreteInclusiveAction"));
 
 	UFUNCTION(BlueprintPure, Category = "LearningAgents", meta = (AdvancedDisplay = 2))
-	static FLearningAgentsActionObjectElement MakeStructAction(ULearningAgentsActionObject* Object, const TMap<FName, FLearningAgentsActionObjectElement>& Elements, const FName Tag = TEXT("Struct"));
+	static FLearningAgentsActionObjectElement MakeStructAction(ULearningAgentsActionObject* Object, const TMap<FName, FLearningAgentsActionObjectElement>& Elements, const FName Tag = TEXT("StructAction"));
 
 	UFUNCTION(BlueprintPure, Category = "LearningAgents", meta = (AdvancedDisplay = 3))
-	static FLearningAgentsActionObjectElement MakeStructActionFromArrays(ULearningAgentsActionObject* Object, const TArray<FName>& ElementNames, const TArray<FLearningAgentsActionObjectElement>& Elements, const FName Tag = TEXT("Struct"));
-	static FLearningAgentsActionObjectElement MakeStructActionFromArrayViews(ULearningAgentsActionObject* Object, const TArrayView<const FName> ElementNames, const TArrayView<const FLearningAgentsActionObjectElement> Elements, const FName Tag = TEXT("Struct"));
+	static FLearningAgentsActionObjectElement MakeStructActionFromArrays(ULearningAgentsActionObject* Object, const TArray<FName>& ElementNames, const TArray<FLearningAgentsActionObjectElement>& Elements, const FName Tag = TEXT("StructAction"));
+	static FLearningAgentsActionObjectElement MakeStructActionFromArrayViews(ULearningAgentsActionObject* Object, const TArrayView<const FName> ElementNames, const TArrayView<const FLearningAgentsActionObjectElement> Elements, const FName Tag = TEXT("StructAction"));
 
 	UFUNCTION(BlueprintPure, Category = "LearningAgents", meta = (AdvancedDisplay = 3))
-	static FLearningAgentsActionObjectElement MakeExclusiveUnionAction(ULearningAgentsActionObject* Object, const FName ElementName, const FLearningAgentsActionObjectElement Element, const FName Tag = TEXT("ExclusiveUnion"));
+	static FLearningAgentsActionObjectElement MakeExclusiveUnionAction(ULearningAgentsActionObject* Object, const FName ElementName, const FLearningAgentsActionObjectElement Element, const FName Tag = TEXT("ExclusiveUnionAction"));
 
 	UFUNCTION(BlueprintPure, Category = "LearningAgents", meta = (AdvancedDisplay = 2))
-	static FLearningAgentsActionObjectElement MakeInclusiveUnionAction(ULearningAgentsActionObject* Object, const TMap<FName, FLearningAgentsActionObjectElement>& Elements, const FName Tag = TEXT("InclusiveUnion"));
+	static FLearningAgentsActionObjectElement MakeInclusiveUnionAction(ULearningAgentsActionObject* Object, const TMap<FName, FLearningAgentsActionObjectElement>& Elements, const FName Tag = TEXT("InclusiveUnionAction"));
 
 	UFUNCTION(BlueprintPure, Category = "LearningAgents", meta = (AdvancedDisplay = 3))
-	static FLearningAgentsActionObjectElement MakeInclusiveUnionActionFromArrays(ULearningAgentsActionObject* Object, const TArray<FName>& ElementNames, const TArray<FLearningAgentsActionObjectElement>& Elements, const FName Tag = TEXT("InclusiveUnion"));
-	static FLearningAgentsActionObjectElement MakeInclusiveUnionActionFromArrayViews(ULearningAgentsActionObject* Object, const TArrayView<const FName> ElementNames, const TArrayView<const FLearningAgentsActionObjectElement> Elements, const FName Tag = TEXT("InclusiveUnion"));
+	static FLearningAgentsActionObjectElement MakeInclusiveUnionActionFromArrays(ULearningAgentsActionObject* Object, const TArray<FName>& ElementNames, const TArray<FLearningAgentsActionObjectElement>& Elements, const FName Tag = TEXT("InclusiveUnionAction"));
+	static FLearningAgentsActionObjectElement MakeInclusiveUnionActionFromArrayViews(ULearningAgentsActionObject* Object, const TArrayView<const FName> ElementNames, const TArrayView<const FLearningAgentsActionObjectElement> Elements, const FName Tag = TEXT("InclusiveUnionAction"));
 
 	UFUNCTION(BlueprintPure, Category = "LearningAgents", meta = (AdvancedDisplay = 2))
-	static FLearningAgentsActionObjectElement MakeStaticArrayAction(ULearningAgentsActionObject* Object, const TArray<FLearningAgentsActionObjectElement>& Elements, const FName Tag = TEXT("StaticArray"));
-	static FLearningAgentsActionObjectElement MakeStaticArrayActionFromArrayView(ULearningAgentsActionObject* Object, const TArrayView<const FLearningAgentsActionObjectElement> Elements, const FName Tag = TEXT("StaticArray"));
+	static FLearningAgentsActionObjectElement MakeStaticArrayAction(ULearningAgentsActionObject* Object, const TArray<FLearningAgentsActionObjectElement>& Elements, const FName Tag = TEXT("StaticArrayAction"));
+	static FLearningAgentsActionObjectElement MakeStaticArrayActionFromArrayView(ULearningAgentsActionObject* Object, const TArrayView<const FLearningAgentsActionObjectElement> Elements, const FName Tag = TEXT("StaticArrayAction"));
 
 	UFUNCTION(BlueprintPure, Category = "LearningAgents", meta = (AdvancedDisplay = 3))
-	static FLearningAgentsActionObjectElement MakePairAction(ULearningAgentsActionObject* Object, const FLearningAgentsActionObjectElement Key, const FLearningAgentsActionObjectElement Value, const FName Tag = TEXT("Pair"));
+	static FLearningAgentsActionObjectElement MakePairAction(ULearningAgentsActionObject* Object, const FLearningAgentsActionObjectElement Key, const FLearningAgentsActionObjectElement Value, const FName Tag = TEXT("PairAction"));
 
 	UFUNCTION(BlueprintPure, Category = "LearningAgents", meta = (AdvancedDisplay = 2))
-	static FLearningAgentsActionObjectElement MakeEnumAction(ULearningAgentsActionObject* Object, const UEnum* Enum, const uint8 EnumValue, const FName Tag = TEXT("Enum"));
+	static FLearningAgentsActionObjectElement MakeEnumAction(ULearningAgentsActionObject* Object, const UEnum* Enum, const uint8 EnumValue, const FName Tag = TEXT("EnumAction"));
 
 	UFUNCTION(BlueprintPure, Category = "LearningAgents", meta = (AdvancedDisplay = 2))
-	static FLearningAgentsActionObjectElement MakeBitmaskAction(ULearningAgentsActionObject* Object, const UEnum* Enum, const int32 BitmaskValue, const FName Tag = TEXT("Bitmask"));
+	static FLearningAgentsActionObjectElement MakeBitmaskAction(ULearningAgentsActionObject* Object, const UEnum* Enum, const int32 BitmaskValue, const FName Tag = TEXT("BitmaskAction"));
 
 	UFUNCTION(BlueprintPure, Category = "LearningAgents", meta = (AdvancedDisplay = 3))
-	static FLearningAgentsActionObjectElement MakeOptionalAction(ULearningAgentsActionObject* Object, const FLearningAgentsActionObjectElement Element, const ELearningAgentsOptionalAction Option, const FName Tag = TEXT("Optional"));
+	static FLearningAgentsActionObjectElement MakeOptionalAction(ULearningAgentsActionObject* Object, const FLearningAgentsActionObjectElement Element, const ELearningAgentsOptionalAction Option, const FName Tag = TEXT("OptionalAction"));
 
 	UFUNCTION(BlueprintPure, Category = "LearningAgents", meta = (AdvancedDisplay = 1))
-	static FLearningAgentsActionObjectElement MakeOptionalNullAction(ULearningAgentsActionObject* Object, const FName Tag = TEXT("Optional"));
+	static FLearningAgentsActionObjectElement MakeOptionalNullAction(ULearningAgentsActionObject* Object, const FName Tag = TEXT("OptionalAction"));
 
 	UFUNCTION(BlueprintPure, Category = "LearningAgents", meta = (AdvancedDisplay = 2))
-	static FLearningAgentsActionObjectElement MakeOptionalValidAction(ULearningAgentsActionObject* Object, const FLearningAgentsActionObjectElement Element, const FName Tag = TEXT("Optional"));
+	static FLearningAgentsActionObjectElement MakeOptionalValidAction(ULearningAgentsActionObject* Object, const FLearningAgentsActionObjectElement Element, const FName Tag = TEXT("OptionalAction"));
 
 	UFUNCTION(BlueprintPure, Category = "LearningAgents", meta = (AdvancedDisplay = 3))
-	static FLearningAgentsActionObjectElement MakeEitherAction(ULearningAgentsActionObject* Object, const FLearningAgentsActionObjectElement Element, const ELearningAgentsEitherAction Either, const FName Tag = TEXT("Either"));
+	static FLearningAgentsActionObjectElement MakeEitherAction(ULearningAgentsActionObject* Object, const FLearningAgentsActionObjectElement Element, const ELearningAgentsEitherAction Either, const FName Tag = TEXT("EitherAction"));
 
 	UFUNCTION(BlueprintPure, Category = "LearningAgents", meta = (AdvancedDisplay = 2, DisplayName = "Make Either A Action"))
-	static FLearningAgentsActionObjectElement MakeEitherAAction(ULearningAgentsActionObject* Object, const FLearningAgentsActionObjectElement A, const FName Tag = TEXT("Either"));
+	static FLearningAgentsActionObjectElement MakeEitherAAction(ULearningAgentsActionObject* Object, const FLearningAgentsActionObjectElement A, const FName Tag = TEXT("EitherAction"));
 
 	UFUNCTION(BlueprintPure, Category = "LearningAgents", meta = (AdvancedDisplay = 2, DisplayName = "Make Either B Action"))
-	static FLearningAgentsActionObjectElement MakeEitherBAction(ULearningAgentsActionObject* Object, const FLearningAgentsActionObjectElement B, const FName Tag = TEXT("Either"));
+	static FLearningAgentsActionObjectElement MakeEitherBAction(ULearningAgentsActionObject* Object, const FLearningAgentsActionObjectElement B, const FName Tag = TEXT("EitherAction"));
 
 	UFUNCTION(BlueprintPure, Category = "LearningAgents", meta = (AdvancedDisplay = 2))
-	static FLearningAgentsActionObjectElement MakeEncodingAction(ULearningAgentsActionObject* Object, const FLearningAgentsActionObjectElement Element, const FName Tag = TEXT("Encoding"));
+	static FLearningAgentsActionObjectElement MakeEncodingAction(ULearningAgentsActionObject* Object, const FLearningAgentsActionObjectElement Element, const FName Tag = TEXT("EncodingAction"));
 
 	UFUNCTION(BlueprintPure, Category = "LearningAgents", meta = (AdvancedDisplay = 2))
-	static FLearningAgentsActionObjectElement MakeBoolAction(ULearningAgentsActionObject* Object, const bool bValue, const FName Tag = TEXT("Bool"));
+	static FLearningAgentsActionObjectElement MakeBoolAction(ULearningAgentsActionObject* Object, const bool bValue, const FName Tag = TEXT("BoolAction"));
 
 	UFUNCTION(BlueprintPure, Category = "LearningAgents", meta = (AdvancedDisplay = 3))
-	static FLearningAgentsActionObjectElement MakeFloatAction(ULearningAgentsActionObject* Object, const float Value, const float FloatScale = 1.0f, const FName Tag = TEXT("Float"));
+	static FLearningAgentsActionObjectElement MakeFloatAction(ULearningAgentsActionObject* Object, const float Value, const float FloatScale = 1.0f, const FName Tag = TEXT("FloatAction"));
 
 	UFUNCTION(BlueprintPure, Category = "LearningAgents", meta = (AdvancedDisplay = 4))
-	static FLearningAgentsActionObjectElement MakeLocationAction(ULearningAgentsActionObject* Object, const FVector Location, const FTransform RelativeTransform = FTransform(), const float LocationScale = 100.0f, const FName Tag = TEXT("Location"));
+	static FLearningAgentsActionObjectElement MakeLocationAction(ULearningAgentsActionObject* Object, const FVector Location, const FTransform RelativeTransform = FTransform(), const float LocationScale = 100.0f, const FName Tag = TEXT("LocationAction"));
 
 	UFUNCTION(BlueprintPure, Category = "LearningAgents", meta = (AdvancedDisplay = 4))
-	static FLearningAgentsActionObjectElement MakeRotationAction(ULearningAgentsActionObject* Object, const FRotator Rotation, const FRotator RelativeRotation = FRotator::ZeroRotator, const float RotationScale = 90.0f, const FName Tag = TEXT("Rotation"));
+	static FLearningAgentsActionObjectElement MakeRotationAction(ULearningAgentsActionObject* Object, const FRotator Rotation, const FRotator RelativeRotation = FRotator::ZeroRotator, const float RotationScale = 90.0f, const FName Tag = TEXT("RotationAction"));
 
 	UFUNCTION(BlueprintPure, Category = "LearningAgents", meta = (AdvancedDisplay = 4))
-	static FLearningAgentsActionObjectElement MakeRotationActionFromQuat(ULearningAgentsActionObject* Object, const FQuat Rotation, const FQuat RelativeRotation, const float RotationScale = 90.0f, const FName Tag = TEXT("Rotation"));
+	static FLearningAgentsActionObjectElement MakeRotationActionFromQuat(ULearningAgentsActionObject* Object, const FQuat Rotation, const FQuat RelativeRotation, const float RotationScale = 90.0f, const FName Tag = TEXT("RotationAction"));
 
 	UFUNCTION(BlueprintPure, Category = "LearningAgents", meta = (AdvancedDisplay = 3))
-	static FLearningAgentsActionObjectElement MakeScaleAction(ULearningAgentsActionObject* Object, const FVector Scale, const FVector RelativeScale = FVector(1, 1, 1), const FName Tag = TEXT("Scale"));
+	static FLearningAgentsActionObjectElement MakeScaleAction(ULearningAgentsActionObject* Object, const FVector Scale, const FVector RelativeScale = FVector(1, 1, 1), const FName Tag = TEXT("ScaleAction"));
 
 	UFUNCTION(BlueprintPure, Category = "LearningAgents", meta = (AdvancedDisplay = 4))
-	static FLearningAgentsActionObjectElement MakeTransformAction(ULearningAgentsActionObject* Object, const FTransform Transform, const FTransform RelativeTransform = FTransform(), const float LocationScale = 100.0f, const FName Tag = TEXT("Transform"));
+	static FLearningAgentsActionObjectElement MakeTransformAction(ULearningAgentsActionObject* Object, const FTransform Transform, const FTransform RelativeTransform = FTransform(), const float LocationScale = 100.0f, const FName Tag = TEXT("TransformAction"));
 
 	UFUNCTION(BlueprintPure, Category = "LearningAgents", meta = (AdvancedDisplay = 4))
-	static FLearningAgentsActionObjectElement MakeAngleAction(ULearningAgentsActionObject* Object, const float Angle, const float RelativeAngle = 0.0f, const float AngleScale = 90.0f, const FName Tag = TEXT("Angle"));
+	static FLearningAgentsActionObjectElement MakeAngleAction(ULearningAgentsActionObject* Object, const float Angle, const float RelativeAngle = 0.0f, const float AngleScale = 90.0f, const FName Tag = TEXT("AngleAction"));
 
 	UFUNCTION(BlueprintPure, Category = "LearningAgents", meta = (AdvancedDisplay = 4))
-	static FLearningAgentsActionObjectElement MakeAngleActionRadians(ULearningAgentsActionObject* Object, const float Angle, const float RelativeAngle = 0.0f, const float AngleScale = 1.57079632679f, const FName Tag = TEXT("Angle"));
+	static FLearningAgentsActionObjectElement MakeAngleActionRadians(ULearningAgentsActionObject* Object, const float Angle, const float RelativeAngle = 0.0f, const float AngleScale = 1.57079632679f, const FName Tag = TEXT("AngleAction"));
 
 	UFUNCTION(BlueprintPure, Category = "LearningAgents", meta = (AdvancedDisplay = 4))
-	static FLearningAgentsActionObjectElement MakeVelocityAction(ULearningAgentsActionObject* Object, const FVector Velocity, const FTransform RelativeTransform = FTransform(), const float VelocityScale = 200.0f, const FName Tag = TEXT("Velocity"));
-
-	UFUNCTION(BlueprintPure, Category = "LearningAgents", meta = (AdvancedDisplay = 3))
-	static FLearningAgentsActionObjectElement MakeDirectionAction(ULearningAgentsActionObject* Object, const FVector Direction, const FTransform RelativeTransform = FTransform(), const FName Tag = TEXT("Direction"));
+	static FLearningAgentsActionObjectElement MakeVelocityAction(ULearningAgentsActionObject* Object, const FVector Velocity, const FTransform RelativeTransform = FTransform(), const float VelocityScale = 200.0f, const FName Tag = TEXT("VelocityAction"));
 
 	UFUNCTION(BlueprintPure, Category = "LearningAgents", meta = (AdvancedDisplay = 3))
-	static FLearningAgentsActionObjectElement MakeSpeedAction(ULearningAgentsActionObject* Object, const float Speed, const float SpeedScale = 200.0f, const FName Tag = TEXT("Speed"));
+	static FLearningAgentsActionObjectElement MakeDirectionAction(ULearningAgentsActionObject* Object, const FVector Direction, const FTransform RelativeTransform = FTransform(), const FName Tag = TEXT("DirectionAction"));
 
 public:
 
+	/**
+	 * Get a null action.
+	 *
+	 * @param Object The Action Object
+	 * @param Element The Action Object Element
+	 * @param Tag The tag of the corresponding action. Must match the tag given during Specify.
+	 * @return true if the provided Element is the correct type, otherwise false.
+	 */
 	UFUNCTION(BlueprintPure = false, Category = "LearningAgents", meta = (AdvancedDisplay = 2, ReturnDisplayName = "Success"))
-	static bool GetNullAction(const ULearningAgentsActionObject* Object, const FLearningAgentsActionObjectElement Element, const FName Tag = TEXT("Null"));
-
-	UFUNCTION(BlueprintPure = false, Category = "LearningAgents", meta = (AdvancedDisplay = 3, ReturnDisplayName = "Success"))
-	static bool GetContinuousActionNum(int32& OutNum, const ULearningAgentsActionObject* Object, const FLearningAgentsActionObjectElement Element, const FName Tag = TEXT("Continuous"));
-
-	UFUNCTION(BlueprintPure = false, Category = "LearningAgents", meta = (AdvancedDisplay = 3, ReturnDisplayName = "Success"))
-	static bool GetContinuousAction(TArray<float>& OutValues, const ULearningAgentsActionObject* Object, const FLearningAgentsActionObjectElement Element, const FName Tag = TEXT("Continuous"));
-	static bool GetContinuousActionToArrayView(TArrayView<float> OutValues, const ULearningAgentsActionObject* Object, const FLearningAgentsActionObjectElement Element, const FName Tag = TEXT("Continuous"));
-
-	UFUNCTION(BlueprintPure = false, Category = "LearningAgents", meta = (AdvancedDisplay = 3, ReturnDisplayName = "Success"))
-	static bool GetExclusiveDiscreteAction(int32& OutIndex, const ULearningAgentsActionObject* Object, const FLearningAgentsActionObjectElement Element, const FName Tag = TEXT("DiscreteExclusive"));
-
-	UFUNCTION(BlueprintPure = false, Category = "LearningAgents", meta = (AdvancedDisplay = 3, ReturnDisplayName = "Success"))
-	static bool GetInclusiveDiscreteActionNum(int32& OutNum, const ULearningAgentsActionObject* Object, const FLearningAgentsActionObjectElement Element, const FName Tag = TEXT("DiscreteInclusive"));
-
-	UFUNCTION(BlueprintPure = false, Category = "LearningAgents", meta = (AdvancedDisplay = 3, ReturnDisplayName = "Success"))
-	static bool GetInclusiveDiscreteAction(TArray<int32>& OutIndices, const ULearningAgentsActionObject* Object, const FLearningAgentsActionObjectElement Element, const FName Tag = TEXT("DiscreteInclusive"));
-	static bool GetInclusiveDiscreteActionToArrayView(TArrayView<int32> OutIndices, const ULearningAgentsActionObject* Object, const FLearningAgentsActionObjectElement Element, const FName Tag = TEXT("DiscreteInclusive"));
-
-	UFUNCTION(BlueprintPure = false, Category = "LearningAgents", meta = (AdvancedDisplay = 3, ReturnDisplayName = "Success"))
-	static bool GetStructActionNum(int32& OutNum, const ULearningAgentsActionObject* Object, const FLearningAgentsActionObjectElement Element, const FName Tag = TEXT("Struct"));
-
-	UFUNCTION(BlueprintPure = false, Category = "LearningAgents", meta = (AdvancedDisplay = 3, ReturnDisplayName = "Success"))
-	static bool GetStructAction(TMap<FName, FLearningAgentsActionObjectElement>& OutElements, const ULearningAgentsActionObject* Object, const FLearningAgentsActionObjectElement Element, const FName Tag = TEXT("Struct"));
-
-	UFUNCTION(BlueprintPure = false, Category = "LearningAgents", meta = (AdvancedDisplay = 4, ReturnDisplayName = "Success"))
-	static bool GetStructActionToArrays(TArray<FName>& OutElementNames, TArray<FLearningAgentsActionObjectElement>& OutElements, const ULearningAgentsActionObject* Object, const FLearningAgentsActionObjectElement Element, const FName Tag = TEXT("Struct"));
-	static bool GetStructActionToArrayViews(TArrayView<FName> OutElementNames, TArrayView<FLearningAgentsActionObjectElement> OutElements, const ULearningAgentsActionObject* Object, const FLearningAgentsActionObjectElement Element, const FName Tag = TEXT("Struct"));
-
-	UFUNCTION(BlueprintPure = false, Category = "LearningAgents", meta = (AdvancedDisplay = 4, ReturnDisplayName = "Success"))
-	static bool GetExclusiveUnionAction(FName& OutElementName, FLearningAgentsActionObjectElement& OutElement, const ULearningAgentsActionObject* Object, const FLearningAgentsActionObjectElement Element, const FName Tag = TEXT("ExclusiveUnion"));
-
-	UFUNCTION(BlueprintPure = false, Category = "LearningAgents", meta = (AdvancedDisplay = 3, ReturnDisplayName = "Success"))
-	static bool GetInclusiveUnionActionNum(int32& OutNum, const ULearningAgentsActionObject* Object, const FLearningAgentsActionObjectElement Element, const FName Tag = TEXT("InclusiveUnion"));
-
-	UFUNCTION(BlueprintPure = false, Category = "LearningAgents", meta = (AdvancedDisplay = 3, ReturnDisplayName = "Success"))
-	static bool GetInclusiveUnionAction(TMap<FName, FLearningAgentsActionObjectElement>& OutElements, const ULearningAgentsActionObject* Object, const FLearningAgentsActionObjectElement Element, const FName Tag = TEXT("InclusiveUnion"));
-
-	UFUNCTION(BlueprintPure = false, Category = "LearningAgents", meta = (AdvancedDisplay = 4, ReturnDisplayName = "Success"))
-	static bool GetInclusiveUnionActionToArrays(TArray<FName>& OutElementNames, TArray<FLearningAgentsActionObjectElement>& OutElements, const ULearningAgentsActionObject* Object, const FLearningAgentsActionObjectElement Element, const FName Tag = TEXT("InclusiveUnion"));
-	static bool GetInclusiveUnionActionToArrayViews(TArrayView<FName> OutElementNames, TArrayView<FLearningAgentsActionObjectElement> OutElements, const ULearningAgentsActionObject* Object, const FLearningAgentsActionObjectElement Element, const FName Tag = TEXT("InclusiveUnion"));
-
-	UFUNCTION(BlueprintPure = false, Category = "LearningAgents", meta = (AdvancedDisplay = 3, ReturnDisplayName = "Success"))
-	static bool GetStaticArrayActionNum(int32& OutNum, const ULearningAgentsActionObject* Object, const FLearningAgentsActionObjectElement Element, const FName Tag = TEXT("StaticArray"));
+	static bool GetNullAction(const ULearningAgentsActionObject* Object, const FLearningAgentsActionObjectElement Element, const FName Tag = TEXT("NullAction"));
 
+	/**
+	 * Get the number of values in a continuous action.
+	 *
+	 * @param OutNum The output number of values.
+	 * @param Object The Action Object
+	 * @param Element The Action Object Element
+	 * @param Tag The tag of the corresponding action. Must match the tag given during Specify.
+	 * @return true if the provided Element is the correct type, otherwise false.
+	 */
 	UFUNCTION(BlueprintPure = false, Category = "LearningAgents", meta = (AdvancedDisplay = 3, ReturnDisplayName = "Success"))
-	static bool GetStaticArrayAction(TArray<FLearningAgentsActionObjectElement>& OutElements, const ULearningAgentsActionObject* Object, const FLearningAgentsActionObjectElement Element, const FName Tag = TEXT("StaticArray"));
-	static bool GetStaticArrayActionToArrayView(TArrayView<FLearningAgentsActionObjectElement> OutElements, const ULearningAgentsActionObject* Object, const FLearningAgentsActionObjectElement Element, const FName Tag = TEXT("StaticArray"));
-
-	UFUNCTION(BlueprintPure = false, Category = "LearningAgents", meta = (AdvancedDisplay = 4, ReturnDisplayName = "Success"))
-	static bool GetPairAction(FLearningAgentsActionObjectElement& OutKey, FLearningAgentsActionObjectElement& OutValue, const ULearningAgentsActionObject* Object, const FLearningAgentsActionObjectElement Element, const FName Tag = TEXT("Pair"));
-
-	UFUNCTION(BlueprintPure = false, Category = "LearningAgents", meta = (AdvancedDisplay = 4, ReturnDisplayName = "Success"))
-	static bool GetEnumAction(uint8& OutEnumValue, const ULearningAgentsActionObject* Object, const FLearningAgentsActionObjectElement Element, const UEnum* Enum, const FName Tag = TEXT("Enum"));
-
-	UFUNCTION(BlueprintPure = false, Category = "LearningAgents", meta = (AdvancedDisplay = 4, ReturnDisplayName = "Success"))
-	static bool GetBitmaskAction(int32& OutBitmaskValue, const ULearningAgentsActionObject* Object, const FLearningAgentsActionObjectElement Element, const UEnum* Enum, const FName Tag = TEXT("Bitmask"));
-
-	UFUNCTION(BlueprintPure = false, Category = "LearningAgents", meta = (AdvancedDisplay = 4, ExpandEnumAsExecs = "OutOption", ReturnDisplayName = "Success"))
-	static bool GetOptionalAction(ELearningAgentsOptionalAction& OutOption, FLearningAgentsActionObjectElement& OutElement, const ULearningAgentsActionObject* Object, const FLearningAgentsActionObjectElement Element, const FName Tag = TEXT("Optional"));
+	static bool GetContinuousActionNum(int32& OutNum, const ULearningAgentsActionObject* Object, const FLearningAgentsActionObjectElement Element, const FName Tag = TEXT("ContinuousAction"));
 
-	UFUNCTION(BlueprintPure = false, Category = "LearningAgents", meta = (AdvancedDisplay = 4, ExpandEnumAsExecs = "OutEither", ReturnDisplayName = "Success"))
-	static bool GetEitherAction(ELearningAgentsEitherAction& OutEither, FLearningAgentsActionObjectElement& OutElement, const ULearningAgentsActionObject* Object, const FLearningAgentsActionObjectElement Element, const FName Tag = TEXT("Either"));
+	/**
+	 * Get the values of a continuous action.
+	 *
+	 * @param OutValues The output values.
+	 * @param Object The Action Object
+	 * @param Element The Action Object Element
+	 * @param Tag The tag of the corresponding action. Must match the tag given during Specify.
+	 * @param bVisualLoggerEnabled When true, debug data will be sent to the visual logger.
+	 * @param VisualLoggerListener The listener object which is making this action. This must be set to use logging.
+	 * @param VisualLoggerAgentId The agent id associated with this action.
+	 * @param VisualLoggerLocation A location for the visual logger information in the world.
+	 * @param VisualLoggerColor The color for the visual logger display.
+	 * @return true if the provided Element is the correct type, otherwise false.
+	 */
+	UFUNCTION(BlueprintPure = false, Category = "LearningAgents", meta = (AdvancedDisplay = 3, ReturnDisplayName = "Success", DefaultToSelf = "VisualLoggerListener"))
+	static bool GetContinuousAction(
+		TArray<float>& OutValues, 
+		const ULearningAgentsActionObject* Object, 
+		const FLearningAgentsActionObjectElement Element, 
+		const FName Tag = TEXT("ContinuousAction"),
+		const bool bVisualLoggerEnabled = false,
+		ULearningAgentsManagerListener* VisualLoggerListener = nullptr,
+		const int32 VisualLoggerAgentId = -1,
+		const FVector VisualLoggerLocation = FVector::ZeroVector,
+		const FLinearColor VisualLoggerColor = FLinearColor::Blue);
+	
+	/**
+	 * Get the values of a continuous action. The OutValues ArrayView must be the correct size.
+	 *
+	 * @param OutValues The output values.
+	 * @param Object The Action Object
+	 * @param Element The Action Object Element
+	 * @param Tag The tag of the corresponding action. Must match the tag given during Specify.
+	 * @param bVisualLoggerEnabled When true, debug data will be sent to the visual logger.
+	 * @param VisualLoggerListener The listener object which is making this action. This must be set to use logging.
+	 * @param VisualLoggerAgentId The agent id associated with this action.
+	 * @param VisualLoggerLocation A location for the visual logger information in the world.
+	 * @param VisualLoggerColor The color for the visual logger display.
+	 * @return true if the provided Element is the correct type, otherwise false.
+	 */
+	static bool GetContinuousActionToArrayView(
+		TArrayView<float> OutValues, 
+		const ULearningAgentsActionObject* Object, 
+		const FLearningAgentsActionObjectElement Element, 
+		const FName Tag = TEXT("ContinuousAction"),
+		const bool bVisualLoggerEnabled = false,
+		ULearningAgentsManagerListener* VisualLoggerListener = nullptr,
+		const int32 VisualLoggerAgentId = -1,
+		const FVector VisualLoggerLocation = FVector::ZeroVector,
+		const FLinearColor VisualLoggerColor = FLinearColor::Blue);
 
-	UFUNCTION(BlueprintPure = false, Category = "LearningAgents", meta = (AdvancedDisplay = 3, ReturnDisplayName = "Success"))
-	static bool GetEncodingAction(FLearningAgentsActionObjectElement& OutElement, const ULearningAgentsActionObject* Object, const FLearningAgentsActionObjectElement Element, const FName Tag = TEXT("Encoding"));
+	/**
+	 * Get the index for an exclusive discrete action.
+	 *
+	 * @param OutIndex The output index.
+	 * @param Object The Action Object
+	 * @param Element The Action Object Element
+	 * @param Tag The tag of the corresponding action. Must match the tag given during Specify.
+	 * @param bVisualLoggerEnabled When true, debug data will be sent to the visual logger.
+	 * @param VisualLoggerListener The listener object which is making this action. This must be set to use logging.
+	 * @param VisualLoggerAgentId The agent id associated with this action.
+	 * @param VisualLoggerLocation A location for the visual logger information in the world.
+	 * @param VisualLoggerColor The color for the visual logger display.
+	 * @return true if the provided Element is the correct type, otherwise false.
+	 */
+	UFUNCTION(BlueprintPure = false, Category = "LearningAgents", meta = (AdvancedDisplay = 3, ReturnDisplayName = "Success", DefaultToSelf = "VisualLoggerListener"))
+	static bool GetExclusiveDiscreteAction(
+		int32& OutIndex, 
+		const ULearningAgentsActionObject* Object, 
+		const FLearningAgentsActionObjectElement Element, 
+		const FName Tag = TEXT("DiscreteExclusiveAction"),
+		const bool bVisualLoggerEnabled = false,
+		ULearningAgentsManagerListener* VisualLoggerListener = nullptr,
+		const int32 VisualLoggerAgentId = -1,
+		const FVector VisualLoggerLocation = FVector::ZeroVector,
+		const FLinearColor VisualLoggerColor = FLinearColor::Blue);
 
+	/**
+	 * Get the number of indices for an inclusive discrete action.
+	 *
+	 * @param OutNum The output number of indices.
+	 * @param Object The Action Object
+	 * @param Element The Action Object Element
+	 * @param Tag The tag of the corresponding action. Must match the tag given during Specify.
+	 * @return true if the provided Element is the correct type, otherwise false.
+	 */
 	UFUNCTION(BlueprintPure = false, Category = "LearningAgents", meta = (AdvancedDisplay = 3, ReturnDisplayName = "Success"))
-	static bool GetBoolAction(bool& bOutValue, const ULearningAgentsActionObject* Object, const FLearningAgentsActionObjectElement Element, const FName Tag = TEXT("Bool"));
+	static bool GetInclusiveDiscreteActionNum(int32& OutNum, const ULearningAgentsActionObject* Object, const FLearningAgentsActionObjectElement Element, const FName Tag = TEXT("DiscreteInclusiveAction"));
 
 	/**
-	 * Get the value for a float action.
+	 * Get the indices for an inclusive discrete action.
 	 *
-	 * @param OutValue The output float value.
+	 * @param OutIndices The output indices.
 	 * @param Object The Action Object
 	 * @param Element The Action Object Element
-	 * @param FloatScale The scale used to control the overall magnitude of the outputted float action.
 	 * @param Tag The tag of the corresponding action. Must match the tag given during Specify.
 	 * @param bVisualLoggerEnabled When true, debug data will be sent to the visual logger.
 	 * @param VisualLoggerListener The listener object which is making this action. This must be set to use logging.
 	 * @param VisualLoggerAgentId The agent id associated with this action.
 	 * @param VisualLoggerLocation A location for the visual logger information in the world.
 	 * @param VisualLoggerColor The color for the visual logger display.
+	 * @return true if the provided Element is the correct type, otherwise false.
 	 */
-	UFUNCTION(BlueprintPure = false, Category = "LearningAgents", meta = (AdvancedDisplay = 4, ReturnDisplayName = "Success", DefaultToSelf = "VisualLoggerListener"))
-	static bool GetFloatAction(
-		float& OutValue, 
+	UFUNCTION(BlueprintPure = false, Category = "LearningAgents", meta = (AdvancedDisplay = 3, ReturnDisplayName = "Success", DefaultToSelf = "VisualLoggerListener"))
+	static bool GetInclusiveDiscreteAction(
+		TArray<int32>& OutIndices, 
 		const ULearningAgentsActionObject* Object, 
 		const FLearningAgentsActionObjectElement Element, 
-		const float FloatScale = 1.0f, 
-		const FName Tag = TEXT("Float"),
+		const FName Tag = TEXT("DiscreteInclusiveAction"),
+		const bool bVisualLoggerEnabled = false,
+		ULearningAgentsManagerListener* VisualLoggerListener = nullptr,
+		const int32 VisualLoggerAgentId = -1,
+		const FVector VisualLoggerLocation = FVector::ZeroVector,
+		const FLinearColor VisualLoggerColor = FLinearColor::Blue);
+	
+	/**
+	 * Get the indices for an inclusive discrete action. The OutIndices ArrayView must be the correct size.
+	 *
+	 * @param OutIndices The output indices.
+	 * @param Object The Action Object
+	 * @param Element The Action Object Element
+	 * @param Tag The tag of the corresponding action. Must match the tag given during Specify.
+	 * @param bVisualLoggerEnabled When true, debug data will be sent to the visual logger.
+	 * @param VisualLoggerListener The listener object which is making this action. This must be set to use logging.
+	 * @param VisualLoggerAgentId The agent id associated with this action.
+	 * @param VisualLoggerLocation A location for the visual logger information in the world.
+	 * @param VisualLoggerColor The color for the visual logger display.
+	 * @return true if the provided Element is the correct type, otherwise false.
+	 */
+	static bool GetInclusiveDiscreteActionToArrayView(
+		TArrayView<int32> OutIndices, 
+		const ULearningAgentsActionObject* Object, 
+		const FLearningAgentsActionObjectElement Element, 
+		const FName Tag = TEXT("DiscreteInclusiveAction"),
 		const bool bVisualLoggerEnabled = false,
 		ULearningAgentsManagerListener* VisualLoggerListener = nullptr,
 		const int32 VisualLoggerAgentId = -1,
 		const FVector VisualLoggerLocation = FVector::ZeroVector,
 		const FLinearColor VisualLoggerColor = FLinearColor::Blue);
 
-	UFUNCTION(BlueprintPure = false, Category = "LearningAgents", meta = (AdvancedDisplay = 5, ReturnDisplayName = "Success"))
-	static bool GetLocationAction(FVector& OutLocation, const ULearningAgentsActionObject* Object, const FLearningAgentsActionObjectElement Element, const FTransform RelativeTransform = FTransform(), const float LocationScale = 100.0f, const FName Tag = TEXT("Location"));
-
-	UFUNCTION(BlueprintPure = false, Category = "LearningAgents", meta = (AdvancedDisplay = 5, ReturnDisplayName = "Success"))
-	static bool GetRotationAction(FRotator& OutRotation, const ULearningAgentsActionObject* Object, const FLearningAgentsActionObjectElement Element, const FRotator RelativeRotation = FRotator::ZeroRotator, const float RotationScale = 90.0f, const FName Tag = TEXT("Rotation"));
+	/**
+	 * Get the number of sub-actions for a struct action.
+	 *
+	 * @param OutNum The output number of sub-actions.
+	 * @param Object The Action Object
+	 * @param Element The Action Object Element
+	 * @param Tag The tag of the corresponding action. Must match the tag given during Specify.
+	 * @return true if the provided Element is the correct type, otherwise false.
+	 */
+	UFUNCTION(BlueprintPure = false, Category = "LearningAgents", meta = (AdvancedDisplay = 3, ReturnDisplayName = "Success"))
+	static bool GetStructActionNum(int32& OutNum, const ULearningAgentsActionObject* Object, const FLearningAgentsActionObjectElement Element, const FName Tag = TEXT("StructAction"));
 
-	UFUNCTION(BlueprintPure = false, Category = "LearningAgents", meta = (AdvancedDisplay = 5, ReturnDisplayName = "Success"))
-	static bool GetRotationActionAsQuat(FQuat& OutRotation, const ULearningAgentsActionObject* Object, const FLearningAgentsActionObjectElement Element, const FQuat RelativeRotation, const float RotationScale = 90.0f, const FName Tag = TEXT("Rotation"));
-
-	UFUNCTION(BlueprintPure = false, Category = "LearningAgents", meta = (AdvancedDisplay = 5, ReturnDisplayName = "Success"))
-	static bool GetScaleAction(FVector& OutScale, const ULearningAgentsActionObject* Object, const FLearningAgentsActionObjectElement Element, const FVector RelativeScale = FVector(1, 1, 1), const float Scale = 1.0f, const FName Tag = TEXT("Scale"));
+	/**
+	 * Get the sub-actions for a struct action.
+	 *
+	 * @param OutElements The output sub-actions.
+	 * @param Object The Action Object
+	 * @param Element The Action Object Element
+	 * @param Tag The tag of the corresponding action. Must match the tag given during Specify.
+	 * @return true if the provided Element is the correct type, otherwise false.
+	 */
+	UFUNCTION(BlueprintPure = false, Category = "LearningAgents", meta = (AdvancedDisplay = 3, ReturnDisplayName = "Success"))
+	static bool GetStructAction(TMap<FName, FLearningAgentsActionObjectElement>& OutElements, const ULearningAgentsActionObject* Object, const FLearningAgentsActionObjectElement Element, const FName Tag = TEXT("StructAction"));
 
-	UFUNCTION(BlueprintPure = false, Category = "LearningAgents", meta = (AdvancedDisplay = 7, ReturnDisplayName = "Success"))
-	static bool GetTransformAction(FTransform& OutTransform, const ULearningAgentsActionObject* Object, const FLearningAgentsActionObjectElement Element, const FTransform RelativeTransform = FTransform(), const float LocationScale = 100.0f, const float RotationScale = 1.0f, const float ScaleScale = 1.0f, const FName Tag = TEXT("Transform"));
+	/**
+	 * Get the sub-actions for a struct action.
+	 *
+	 * @param OutElementNames The output sub-action names.
+	 * @param OutElements The output sub-actions.
+	 * @param Object The Action Object
+	 * @param Element The Action Object Element
+	 * @param Tag The tag of the corresponding action. Must match the tag given during Specify.
+	 * @return true if the provided Element is the correct type, otherwise false.
+	 */
+	UFUNCTION(BlueprintPure = false, Category = "LearningAgents", meta = (AdvancedDisplay = 4, ReturnDisplayName = "Success"))
+	static bool GetStructActionToArrays(TArray<FName>& OutElementNames, TArray<FLearningAgentsActionObjectElement>& OutElements, const ULearningAgentsActionObject* Object, const FLearningAgentsActionObjectElement Element, const FName Tag = TEXT("StructAction"));
+	
+	/**
+	 * Get the sub-actions for a struct action. The OutElementNames and OutElements ArrayViews must be the correct size.
+	 *
+	 * @param OutElementNames The output sub-action names.
+	 * @param OutElements The output sub-actions.
+	 * @param Object The Action Object
+	 * @param Element The Action Object Element
+	 * @param Tag The tag of the corresponding action. Must match the tag given during Specify.
+	 * @return true if the provided Element is the correct type, otherwise false.
+	 */
+	static bool GetStructActionToArrayViews(TArrayView<FName> OutElementNames, TArrayView<FLearningAgentsActionObjectElement> OutElements, const ULearningAgentsActionObject* Object, const FLearningAgentsActionObjectElement Element, const FName Tag = TEXT("StructAction"));
 
-	UFUNCTION(BlueprintPure = false, Category = "LearningAgents", meta = (AdvancedDisplay = 5, ReturnDisplayName = "Success"))
-	static bool GetAngleAction(float& OutAngle, const ULearningAgentsActionObject* Object, const FLearningAgentsActionObjectElement Element, const float RelativeAngle = 0.0f, const float AngleScale = 90.0f, const FName Tag = TEXT("Angle"));
+	/**
+	 * Get the chosen sub-action for an exclusive union action.
+	 *
+	 * @param OutElement The output sub-action.
+	 * @param Object The Action Object
+	 * @param Element The Action Object Element
+	 * @param Tag The tag of the corresponding action. Must match the tag given during Specify.
+	 * @return true if the provided Element is the correct type, otherwise false.
+	 */
+	UFUNCTION(BlueprintPure = false, Category = "LearningAgents", meta = (AdvancedDisplay = 4, ReturnDisplayName = "Success"))
+	static bool GetExclusiveUnionAction(FName& OutElementName, FLearningAgentsActionObjectElement& OutElement, const ULearningAgentsActionObject* Object, const FLearningAgentsActionObjectElement Element, const FName Tag = TEXT("ExclusiveUnionAction"));
 
-	UFUNCTION(BlueprintPure = false, Category = "LearningAgents", meta = (AdvancedDisplay = 5, ReturnDisplayName = "Success"))
-	static bool GetAngleActionRadians(float& OutAngle, const ULearningAgentsActionObject* Object, const FLearningAgentsActionObjectElement Element, const float RelativeAngle = 0.0f, const float AngleScale = 1.57079632679f, const FName Tag = TEXT("Angle"));
+	/**
+	 * Get the number of sub-actions for an inclusive union action.
+	 *
+	 * @param OutNum The output number of sub-actions.
+	 * @param Object The Action Object
+	 * @param Element The Action Object Element
+	 * @param Tag The tag of the corresponding action. Must match the tag given during Specify.
+	 * @return true if the provided Element is the correct type, otherwise false.
+	 */
+	UFUNCTION(BlueprintPure = false, Category = "LearningAgents", meta = (AdvancedDisplay = 3, ReturnDisplayName = "Success"))
+	static bool GetInclusiveUnionActionNum(int32& OutNum, const ULearningAgentsActionObject* Object, const FLearningAgentsActionObjectElement Element, const FName Tag = TEXT("InclusiveUnionAction"));
 
-	UFUNCTION(BlueprintPure = false, Category = "LearningAgents", meta = (AdvancedDisplay = 5, ReturnDisplayName = "Success"))
-	static bool GetVelocityAction(FVector& OutVelocity, const ULearningAgentsActionObject* Object, const FLearningAgentsActionObjectElement Element, const FTransform RelativeTransform = FTransform(), const float VelocityScale = 200.0f, const FName Tag = TEXT("Velocity"));
+	/**
+	 * Get the chosen sub-actions for an inclusive union action.
+	 *
+	 * @param OutElements The output sub-actions.
+	 * @param Object The Action Object
+	 * @param Element The Action Object Element
+	 * @param Tag The tag of the corresponding action. Must match the tag given during Specify.
+	 * @return true if the provided Element is the correct type, otherwise false.
+	 */
+	UFUNCTION(BlueprintPure = false, Category = "LearningAgents", meta = (AdvancedDisplay = 3, ReturnDisplayName = "Success"))
+	static bool GetInclusiveUnionAction(TMap<FName, FLearningAgentsActionObjectElement>& OutElements, const ULearningAgentsActionObject* Object, const FLearningAgentsActionObjectElement Element, const FName Tag = TEXT("InclusiveUnionAction"));
 
+	/**
+	 * Get the chosen sub-actions for an inclusive union action.
+	 *
+	 * @param OutElementNames The output sub-action names.
+	 * @param OutElements The output sub-actions.
+	 * @param Object The Action Object
+	 * @param Element The Action Object Element
+	 * @param Tag The tag of the corresponding action. Must match the tag given during Specify.
+	 * @return true if the provided Element is the correct type, otherwise false.
+	 */
 	UFUNCTION(BlueprintPure = false, Category = "LearningAgents", meta = (AdvancedDisplay = 4, ReturnDisplayName = "Success"))
-	static bool GetDirectionAction(FVector& OutDirection, const ULearningAgentsActionObject* Object, const FLearningAgentsActionObjectElement Element, const FTransform RelativeTransform = FTransform(), const FName Tag = TEXT("Direction"));
+	static bool GetInclusiveUnionActionToArrays(TArray<FName>& OutElementNames, TArray<FLearningAgentsActionObjectElement>& OutElements, const ULearningAgentsActionObject* Object, const FLearningAgentsActionObjectElement Element, const FName Tag = TEXT("InclusiveUnionAction"));
+	
+	/**
+	 * Get the chosen sub-actions for an inclusive union action. The OutElementNames and OutElements ArrayViews must be the correct size.
+	 *
+	 * @param OutElementNames The output sub-action names.
+	 * @param OutElements The output sub-actions.
+	 * @param Object The Action Object
+	 * @param Element The Action Object Element
+	 * @param Tag The tag of the corresponding action. Must match the tag given during Specify.
+	 * @return true if the provided Element is the correct type, otherwise false.
+	 */
+	static bool GetInclusiveUnionActionToArrayViews(TArrayView<FName> OutElementNames, TArrayView<FLearningAgentsActionObjectElement> OutElements, const ULearningAgentsActionObject* Object, const FLearningAgentsActionObjectElement Element, const FName Tag = TEXT("InclusiveUnionAction"));
 
+	/**
+	 * Get the number of entries in a static array action.
+	 *
+	 * @param OutNum The output number of entries.
+	 * @param Object The Action Object
+	 * @param Element The Action Object Element
+	 * @param Tag The tag of the corresponding action. Must match the tag given during Specify.
+	 * @return true if the provided Element is the correct type, otherwise false.
+	 */
+	UFUNCTION(BlueprintPure = false, Category = "LearningAgents", meta = (AdvancedDisplay = 3, ReturnDisplayName = "Success"))
+	static bool GetStaticArrayActionNum(int32& OutNum, const ULearningAgentsActionObject* Object, const FLearningAgentsActionObjectElement Element, const FName Tag = TEXT("StaticArrayAction"));
+
+	/**
+	 * Get the entries of a static array action.
+	 *
+	 * @param OutElements The output sub-elements.
+	 * @param Object The Action Object
+	 * @param Element The Action Object Element
+	 * @param Tag The tag of the corresponding action. Must match the tag given during Specify.
+	 * @return true if the provided Element is the correct type, otherwise false.
+	 */
+	UFUNCTION(BlueprintPure = false, Category = "LearningAgents", meta = (AdvancedDisplay = 3, ReturnDisplayName = "Success"))
+	static bool GetStaticArrayAction(TArray<FLearningAgentsActionObjectElement>& OutElements, const ULearningAgentsActionObject* Object, const FLearningAgentsActionObjectElement Element, const FName Tag = TEXT("StaticArrayAction"));
+	
+	/**
+	 * Get the entries of a static array action. The OutElements ArrayView must be the correct size.
+	 *
+	 * @param OutElements The output sub-elements.
+	 * @param Object The Action Object
+	 * @param Element The Action Object Element
+	 * @param Tag The tag of the corresponding action. Must match the tag given during Specify.
+	 * @return true if the provided Element is the correct type, otherwise false.
+	 */
+	static bool GetStaticArrayActionToArrayView(TArrayView<FLearningAgentsActionObjectElement> OutElements, const ULearningAgentsActionObject* Object, const FLearningAgentsActionObjectElement Element, const FName Tag = TEXT("StaticArrayAction"));
+
+	/**
+	 * Get the sub-actions of a pair action.
+	 *
+	 * @param OutKey The output key sub-element.
+	 * @param OutValue The output value sub-element.
+	 * @param Object The Action Object
+	 * @param Element The Action Object Element
+	 * @param Tag The tag of the corresponding action. Must match the tag given during Specify.
+	 * @return true if the provided Element is the correct type, otherwise false.
+	 */
 	UFUNCTION(BlueprintPure = false, Category = "LearningAgents", meta = (AdvancedDisplay = 4, ReturnDisplayName = "Success"))
-	static bool GetSpeedAction(float& OutSpeed, const ULearningAgentsActionObject* Object, const FLearningAgentsActionObjectElement Element, const float SpeedScale = 200.0f, const FName Tag = TEXT("Speed"));
+	static bool GetPairAction(FLearningAgentsActionObjectElement& OutKey, FLearningAgentsActionObjectElement& OutValue, const ULearningAgentsActionObject* Object, const FLearningAgentsActionObjectElement Element, const FName Tag = TEXT("PairAction"));
+
+	/**
+	 * Get the enum value of an enum action.
+	 *
+	 * @param OutEnumValue The output enum value.
+	 * @param Object The Action Object
+	 * @param Element The Action Object Element
+	 * @param Enum The Enum
+	 * @param Tag The tag of the corresponding action. Must match the tag given during Specify.
+	 * @param bVisualLoggerEnabled When true, debug data will be sent to the visual logger.
+	 * @param VisualLoggerListener The listener object which is making this action. This must be set to use logging.
+	 * @param VisualLoggerAgentId The agent id associated with this action.
+	 * @param VisualLoggerLocation A location for the visual logger information in the world.
+	 * @param VisualLoggerColor The color for the visual logger display.
+	 * @return true if the provided Element is the correct type, otherwise false.
+	 */
+	UFUNCTION(BlueprintPure = false, Category = "LearningAgents", meta = (AdvancedDisplay = 4, ReturnDisplayName = "Success", DefaultToSelf = "VisualLoggerListener"))
+	static bool GetEnumAction(
+		uint8& OutEnumValue, 
+		const ULearningAgentsActionObject* Object, 
+		const FLearningAgentsActionObjectElement Element, 
+		const UEnum* Enum, 
+		const FName Tag = TEXT("EnumAction"),
+		const bool bVisualLoggerEnabled = false,
+		ULearningAgentsManagerListener* VisualLoggerListener = nullptr,
+		const int32 VisualLoggerAgentId = -1,
+		const FVector VisualLoggerLocation = FVector::ZeroVector,
+		const FLinearColor VisualLoggerColor = FLinearColor::Blue);
+
+	/**
+	 * Get the bitmask value of a bitmask action.
+	 *
+	 * @param OutBitmaskValue The output bitmask value.
+	 * @param Object The Action Object
+	 * @param Element The Action Object Element
+	 * @param Enum The Enum
+	 * @param Tag The tag of the corresponding action. Must match the tag given during Specify.
+	 * @param bVisualLoggerEnabled When true, debug data will be sent to the visual logger.
+	 * @param VisualLoggerListener The listener object which is making this action. This must be set to use logging.
+	 * @param VisualLoggerAgentId The agent id associated with this action.
+	 * @param VisualLoggerLocation A location for the visual logger information in the world.
+	 * @param VisualLoggerColor The color for the visual logger display.
+	 * @return true if the provided Element is the correct type, otherwise false.
+	 */
+	UFUNCTION(BlueprintPure = false, Category = "LearningAgents", meta = (AdvancedDisplay = 4, ReturnDisplayName = "Success", DefaultToSelf = "VisualLoggerListener"))
+	static bool GetBitmaskAction(
+		int32& OutBitmaskValue, 
+		const ULearningAgentsActionObject* Object, 
+		const FLearningAgentsActionObjectElement Element, 
+		const UEnum* Enum, 
+		const FName Tag = TEXT("BitmaskAction"),
+		const bool bVisualLoggerEnabled = false,
+		ULearningAgentsManagerListener* VisualLoggerListener = nullptr,
+		const int32 VisualLoggerAgentId = -1,
+		const FVector VisualLoggerLocation = FVector::ZeroVector,
+		const FLinearColor VisualLoggerColor = FLinearColor::Blue);
+
+	/**
+	 * Get the sub-action of an option action.
+	 *
+	 * @param OutOption The output optional specifier.
+	 * @param OutElement The output sub-action.
+	 * @param Object The Action Object
+	 * @param Element The Action Object Element
+	 * @param Tag The tag of the corresponding action. Must match the tag given during Specify.
+	 * @return true if the provided Element is the correct type, otherwise false.
+	 */
+	UFUNCTION(BlueprintPure = false, Category = "LearningAgents", meta = (AdvancedDisplay = 4, ExpandEnumAsExecs = "OutOption", ReturnDisplayName = "Success"))
+	static bool GetOptionalAction(ELearningAgentsOptionalAction& OutOption, FLearningAgentsActionObjectElement& OutElement, const ULearningAgentsActionObject* Object, const FLearningAgentsActionObjectElement Element, const FName Tag = TEXT("OptionalAction"));
+
+	/**
+	 * Get the sub-action of an either action.
+	 *
+	 * @param OutEither The output either specifier.
+	 * @param OutElement The output sub-action.
+	 * @param Object The Action Object
+	 * @param Element The Action Object Element
+	 * @param Tag The tag of the corresponding action. Must match the tag given during Specify.
+	 * @return true if the provided Element is the correct type, otherwise false.
+	 */
+	UFUNCTION(BlueprintPure = false, Category = "LearningAgents", meta = (AdvancedDisplay = 4, ExpandEnumAsExecs = "OutEither", ReturnDisplayName = "Success"))
+	static bool GetEitherAction(ELearningAgentsEitherAction& OutEither, FLearningAgentsActionObjectElement& OutElement, const ULearningAgentsActionObject* Object, const FLearningAgentsActionObjectElement Element, const FName Tag = TEXT("EitherAction"));
+
+	/**
+	 * Get the sub-action of an encoding action.
+	 *
+	 * @param OutElement The output sub-action.
+	 * @param Object The Action Object
+	 * @param Element The Action Object Element
+	 * @param Tag The tag of the corresponding action. Must match the tag given during Specify.
+	 * @return true if the provided Element is the correct type, otherwise false.
+	 */
+	UFUNCTION(BlueprintPure = false, Category = "LearningAgents", meta = (AdvancedDisplay = 3, ReturnDisplayName = "Success"))
+	static bool GetEncodingAction(FLearningAgentsActionObjectElement& OutElement, const ULearningAgentsActionObject* Object, const FLearningAgentsActionObjectElement Element, const FName Tag = TEXT("EncodingAction"));
+
+	/**
+	 * Get the value for a bool action.
+	 *
+	 * @param bOutValue The output bool value.
+	 * @param Object The Action Object
+	 * @param Element The Action Object Element
+	 * @param Tag The tag of the corresponding action. Must match the tag given during Specify.
+	 * @param bVisualLoggerEnabled When true, debug data will be sent to the visual logger.
+	 * @param VisualLoggerListener The listener object which is making this action. This must be set to use logging.
+	 * @param VisualLoggerAgentId The agent id associated with this action.
+	 * @param VisualLoggerLocation A location for the visual logger information in the world.
+	 * @param VisualLoggerColor The color for the visual logger display.
+	 * @return true if the provided Element is the correct type, otherwise false.
+	 */
+	UFUNCTION(BlueprintPure = false, Category = "LearningAgents", meta = (AdvancedDisplay = 3, ReturnDisplayName = "Success", DefaultToSelf = "VisualLoggerListener"))
+	static bool GetBoolAction(
+		bool& bOutValue, 
+		const ULearningAgentsActionObject* Object, 
+		const FLearningAgentsActionObjectElement Element, 
+		const FName Tag = TEXT("BoolAction"),
+		const bool bVisualLoggerEnabled = false,
+		ULearningAgentsManagerListener* VisualLoggerListener = nullptr,
+		const int32 VisualLoggerAgentId = -1,
+		const FVector VisualLoggerLocation = FVector::ZeroVector,
+		const FLinearColor VisualLoggerColor = FLinearColor::Blue);
+
+	/**
+	 * Get the value for a float action.
+	 *
+	 * @param OutValue The output float value.
+	 * @param Object The Action Object
+	 * @param Element The Action Object Element
+	 * @param FloatScale The scale used to control the overall magnitude of the outputted float action.
+	 * @param Tag The tag of the corresponding action. Must match the tag given during Specify.
+	 * @param bVisualLoggerEnabled When true, debug data will be sent to the visual logger.
+	 * @param VisualLoggerListener The listener object which is making this action. This must be set to use logging.
+	 * @param VisualLoggerAgentId The agent id associated with this action.
+	 * @param VisualLoggerLocation A location for the visual logger information in the world.
+	 * @param VisualLoggerColor The color for the visual logger display.
+	 * @return true if the provided Element is the correct type, otherwise false.
+	 */
+	UFUNCTION(BlueprintPure = false, Category = "LearningAgents", meta = (AdvancedDisplay = 4, ReturnDisplayName = "Success", DefaultToSelf = "VisualLoggerListener"))
+	static bool GetFloatAction(
+		float& OutValue, 
+		const ULearningAgentsActionObject* Object, 
+		const FLearningAgentsActionObjectElement Element, 
+		const float FloatScale = 1.0f, 
+		const FName Tag = TEXT("FloatAction"),
+		const bool bVisualLoggerEnabled = false,
+		ULearningAgentsManagerListener* VisualLoggerListener = nullptr,
+		const int32 VisualLoggerAgentId = -1,
+		const FVector VisualLoggerLocation = FVector::ZeroVector,
+		const FLinearColor VisualLoggerColor = FLinearColor::Blue);
+
+	/**
+	 * Get the value for a location action.
+	 *
+	 * @param OutLocation The output location value.
+	 * @param Object The Action Object
+	 * @param Element The Action Object Element
+	 * @param RelativeTransform The relative transform to transform the location by.
+	 * @param LocationScale The scale used to control the overall magnitude of the outputted location action.
+	 * @param Tag The tag of the corresponding action. Must match the tag given during Specify.
+	 * @param bVisualLoggerEnabled When true, debug data will be sent to the visual logger.
+	 * @param VisualLoggerListener The listener object which is making this action. This must be set to use logging.
+	 * @param VisualLoggerAgentId The agent id associated with this action.
+	 * @param VisualLoggerLocation A location for the visual logger information in the world.
+	 * @param VisualLoggerColor The color for the visual logger display.
+	 * @return true if the provided Element is the correct type, otherwise false.
+	 */
+	UFUNCTION(BlueprintPure = false, Category = "LearningAgents", meta = (AdvancedDisplay = 5, ReturnDisplayName = "Success", DefaultToSelf = "VisualLoggerListener"))
+	static bool GetLocationAction(
+		FVector& OutLocation, 
+		const ULearningAgentsActionObject* Object, 
+		const FLearningAgentsActionObjectElement Element, 
+		const FTransform RelativeTransform = FTransform(), 
+		const float LocationScale = 100.0f, 
+		const FName Tag = TEXT("LocationAction"),
+		const bool bVisualLoggerEnabled = false,
+		ULearningAgentsManagerListener* VisualLoggerListener = nullptr,
+		const int32 VisualLoggerAgentId = -1,
+		const FVector VisualLoggerLocation = FVector::ZeroVector,
+		const FLinearColor VisualLoggerColor = FLinearColor::Blue);
+
+	/**
+	 * Get the value for a rotation action.
+	 *
+	 * @param OutRotation The output rotation value.
+	 * @param Object The Action Object
+	 * @param Element The Action Object Element
+	 * @param RelativeRotation The relative rotation to transform the rotation by.
+	 * @param RotationScale The scale used to control the overall magnitude of the outputted rotation action.
+	 * @param Tag The tag of the corresponding action. Must match the tag given during Specify.
+	 * @param bVisualLoggerEnabled When true, debug data will be sent to the visual logger.
+	 * @param VisualLoggerListener The listener object which is making this action. This must be set to use logging.
+	 * @param VisualLoggerAgentId The agent id associated with this action.
+	 * @param VisualLoggerRotationLocation A location for the visual logger to display the rotation in the world.
+	 * @param VisualLoggerLocation A location for the visual logger information in the world.
+	 * @param VisualLoggerColor The color for the visual logger display.
+	 * @return true if the provided Element is the correct type, otherwise false.
+	 */
+	UFUNCTION(BlueprintPure = false, Category = "LearningAgents", meta = (AdvancedDisplay = 5, ReturnDisplayName = "Success", DefaultToSelf = "VisualLoggerListener"))
+	static bool GetRotationAction(
+		FRotator& OutRotation, 
+		const ULearningAgentsActionObject* Object, 
+		const FLearningAgentsActionObjectElement Element, 
+		const FRotator RelativeRotation = FRotator::ZeroRotator, 
+		const float RotationScale = 90.0f, 
+		const FName Tag = TEXT("RotationAction"),
+		const bool bVisualLoggerEnabled = false,
+		ULearningAgentsManagerListener* VisualLoggerListener = nullptr,
+		const int32 VisualLoggerAgentId = -1,
+		const FVector VisualLoggerRotationLocation = FVector::ZeroVector,
+		const FVector VisualLoggerLocation = FVector::ZeroVector,
+		const FLinearColor VisualLoggerColor = FLinearColor::Blue);
+
+	/**
+	 * Get the value for a rotation action as a quaternion.
+	 *
+	 * @param OutRotation The output rotation value.
+	 * @param Object The Action Object
+	 * @param Element The Action Object Element
+	 * @param RelativeRotation The relative rotation to transform the rotation by.
+	 * @param RotationScale The scale used to control the overall magnitude of the outputted rotation action.
+	 * @param Tag The tag of the corresponding action. Must match the tag given during Specify.
+	 * @param bVisualLoggerEnabled When true, debug data will be sent to the visual logger.
+	 * @param VisualLoggerListener The listener object which is making this action. This must be set to use logging.
+	 * @param VisualLoggerAgentId The agent id associated with this action.
+	 * @param VisualLoggerRotationLocation A location for the visual logger to display the rotation in the world.
+	 * @param VisualLoggerLocation A location for the visual logger information in the world.
+	 * @param VisualLoggerColor The color for the visual logger display.
+	 * @return true if the provided Element is the correct type, otherwise false.
+	 */
+	UFUNCTION(BlueprintPure = false, Category = "LearningAgents", meta = (AdvancedDisplay = 5, ReturnDisplayName = "Success", DefaultToSelf = "VisualLoggerListener"))
+	static bool GetRotationActionAsQuat(
+		FQuat& OutRotation, 
+		const ULearningAgentsActionObject* Object, 
+		const FLearningAgentsActionObjectElement Element, 
+		const FQuat RelativeRotation, 
+		const float RotationScale = 90.0f, 
+		const FName Tag = TEXT("RotationAction"),
+		const bool bVisualLoggerEnabled = false,
+		ULearningAgentsManagerListener* VisualLoggerListener = nullptr,
+		const int32 VisualLoggerAgentId = -1,
+		const FVector VisualLoggerRotationLocation = FVector::ZeroVector,
+		const FVector VisualLoggerLocation = FVector::ZeroVector,
+		const FLinearColor VisualLoggerColor = FLinearColor::Blue);
+
+	/**
+	 * Get the value for a scale action.
+	 *
+	 * @param OutScale The output scale value.
+	 * @param Object The Action Object
+	 * @param Element The Action Object Element
+	 * @param RelativeScale The relative scale to transform the scale by.
+	 * @param Scale The scale used to control the overall magnitude of the outputted scale action.
+	 * @param Tag The tag of the corresponding action. Must match the tag given during Specify.
+	 * @param bVisualLoggerEnabled When true, debug data will be sent to the visual logger.
+	 * @param VisualLoggerListener The listener object which is making this action. This must be set to use logging.
+	 * @param VisualLoggerAgentId The agent id associated with this action.
+	 * @param VisualLoggerLocation A location for the visual logger information in the world.
+	 * @param VisualLoggerColor The color for the visual logger display.
+	 * @return true if the provided Element is the correct type, otherwise false.
+	 */
+	UFUNCTION(BlueprintPure = false, Category = "LearningAgents", meta = (AdvancedDisplay = 5, ReturnDisplayName = "Success", DefaultToSelf = "VisualLoggerListener"))
+	static bool GetScaleAction(
+		FVector& OutScale, 
+		const ULearningAgentsActionObject* Object, 
+		const FLearningAgentsActionObjectElement Element, 
+		const FVector RelativeScale = FVector(1, 1, 1), 
+		const float Scale = 1.0f, 
+		const FName Tag = TEXT("ScaleAction"),
+		const bool bVisualLoggerEnabled = false,
+		ULearningAgentsManagerListener* VisualLoggerListener = nullptr,
+		const int32 VisualLoggerAgentId = -1,
+		const FVector VisualLoggerLocation = FVector::ZeroVector,
+		const FLinearColor VisualLoggerColor = FLinearColor::Blue);
+
+	/**
+	 * Get the value for a transform action.
+	 *
+	 * @param OutTransform The output transform value.
+	 * @param Object The Action Object
+	 * @param Element The Action Object Element
+	 * @param RelativeTransform The relative transform.
+	 * @param LocationScale The scale used to control the overall magnitude of the outputted location action.
+	 * @param RotationScale The scale used to control the overall magnitude of the outputted rotation action.
+	 * @param ScaleScale The scale used to control the overall magnitude of the outputted scale action.
+	 * @param Tag The tag of the corresponding action. Must match the tag given during Specify.
+	 * @param bVisualLoggerEnabled When true, debug data will be sent to the visual logger.
+	 * @param VisualLoggerListener The listener object which is making this action. This must be set to use logging.
+	 * @param VisualLoggerAgentId The agent id associated with this action.
+	 * @param VisualLoggerLocation A location for the visual logger information in the world.
+	 * @param VisualLoggerColor The color for the visual logger display.
+	 * @return true if the provided Element is the correct type, otherwise false.
+	 */
+	UFUNCTION(BlueprintPure = false, Category = "LearningAgents", meta = (AdvancedDisplay = 7, ReturnDisplayName = "Success", DefaultToSelf = "VisualLoggerListener"))
+	static bool GetTransformAction(
+		FTransform& OutTransform, 
+		const ULearningAgentsActionObject* Object, 
+		const FLearningAgentsActionObjectElement Element, 
+		const FTransform RelativeTransform = FTransform(), 
+		const float LocationScale = 100.0f, 
+		const float RotationScale = 1.0f, 
+		const float ScaleScale = 1.0f, 
+		const FName Tag = TEXT("TransformAction"),
+		const bool bVisualLoggerEnabled = false,
+		ULearningAgentsManagerListener* VisualLoggerListener = nullptr,
+		const int32 VisualLoggerAgentId = -1,
+		const FVector VisualLoggerLocation = FVector::ZeroVector,
+		const FLinearColor VisualLoggerColor = FLinearColor::Blue);
+
+	/**
+	 * Get the value for an angle action. Returned angle is in degrees.
+	 *
+	 * @param OutAngle The output angle value.
+	 * @param Object The Action Object
+	 * @param Element The Action Object Element
+	 * @param RelativeAngle The relative angle to transform the angle by.
+	 * @param AngleScale The scale used to control the overall magnitude of the outputted scale action.
+	 * @param Tag The tag of the corresponding action. Must match the tag given during Specify.
+	 * @param bVisualLoggerEnabled When true, debug data will be sent to the visual logger.
+	 * @param VisualLoggerListener The listener object which is making this action. This must be set to use logging.
+	 * @param VisualLoggerAgentId The agent id associated with this action.
+	 * @param VisualLoggerAngleLocation A location for the visual logger to display the angle in the world.
+	 * @param VisualLoggerLocation A location for the visual logger information in the world.
+	 * @param VisualLoggerColor The color for the visual logger display.
+	 * @return true if the provided Element is the correct type, otherwise false.
+	 */
+	UFUNCTION(BlueprintPure = false, Category = "LearningAgents", meta = (AdvancedDisplay = 5, ReturnDisplayName = "Success", DefaultToSelf = "VisualLoggerListener"))
+	static bool GetAngleAction(
+		float& OutAngle, 
+		const ULearningAgentsActionObject* Object, 
+		const FLearningAgentsActionObjectElement Element, 
+		const float RelativeAngle = 0.0f, 
+		const float AngleScale = 90.0f, 
+		const FName Tag = TEXT("AngleAction"),
+		const bool bVisualLoggerEnabled = false,
+		ULearningAgentsManagerListener* VisualLoggerListener = nullptr,
+		const int32 VisualLoggerAgentId = -1,
+		const FVector VisualLoggerAngleLocation = FVector::ZeroVector,
+		const FVector VisualLoggerLocation = FVector::ZeroVector,
+		const FLinearColor VisualLoggerColor = FLinearColor::Blue);
+
+	/**
+	 * Get the value for an angle action. Returned angle is in radians.
+	 *
+	 * @param OutAngle The output angle value.
+	 * @param Object The Action Object
+	 * @param Element The Action Object Element
+	 * @param RelativeAngle The relative angle to transform the angle by.
+	 * @param AngleScale The scale used to control the overall magnitude of the outputted scale action.
+	 * @param Tag The tag of the corresponding action. Must match the tag given during Specify.
+	 * @param bVisualLoggerEnabled When true, debug data will be sent to the visual logger.
+	 * @param VisualLoggerListener The listener object which is making this action. This must be set to use logging.
+	 * @param VisualLoggerAgentId The agent id associated with this action.
+	 * @param VisualLoggerAngleLocation A location for the visual logger to display the angle in the world.
+	 * @param VisualLoggerLocation A location for the visual logger information in the world.
+	 * @param VisualLoggerColor The color for the visual logger display.
+	 * @return true if the provided Element is the correct type, otherwise false.
+	 */
+	UFUNCTION(BlueprintPure = false, Category = "LearningAgents", meta = (AdvancedDisplay = 5, ReturnDisplayName = "Success", DefaultToSelf = "VisualLoggerListener"))
+	static bool GetAngleActionRadians(
+		float& OutAngle, 
+		const ULearningAgentsActionObject* Object,
+		const FLearningAgentsActionObjectElement Element, 
+		const float RelativeAngle = 0.0f,
+		const float AngleScale = 1.57079632679f, 
+		const FName Tag = TEXT("AngleAction"),
+		const bool bVisualLoggerEnabled = false,
+		ULearningAgentsManagerListener* VisualLoggerListener = nullptr,
+		const int32 VisualLoggerAgentId = -1,
+		const FVector VisualLoggerAngleLocation = FVector::ZeroVector,
+		const FVector VisualLoggerLocation = FVector::ZeroVector,
+		const FLinearColor VisualLoggerColor = FLinearColor::Blue);
+
+	/**
+	 * Get the value for a velocity action.
+	 *
+	 * @param OutVelocity The output velocity value.
+	 * @param Object The Action Object
+	 * @param Element The Action Object Element
+	 * @param RelativeTransform The relative transform to transform the velocity by.
+	 * @param VelocityScale The scale used to control the overall magnitude of the outputted velocity action.
+	 * @param Tag The tag of the corresponding action. Must match the tag given during Specify.
+	 * @param bVisualLoggerEnabled When true, debug data will be sent to the visual logger.
+	 * @param VisualLoggerListener The listener object which is making this action. This must be set to use logging.
+	 * @param VisualLoggerAgentId The agent id associated with this action.
+	 * @param VisualLoggerVelocityLocation A location for the visual logger to display the velocity in the world.
+	 * @param VisualLoggerLocation A location for the visual logger information in the world.
+	 * @param VisualLoggerColor The color for the visual logger display.
+	 * @return true if the provided Element is the correct type, otherwise false.
+	 */
+	UFUNCTION(BlueprintPure = false, Category = "LearningAgents", meta = (AdvancedDisplay = 5, ReturnDisplayName = "Success", DefaultToSelf = "VisualLoggerListener"))
+	static bool GetVelocityAction(
+		FVector& OutVelocity, 
+		const ULearningAgentsActionObject* Object, 
+		const FLearningAgentsActionObjectElement Element, 
+		const FTransform RelativeTransform = FTransform(), 
+		const float VelocityScale = 200.0f, 
+		const FName Tag = TEXT("VelocityAction"),
+		const bool bVisualLoggerEnabled = false,
+		ULearningAgentsManagerListener* VisualLoggerListener = nullptr,
+		const int32 VisualLoggerAgentId = -1,
+		const FVector VisualLoggerVelocityLocation = FVector::ZeroVector,
+		const FVector VisualLoggerLocation = FVector::ZeroVector,
+		const FLinearColor VisualLoggerColor = FLinearColor::Blue);
+
+	/**
+	 * Get the value for a direction action.
+	 *
+	 * @param OutDirection The output direction value.
+	 * @param Object The Action Object
+	 * @param Element The Action Object Element
+	 * @param RelativeTransform The relative transform to transform the direction by.
+	 * @param Tag The tag of the corresponding action. Must match the tag given during Specify.
+	 * @param bVisualLoggerEnabled When true, debug data will be sent to the visual logger.
+	 * @param VisualLoggerListener The listener object which is making this action. This must be set to use logging.
+	 * @param VisualLoggerAgentId The agent id associated with this action.
+	 * @param VisualLoggerDirectionLocation A location for the visual logger to display the direction in the world.
+	 * @param VisualLoggerLocation A location for the visual logger information in the world.
+	 * @param VisualLoggerArrowLength The length of the arrow to display to represent the direction.
+	 * @param VisualLoggerColor The color for the visual logger display.
+	 * @return true if the provided Element is the correct type, otherwise false.
+	 */
+	UFUNCTION(BlueprintPure = false, Category = "LearningAgents", meta = (AdvancedDisplay = 4, ReturnDisplayName = "Success", DefaultToSelf = "VisualLoggerListener"))
+	static bool GetDirectionAction(
+		FVector& OutDirection, 
+		const ULearningAgentsActionObject* Object, 
+		const FLearningAgentsActionObjectElement Element, 
+		const FTransform RelativeTransform = FTransform(),
+		const FName Tag = TEXT("DirectionAction"),
+		const bool bVisualLoggerEnabled = false,
+		ULearningAgentsManagerListener* VisualLoggerListener = nullptr,
+		const int32 VisualLoggerAgentId = -1,
+		const FVector VisualLoggerDirectionLocation = FVector::ZeroVector,
+		const FVector VisualLoggerLocation = FVector::ZeroVector,
+		const float VisualLoggerArrowLength = 100.0f,
+		const FLinearColor VisualLoggerColor = FLinearColor::Blue);
 };
 
 
diff --git a/Engine/Plugins/Experimental/LearningAgents/Source/LearningAgents/Public/LearningAgentsObservations.h b/Engine/Plugins/Experimental/LearningAgents/Source/LearningAgents/Public/LearningAgentsObservations.h
--- a/Engine/Plugins/Experimental/LearningAgents/Source/LearningAgents/Public/LearningAgentsObservations.h
+++ b/Engine/Plugins/Experimental/LearningAgents/Source/LearningAgents/Public/LearningAgentsObservations.h
@@ -117,7 +117,7 @@ public:
 	static UEnum* FindEnumByName(const FString& Name);
 
 	/**
-	 * Validates that the given object matches the schema. Will log errors on objects that don't match.
+	 * Validates that the given observation object matches the schema. Will log errors on objects that don't match.
 	 *
 	 * @param Schema				Observation Schema
 	 * @param SchemaElement			Observation Schema Element
@@ -126,7 +126,7 @@ public:
 	 * @returns						true if the object matches the schema
 	 */
 	UFUNCTION(BlueprintPure = false, Category = "LearningAgents")
-	static bool ValidateObjectMatchesSchema(
+	static bool ValidateObservationObjectMatchesSchema(
 		const ULearningAgentsObservationSchema* Schema,
 		const FLearningAgentsObservationSchemaElement SchemaElement,
 		const ULearningAgentsObservationObject* Object,
@@ -151,7 +151,7 @@ public:
 	 * @return The newly created observation schema element.
 	 */
 	UFUNCTION(BlueprintPure, Category = "LearningAgents", meta = (AdvancedDisplay = 1))
-	static FLearningAgentsObservationSchemaElement SpecifyNullObservation(ULearningAgentsObservationSchema* Schema, const FName Tag = TEXT("Null"));
+	static FLearningAgentsObservationSchemaElement SpecifyNullObservation(ULearningAgentsObservationSchema* Schema, const FName Tag = TEXT("NullObservation"));
 
 	/**
 	 * Specifies a new continuous observation. This represents an observation made up of several float values.
@@ -162,7 +162,7 @@ public:
 	 * @return The newly created observation schema element.
 	 */
 	UFUNCTION(BlueprintPure, Category = "LearningAgents", meta = (AdvancedDisplay = 2))
-	static FLearningAgentsObservationSchemaElement SpecifyContinuousObservation(ULearningAgentsObservationSchema* Schema, const int32 Size, const FName Tag = TEXT("Continuous"));
+	static FLearningAgentsObservationSchemaElement SpecifyContinuousObservation(ULearningAgentsObservationSchema* Schema, const int32 Size, const FName Tag = TEXT("ContinuousObservation"));
 
 	/**
 	 * Specifies a new exclusive discrete observation. This represents a discrete observation which is an exclusive selection from multiple choices.
@@ -173,7 +173,7 @@ public:
 	 * @return The newly created observation schema element.
 	 */
 	UFUNCTION(BlueprintPure, Category = "LearningAgents", meta = (AdvancedDisplay = 2))
-	static FLearningAgentsObservationSchemaElement SpecifyExclusiveDiscreteObservation(ULearningAgentsObservationSchema* Schema, const int32 Size, const FName Tag = TEXT("ExclusiveDiscrete"));
+	static FLearningAgentsObservationSchemaElement SpecifyExclusiveDiscreteObservation(ULearningAgentsObservationSchema* Schema, const int32 Size, const FName Tag = TEXT("ExclusiveDiscreteObservation"));
 
 	/**
 	 * Specifies a new inclusive discrete observation. This represents a discrete observation which is an inclusive selection from multiple choices.
@@ -184,7 +184,7 @@ public:
 	 * @return The newly created observation schema element.
 	 */
 	UFUNCTION(BlueprintPure, Category = "LearningAgents", meta = (AdvancedDisplay = 2))
-	static FLearningAgentsObservationSchemaElement SpecifyInclusiveDiscreteObservation(ULearningAgentsObservationSchema* Schema, const int32 Size, const FName Tag = TEXT("InclusiveDiscrete"));
+	static FLearningAgentsObservationSchemaElement SpecifyInclusiveDiscreteObservation(ULearningAgentsObservationSchema* Schema, const int32 Size, const FName Tag = TEXT("InclusiveDiscreteObservation"));
 
 	/**
 	 * Specifies a new index observation. This represents an index into something such as an array.
@@ -195,7 +195,7 @@ public:
 	 * @return The newly created observation schema element.
 	 */
 	UFUNCTION(BlueprintPure, Category = "LearningAgents", meta = (AdvancedDisplay = 2))
-	static FLearningAgentsObservationSchemaElement SpecifyIndexObservation(ULearningAgentsObservationSchema* Schema, const int32 Size, const FName Tag = TEXT("Index"));
+	static FLearningAgentsObservationSchemaElement SpecifyIndexObservation(ULearningAgentsObservationSchema* Schema, const int32 Size, const FName Tag = TEXT("IndexObservation"));
 
 	/**
 	 * Specifies a new count observation. This represents a count of something such as the size of an array.
@@ -205,7 +205,7 @@ public:
 	 * @return The newly created observation schema element.
 	 */
 	UFUNCTION(BlueprintPure, Category = "LearningAgents", meta = (AdvancedDisplay = 1))
-	static FLearningAgentsObservationSchemaElement SpecifyCountObservation(ULearningAgentsObservationSchema* Schema, const FName Tag = TEXT("Count"));
+	static FLearningAgentsObservationSchemaElement SpecifyCountObservation(ULearningAgentsObservationSchema* Schema, const FName Tag = TEXT("CountObservation"));
 
 	/**
 	 * Specifies a new struct observation. This represents a group of named sub-observations.
@@ -216,7 +216,7 @@ public:
 	 * @return The newly created observation schema element.
 	 */
 	UFUNCTION(BlueprintPure, Category = "LearningAgents", meta = (AdvancedDisplay = 2))
-	static FLearningAgentsObservationSchemaElement SpecifyStructObservation(ULearningAgentsObservationSchema* Schema, const TMap<FName, FLearningAgentsObservationSchemaElement>& Elements, const FName Tag = TEXT("Struct"));
+	static FLearningAgentsObservationSchemaElement SpecifyStructObservation(ULearningAgentsObservationSchema* Schema, const TMap<FName, FLearningAgentsObservationSchemaElement>& Elements, const FName Tag = TEXT("StructObservation"));
 
 	/**
 	 * Specifies a new struct observation. This represents a group of named sub-observations.
@@ -228,7 +228,7 @@ public:
 	 * @return The newly created observation schema element.
 	 */
 	UFUNCTION(BlueprintPure, Category = "LearningAgents", meta = (AdvancedDisplay = 3))
-	static FLearningAgentsObservationSchemaElement SpecifyStructObservationFromArrays(ULearningAgentsObservationSchema* Schema, const TArray<FName>& ElementNames, const TArray<FLearningAgentsObservationSchemaElement>& Elements, const FName Tag = TEXT("Struct"));
+	static FLearningAgentsObservationSchemaElement SpecifyStructObservationFromArrays(ULearningAgentsObservationSchema* Schema, const TArray<FName>& ElementNames, const TArray<FLearningAgentsObservationSchemaElement>& Elements, const FName Tag = TEXT("StructObservation"));
 	
 	/**
 	 * Specifies a new struct observation. This represents a group of named sub-observations.
@@ -239,7 +239,7 @@ public:
 	 * @param Tag The tag of this new observation. Used during observation object validation and debugging.
 	 * @return The newly created observation schema element.
 	 */
-	static FLearningAgentsObservationSchemaElement SpecifyStructObservationFromArrayViews(ULearningAgentsObservationSchema* Schema, const TArrayView<const FName> ElementNames, const TArrayView<const FLearningAgentsObservationSchemaElement> Elements, const FName Tag = TEXT("Struct"));
+	static FLearningAgentsObservationSchemaElement SpecifyStructObservationFromArrayViews(ULearningAgentsObservationSchema* Schema, const TArrayView<const FName> ElementNames, const TArrayView<const FLearningAgentsObservationSchemaElement> Elements, const FName Tag = TEXT("StructObservation"));
 
 	/**
 	 * Specifies a new exclusive union observation. This represents an observation which is exclusively chosen from a set of named sub-observations.
@@ -252,7 +252,7 @@ public:
 	 * @return The newly created observation schema element.
 	 */
 	UFUNCTION(BlueprintPure, Category = "LearningAgents", meta = (AdvancedDisplay = 2))
-	static FLearningAgentsObservationSchemaElement SpecifyExclusiveUnionObservation(ULearningAgentsObservationSchema* Schema, const TMap<FName, FLearningAgentsObservationSchemaElement>& Elements, const int32 EncodingSize = 128, const FName Tag = TEXT("ExclusiveUnion"));
+	static FLearningAgentsObservationSchemaElement SpecifyExclusiveUnionObservation(ULearningAgentsObservationSchema* Schema, const TMap<FName, FLearningAgentsObservationSchemaElement>& Elements, const int32 EncodingSize = 128, const FName Tag = TEXT("ExclusiveUnionObservation"));
 
 	/**
 	 * Specifies a new exclusive union observation. This represents an observation which is exclusively chosen from a set of named sub-observations.
@@ -266,7 +266,7 @@ public:
 	 * @return The newly created observation schema element.
 	 */
 	UFUNCTION(BlueprintPure, Category = "LearningAgents", meta = (AdvancedDisplay = 3))
-	static FLearningAgentsObservationSchemaElement SpecifyExclusiveUnionObservationFromArrays(ULearningAgentsObservationSchema* Schema, const TArray<FName>& ElementNames, const TArray<FLearningAgentsObservationSchemaElement>& Elements, const int32 EncodingSize = 128, const FName Tag = TEXT("ExclusiveUnion"));
+	static FLearningAgentsObservationSchemaElement SpecifyExclusiveUnionObservationFromArrays(ULearningAgentsObservationSchema* Schema, const TArray<FName>& ElementNames, const TArray<FLearningAgentsObservationSchemaElement>& Elements, const int32 EncodingSize = 128, const FName Tag = TEXT("ExclusiveUnionObservation"));
 	
 	/**
 	 * Specifies a new exclusive union observation. This represents an observation which is exclusively chosen from a set of named sub-observations.
@@ -279,7 +279,7 @@ public:
 	 * @param Tag The tag of this new observation. Used during observation object validation and debugging.
 	 * @return The newly created observation schema element.
 	 */
-	static FLearningAgentsObservationSchemaElement SpecifyExclusiveUnionObservationFromArrayViews(ULearningAgentsObservationSchema* Schema, const TArrayView<const FName> ElementNames, const TArrayView<const FLearningAgentsObservationSchemaElement> Elements, const int32 EncodingSize = 128, const FName Tag = TEXT("ExclusiveUnion"));
+	static FLearningAgentsObservationSchemaElement SpecifyExclusiveUnionObservationFromArrayViews(ULearningAgentsObservationSchema* Schema, const TArrayView<const FName> ElementNames, const TArrayView<const FLearningAgentsObservationSchemaElement> Elements, const int32 EncodingSize = 128, const FName Tag = TEXT("ExclusiveUnionObservation"));
 
 	/**
 	 * Specifies a new inclusive union observation. This represents an observation which is inclusively chosen from a set of named sub-observations.
@@ -296,7 +296,7 @@ public:
 	 * @return The newly created observation schema element.
 	 */
 	UFUNCTION(BlueprintPure, Category = "LearningAgents", meta = (AdvancedDisplay = 2))
-	static FLearningAgentsObservationSchemaElement SpecifyInclusiveUnionObservation(ULearningAgentsObservationSchema* Schema, const TMap<FName, FLearningAgentsObservationSchemaElement>& Elements, const int32 AttentionEncodingSize = 32, const int32 AttentionHeadNum = 4, const int32 ValueEncodingSize = 32, const FName Tag = TEXT("InclusiveUnion"));
+	static FLearningAgentsObservationSchemaElement SpecifyInclusiveUnionObservation(ULearningAgentsObservationSchema* Schema, const TMap<FName, FLearningAgentsObservationSchemaElement>& Elements, const int32 AttentionEncodingSize = 32, const int32 AttentionHeadNum = 4, const int32 ValueEncodingSize = 32, const FName Tag = TEXT("InclusiveUnionObservation"));
 
 	/**
 	 * Specifies a new inclusive union observation. This represents an observation which is inclusively chosen from a set of named sub-observations.
@@ -314,7 +314,7 @@ public:
 	 * @return The newly created observation schema element.
 	 */
 	UFUNCTION(BlueprintPure, Category = "LearningAgents", meta = (AdvancedDisplay = 3))
-	static FLearningAgentsObservationSchemaElement SpecifyInclusiveUnionObservationFromArrays(ULearningAgentsObservationSchema* Schema, const TArray<FName>& ElementNames, const TArray<FLearningAgentsObservationSchemaElement>& Elements, const int32 AttentionEncodingSize = 32, const int32 AttentionHeadNum = 4, const int32 ValueEncodingSize = 32, const FName Tag = TEXT("InclusiveUnion"));
+	static FLearningAgentsObservationSchemaElement SpecifyInclusiveUnionObservationFromArrays(ULearningAgentsObservationSchema* Schema, const TArray<FName>& ElementNames, const TArray<FLearningAgentsObservationSchemaElement>& Elements, const int32 AttentionEncodingSize = 32, const int32 AttentionHeadNum = 4, const int32 ValueEncodingSize = 32, const FName Tag = TEXT("InclusiveUnionObservation"));
 	
 	/**
 	 * Specifies a new inclusive union observation. This represents an observation which is inclusively chosen from a set of named sub-observations.
@@ -331,7 +331,7 @@ public:
 	 * @param Tag The tag of this new observation. Used during observation object validation and debugging.
 	 * @return The newly created observation schema element.
 	 */
-	static FLearningAgentsObservationSchemaElement SpecifyInclusiveUnionObservationFromArrayViews(ULearningAgentsObservationSchema* Schema, const TArrayView<const FName> ElementNames, const TArrayView<const FLearningAgentsObservationSchemaElement> Elements, const int32 AttentionEncodingSize = 32, const int32 AttentionHeadNum = 4, const int32 ValueEncodingSize = 32, const FName Tag = TEXT("InclusiveUnion"));
+	static FLearningAgentsObservationSchemaElement SpecifyInclusiveUnionObservationFromArrayViews(ULearningAgentsObservationSchema* Schema, const TArrayView<const FName> ElementNames, const TArrayView<const FLearningAgentsObservationSchemaElement> Elements, const int32 AttentionEncodingSize = 32, const int32 AttentionHeadNum = 4, const int32 ValueEncodingSize = 32, const FName Tag = TEXT("InclusiveUnionObservation"));
 
 	/**
 	 * Specifies a new static array observation. This represents an observation made up of a fixed-size array of some other observation.
@@ -343,7 +343,7 @@ public:
 	 * @return The newly created observation schema element.
 	 */
 	UFUNCTION(BlueprintPure, Category = "LearningAgents", meta = (AdvancedDisplay = 3))
-	static FLearningAgentsObservationSchemaElement SpecifyStaticArrayObservation(ULearningAgentsObservationSchema* Schema, const FLearningAgentsObservationSchemaElement Element, const int32 Num, const FName Tag = TEXT("StaticArray"));
+	static FLearningAgentsObservationSchemaElement SpecifyStaticArrayObservation(ULearningAgentsObservationSchema* Schema, const FLearningAgentsObservationSchemaElement Element, const int32 Num, const FName Tag = TEXT("StaticArrayObservation"));
 
 	/**
 	 * Specifies a new set observation. This represents an observation made up of a Set of some other observation. This Set can be variable in size 
@@ -360,7 +360,7 @@ public:
 	 * @return The newly created observation schema element.
 	 */
 	UFUNCTION(BlueprintPure, Category = "LearningAgents", meta = (AdvancedDisplay = 3))
-	static FLearningAgentsObservationSchemaElement SpecifySetObservation(ULearningAgentsObservationSchema* Schema, const FLearningAgentsObservationSchemaElement Element, const int32 MaxNum, const int32 AttentionEncodingSize = 32, const int32 AttentionHeadNum = 4, const int32 ValueEncodingSize = 32, const FName Tag = TEXT("Set"));
+	static FLearningAgentsObservationSchemaElement SpecifySetObservation(ULearningAgentsObservationSchema* Schema, const FLearningAgentsObservationSchemaElement Element, const int32 MaxNum, const int32 AttentionEncodingSize = 32, const int32 AttentionHeadNum = 4, const int32 ValueEncodingSize = 32, const FName Tag = TEXT("SetObservation"));
 
 	/**
 	 * Specifies a new pair observation. This represents an observation made up of two sub-observations.
@@ -372,7 +372,7 @@ public:
 	 * @return The newly created observation schema element.
 	 */
 	UFUNCTION(BlueprintPure, Category = "LearningAgents", meta = (AdvancedDisplay = 3))
-	static FLearningAgentsObservationSchemaElement SpecifyPairObservation(ULearningAgentsObservationSchema* Schema, const FLearningAgentsObservationSchemaElement Key, const FLearningAgentsObservationSchemaElement Value, const FName Tag = TEXT("Pair"));
+	static FLearningAgentsObservationSchemaElement SpecifyPairObservation(ULearningAgentsObservationSchema* Schema, const FLearningAgentsObservationSchemaElement Key, const FLearningAgentsObservationSchemaElement Value, const FName Tag = TEXT("PairObservation"));
 
 	/**
 	 * Specifies a new array observation. This represents an observation made up of an Array of some other observation. This Array can be variable in 
@@ -389,7 +389,7 @@ public:
 	 * @return The newly created observation schema element.
 	 */
 	UFUNCTION(BlueprintPure, Category = "LearningAgents", meta = (AdvancedDisplay = 3))
-	static FLearningAgentsObservationSchemaElement SpecifyArrayObservation(ULearningAgentsObservationSchema* Schema, const FLearningAgentsObservationSchemaElement Element, const int32 MaxNum, const int32 AttentionEncodingSize = 32, const int32 AttentionHeadNum = 4, const int32 ValueEncodingSize = 32, const FName Tag = TEXT("Array"));
+	static FLearningAgentsObservationSchemaElement SpecifyArrayObservation(ULearningAgentsObservationSchema* Schema, const FLearningAgentsObservationSchemaElement Element, const int32 MaxNum, const int32 AttentionEncodingSize = 32, const int32 AttentionHeadNum = 4, const int32 ValueEncodingSize = 32, const FName Tag = TEXT("ArrayObservation"));
 
 	/**
 	 * Specifies a new map observation. This represents an observation made up of a Map of some other key and pair observations. This Map can be 
@@ -407,7 +407,7 @@ public:
 	 * @return The newly created observation schema element.
 	 */
 	UFUNCTION(BlueprintPure, Category = "LearningAgents", meta = (AdvancedDisplay = 4))
-	static FLearningAgentsObservationSchemaElement SpecifyMapObservation(ULearningAgentsObservationSchema* Schema, const FLearningAgentsObservationSchemaElement KeyElement, const FLearningAgentsObservationSchemaElement ValueElement, const int32 MaxNum, const int32 AttentionEncodingSize = 32, const int32 AttentionHeadNum = 4, const int32 ValueEncodingSize = 32, const FName Tag = TEXT("Map"));
+	static FLearningAgentsObservationSchemaElement SpecifyMapObservation(ULearningAgentsObservationSchema* Schema, const FLearningAgentsObservationSchemaElement KeyElement, const FLearningAgentsObservationSchemaElement ValueElement, const int32 MaxNum, const int32 AttentionEncodingSize = 32, const int32 AttentionHeadNum = 4, const int32 ValueEncodingSize = 32, const FName Tag = TEXT("MapObservation"));
 
 	/**
 	 * Specifies a new enum observation. This represents an exclusive choice from elements of the given Enum. To use this with an Enum defined in C++ 
@@ -418,7 +418,7 @@ public:
 	 * @return The newly created observation schema element.
 	 */
 	UFUNCTION(BlueprintPure, Category = "LearningAgents", meta = (AdvancedDisplay = 2))
-	static FLearningAgentsObservationSchemaElement SpecifyEnumObservation(ULearningAgentsObservationSchema* Schema, const UEnum* Enum, const FName Tag = TEXT("Enum"));
+	static FLearningAgentsObservationSchemaElement SpecifyEnumObservation(ULearningAgentsObservationSchema* Schema, const UEnum* Enum, const FName Tag = TEXT("EnumObservation"));
 
 	/**
 	 * Specifies a new bitmask observation. This represents an inclusive choice from elements of the given Enum. To use this with an Enum defined in 
@@ -429,7 +429,7 @@ public:
 	 * @return The newly created observation schema element.
 	 */
 	UFUNCTION(BlueprintPure, Category = "LearningAgents", meta = (AdvancedDisplay = 2))
-	static FLearningAgentsObservationSchemaElement SpecifyBitmaskObservation(ULearningAgentsObservationSchema* Schema, const UEnum* Enum, const FName Tag = TEXT("Bitmask"));
+	static FLearningAgentsObservationSchemaElement SpecifyBitmaskObservation(ULearningAgentsObservationSchema* Schema, const UEnum* Enum, const FName Tag = TEXT("BitmaskObservation"));
 
 	/**
 	 * Specifies a new optional observation. This represents an observation which may or may not be provided.
@@ -440,7 +440,7 @@ public:
 	 * @return The newly created observation schema element.
 	 */
 	UFUNCTION(BlueprintPure, Category = "LearningAgents", meta = (AdvancedDisplay = 2))
-	static FLearningAgentsObservationSchemaElement SpecifyOptionalObservation(ULearningAgentsObservationSchema* Schema, const FLearningAgentsObservationSchemaElement Element, const int32 EncodingSize = 128, const FName Tag = TEXT("Optional"));
+	static FLearningAgentsObservationSchemaElement SpecifyOptionalObservation(ULearningAgentsObservationSchema* Schema, const FLearningAgentsObservationSchemaElement Element, const int32 EncodingSize = 128, const FName Tag = TEXT("OptionalObservation"));
 
 	/**
 	 * Specifies a new either observation. This represents an observation which will be either sub-observation A or sub-observation B.
@@ -452,7 +452,7 @@ public:
 	 * @return The newly created observation schema element.
 	 */
 	UFUNCTION(BlueprintPure, Category = "LearningAgents", meta = (AdvancedDisplay = 3))
-	static FLearningAgentsObservationSchemaElement SpecifyEitherObservation(ULearningAgentsObservationSchema* Schema, const FLearningAgentsObservationSchemaElement A, const FLearningAgentsObservationSchemaElement B, const int32 EncodingSize = 128, const FName Tag = TEXT("Either"));
+	static FLearningAgentsObservationSchemaElement SpecifyEitherObservation(ULearningAgentsObservationSchema* Schema, const FLearningAgentsObservationSchemaElement A, const FLearningAgentsObservationSchemaElement B, const int32 EncodingSize = 128, const FName Tag = TEXT("EitherObservation"));
 
 	/**
 	 * Specifies a new encoding observation. This represents an observation which will be an encoding of another sub-observation using a small neural
@@ -466,7 +466,7 @@ public:
 	 * @return The newly created observation schema element.
 	 */
 	UFUNCTION(BlueprintPure, Category = "LearningAgents", meta = (AdvancedDisplay = 2))
-	static FLearningAgentsObservationSchemaElement SpecifyEncodingObservation(ULearningAgentsObservationSchema* Schema, const FLearningAgentsObservationSchemaElement Element, const int32 EncodingSize = 128, const int32 LayerNum = 1, const ELearningAgentsActivationFunction ActivationFunction = ELearningAgentsActivationFunction::ELU, const FName Tag = TEXT("Encoding"));
+	static FLearningAgentsObservationSchemaElement SpecifyEncodingObservation(ULearningAgentsObservationSchema* Schema, const FLearningAgentsObservationSchemaElement Element, const int32 EncodingSize = 128, const int32 LayerNum = 1, const ELearningAgentsActivationFunction ActivationFunction = ELearningAgentsActivationFunction::ELU, const FName Tag = TEXT("EncodingObservation"));
 
 	/**
 	 * Specifies a new bool observation. A true or false observation.
@@ -476,7 +476,7 @@ public:
 	 * @return The newly created observation schema element.
 	 */
 	UFUNCTION(BlueprintPure, Category = "LearningAgents", meta = (AdvancedDisplay = 1))
-	static FLearningAgentsObservationSchemaElement SpecifyBoolObservation(ULearningAgentsObservationSchema* Schema, const FName Tag = TEXT("Bool"));
+	static FLearningAgentsObservationSchemaElement SpecifyBoolObservation(ULearningAgentsObservationSchema* Schema, const FName Tag = TEXT("BoolObservation"));
 
 	/**
 	 * Specifies a new float observation. A simple observation which can be used as a catch-all for situations where a
@@ -487,7 +487,7 @@ public:
 	 * @return The newly created observation schema element.
 	 */
 	UFUNCTION(BlueprintPure, Category = "LearningAgents", meta = (AdvancedDisplay = 1))
-	static FLearningAgentsObservationSchemaElement SpecifyFloatObservation(ULearningAgentsObservationSchema* Schema, const FName Tag = TEXT("Float"));
+	static FLearningAgentsObservationSchemaElement SpecifyFloatObservation(ULearningAgentsObservationSchema* Schema, const FName Tag = TEXT("FloatObservation"));
 
 	/**
 	 * Specifies a new location observation. Allows an agent to observe the location of some entity.
@@ -497,7 +497,7 @@ public:
 	 * @return The newly created observation schema element.
 	 */
 	UFUNCTION(BlueprintPure, Category = "LearningAgents", meta = (AdvancedDisplay = 1))
-	static FLearningAgentsObservationSchemaElement SpecifyLocationObservation(ULearningAgentsObservationSchema* Schema, const FName Tag = TEXT("Location"));
+	static FLearningAgentsObservationSchemaElement SpecifyLocationObservation(ULearningAgentsObservationSchema* Schema, const FName Tag = TEXT("LocationObservation"));
 
 	/**
 	 * Specifies a new rotation observation. Allows an agent to observe the rotation of some entity. Rotations are encoded as two columns of the 
@@ -508,7 +508,7 @@ public:
 	 * @return The newly created observation schema element.
 	 */
 	UFUNCTION(BlueprintPure, Category = "LearningAgents", meta = (AdvancedDisplay = 1))
-	static FLearningAgentsObservationSchemaElement SpecifyRotationObservation(ULearningAgentsObservationSchema* Schema, const FName Tag = TEXT("Rotation"));
+	static FLearningAgentsObservationSchemaElement SpecifyRotationObservation(ULearningAgentsObservationSchema* Schema, const FName Tag = TEXT("RotationObservation"));
 
 	/**
 	 * Specifies a new scale observation. Allows an agent to observe the scale of some entity. Negative scales are not supported by this observation 
@@ -519,7 +519,7 @@ public:
 	 * @return The newly created observation schema element.
 	 */
 	UFUNCTION(BlueprintPure, Category = "LearningAgents", meta = (AdvancedDisplay = 1))
-	static FLearningAgentsObservationSchemaElement SpecifyScaleObservation(ULearningAgentsObservationSchema* Schema, const FName Tag = TEXT("Scale"));
+	static FLearningAgentsObservationSchemaElement SpecifyScaleObservation(ULearningAgentsObservationSchema* Schema, const FName Tag = TEXT("ScaleObservation"));
 
 	/**
 	 * Specifies a new transform observation. Allows an agent to observe the transform of some entity.
@@ -529,7 +529,7 @@ public:
 	 * @return The newly created observation schema element.
 	 */
 	UFUNCTION(BlueprintPure, Category = "LearningAgents", meta = (AdvancedDisplay = 1))
-	static FLearningAgentsObservationSchemaElement SpecifyTransformObservation(ULearningAgentsObservationSchema* Schema, const FName Tag = TEXT("Transform"));
+	static FLearningAgentsObservationSchemaElement SpecifyTransformObservation(ULearningAgentsObservationSchema* Schema, const FName Tag = TEXT("TransformObservation"));
 
 	/**
 	 * Specifies a new angle observation. This will be encoded as a 2-dimension Cartesian vector so that 0 and 350 are close to each other in the 
@@ -540,7 +540,7 @@ public:
 	 * @return The newly created observation schema element.
 	 */
 	UFUNCTION(BlueprintPure, Category = "LearningAgents", meta = (AdvancedDisplay = 1))
-	static FLearningAgentsObservationSchemaElement SpecifyAngleObservation(ULearningAgentsObservationSchema* Schema, const FName Tag = TEXT("Angle"));
+	static FLearningAgentsObservationSchemaElement SpecifyAngleObservation(ULearningAgentsObservationSchema* Schema, const FName Tag = TEXT("AngleObservation"));
 
 	/**
 	 * Specifies a new velocity observation. Allows an agent to observe the velocity of some entity.
@@ -550,7 +550,7 @@ public:
 	 * @return The newly created observation schema element.
 	 */
 	UFUNCTION(BlueprintPure, Category = "LearningAgents", meta = (AdvancedDisplay = 1))
-	static FLearningAgentsObservationSchemaElement SpecifyVelocityObservation(ULearningAgentsObservationSchema* Schema, const FName Tag = TEXT("Velocity"));
+	static FLearningAgentsObservationSchemaElement SpecifyVelocityObservation(ULearningAgentsObservationSchema* Schema, const FName Tag = TEXT("VelocityObservation"));
 
 	/**
 	 * Specifies a new direction observation. Allows an agent to observe the direction of some entity.
@@ -560,7 +560,7 @@ public:
 	 * @return The newly created observation schema element.
 	 */
 	UFUNCTION(BlueprintPure, Category = "LearningAgents", meta = (AdvancedDisplay = 1))
-	static FLearningAgentsObservationSchemaElement SpecifyDirectionObservation(ULearningAgentsObservationSchema* Schema, const FName Tag = TEXT("Direction"));
+	static FLearningAgentsObservationSchemaElement SpecifyDirectionObservation(ULearningAgentsObservationSchema* Schema, const FName Tag = TEXT("DirectionObservation"));
 
 	/**
 	 * Specifies a new location along spline observation. This observes the location of the spline at the given distance along that spline.
@@ -570,7 +570,7 @@ public:
 	 * @return The newly created observation schema element.
 	 */
 	UFUNCTION(BlueprintPure, Category = "LearningAgents", meta = (AdvancedDisplay = 1))
-	static FLearningAgentsObservationSchemaElement SpecifyLocationAlongSplineObservation(ULearningAgentsObservationSchema* Schema, const FName Tag = TEXT("LocationAlongSpline"));
+	static FLearningAgentsObservationSchemaElement SpecifyLocationAlongSplineObservation(ULearningAgentsObservationSchema* Schema, const FName Tag = TEXT("LocationAlongSplineObservation"));
 
 	/**
 	 * Specifies a new proportion along spline observation. This observes the proportion along a spline at the given distance. For looped splines 
@@ -582,7 +582,7 @@ public:
 	 * @return The newly created observation schema element.
 	 */
 	UFUNCTION(BlueprintPure, Category = "LearningAgents", meta = (AdvancedDisplay = 1))
-	static FLearningAgentsObservationSchemaElement SpecifyProportionAlongSplineObservation(ULearningAgentsObservationSchema* Schema, const FName Tag = TEXT("ProportionAlongSpline"));
+	static FLearningAgentsObservationSchemaElement SpecifyProportionAlongSplineObservation(ULearningAgentsObservationSchema* Schema, const FName Tag = TEXT("ProportionAlongSplineObservation"));
 
 	/**
 	 * Specifies a new direction along spline observation. This observes the direction of the spline at the given distance along that spline.
@@ -592,7 +592,7 @@ public:
 	 * @return The newly created observation schema element.
 	 */
 	UFUNCTION(BlueprintPure, Category = "LearningAgents", meta = (AdvancedDisplay = 1))
-	static FLearningAgentsObservationSchemaElement SpecifyDirectionAlongSplineObservation(ULearningAgentsObservationSchema* Schema, const FName Tag = TEXT("DirectionAlongSpline"));
+	static FLearningAgentsObservationSchemaElement SpecifyDirectionAlongSplineObservation(ULearningAgentsObservationSchema* Schema, const FName Tag = TEXT("DirectionAlongSplineObservation"));
 
 	/**
 	 * Specifies a new proportion along ray observation. This observes how far a you can travel along a ray before collision. Rays that can travel 
@@ -603,7 +603,7 @@ public:
 	 * @return The newly created observation schema element.
 	 */
 	UFUNCTION(BlueprintPure, Category = "LearningAgents", meta = (AdvancedDisplay = 1))
-	static FLearningAgentsObservationSchemaElement SpecifyProportionAlongRayObservation(ULearningAgentsObservationSchema* Schema, const FName Tag = TEXT("ProportionAlongRay"));
+	static FLearningAgentsObservationSchemaElement SpecifyProportionAlongRayObservation(ULearningAgentsObservationSchema* Schema, const FName Tag = TEXT("ProportionAlongRayObservation"));
 
 public:
 
@@ -615,7 +615,7 @@ public:
 	 * @return The newly created observation object element.
 	 */
 	UFUNCTION(BlueprintPure, Category = "LearningAgents", meta = (AdvancedDisplay = 1))
-	static FLearningAgentsObservationObjectElement MakeNullObservation(ULearningAgentsObservationObject* Object, const FName Tag = TEXT("Null"));
+	static FLearningAgentsObservationObjectElement MakeNullObservation(ULearningAgentsObservationObject* Object, const FName Tag = TEXT("NullObservation"));
 
 	/**
 	 * Make a new continuous observation. The size of Values must match the Size given during Specify.
@@ -634,7 +634,7 @@ public:
 	static FLearningAgentsObservationObjectElement MakeContinuousObservation(
 		ULearningAgentsObservationObject* Object, 
 		const TArray<float>& Values, 
-		const FName Tag = TEXT("Continuous"),
+		const FName Tag = TEXT("ContinuousObservation"),
 		const bool bVisualLoggerEnabled = false,
 		ULearningAgentsManagerListener* VisualLoggerListener = nullptr,
 		const int32 VisualLoggerAgentId = -1,
@@ -657,7 +657,7 @@ public:
 	static FLearningAgentsObservationObjectElement MakeContinuousObservationFromArrayView(
 		ULearningAgentsObservationObject* Object, 
 		const TArrayView<const float> Values, 
-		const FName Tag = TEXT("Continuous"),
+		const FName Tag = TEXT("ContinuousObservation"),
 		const bool bVisualLoggerEnabled = false,
 		ULearningAgentsManagerListener* VisualLoggerListener = nullptr,
 		const int32 VisualLoggerAgentId = -1,
@@ -683,7 +683,7 @@ public:
 		ULearningAgentsObservationObject* Object, 
 		const int32 DiscreteIndex, 
 		const int32 Size,
-		const FName Tag = TEXT("ExclusiveDiscrete"),
+		const FName Tag = TEXT("ExclusiveDiscreteObservation"),
 		const bool bVisualLoggerEnabled = false,
 		ULearningAgentsManagerListener* VisualLoggerListener = nullptr,
 		const int32 VisualLoggerAgentId = -1,
@@ -709,7 +709,7 @@ public:
 		ULearningAgentsObservationObject* Object, 
 		const TArray<int32>& DiscreteIndices, 
 		const int32 Size, 
-		const FName Tag = TEXT("InclusiveDiscrete"),
+		const FName Tag = TEXT("InclusiveDiscreteObservation"),
 		const bool bVisualLoggerEnabled = false,
 		ULearningAgentsManagerListener* VisualLoggerListener = nullptr,
 		const int32 VisualLoggerAgentId = -1,
@@ -734,7 +734,7 @@ public:
 		ULearningAgentsObservationObject* Object, 
 		const TArrayView<const int32> DiscreteIndices, 
 		const int32 Size, 
-		const FName Tag = TEXT("InclusiveDiscrete"),
+		const FName Tag = TEXT("InclusiveDiscreteObservation"),
 		const bool bVisualLoggerEnabled = false,
 		ULearningAgentsManagerListener* VisualLoggerListener = nullptr,
 		const int32 VisualLoggerAgentId = -1,
@@ -760,7 +760,7 @@ public:
 		ULearningAgentsObservationObject* Object, 
 		const int32 Index, 
 		const int32 Size, 
-		const FName Tag = TEXT("Index"),
+		const FName Tag = TEXT("IndexObservation"),
 		const bool bVisualLoggerEnabled = false,
 		ULearningAgentsManagerListener* VisualLoggerListener = nullptr,
 		const int32 VisualLoggerAgentId = -1,
@@ -786,7 +786,7 @@ public:
 		ULearningAgentsObservationObject* Object, 
 		const int32 Num, 
 		const int32 MaxNum, 
-		const FName Tag = TEXT("Count"),
+		const FName Tag = TEXT("CountObservation"),
 		const bool bVisualLoggerEnabled = false,
 		ULearningAgentsManagerListener* VisualLoggerListener = nullptr,
 		const int32 VisualLoggerAgentId = -1,
@@ -802,7 +802,7 @@ public:
 	 * @return The newly created observation object element.
 	 */
 	UFUNCTION(BlueprintPure, Category = "LearningAgents", meta = (AdvancedDisplay = 2))
-	static FLearningAgentsObservationObjectElement MakeStructObservation(ULearningAgentsObservationObject* Object, const TMap<FName, FLearningAgentsObservationObjectElement>& Elements, const FName Tag = TEXT("Struct"));
+	static FLearningAgentsObservationObjectElement MakeStructObservation(ULearningAgentsObservationObject* Object, const TMap<FName, FLearningAgentsObservationObjectElement>& Elements, const FName Tag = TEXT("StructObservation"));
 
 	/**
 	 * Make a new struct observation.
@@ -814,7 +814,7 @@ public:
 	 * @return The newly created observation object element.
 	 */
 	UFUNCTION(BlueprintPure, Category = "LearningAgents", meta = (AdvancedDisplay = 3))
-	static FLearningAgentsObservationObjectElement MakeStructObservationFromArrays(ULearningAgentsObservationObject* Object, const TArray<FName>& ElementNames, const TArray<FLearningAgentsObservationObjectElement>& Elements, const FName Tag = TEXT("Struct"));
+	static FLearningAgentsObservationObjectElement MakeStructObservationFromArrays(ULearningAgentsObservationObject* Object, const TArray<FName>& ElementNames, const TArray<FLearningAgentsObservationObjectElement>& Elements, const FName Tag = TEXT("StructObservation"));
 	
 	/**
 	 * Make a new struct observation.
@@ -825,7 +825,7 @@ public:
 	 * @param Tag The tag of the corresponding observation. Must match the tag given during Specify.
 	 * @return The newly created observation object element.
 	 */
-	static FLearningAgentsObservationObjectElement MakeStructObservationFromArrayViews(ULearningAgentsObservationObject* Object, const TArrayView<const FName> ElementNames, const TArrayView<const FLearningAgentsObservationObjectElement> Elements, const FName Tag = TEXT("Struct"));
+	static FLearningAgentsObservationObjectElement MakeStructObservationFromArrayViews(ULearningAgentsObservationObject* Object, const TArrayView<const FName> ElementNames, const TArrayView<const FLearningAgentsObservationObjectElement> Elements, const FName Tag = TEXT("StructObservation"));
 
 	/**
 	 * Make a new exclusive union observation.
@@ -837,7 +837,7 @@ public:
 	 * @return The newly created observation object element.
 	 */
 	UFUNCTION(BlueprintPure, Category = "LearningAgents", meta = (AdvancedDisplay = 3))
-	static FLearningAgentsObservationObjectElement MakeExclusiveUnionObservation(ULearningAgentsObservationObject* Object, const FName ElementName, const FLearningAgentsObservationObjectElement Element, const FName Tag = TEXT("ExclusiveUnion"));
+	static FLearningAgentsObservationObjectElement MakeExclusiveUnionObservation(ULearningAgentsObservationObject* Object, const FName ElementName, const FLearningAgentsObservationObjectElement Element, const FName Tag = TEXT("ExclusiveUnionObservation"));
 
 	/**
 	 * Make a new inclusive union observation.
@@ -848,7 +848,7 @@ public:
 	 * @return The newly created observation object element.
 	 */
 	UFUNCTION(BlueprintPure, Category = "LearningAgents", meta = (AdvancedDisplay = 2))
-	static FLearningAgentsObservationObjectElement MakeInclusiveUnionObservation(ULearningAgentsObservationObject* Object, const TMap<FName, FLearningAgentsObservationObjectElement>& Elements, const FName Tag = TEXT("InclusiveUnion"));
+	static FLearningAgentsObservationObjectElement MakeInclusiveUnionObservation(ULearningAgentsObservationObject* Object, const TMap<FName, FLearningAgentsObservationObjectElement>& Elements, const FName Tag = TEXT("InclusiveUnionObservation"));
 
 	/**
 	 * Make a new inclusive union observation.
@@ -860,7 +860,7 @@ public:
 	 * @return The newly created observation object element.
 	 */
 	UFUNCTION(BlueprintPure, Category = "LearningAgents", meta = (AdvancedDisplay = 3))
-	static FLearningAgentsObservationObjectElement MakeInclusiveUnionObservationFromArrays(ULearningAgentsObservationObject* Object, const TArray<FName>& ElementNames, const TArray<FLearningAgentsObservationObjectElement>& Elements, const FName Tag = TEXT("InclusiveUnion"));
+	static FLearningAgentsObservationObjectElement MakeInclusiveUnionObservationFromArrays(ULearningAgentsObservationObject* Object, const TArray<FName>& ElementNames, const TArray<FLearningAgentsObservationObjectElement>& Elements, const FName Tag = TEXT("InclusiveUnionObservation"));
 	
 	/**
 	 * Make a new inclusive union observation.
@@ -871,7 +871,7 @@ public:
 	 * @param Tag The tag of the corresponding observation. Must match the tag given during Specify.
 	 * @return The newly created observation object element.
 	 */
-	static FLearningAgentsObservationObjectElement MakeInclusiveUnionObservationFromArrayViews(ULearningAgentsObservationObject* Object, const TArrayView<const FName> ElementNames, const TArrayView<const FLearningAgentsObservationObjectElement> Elements, const FName Tag = TEXT("InclusiveUnion"));
+	static FLearningAgentsObservationObjectElement MakeInclusiveUnionObservationFromArrayViews(ULearningAgentsObservationObject* Object, const TArrayView<const FName> ElementNames, const TArrayView<const FLearningAgentsObservationObjectElement> Elements, const FName Tag = TEXT("InclusiveUnionObservation"));
 
 	/**
 	 * Make a new static array observation.
@@ -882,7 +882,7 @@ public:
 	 * @return The newly created observation object element.
 	 */
 	UFUNCTION(BlueprintPure, Category = "LearningAgents", meta = (AdvancedDisplay = 2))
-	static FLearningAgentsObservationObjectElement MakeStaticArrayObservation(ULearningAgentsObservationObject* Object, const TArray<FLearningAgentsObservationObjectElement>& Elements, const FName Tag = TEXT("StaticArray"));
+	static FLearningAgentsObservationObjectElement MakeStaticArrayObservation(ULearningAgentsObservationObject* Object, const TArray<FLearningAgentsObservationObjectElement>& Elements, const FName Tag = TEXT("StaticArrayObservation"));
 
 	/**
 	 * Make a new static array observation.
@@ -892,7 +892,7 @@ public:
 	 * @param Tag The tag of the corresponding observation. Must match the tag given during Specify.
 	 * @return The newly created observation object element.
 	 */
-	static FLearningAgentsObservationObjectElement MakeStaticArrayObservationFromArrayView(ULearningAgentsObservationObject* Object, const TArrayView<const FLearningAgentsObservationObjectElement> Elements, const FName Tag = TEXT("StaticArray"));
+	static FLearningAgentsObservationObjectElement MakeStaticArrayObservationFromArrayView(ULearningAgentsObservationObject* Object, const TArrayView<const FLearningAgentsObservationObjectElement> Elements, const FName Tag = TEXT("StaticArrayObservation"));
 
 	/**
 	 * Make a new set observation.
@@ -903,7 +903,7 @@ public:
 	 * @return The newly created observation object element.
 	 */
 	UFUNCTION(BlueprintPure, Category = "LearningAgents", meta = (AdvancedDisplay = 2))
-	static FLearningAgentsObservationObjectElement MakeSetObservation(ULearningAgentsObservationObject* Object, const TSet<FLearningAgentsObservationObjectElement>& Elements, const FName Tag = TEXT("Set"));
+	static FLearningAgentsObservationObjectElement MakeSetObservation(ULearningAgentsObservationObject* Object, const TSet<FLearningAgentsObservationObjectElement>& Elements, const FName Tag = TEXT("SetObservation"));
 
 	/**
 	 * Make a new set observation.
@@ -914,7 +914,7 @@ public:
 	 * @return The newly created observation object element.
 	 */
 	UFUNCTION(BlueprintPure, Category = "LearningAgents", meta = (AdvancedDisplay = 2))
-	static FLearningAgentsObservationObjectElement MakeSetObservationFromArray(ULearningAgentsObservationObject* Object, const TArray<FLearningAgentsObservationObjectElement>& Elements, const FName Tag = TEXT("Set"));
+	static FLearningAgentsObservationObjectElement MakeSetObservationFromArray(ULearningAgentsObservationObject* Object, const TArray<FLearningAgentsObservationObjectElement>& Elements, const FName Tag = TEXT("SetObservation"));
 
 	/**
 	 * Make a new set observation.
@@ -924,7 +924,7 @@ public:
 	 * @param Tag The tag of the corresponding observation. Must match the tag given during Specify.
 	 * @return The newly created observation object element.
 	 */
-	static FLearningAgentsObservationObjectElement MakeSetObservationFromArrayView(ULearningAgentsObservationObject* Object, const TArrayView<const FLearningAgentsObservationObjectElement> Elements, const FName Tag = TEXT("Set"));
+	static FLearningAgentsObservationObjectElement MakeSetObservationFromArrayView(ULearningAgentsObservationObject* Object, const TArrayView<const FLearningAgentsObservationObjectElement> Elements, const FName Tag = TEXT("SetObservation"));
 
 	/**
 	 * Make a new pair observation.
@@ -936,7 +936,7 @@ public:
 	 * @return The newly created observation object element.
 	 */
 	UFUNCTION(BlueprintPure, Category = "LearningAgents", meta = (AdvancedDisplay = 3))
-	static FLearningAgentsObservationObjectElement MakePairObservation(ULearningAgentsObservationObject* Object, const FLearningAgentsObservationObjectElement Key, const FLearningAgentsObservationObjectElement Value, const FName Tag = TEXT("Pair"));
+	static FLearningAgentsObservationObjectElement MakePairObservation(ULearningAgentsObservationObject* Object, const FLearningAgentsObservationObjectElement Key, const FLearningAgentsObservationObjectElement Value, const FName Tag = TEXT("PairObservation"));
 
 	/**
 	 * Make a new array observation.
@@ -947,7 +947,7 @@ public:
 	 * @return The newly created observation object element.
 	 */
 	UFUNCTION(BlueprintPure, Category = "LearningAgents", meta = (AdvancedDisplay = 2))
-	static FLearningAgentsObservationObjectElement MakeArrayObservation(ULearningAgentsObservationObject* Object, const TArray<FLearningAgentsObservationObjectElement>& Elements, const FName Tag = TEXT("Array"));
+	static FLearningAgentsObservationObjectElement MakeArrayObservation(ULearningAgentsObservationObject* Object, const TArray<FLearningAgentsObservationObjectElement>& Elements, const FName Tag = TEXT("ArrayObservation"));
 
 	/**
 	 * Make a new array observation.
@@ -957,7 +957,7 @@ public:
 	 * @param Tag The tag of the corresponding observation. Must match the tag given during Specify.
 	 * @return The newly created observation object element.
 	 */
-	static FLearningAgentsObservationObjectElement MakeArrayObservationFromArrayView(ULearningAgentsObservationObject* Object, const TArrayView<const FLearningAgentsObservationObjectElement> Elements, const FName Tag = TEXT("Array"));
+	static FLearningAgentsObservationObjectElement MakeArrayObservationFromArrayView(ULearningAgentsObservationObject* Object, const TArrayView<const FLearningAgentsObservationObjectElement> Elements, const FName Tag = TEXT("ArrayObservation"));
 
 	/**
 	 * Make a new map observation.
@@ -968,7 +968,7 @@ public:
 	 * @return The newly created observation object element.
 	 */
 	UFUNCTION(BlueprintPure, Category = "LearningAgents", meta = (AdvancedDisplay = 2))
-	static FLearningAgentsObservationObjectElement MakeMapObservation(ULearningAgentsObservationObject* Object, const TMap<FLearningAgentsObservationObjectElement, FLearningAgentsObservationObjectElement>& Map, const FName Tag = TEXT("Map"));
+	static FLearningAgentsObservationObjectElement MakeMapObservation(ULearningAgentsObservationObject* Object, const TMap<FLearningAgentsObservationObjectElement, FLearningAgentsObservationObjectElement>& Map, const FName Tag = TEXT("MapObservation"));
 
 	/**
 	 * Make a new map observation.
@@ -979,7 +979,7 @@ public:
 	 * @return The newly created observation object element.
 	 */
 	UFUNCTION(BlueprintPure, Category = "LearningAgents", meta = (AdvancedDisplay = 3))
-	static FLearningAgentsObservationObjectElement MakeMapObservationFromArrays(ULearningAgentsObservationObject* Object, const TArray<FLearningAgentsObservationObjectElement>& Keys, const TArray<FLearningAgentsObservationObjectElement>& Values, const FName Tag = TEXT("Map"));
+	static FLearningAgentsObservationObjectElement MakeMapObservationFromArrays(ULearningAgentsObservationObject* Object, const TArray<FLearningAgentsObservationObjectElement>& Keys, const TArray<FLearningAgentsObservationObjectElement>& Values, const FName Tag = TEXT("MapObservation"));
 	
 	/**
 	 * Make a new map observation.
@@ -989,7 +989,7 @@ public:
 	 * @param Tag The tag of the corresponding observation. Must match the tag given during Specify.
 	 * @return The newly created observation object element.
 	 */
-	static FLearningAgentsObservationObjectElement MakeMapObservationFromArrayViews(ULearningAgentsObservationObject* Object, const TArrayView<const FLearningAgentsObservationObjectElement> Keys, const TArrayView<const FLearningAgentsObservationObjectElement> Values, const FName Tag = TEXT("Map"));
+	static FLearningAgentsObservationObjectElement MakeMapObservationFromArrayViews(ULearningAgentsObservationObject* Object, const TArrayView<const FLearningAgentsObservationObjectElement> Keys, const TArrayView<const FLearningAgentsObservationObjectElement> Values, const FName Tag = TEXT("MapObservation"));
 
 	/**
 	 * Make a new enum observation.
@@ -1010,7 +1010,7 @@ public:
 		ULearningAgentsObservationObject* Object, 
 		const UEnum* Enum, 
 		const uint8 EnumValue, 
-		const FName Tag = TEXT("Enum"),
+		const FName Tag = TEXT("EnumObservation"),
 		const bool bVisualLoggerEnabled = false,
 		ULearningAgentsManagerListener* VisualLoggerListener = nullptr,
 		const int32 VisualLoggerAgentId = -1,
@@ -1036,7 +1036,7 @@ public:
 		ULearningAgentsObservationObject* Object, 
 		const UEnum* Enum, 
 		const int32 BitmaskValue, 
-		const FName Tag = TEXT("Bitmask"),
+		const FName Tag = TEXT("BitmaskObservation"),
 		const bool bVisualLoggerEnabled = false,
 		ULearningAgentsManagerListener* VisualLoggerListener = nullptr,
 		const int32 VisualLoggerAgentId = -1,
@@ -1053,7 +1053,7 @@ public:
 	 * @return The newly created observation object element.
 	 */
 	UFUNCTION(BlueprintPure, Category = "LearningAgents", meta = (AdvancedDisplay = 3))
-	static FLearningAgentsObservationObjectElement MakeOptionalObservation(ULearningAgentsObservationObject* Object, const FLearningAgentsObservationObjectElement Element, const ELearningAgentsOptionalObservation Option, const FName Tag = TEXT("Optional"));
+	static FLearningAgentsObservationObjectElement MakeOptionalObservation(ULearningAgentsObservationObject* Object, const FLearningAgentsObservationObjectElement Element, const ELearningAgentsOptionalObservation Option, const FName Tag = TEXT("OptionalObservation"));
 
 	/**
 	 * Make a new null optional observation. Use this to provide a null optional observation.
@@ -1063,7 +1063,7 @@ public:
 	 * @return The newly created observation object element.
 	 */
 	UFUNCTION(BlueprintPure, Category = "LearningAgents", meta = (AdvancedDisplay = 1))
-	static FLearningAgentsObservationObjectElement MakeOptionalNullObservation(ULearningAgentsObservationObject* Object, const FName Tag = TEXT("Optional"));
+	static FLearningAgentsObservationObjectElement MakeOptionalNullObservation(ULearningAgentsObservationObject* Object, const FName Tag = TEXT("OptionalObservation"));
 
 	/**
 	 * Make a new valid optional observation. Use this to provide a valid optional observation.
@@ -1073,7 +1073,7 @@ public:
 	 * @return The newly created observation object element.
 	 */
 	UFUNCTION(BlueprintPure, Category = "LearningAgents", meta = (AdvancedDisplay = 2))
-	static FLearningAgentsObservationObjectElement MakeOptionalValidObservation(ULearningAgentsObservationObject* Object, const FLearningAgentsObservationObjectElement Element, const FName Tag = TEXT("Optional"));
+	static FLearningAgentsObservationObjectElement MakeOptionalValidObservation(ULearningAgentsObservationObject* Object, const FLearningAgentsObservationObjectElement Element, const FName Tag = TEXT("OptionalObservation"));
 
 	/**
 	 * Make a new either observation.
@@ -1085,7 +1085,7 @@ public:
 	 * @return The newly created observation object element.
 	 */
 	UFUNCTION(BlueprintPure, Category = "LearningAgents", meta = (AdvancedDisplay = 3))
-	static FLearningAgentsObservationObjectElement MakeEitherObservation(ULearningAgentsObservationObject* Object, const FLearningAgentsObservationObjectElement Element, const ELearningAgentsEitherObservation Either, const FName Tag = TEXT("Either"));
+	static FLearningAgentsObservationObjectElement MakeEitherObservation(ULearningAgentsObservationObject* Object, const FLearningAgentsObservationObjectElement Element, const ELearningAgentsEitherObservation Either, const FName Tag = TEXT("EitherObservation"));
 
 	/**
 	 * Make a new either A observation. Use this to provide option A.
@@ -1096,7 +1096,7 @@ public:
 	 * @return The newly created observation object element.
 	 */
 	UFUNCTION(BlueprintPure, Category = "LearningAgents", meta = (AdvancedDisplay = 2, DisplayName = "Make Either A Observation"))
-	static FLearningAgentsObservationObjectElement MakeEitherAObservation(ULearningAgentsObservationObject* Object, const FLearningAgentsObservationObjectElement A, const FName Tag = TEXT("Either"));
+	static FLearningAgentsObservationObjectElement MakeEitherAObservation(ULearningAgentsObservationObject* Object, const FLearningAgentsObservationObjectElement A, const FName Tag = TEXT("EitherObservation"));
 
 	/**
 	 * Make a new either B observation. Use this to provide option B.
@@ -1107,7 +1107,7 @@ public:
 	 * @return The newly created observation object element.
 	 */
 	UFUNCTION(BlueprintPure, Category = "LearningAgents", meta = (AdvancedDisplay = 2, DisplayName = "Make Either B Observation"))
-	static FLearningAgentsObservationObjectElement MakeEitherBObservation(ULearningAgentsObservationObject* Object, const FLearningAgentsObservationObjectElement B, const FName Tag = TEXT("Either"));
+	static FLearningAgentsObservationObjectElement MakeEitherBObservation(ULearningAgentsObservationObject* Object, const FLearningAgentsObservationObjectElement B, const FName Tag = TEXT("EitherObservation"));
 
 	/**
 	 * Make a new encoding observation. This must be used in conjunction with SpecifyEncodingObservation.
@@ -1118,7 +1118,7 @@ public:
 	 * @return The newly created observation object element.
 	 */
 	UFUNCTION(BlueprintPure, Category = "LearningAgents", meta = (AdvancedDisplay = 2))
-	static FLearningAgentsObservationObjectElement MakeEncodingObservation(ULearningAgentsObservationObject* Object, const FLearningAgentsObservationObjectElement Element, const FName Tag = TEXT("Encoding"));
+	static FLearningAgentsObservationObjectElement MakeEncodingObservation(ULearningAgentsObservationObject* Object, const FLearningAgentsObservationObjectElement Element, const FName Tag = TEXT("EncodingObservation"));
 
 	/**
 	 * Make a new bool observation.
@@ -1137,7 +1137,7 @@ public:
 	static FLearningAgentsObservationObjectElement MakeBoolObservation(
 		ULearningAgentsObservationObject* Object,
 		const bool bValue,
-		const FName Tag = TEXT("Bool"),
+		const FName Tag = TEXT("BoolObservation"),
 		const bool bVisualLoggerEnabled = false,
 		ULearningAgentsManagerListener* VisualLoggerListener = nullptr,
 		const int32 VisualLoggerAgentId = -1,
@@ -1163,7 +1163,7 @@ public:
 		ULearningAgentsObservationObject* Object,
 		const float Value,
 		const float FloatScale = 1.0f,
-		const FName Tag = TEXT("Float"),
+		const FName Tag = TEXT("FloatObservation"),
 		const bool bVisualLoggerEnabled = false,
 		ULearningAgentsManagerListener* VisualLoggerListener = nullptr,
 		const int32 VisualLoggerAgentId = -1,
@@ -1191,7 +1191,7 @@ public:
 		const FVector Location,
 		const FTransform RelativeTransform = FTransform(),
 		const float LocationScale = 100.0f,
-		const FName Tag = TEXT("Location"),
+		const FName Tag = TEXT("LocationObservation"),
 		const bool bVisualLoggerEnabled = false,
 		ULearningAgentsManagerListener* VisualLoggerListener = nullptr,
 		const int32 VisualLoggerAgentId = -1,
@@ -1218,7 +1218,7 @@ public:
 		ULearningAgentsObservationObject* Object,
 		const FRotator Rotation,
 		const FRotator RelativeRotation = FRotator::ZeroRotator,
-		const FName Tag = TEXT("Rotation"),
+		const FName Tag = TEXT("RotationObservation"),
 		const bool bVisualLoggerEnabled = false,
 		ULearningAgentsManagerListener* VisualLoggerListener = nullptr,
 		const int32 VisualLoggerAgentId = -1,
@@ -1246,7 +1246,7 @@ public:
 		ULearningAgentsObservationObject* Object,
 		const FQuat Rotation,
 		const FQuat RelativeRotation,
-		const FName Tag = TEXT("Rotation"),
+		const FName Tag = TEXT("RotationObservation"),
 		const bool bVisualLoggerEnabled = false,
 		ULearningAgentsManagerListener* VisualLoggerListener = nullptr,
 		const int32 VisualLoggerAgentId = -1,
@@ -1274,7 +1274,7 @@ public:
 		ULearningAgentsObservationObject* Object,
 		const FVector Scale,
 		const FVector RelativeScale = FVector(1, 1, 1),
-		const FName Tag = TEXT("Scale"),
+		const FName Tag = TEXT("ScaleObservation"),
 		const bool bVisualLoggerEnabled = false,
 		ULearningAgentsManagerListener* VisualLoggerListener = nullptr,
 		const int32 VisualLoggerAgentId = -1,
@@ -1303,7 +1303,7 @@ public:
 		const FTransform Transform,
 		const FTransform RelativeTransform = FTransform(),
 		const float LocationScale = 100.0f,
-		const FName Tag = TEXT("Transform"),
+		const FName Tag = TEXT("TransformObservation"),
 		const bool bVisualLoggerEnabled = false,
 		ULearningAgentsManagerListener* VisualLoggerListener = nullptr,
 		const int32 VisualLoggerAgentId = -1,
@@ -1329,7 +1329,7 @@ public:
 		ULearningAgentsObservationObject* Object, 
 		const float Angle, 
 		const float RelativeAngle = 0.0f, 
-		const FName Tag = TEXT("Angle"),
+		const FName Tag = TEXT("AngleObservation"),
 		const bool bVisualLoggerEnabled = false,
 		ULearningAgentsManagerListener* VisualLoggerListener = nullptr,
 		const int32 VisualLoggerAgentId = -1,
@@ -1355,7 +1355,7 @@ public:
 		ULearningAgentsObservationObject* Object, 
 		const float Angle, 
 		const float RelativeAngle = 0.0f, 
-		const FName Tag = TEXT("Angle"),
+		const FName Tag = TEXT("AngleObservation"),
 		const bool bVisualLoggerEnabled = false,
 		ULearningAgentsManagerListener* VisualLoggerListener = nullptr,
 		const int32 VisualLoggerAgentId = -1,
@@ -1384,7 +1384,7 @@ public:
 		const FVector Velocity, 
 		const FTransform RelativeTransform = FTransform(), 
 		const float VelocityScale = 200.0f, 
-		const FName Tag = TEXT("Velocity"),
+		const FName Tag = TEXT("VelocityObservation"),
 		const bool bVisualLoggerEnabled = false,
 		ULearningAgentsManagerListener* VisualLoggerListener = nullptr,
 		const int32 VisualLoggerAgentId = -1,
@@ -1413,7 +1413,7 @@ public:
 		ULearningAgentsObservationObject* Object,
 		const FVector Direction,
 		const FTransform RelativeTransform = FTransform(),
-		const FName Tag = TEXT("Direction"),
+		const FName Tag = TEXT("DirectionObservation"),
 		const bool bVisualLoggerEnabled = false,
 		ULearningAgentsManagerListener* VisualLoggerListener = nullptr,
 		const int32 VisualLoggerAgentId = -1,
@@ -1445,7 +1445,7 @@ public:
 		const float DistanceAlongSpline,
 		const FTransform RelativeTransform = FTransform(),
 		const float LocationScale = 100.0f,
-		const FName Tag = TEXT("LocationAlongSpline"),
+		const FName Tag = TEXT("LocationAlongSplineObservation"),
 		const bool bVisualLoggerEnabled = false,
 		ULearningAgentsManagerListener* VisualLoggerListener = nullptr,
 		const int32 VisualLoggerAgentId = -1,
@@ -1471,7 +1471,7 @@ public:
 		ULearningAgentsObservationObject* Object, 
 		const USplineComponent* SplineComponent, 
 		const float DistanceAlongSpline, 
-		const FName Tag = TEXT("ProportionAlongSpline"),
+		const FName Tag = TEXT("ProportionAlongSplineObservation"),
 		const bool bVisualLoggerEnabled = false,
 		ULearningAgentsManagerListener* VisualLoggerListener = nullptr,
 		const int32 VisualLoggerAgentId = -1,
@@ -1500,7 +1500,7 @@ public:
 		const USplineComponent* SplineComponent,
 		const float DistanceAlongSpline,
 		const FTransform RelativeTransform = FTransform(),
-		const FName Tag = TEXT("DirectionAlongSpline"),
+		const FName Tag = TEXT("DirectionAlongSplineObservation"),
 		const bool bVisualLoggerEnabled = false,
 		ULearningAgentsManagerListener* VisualLoggerListener = nullptr,
 		const int32 VisualLoggerAgentId = -1,
@@ -1531,7 +1531,7 @@ public:
 		const FVector RayEnd, 
 		const FTransform RayTransform = FTransform(), 
 		const ECollisionChannel CollisionChannel = ECollisionChannel::ECC_WorldStatic, 
-		const FName Tag = TEXT("ProportionAlongRay"),
+		const FName Tag = TEXT("ProportionAlongRayObservation"),
 		const bool bVisualLoggerEnabled = false,
 		ULearningAgentsManagerListener* VisualLoggerListener = nullptr,
 		const int32 VisualLoggerAgentId = -1,
@@ -1543,154 +1543,154 @@ public:
 	// Get Basic Observations
 
 	UFUNCTION(BlueprintPure = false, Category = "LearningAgents", meta = (AdvancedDisplay = 2, ReturnDisplayName = "Success"))
-	static bool GetNullObservation(const ULearningAgentsObservationObject* Object, const FLearningAgentsObservationObjectElement Element, const FName Tag = TEXT("Null"));
+	static bool GetNullObservation(const ULearningAgentsObservationObject* Object, const FLearningAgentsObservationObjectElement Element, const FName Tag = TEXT("NullObservation"));
 
 	UFUNCTION(BlueprintPure = false, Category = "LearningAgents", meta = (AdvancedDisplay = 3, ReturnDisplayName = "Success"))
-	static bool GetContinuousObservationNum(int32& OutNum, const ULearningAgentsObservationObject* Object, const FLearningAgentsObservationObjectElement Element, const FName Tag = TEXT("Continuous"));
+	static bool GetContinuousObservationNum(int32& OutNum, const ULearningAgentsObservationObject* Object, const FLearningAgentsObservationObjectElement Element, const FName Tag = TEXT("ContinuousObservation"));
 
 	UFUNCTION(BlueprintPure = false, Category = "LearningAgents", meta = (AdvancedDisplay = 3, ReturnDisplayName = "Success"))
-	static bool GetContinuousObservation(TArray<float>& OutValues, const ULearningAgentsObservationObject* Object, const FLearningAgentsObservationObjectElement Element, const FName Tag = TEXT("Continuous"));
-	static bool GetContinuousObservationToArrayView(TArrayView<float> OutValues, const ULearningAgentsObservationObject* Object, const FLearningAgentsObservationObjectElement Element, const FName Tag = TEXT("Continuous"));
+	static bool GetContinuousObservation(TArray<float>& OutValues, const ULearningAgentsObservationObject* Object, const FLearningAgentsObservationObjectElement Element, const FName Tag = TEXT("ContinuousObservation"));
+	static bool GetContinuousObservationToArrayView(TArrayView<float> OutValues, const ULearningAgentsObservationObject* Object, const FLearningAgentsObservationObjectElement Element, const FName Tag = TEXT("ContinuousObservation"));
 
 	UFUNCTION(BlueprintPure = false, Category = "LearningAgents", meta = (AdvancedDisplay = 3, ReturnDisplayName = "Success"))
-	static bool GetExclusiveDiscreteObservation(int32& OutIndex, const ULearningAgentsObservationObject* Object, const FLearningAgentsObservationObjectElement Element, const FName Tag = TEXT("ExclusiveDiscrete"));
+	static bool GetExclusiveDiscreteObservation(int32& OutIndex, const ULearningAgentsObservationObject* Object, const FLearningAgentsObservationObjectElement Element, const FName Tag = TEXT("ExclusiveDiscreteObservation"));
 
 	UFUNCTION(BlueprintPure = false, Category = "LearningAgents", meta = (AdvancedDisplay = 3, ReturnDisplayName = "Success"))
-	static bool GetInclusiveDiscreteObservationNum(int32& OutNum, const ULearningAgentsObservationObject* Object, const FLearningAgentsObservationObjectElement Element, const FName Tag = TEXT("InclusiveDiscrete"));
+	static bool GetInclusiveDiscreteObservationNum(int32& OutNum, const ULearningAgentsObservationObject* Object, const FLearningAgentsObservationObjectElement Element, const FName Tag = TEXT("InclusiveDiscreteObservation"));
 
 	UFUNCTION(BlueprintPure = false, Category = "LearningAgents", meta = (AdvancedDisplay = 3, ReturnDisplayName = "Success"))
-	static bool GetInclusiveDiscreteObservation(TArray<int32>& OutIndices, const ULearningAgentsObservationObject* Object, const FLearningAgentsObservationObjectElement Element, const FName Tag = TEXT("InclusiveDiscrete"));
-	static bool GetInclusiveDiscreteObservationToArrayView(TArrayView<int32> OutIndices, const ULearningAgentsObservationObject* Object, const FLearningAgentsObservationObjectElement Element, const FName Tag = TEXT("InclusiveDiscrete"));
+	static bool GetInclusiveDiscreteObservation(TArray<int32>& OutIndices, const ULearningAgentsObservationObject* Object, const FLearningAgentsObservationObjectElement Element, const FName Tag = TEXT("InclusiveDiscreteObservation"));
+	static bool GetInclusiveDiscreteObservationToArrayView(TArrayView<int32> OutIndices, const ULearningAgentsObservationObject* Object, const FLearningAgentsObservationObjectElement Element, const FName Tag = TEXT("InclusiveDiscreteObservation"));
 
 	UFUNCTION(BlueprintPure = false, Category = "LearningAgents", meta = (AdvancedDisplay = 3, ReturnDisplayName = "Success"))
-	static bool GetIndexObservation(int32& OutIndex, const ULearningAgentsObservationObject* Object, const FLearningAgentsObservationObjectElement Element, const FName Tag = TEXT("Index"));
+	static bool GetIndexObservation(int32& OutIndex, const ULearningAgentsObservationObject* Object, const FLearningAgentsObservationObjectElement Element, const FName Tag = TEXT("IndexObservation"));
 
 	UFUNCTION(BlueprintPure = false, Category = "LearningAgents", meta = (AdvancedDisplay = 3, ReturnDisplayName = "Success"))
-	static bool GetCountObservation(int32& OutNum, const ULearningAgentsObservationObject* Object, const FLearningAgentsObservationObjectElement Element, const int32 MaxNum, const FName Tag = TEXT("Count"));
+	static bool GetCountObservation(int32& OutNum, const ULearningAgentsObservationObject* Object, const FLearningAgentsObservationObjectElement Element, const int32 MaxNum, const FName Tag = TEXT("CountObservation"));
 
 	UFUNCTION(BlueprintPure = false, Category = "LearningAgents", meta = (AdvancedDisplay = 3, ReturnDisplayName = "Success"))
-	static bool GetStructObservationNum(int32& OutNum, const ULearningAgentsObservationObject* Object, const FLearningAgentsObservationObjectElement Element, const FName Tag = TEXT("Struct"));
+	static bool GetStructObservationNum(int32& OutNum, const ULearningAgentsObservationObject* Object, const FLearningAgentsObservationObjectElement Element, const FName Tag = TEXT("StructObservation"));
 
 	UFUNCTION(BlueprintPure = false, Category = "LearningAgents", meta = (AdvancedDisplay = 3, ReturnDisplayName = "Success"))
-	static bool GetStructObservation(TMap<FName, FLearningAgentsObservationObjectElement>& OutElements, const ULearningAgentsObservationObject* Object, const FLearningAgentsObservationObjectElement Element, const FName Tag = TEXT("Struct"));
+	static bool GetStructObservation(TMap<FName, FLearningAgentsObservationObjectElement>& OutElements, const ULearningAgentsObservationObject* Object, const FLearningAgentsObservationObjectElement Element, const FName Tag = TEXT("StructObservation"));
 
 	UFUNCTION(BlueprintPure = false, Category = "LearningAgents", meta = (AdvancedDisplay = 4, ReturnDisplayName = "Success"))
-	static bool GetStructObservationToArrays(TArray<FName>& OutElementNames, TArray<FLearningAgentsObservationObjectElement>& OutElements, const ULearningAgentsObservationObject* Object, const FLearningAgentsObservationObjectElement Element, const FName Tag = TEXT("Struct"));
-	static bool GetStructObservationToArrayViews(TArrayView<FName> OutElementNames, TArrayView<FLearningAgentsObservationObjectElement> OutElements, const ULearningAgentsObservationObject* Object, const FLearningAgentsObservationObjectElement Element, const FName Tag = TEXT("Struct"));
+	static bool GetStructObservationToArrays(TArray<FName>& OutElementNames, TArray<FLearningAgentsObservationObjectElement>& OutElements, const ULearningAgentsObservationObject* Object, const FLearningAgentsObservationObjectElement Element, const FName Tag = TEXT("StructObservation"));
+	static bool GetStructObservationToArrayViews(TArrayView<FName> OutElementNames, TArrayView<FLearningAgentsObservationObjectElement> OutElements, const ULearningAgentsObservationObject* Object, const FLearningAgentsObservationObjectElement Element, const FName Tag = TEXT("StructObservation"));
 
 	UFUNCTION(BlueprintPure = false, Category = "LearningAgents", meta = (AdvancedDisplay = 4, ReturnDisplayName = "Success"))
-	static bool GetExclusiveUnionObservation(FName& OutElementName, FLearningAgentsObservationObjectElement& OutElement, const ULearningAgentsObservationObject* Object, const FLearningAgentsObservationObjectElement Element, const FName Tag = TEXT("ExclusiveUnion"));
+	static bool GetExclusiveUnionObservation(FName& OutElementName, FLearningAgentsObservationObjectElement& OutElement, const ULearningAgentsObservationObject* Object, const FLearningAgentsObservationObjectElement Element, const FName Tag = TEXT("ExclusiveUnionObservation"));
 
 	UFUNCTION(BlueprintPure = false, Category = "LearningAgents", meta = (AdvancedDisplay = 3, ReturnDisplayName = "Success"))
-	static bool GetInclusiveUnionObservationNum(int32& OutNum, const ULearningAgentsObservationObject* Object, const FLearningAgentsObservationObjectElement Element, const FName Tag = TEXT("InclusiveUnion"));
+	static bool GetInclusiveUnionObservationNum(int32& OutNum, const ULearningAgentsObservationObject* Object, const FLearningAgentsObservationObjectElement Element, const FName Tag = TEXT("InclusiveUnionObservation"));
 
 	UFUNCTION(BlueprintPure = false, Category = "LearningAgents", meta = (AdvancedDisplay = 3, ReturnDisplayName = "Success"))
-	static bool GetInclusiveUnionObservation(TMap<FName, FLearningAgentsObservationObjectElement>& OutElements, const ULearningAgentsObservationObject* Object, const FLearningAgentsObservationObjectElement Element, const FName Tag = TEXT("InclusiveUnion"));
+	static bool GetInclusiveUnionObservation(TMap<FName, FLearningAgentsObservationObjectElement>& OutElements, const ULearningAgentsObservationObject* Object, const FLearningAgentsObservationObjectElement Element, const FName Tag = TEXT("InclusiveUnionObservation"));
 
 	UFUNCTION(BlueprintPure = false, Category = "LearningAgents", meta = (AdvancedDisplay = 4, ReturnDisplayName = "Success"))
-	static bool GetInclusiveUnionObservationToArrays(TArray<FName>& OutElementNames, TArray<FLearningAgentsObservationObjectElement>& OutElements, const ULearningAgentsObservationObject* Object, const FLearningAgentsObservationObjectElement Element, const FName Tag = TEXT("InclusiveUnion"));
-	static bool GetInclusiveUnionObservationToArrayViews(TArrayView<FName> OutElementNames, TArrayView<FLearningAgentsObservationObjectElement> OutElements, const ULearningAgentsObservationObject* Object, const FLearningAgentsObservationObjectElement Element, const FName Tag = TEXT("InclusiveUnion"));
+	static bool GetInclusiveUnionObservationToArrays(TArray<FName>& OutElementNames, TArray<FLearningAgentsObservationObjectElement>& OutElements, const ULearningAgentsObservationObject* Object, const FLearningAgentsObservationObjectElement Element, const FName Tag = TEXT("InclusiveUnionObservation"));
+	static bool GetInclusiveUnionObservationToArrayViews(TArrayView<FName> OutElementNames, TArrayView<FLearningAgentsObservationObjectElement> OutElements, const ULearningAgentsObservationObject* Object, const FLearningAgentsObservationObjectElement Element, const FName Tag = TEXT("InclusiveUnionObservation"));
 
 	UFUNCTION(BlueprintPure = false, Category = "LearningAgents", meta = (AdvancedDisplay = 3, ReturnDisplayName = "Success"))
-	static bool GetStaticArrayObservationNum(int32& OutNum, const ULearningAgentsObservationObject* Object, const FLearningAgentsObservationObjectElement Element, const FName Tag = TEXT("StaticArray"));
+	static bool GetStaticArrayObservationNum(int32& OutNum, const ULearningAgentsObservationObject* Object, const FLearningAgentsObservationObjectElement Element, const FName Tag = TEXT("StaticArrayObservation"));
 
 	UFUNCTION(BlueprintPure = false, Category = "LearningAgents", meta = (AdvancedDisplay = 3, ReturnDisplayName = "Success"))
-	static bool GetStaticArrayObservation(TArray<FLearningAgentsObservationObjectElement>& OutElements, const ULearningAgentsObservationObject* Object, const FLearningAgentsObservationObjectElement Element, const FName Tag = TEXT("StaticArray"));
-	static bool GetStaticArrayObservationToArrayView(TArrayView<FLearningAgentsObservationObjectElement> OutElements, const ULearningAgentsObservationObject* Object, const FLearningAgentsObservationObjectElement Element, const FName Tag = TEXT("StaticArray"));
+	static bool GetStaticArrayObservation(TArray<FLearningAgentsObservationObjectElement>& OutElements, const ULearningAgentsObservationObject* Object, const FLearningAgentsObservationObjectElement Element, const FName Tag = TEXT("StaticArrayObservation"));
+	static bool GetStaticArrayObservationToArrayView(TArrayView<FLearningAgentsObservationObjectElement> OutElements, const ULearningAgentsObservationObject* Object, const FLearningAgentsObservationObjectElement Element, const FName Tag = TEXT("StaticArrayObservation"));
 
 	UFUNCTION(BlueprintPure = false, Category = "LearningAgents", meta = (AdvancedDisplay = 3, ReturnDisplayName = "Success"))
-	static bool GetSetObservationNum(int32& OutNum, const ULearningAgentsObservationObject* Object, const FLearningAgentsObservationObjectElement Element, const FName Tag = TEXT("Set"));
+	static bool GetSetObservationNum(int32& OutNum, const ULearningAgentsObservationObject* Object, const FLearningAgentsObservationObjectElement Element, const FName Tag = TEXT("SetObservation"));
 
 	UFUNCTION(BlueprintPure = false, Category = "LearningAgents", meta = (AdvancedDisplay = 3, ReturnDisplayName = "Success"))
-	static bool GetSetObservation(TSet<FLearningAgentsObservationObjectElement>& OutElements, const ULearningAgentsObservationObject* Object, const FLearningAgentsObservationObjectElement Element, const FName Tag = TEXT("Set"));
+	static bool GetSetObservation(TSet<FLearningAgentsObservationObjectElement>& OutElements, const ULearningAgentsObservationObject* Object, const FLearningAgentsObservationObjectElement Element, const FName Tag = TEXT("SetObservation"));
 
 	UFUNCTION(BlueprintPure = false, Category = "LearningAgents", meta = (AdvancedDisplay = 3, ReturnDisplayName = "Success"))
-	static bool GetSetObservationToArray(TArray<FLearningAgentsObservationObjectElement>& OutElements, const ULearningAgentsObservationObject* Object, const FLearningAgentsObservationObjectElement Element, const FName Tag = TEXT("Set"));
-	static bool GetSetObservationToArrayView(TArrayView<FLearningAgentsObservationObjectElement> OutElements, const ULearningAgentsObservationObject* Object, const FLearningAgentsObservationObjectElement Element, const FName Tag = TEXT("Set"));
+	static bool GetSetObservationToArray(TArray<FLearningAgentsObservationObjectElement>& OutElements, const ULearningAgentsObservationObject* Object, const FLearningAgentsObservationObjectElement Element, const FName Tag = TEXT("SetObservation"));
+	static bool GetSetObservationToArrayView(TArrayView<FLearningAgentsObservationObjectElement> OutElements, const ULearningAgentsObservationObject* Object, const FLearningAgentsObservationObjectElement Element, const FName Tag = TEXT("SetObservation"));
 
 	UFUNCTION(BlueprintPure = false, Category = "LearningAgents", meta = (AdvancedDisplay = 4, ReturnDisplayName = "Success"))
-	static bool GetPairObservation(FLearningAgentsObservationObjectElement& OutKey, FLearningAgentsObservationObjectElement& OutValue, const ULearningAgentsObservationObject* Object, const FLearningAgentsObservationObjectElement Element, const FName Tag = TEXT("Pair"));
+	static bool GetPairObservation(FLearningAgentsObservationObjectElement& OutKey, FLearningAgentsObservationObjectElement& OutValue, const ULearningAgentsObservationObject* Object, const FLearningAgentsObservationObjectElement Element, const FName Tag = TEXT("PairObservation"));
 
 	UFUNCTION(BlueprintPure = false, Category = "LearningAgents", meta = (AdvancedDisplay = 3, ReturnDisplayName = "Success"))
-	static bool GetArrayObservationNum(int32& OutNum, const ULearningAgentsObservationObject* Object, const FLearningAgentsObservationObjectElement Element, const FName Tag = TEXT("Array"));
+	static bool GetArrayObservationNum(int32& OutNum, const ULearningAgentsObservationObject* Object, const FLearningAgentsObservationObjectElement Element, const FName Tag = TEXT("ArrayObservation"));
 
 	UFUNCTION(BlueprintPure = false, Category = "LearningAgents", meta = (AdvancedDisplay = 3, ReturnDisplayName = "Success"))
-	static bool GetArrayObservation(TArray<FLearningAgentsObservationObjectElement>& OutElements, const ULearningAgentsObservationObject* Object, const FLearningAgentsObservationObjectElement Element, const FName Tag = TEXT("Array"));
-	static bool GetArrayObservationToArrayView(TArrayView<FLearningAgentsObservationObjectElement> OutElements, const ULearningAgentsObservationObject* Object, const FLearningAgentsObservationObjectElement Element, const FName Tag = TEXT("Array"));
+	static bool GetArrayObservation(TArray<FLearningAgentsObservationObjectElement>& OutElements, const ULearningAgentsObservationObject* Object, const FLearningAgentsObservationObjectElement Element, const FName Tag = TEXT("ArrayObservation"));
+	static bool GetArrayObservationToArrayView(TArrayView<FLearningAgentsObservationObjectElement> OutElements, const ULearningAgentsObservationObject* Object, const FLearningAgentsObservationObjectElement Element, const FName Tag = TEXT("ArrayObservation"));
 
 	UFUNCTION(BlueprintPure = false, Category = "LearningAgents", meta = (AdvancedDisplay = 3, ReturnDisplayName = "Success"))
-	static bool GetMapObservationNum(int32& OutNum, const ULearningAgentsObservationObject* Object, const FLearningAgentsObservationObjectElement Element, const FName Tag = TEXT("Map"));
+	static bool GetMapObservationNum(int32& OutNum, const ULearningAgentsObservationObject* Object, const FLearningAgentsObservationObjectElement Element, const FName Tag = TEXT("MapObservation"));
 
 	UFUNCTION(BlueprintPure = false, Category = "LearningAgents", meta = (AdvancedDisplay = 3, ReturnDisplayName = "Success"))
-	static bool GetMapObservation(TMap<FLearningAgentsObservationObjectElement, FLearningAgentsObservationObjectElement>& OutElements, const ULearningAgentsObservationObject* Object, const FLearningAgentsObservationObjectElement Element, const FName Tag = TEXT("Map"));
+	static bool GetMapObservation(TMap<FLearningAgentsObservationObjectElement, FLearningAgentsObservationObjectElement>& OutElements, const ULearningAgentsObservationObject* Object, const FLearningAgentsObservationObjectElement Element, const FName Tag = TEXT("MapObservation"));
 
 	UFUNCTION(BlueprintPure = false, Category = "LearningAgents", meta = (AdvancedDisplay = 4, ReturnDisplayName = "Success"))
-	static bool GetMapObservationToArrays(TArray<FLearningAgentsObservationObjectElement>& OutKeys, TArray<FLearningAgentsObservationObjectElement>& OutValues, const ULearningAgentsObservationObject* Object, const FLearningAgentsObservationObjectElement Element, const FName Tag = TEXT("Map"));
-	static bool GetMapObservationToArrayViews(TArrayView<FLearningAgentsObservationObjectElement> OutKeys, TArrayView<FLearningAgentsObservationObjectElement> OutValues, const ULearningAgentsObservationObject* Object, const FLearningAgentsObservationObjectElement Element, const FName Tag = TEXT("Map"));
+	static bool GetMapObservationToArrays(TArray<FLearningAgentsObservationObjectElement>& OutKeys, TArray<FLearningAgentsObservationObjectElement>& OutValues, const ULearningAgentsObservationObject* Object, const FLearningAgentsObservationObjectElement Element, const FName Tag = TEXT("MapObservation"));
+	static bool GetMapObservationToArrayViews(TArrayView<FLearningAgentsObservationObjectElement> OutKeys, TArrayView<FLearningAgentsObservationObjectElement> OutValues, const ULearningAgentsObservationObject* Object, const FLearningAgentsObservationObjectElement Element, const FName Tag = TEXT("MapObservation"));
 
 	UFUNCTION(BlueprintPure = false, Category = "LearningAgents", meta = (AdvancedDisplay = 4, ReturnDisplayName = "Success"))
-	static bool GetEnumObservation(uint8& OutEnumValue, const ULearningAgentsObservationObject* Object, const FLearningAgentsObservationObjectElement Element, const UEnum* Enum, const FName Tag = TEXT("Enum"));
+	static bool GetEnumObservation(uint8& OutEnumValue, const ULearningAgentsObservationObject* Object, const FLearningAgentsObservationObjectElement Element, const UEnum* Enum, const FName Tag = TEXT("EnumObservation"));
 
 	UFUNCTION(BlueprintPure = false, Category = "LearningAgents", meta = (AdvancedDisplay = 4, ReturnDisplayName = "Success"))
-	static bool GetBitmaskObservation(int32& OutBitmaskValue, const ULearningAgentsObservationObject* Object, const FLearningAgentsObservationObjectElement Element, const UEnum* Enum, const FName Tag = TEXT("Bitmask"));
+	static bool GetBitmaskObservation(int32& OutBitmaskValue, const ULearningAgentsObservationObject* Object, const FLearningAgentsObservationObjectElement Element, const UEnum* Enum, const FName Tag = TEXT("BitmaskObservation"));
 
 	UFUNCTION(BlueprintPure = false, Category = "LearningAgents", meta = (AdvancedDisplay = 3, ExpandEnumAsExecs = "OutOption", ReturnDisplayName = "Success"))
-	static bool GetOptionalObservation(ELearningAgentsOptionalObservation& OutOption, FLearningAgentsObservationObjectElement& OutElement, const ULearningAgentsObservationObject* Object, const FLearningAgentsObservationObjectElement Element, const FName Tag = TEXT("Optional"));
+	static bool GetOptionalObservation(ELearningAgentsOptionalObservation& OutOption, FLearningAgentsObservationObjectElement& OutElement, const ULearningAgentsObservationObject* Object, const FLearningAgentsObservationObjectElement Element, const FName Tag = TEXT("OptionalObservation"));
 
 	UFUNCTION(BlueprintPure = false, Category = "LearningAgents", meta = (AdvancedDisplay = 3, ExpandEnumAsExecs = "OutEither", ReturnDisplayName = "Success"))
-	static bool GetEitherObservation(ELearningAgentsEitherObservation& OutEither, FLearningAgentsObservationObjectElement& OutElement, const ULearningAgentsObservationObject* Object, const FLearningAgentsObservationObjectElement Element, const FName Tag = TEXT("Either"));
+	static bool GetEitherObservation(ELearningAgentsEitherObservation& OutEither, FLearningAgentsObservationObjectElement& OutElement, const ULearningAgentsObservationObject* Object, const FLearningAgentsObservationObjectElement Element, const FName Tag = TEXT("EitherObservation"));
 
 	UFUNCTION(BlueprintPure = false, Category = "LearningAgents", meta = (AdvancedDisplay = 3, ReturnDisplayName = "Success"))
-	static bool GetEncodingObservation(FLearningAgentsObservationObjectElement& OutElement, const ULearningAgentsObservationObject* Object, const FLearningAgentsObservationObjectElement Element, const FName Tag = TEXT("Encoding"));
+	static bool GetEncodingObservation(FLearningAgentsObservationObjectElement& OutElement, const ULearningAgentsObservationObject* Object, const FLearningAgentsObservationObjectElement Element, const FName Tag = TEXT("EncodingObservation"));
 
 	UFUNCTION(BlueprintPure = false, Category = "LearningAgents", meta = (AdvancedDisplay = 3, ReturnDisplayName = "Success"))
-	static bool GetBoolObservation(bool& bOutValue, const ULearningAgentsObservationObject* Object, const FLearningAgentsObservationObjectElement Element, const FName Tag = TEXT("Bool"));
+	static bool GetBoolObservation(bool& bOutValue, const ULearningAgentsObservationObject* Object, const FLearningAgentsObservationObjectElement Element, const FName Tag = TEXT("BoolObservation"));
 
 	UFUNCTION(BlueprintPure = false, Category = "LearningAgents", meta = (AdvancedDisplay = 4, ReturnDisplayName = "Success"))
-	static bool GetFloatObservation(float& OutValue, const ULearningAgentsObservationObject* Object, const FLearningAgentsObservationObjectElement Element, const float FloatScale = 1.0f, const FName Tag = TEXT("Float"));
+	static bool GetFloatObservation(float& OutValue, const ULearningAgentsObservationObject* Object, const FLearningAgentsObservationObjectElement Element, const float FloatScale = 1.0f, const FName Tag = TEXT("FloatObservation"));
 
 	UFUNCTION(BlueprintPure = false, Category = "LearningAgents", meta = (AdvancedDisplay = 5, ReturnDisplayName = "Success"))
-	static bool GetLocationObservation(FVector& OutLocation, const ULearningAgentsObservationObject* Object, const FLearningAgentsObservationObjectElement Element, const FTransform RelativeTransform = FTransform(), const float LocationScale = 100.0f, const FName Tag = TEXT("Location"));
+	static bool GetLocationObservation(FVector& OutLocation, const ULearningAgentsObservationObject* Object, const FLearningAgentsObservationObjectElement Element, const FTransform RelativeTransform = FTransform(), const float LocationScale = 100.0f, const FName Tag = TEXT("LocationObservation"));
 
 	UFUNCTION(BlueprintPure = false, Category = "LearningAgents", meta = (AdvancedDisplay = 4, ReturnDisplayName = "Success"))
-	static bool GetRotationObservation(FRotator& OutRotation, const ULearningAgentsObservationObject* Object, const FLearningAgentsObservationObjectElement Element, const FRotator RelativeRotation = FRotator::ZeroRotator, const FName Tag = TEXT("Rotation"));
+	static bool GetRotationObservation(FRotator& OutRotation, const ULearningAgentsObservationObject* Object, const FLearningAgentsObservationObjectElement Element, const FRotator RelativeRotation = FRotator::ZeroRotator, const FName Tag = TEXT("RotationObservation"));
 
 	UFUNCTION(BlueprintPure = false, Category = "LearningAgents", meta = (AdvancedDisplay = 4, ReturnDisplayName = "Success"))
-	static bool GetRotationObservationAsQuat(FQuat& OutRotation, const ULearningAgentsObservationObject* Object, const FLearningAgentsObservationObjectElement Element, const FQuat RelativeRotation, const FName Tag = TEXT("Rotation"));
+	static bool GetRotationObservationAsQuat(FQuat& OutRotation, const ULearningAgentsObservationObject* Object, const FLearningAgentsObservationObjectElement Element, const FQuat RelativeRotation, const FName Tag = TEXT("RotationObservation"));
 
 	UFUNCTION(BlueprintPure = false, Category = "LearningAgents", meta = (AdvancedDisplay = 4, ReturnDisplayName = "Success"))
-	static bool GetScaleObservation(FVector& OutScale, const ULearningAgentsObservationObject* Object, const FLearningAgentsObservationObjectElement Element, const FVector RelativeScale = FVector(1, 1, 1), const FName Tag = TEXT("Scale"));
+	static bool GetScaleObservation(FVector& OutScale, const ULearningAgentsObservationObject* Object, const FLearningAgentsObservationObjectElement Element, const FVector RelativeScale = FVector(1, 1, 1), const FName Tag = TEXT("ScaleObservation"));
 
 	UFUNCTION(BlueprintPure = false, Category = "LearningAgents", meta = (AdvancedDisplay = 5, ReturnDisplayName = "Success"))
-	static bool GetTransformObservation(FTransform& OutTransform, const ULearningAgentsObservationObject* Object, const FLearningAgentsObservationObjectElement Element, const FTransform RelativeTransform = FTransform(), const float LocationScale = 100.0f, const FName Tag = TEXT("Transform"));
+	static bool GetTransformObservation(FTransform& OutTransform, const ULearningAgentsObservationObject* Object, const FLearningAgentsObservationObjectElement Element, const FTransform RelativeTransform = FTransform(), const float LocationScale = 100.0f, const FName Tag = TEXT("TransformObservation"));
 
 	UFUNCTION(BlueprintPure = false, Category = "LearningAgents", meta = (AdvancedDisplay = 4, ReturnDisplayName = "Success"))
-	static bool GetAngleObservation(float& OutAngle, const ULearningAgentsObservationObject* Object, const FLearningAgentsObservationObjectElement Element, const float RelativeAngle = 0.0f, const FName Tag = TEXT("Angle"));
+	static bool GetAngleObservation(float& OutAngle, const ULearningAgentsObservationObject* Object, const FLearningAgentsObservationObjectElement Element, const float RelativeAngle = 0.0f, const FName Tag = TEXT("AngleObservation"));
 
 	UFUNCTION(BlueprintPure = false, Category = "LearningAgents", meta = (AdvancedDisplay = 4, ReturnDisplayName = "Success"))
-	static bool GetAngleObservationRadians(float& OutAngle, const ULearningAgentsObservationObject* Object, const FLearningAgentsObservationObjectElement Element, const float RelativeAngle = 0.0f, const FName Tag = TEXT("Angle"));
+	static bool GetAngleObservationRadians(float& OutAngle, const ULearningAgentsObservationObject* Object, const FLearningAgentsObservationObjectElement Element, const float RelativeAngle = 0.0f, const FName Tag = TEXT("AngleObservation"));
 
 	UFUNCTION(BlueprintPure = false, Category = "LearningAgents", meta = (AdvancedDisplay = 5, ReturnDisplayName = "Success"))
-	static bool GetVelocityObservation(FVector& OutVelocity, const ULearningAgentsObservationObject* Object, const FLearningAgentsObservationObjectElement Element, const FTransform RelativeTransform = FTransform(), const float VelocityScale = 200.0f, const FName Tag = TEXT("Velocity"));
+	static bool GetVelocityObservation(FVector& OutVelocity, const ULearningAgentsObservationObject* Object, const FLearningAgentsObservationObjectElement Element, const FTransform RelativeTransform = FTransform(), const float VelocityScale = 200.0f, const FName Tag = TEXT("VelocityObservation"));
 
 	UFUNCTION(BlueprintPure = false, Category = "LearningAgents", meta = (AdvancedDisplay = 4, ReturnDisplayName = "Success"))
-	static bool GetDirectionObservation(FVector& OutDirection, const ULearningAgentsObservationObject* Object, const FLearningAgentsObservationObjectElement Element, const FTransform RelativeTransform = FTransform(), const FName Tag = TEXT("Direction"));
+	static bool GetDirectionObservation(FVector& OutDirection, const ULearningAgentsObservationObject* Object, const FLearningAgentsObservationObjectElement Element, const FTransform RelativeTransform = FTransform(), const FName Tag = TEXT("DirectionObservation"));
 
 	// Get Spline Observations
 
 	UFUNCTION(BlueprintPure = false, Category = "LearningAgents", meta = (AdvancedDisplay = 5, ReturnDisplayName = "Success"))
-	static bool GetLocationAlongSplineObservation(FVector& OutLocation, const ULearningAgentsObservationObject* Object, const FLearningAgentsObservationObjectElement Element, const FTransform RelativeTransform = FTransform(), const float LocationScale = 100.0f, const FName Tag = TEXT("LocationAlongSpline"));
+	static bool GetLocationAlongSplineObservation(FVector& OutLocation, const ULearningAgentsObservationObject* Object, const FLearningAgentsObservationObjectElement Element, const FTransform RelativeTransform = FTransform(), const float LocationScale = 100.0f, const FName Tag = TEXT("LocationAlongSplineObservation"));
 
 	UFUNCTION(BlueprintPure = false, Category = "LearningAgents", meta = (AdvancedDisplay = 5, ReturnDisplayName = "Success"))
-	static bool GetProportionAlongSplineObservation(bool& bOutIsClosedLoop, float& OutAngle, float& OutPropotion, const ULearningAgentsObservationObject* Object, const FLearningAgentsObservationObjectElement Element, const FName Tag = TEXT("ProportionAlongSpline"));
+	static bool GetProportionAlongSplineObservation(bool& bOutIsClosedLoop, float& OutAngle, float& OutPropotion, const ULearningAgentsObservationObject* Object, const FLearningAgentsObservationObjectElement Element, const FName Tag = TEXT("ProportionAlongSplineObservation"));
 
 	UFUNCTION(BlueprintPure = false, Category = "LearningAgents", meta = (AdvancedDisplay = 4, ReturnDisplayName = "Success"))
-	static bool GetDirectionAlongSplineObservation(FVector& OutDirection, const ULearningAgentsObservationObject* Object, const FLearningAgentsObservationObjectElement Element, const FTransform RelativeTransform = FTransform(), const FName Tag = TEXT("DirectionAlongSpline"));
+	static bool GetDirectionAlongSplineObservation(FVector& OutDirection, const ULearningAgentsObservationObject* Object, const FLearningAgentsObservationObjectElement Element, const FTransform RelativeTransform = FTransform(), const FName Tag = TEXT("DirectionAlongSplineObservation"));
 
 	// Get Ray Cast Observations
 
 	UFUNCTION(BlueprintPure, Category = "LearningAgents", meta = (AdvancedDisplay = 3, ReturnDisplayName = "Success"))
-	static bool GetProportionAlongRayObservation(float& OutProportion, const ULearningAgentsObservationObject* Object, const FLearningAgentsObservationObjectElement Element, const FName Tag = TEXT("ProportionAlongRay"));
+	static bool GetProportionAlongRayObservation(float& OutProportion, const ULearningAgentsObservationObject* Object, const FLearningAgentsObservationObjectElement Element, const FName Tag = TEXT("ProportionAlongRayObservation"));
 
 };
 
diff --git a/Engine/Plugins/Experimental/LearningAgents/Source/LearningAgentsTraining/Public/LearningAgentsCompletions.h b/Engine/Plugins/Experimental/LearningAgents/Source/LearningAgentsTraining/Public/LearningAgentsCompletions.h
--- a/Engine/Plugins/Experimental/LearningAgents/Source/LearningAgentsTraining/Public/LearningAgentsCompletions.h
+++ b/Engine/Plugins/Experimental/LearningAgents/Source/LearningAgentsTraining/Public/LearningAgentsCompletions.h
@@ -111,7 +111,7 @@ public:
 	static ELearningAgentsCompletion MakeCompletionOnCondition(
 		const bool bCondition, 
 		const ELearningAgentsCompletion CompletionType = ELearningAgentsCompletion::Termination,
-		const FName Tag = TEXT("CompletionOnCondition"),
+		const FName Tag = TEXT("ConditionCompletion"),
 		const bool bVisualLoggerEnabled = false,
 		ULearningAgentsManagerListener* VisualLoggerListener = nullptr,
 		const int32 VisualLoggerAgentId = -1,
@@ -137,7 +137,7 @@ public:
 		const float Time, 
 		const float TimeThreshold = 10.0f, 
 		const ELearningAgentsCompletion CompletionType = ELearningAgentsCompletion::Truncation,
-		const FName Tag = TEXT("CompletionOnTimeElapsed"),
+		const FName Tag = TEXT("TimeElapsedCompletion"),
 		const bool bVisualLoggerEnabled = false,
 		ULearningAgentsManagerListener* VisualLoggerListener = nullptr,
 		const int32 VisualLoggerAgentId = -1,
@@ -163,7 +163,7 @@ public:
 		const int32 EpisodeSteps, 
 		const int32 MaxEpisodeSteps = 64, 
 		const ELearningAgentsCompletion CompletionType = ELearningAgentsCompletion::Truncation,
-		const FName Tag = TEXT("CompletionOnEpisodeStepsRecorded"),
+		const FName Tag = TEXT("EpisodeStepsRecordedCompletion"),
 		const bool bVisualLoggerEnabled = false,
 		ULearningAgentsManagerListener* VisualLoggerListener = nullptr,
 		const int32 VisualLoggerAgentId = -1,
@@ -191,7 +191,7 @@ public:
 		const FVector LocationB, 
 		const float DistanceThreshold = 100.0f, 
 		const ELearningAgentsCompletion CompletionType = ELearningAgentsCompletion::Termination,
-		const FName Tag = TEXT("CompletionOnLocationDifferenceBelowThreshold"),
+		const FName Tag = TEXT("LocationDifferenceBelowThresholdCompletion"),
 		const bool bVisualLoggerEnabled = false,
 		ULearningAgentsManagerListener* VisualLoggerListener = nullptr,
 		const int32 VisualLoggerAgentId = -1,
@@ -219,7 +219,7 @@ public:
 		const FVector LocationB, 
 		const float DistanceThreshold = 100.0f, 
 		const ELearningAgentsCompletion CompletionType = ELearningAgentsCompletion::Termination,
-		const FName Tag = TEXT("CompletionOnLocationDifferenceAboveThreshold"),
+		const FName Tag = TEXT("LocationDifferenceAboveThresholdCompletion"),
 		const bool bVisualLoggerEnabled = false,
 		ULearningAgentsManagerListener* VisualLoggerListener = nullptr,
 		const int32 VisualLoggerAgentId = -1,
@@ -249,7 +249,7 @@ public:
 		const FVector BoundsMins = FVector(-100.0f, -100.0f, -100.0f),
 		const FVector BoundsMaxs = FVector(+100.0f, +100.0f, +100.0f),
 		const ELearningAgentsCompletion CompletionType = ELearningAgentsCompletion::Termination,
-		const FName Tag = TEXT("CompletionOnLocationOutsideBounds"),
+		const FName Tag = TEXT("LocationOutsideBoundsCompletion"),
 		const bool bVisualLoggerEnabled = false,
 		ULearningAgentsManagerListener* VisualLoggerListener = nullptr,
 		const int32 VisualLoggerAgentId = -1,
diff --git a/Engine/Plugins/Experimental/LearningAgents/Source/LearningAgentsTraining/Public/LearningAgentsRewards.h b/Engine/Plugins/Experimental/LearningAgents/Source/LearningAgentsTraining/Public/LearningAgentsRewards.h
--- a/Engine/Plugins/Experimental/LearningAgents/Source/LearningAgentsTraining/Public/LearningAgentsRewards.h
+++ b/Engine/Plugins/Experimental/LearningAgents/Source/LearningAgentsTraining/Public/LearningAgentsRewards.h
@@ -57,7 +57,7 @@ public:
 	static float MakeRewardOnCondition(
 		const bool bCondition, 
 		const float RewardScale = 1.0f,
-		const FName Tag = TEXT("RewardOnCondition"),
+		const FName Tag = TEXT("ConditionReward"),
 		const bool bVisualLoggerEnabled = false,
 		ULearningAgentsManagerListener* VisualLoggerListener = nullptr,
 		const int32 VisualLoggerAgentId = -1,
@@ -85,7 +85,7 @@ public:
 		const FVector LocationB, 
 		const float DistanceThreshold = 100.0f, 
 		const float RewardScale = 1.0f,
-		const FName Tag = TEXT("RewardOnLocationDifferenceBelowThreshold"),
+		const FName Tag = TEXT("LocationDifferenceBelowThresholdReward"),
 		const bool bVisualLoggerEnabled = false,
 		ULearningAgentsManagerListener* VisualLoggerListener = nullptr,
 		const int32 VisualLoggerAgentId = -1,
@@ -113,7 +113,7 @@ public:
 		const FVector LocationB,
 		const float DistanceThreshold = 100.0f,
 		const float RewardScale = 1.0f,
-		const FName Tag = TEXT("RewardOnLocationDifferenceAboveThreshold"),
+		const FName Tag = TEXT("LocationDifferenceAboveThresholdReward"),
 		const bool bVisualLoggerEnabled = false,
 		ULearningAgentsManagerListener* VisualLoggerListener = nullptr,
 		const int32 VisualLoggerAgentId = -1,
@@ -141,7 +141,7 @@ public:
 		const FVector LocationB, 
 		const float LocationScale = 100.0f, 
 		const float RewardScale = 1.0f,
-		const FName Tag = TEXT("RewardFromLocationSimilarity"),
+		const FName Tag = TEXT("LocationSimilarityReward"),
 		const bool bVisualLoggerEnabled = false,
 		ULearningAgentsManagerListener* VisualLoggerListener = nullptr,
 		const int32 VisualLoggerAgentId = -1,
@@ -169,7 +169,7 @@ public:
 		const FRotator RotationA, 
 		const FRotator RotationB, 
 		const float RewardScale = 1.0f,
-		const FName Tag = TEXT("RewardFromRotationSimilarity"),
+		const FName Tag = TEXT("RotationSimilarityReward"),
 		const bool bVisualLoggerEnabled = false,
 		ULearningAgentsManagerListener* VisualLoggerListener = nullptr,
 		const int32 VisualLoggerAgentId = -1,
@@ -199,7 +199,7 @@ public:
 		const FQuat RotationA, 
 		const FQuat RotationB, 
 		const float RewardScale = 1.0f,
-		const FName Tag = TEXT("RewardFromRotationSimilarity"),
+		const FName Tag = TEXT("RotationSimilarityReward"),
 		const bool bVisualLoggerEnabled = false,
 		ULearningAgentsManagerListener* VisualLoggerListener = nullptr,
 		const int32 VisualLoggerAgentId = -1,
@@ -229,7 +229,7 @@ public:
 		const float AngleA, 
 		const float AngleB, 
 		const float RewardScale = 1.0f,
-		const FName Tag = TEXT("RewardFromAngleSimilarity"),
+		const FName Tag = TEXT("AngleSimilarityReward"),
 		const bool bVisualLoggerEnabled = false,
 		ULearningAgentsManagerListener* VisualLoggerListener = nullptr,
 		const int32 VisualLoggerAgentId = -1,
@@ -260,7 +260,7 @@ public:
 		const FVector DirectionA, 
 		const FVector DirectionB, 
 		const float RewardScale = 1.0f,
-		const FName Tag = TEXT("RewardFromDirectionSimilarity"),
+		const FName Tag = TEXT("DirectionSimilarityReward"),
 		const bool bVisualLoggerEnabled = false,
 		ULearningAgentsManagerListener* VisualLoggerListener = nullptr,
 		const int32 VisualLoggerAgentId = -1,
@@ -295,7 +295,7 @@ public:
 		const float VelocityScale = 200.0f, 
 		const float RewardScale = 1.0f, 
 		const float FiniteDifferenceDelta = 10.0f,
-		const FName Tag = TEXT("RewardFromVelocityAlongSpline"),
+		const FName Tag = TEXT("VelocityAlongSplineReward"),
 		const bool bVisualLoggerEnabled = false,
 		ULearningAgentsManagerListener* VisualLoggerListener = nullptr,
 		const int32 VisualLoggerAgentId = -1,
