LearningAgents: Fixed bug in policy network size check. Removed IndexObservation. Fixed bug with ArrayObservation. Fixed bug with ProportionAlongSplineObservation.

#virtualized

diff --git a/Engine/Plugins/Experimental/LearningAgents/Source/LearningAgents/Private/LearningAgentsObservations.cpp b/Engine/Plugins/Experimental/LearningAgents/Source/LearningAgents/Private/LearningAgentsObservations.cpp
--- a/Engine/Plugins/Experimental/LearningAgents/Source/LearningAgents/Private/LearningAgentsObservations.cpp
+++ b/Engine/Plugins/Experimental/LearningAgents/Source/LearningAgents/Private/LearningAgentsObservations.cpp
@@ -523,11 +523,6 @@ FLearningAgentsObservationSchemaElement ULearningAgentsObservations::SpecifyIncl
 	return SpecifyContinuousObservation(Schema, Size, Tag);
 }
 
-FLearningAgentsObservationSchemaElement ULearningAgentsObservations::SpecifyIndexObservation(ULearningAgentsObservationSchema* Schema, const int32 Size, const FName Tag)
-{
-	return SpecifyContinuousObservation(Schema, Size, Tag);
-}
-
 FLearningAgentsObservationSchemaElement ULearningAgentsObservations::SpecifyCountObservation(ULearningAgentsObservationSchema* Schema, const FName Tag)
 {
 	return SpecifyContinuousObservation(Schema, 1, Tag);
@@ -902,7 +897,7 @@ FLearningAgentsObservationSchemaElement ULearningAgentsObservations::SpecifyPair
 
 FLearningAgentsObservationSchemaElement ULearningAgentsObservations::SpecifyArrayObservation(ULearningAgentsObservationSchema* Schema, const FLearningAgentsObservationSchemaElement Element, const int32 MaxNum, const int32 AttentionEncodingSize, const int32 AttentionHeadNum, const int32 ValueEncodingSize, const FName Tag)
 {
-	return SpecifySetObservation(Schema, SpecifyPairObservation(Schema, SpecifyIndexObservation(Schema, MaxNum), Element), MaxNum, AttentionEncodingSize, AttentionHeadNum, ValueEncodingSize);
+	return SpecifySetObservation(Schema, SpecifyPairObservation(Schema, SpecifyCountObservation(Schema), Element), MaxNum, AttentionEncodingSize, AttentionHeadNum, ValueEncodingSize);
 }
 
 FLearningAgentsObservationSchemaElement ULearningAgentsObservations::SpecifyMapObservation(ULearningAgentsObservationSchema* Schema, const FLearningAgentsObservationSchemaElement KeyElement, const FLearningAgentsObservationSchemaElement ValueElement, const int32 MaxNum, const int32 AttentionEncodingSize, const int32 AttentionHeadNum, const int32 ValueEncodingSize, const FName Tag)
@@ -1244,51 +1239,6 @@ FLearningAgentsObservationObjectElement ULearningAgentsObservations::MakeInclusi
 	return MakeContinuousObservationFromArrayView(Object, Values, Tag);
 }
 
-FLearningAgentsObservationObjectElement ULearningAgentsObservations::MakeIndexObservation(
-	ULearningAgentsObservationObject* Object, 
-	const int32 Index, 
-	const int32 Size, 
-	const FName Tag,
-	const bool bVisualLoggerEnabled,
-	ULearningAgentsManagerListener* VisualLoggerListener,
-	const int32 VisualLoggerAgentId,
-	const FVector VisualLoggerLocation,
-	const FLinearColor VisualLoggerColor)
-{
-	if (Index < 0 || Index >= Size)
-	{
-		UE_LOG(LogLearning, Error, TEXT("MakeIndexObservation: Discrete index out of range: Got %i, expected <= %i."), Index, Size);
-		return FLearningAgentsObservationObjectElement();
-	}
-
-	TArray<float, TInlineAllocator<32>> Values;
-	Values.Init(0.0f, Size);
-
-	for (int32 Idx = 0; Idx < Index; Idx++)
-	{
-		Values[Idx] = 1.0f;
-	}
-
-#if UE_LEARNING_AGENTS_ENABLE_VISUAL_LOG
-	if (bVisualLoggerEnabled && VisualLoggerListener)
-	{
-		const ULearningAgentsVisualLoggerObject* VisualLoggerObject = VisualLoggerListener->GetOrAddVisualLoggerObject(Tag);
-
-		UE_LEARNING_AGENTS_VLOG_STRING(VisualLoggerObject, LogLearning, Display, VisualLoggerLocation,
-			VisualLoggerColor.ToFColor(true),
-			TEXT("Listener: %s\nTag: %s\nAgent Id: % 3i\nIndex: [%i]\nSize: [%i]\nEncoded: %s"),
-			*VisualLoggerListener->GetName(),
-			*Tag.ToString(),
-			VisualLoggerAgentId,
-			Index,
-			Size,
-			*UE::Learning::Array::FormatFloat(Values, Size));
-	}
-#endif
-
-	return MakeContinuousObservationFromArrayView(Object, Values, Tag);
-}
-
 FLearningAgentsObservationObjectElement ULearningAgentsObservations::MakeCountObservation(
 	ULearningAgentsObservationObject* Object, 
 	const int32 Num, 
@@ -1570,12 +1520,12 @@ FLearningAgentsObservationObjectElement ULearningAgentsObservations::MakePairObs
 	return MakeStructObservationFromArrayViews(Object, { TEXT("Key"), TEXT("Value") }, { Key, Value }, Tag);
 }
 
-FLearningAgentsObservationObjectElement ULearningAgentsObservations::MakeArrayObservation(ULearningAgentsObservationObject* Object, const TArray<FLearningAgentsObservationObjectElement>& Elements, const FName Tag)
+FLearningAgentsObservationObjectElement ULearningAgentsObservations::MakeArrayObservation(ULearningAgentsObservationObject* Object, const TArray<FLearningAgentsObservationObjectElement>& Elements, const int32 MaxNum, const FName Tag)
 {
-	return MakeArrayObservationFromArrayView(Object, Elements, Tag);
+	return MakeArrayObservationFromArrayView(Object, Elements, MaxNum, Tag);
 }
 
-FLearningAgentsObservationObjectElement ULearningAgentsObservations::MakeArrayObservationFromArrayView(ULearningAgentsObservationObject* Object, const TArrayView<const FLearningAgentsObservationObjectElement> Elements, const FName Tag)
+FLearningAgentsObservationObjectElement ULearningAgentsObservations::MakeArrayObservationFromArrayView(ULearningAgentsObservationObject* Object, const TArrayView<const FLearningAgentsObservationObjectElement> Elements, const int32 MaxNum, const FName Tag)
 {
 	if (!Object)
 	{
@@ -1596,7 +1546,7 @@ FLearningAgentsObservationObjectElement ULearningAgentsObservations::MakeArrayOb
 			return FLearningAgentsObservationObjectElement();
 		}
 
-		SubElements.Add(MakePairObservation(Object, MakeIndexObservation(Object, ElementIdx, Elements.Num()), Element));
+		SubElements.Add(MakePairObservation(Object, MakeCountObservation(Object, ElementIdx, MaxNum), Element));
 	}
 
 	return MakeSetObservationFromArrayView(Object, SubElements);
@@ -2366,7 +2316,7 @@ FLearningAgentsObservationObjectElement ULearningAgentsObservations::MakeProport
 			Object, 
 			Angle,
 			0.0f,
-			TEXT("Angle"),
+			TEXT("AngleObservation"),
 			bVisualLoggerEnabled,
 			VisualLoggerListener,
 			VisualLoggerAgentId,
@@ -2381,7 +2331,7 @@ FLearningAgentsObservationObjectElement ULearningAgentsObservations::MakeProport
 			Object, 
 			Proportion, 
 			1.0f, 
-			TEXT("Float"),
+			TEXT("FloatObservation"),
 			bVisualLoggerEnabled, 
 			VisualLoggerListener, 
 			VisualLoggerAgentId, 
@@ -2752,32 +2702,6 @@ bool ULearningAgentsObservations::GetInclusiveDiscreteObservationToArrayView(TAr
 	return true;
 }
 
-bool ULearningAgentsObservations::GetIndexObservation(int32& OutIndex, const ULearningAgentsObservationObject* Object, const FLearningAgentsObservationObjectElement Element, const FName Tag)
-{
-	int32 DiscreteValueNum;
-	if (!GetContinuousObservationNum(DiscreteValueNum, Object, Element, Tag))
-	{
-		OutIndex = INDEX_NONE;
-		return false;
-	}
-
-	TArray<float, TInlineAllocator<32>> OneHot;
-	OneHot.SetNumUninitialized(DiscreteValueNum);
-	if (!GetContinuousObservationToArrayView(OneHot, Object, Element, Tag))
-	{
-		OutIndex = INDEX_NONE;
-		return false;
-	}
-
-	OutIndex = 0;
-	for (int32 Idx = 0; Idx < DiscreteValueNum; Idx++)
-	{
-		if (OneHot[Idx]) { OutIndex = Idx; }
-	}
-
-	return true;
-}
-
 bool ULearningAgentsObservations::GetCountObservation(int32& OutNum, const ULearningAgentsObservationObject* Object, const FLearningAgentsObservationObjectElement Element, const int32 MaxNum, const FName Tag)
 {
 	float FloatNum = 0.0f;
@@ -3455,7 +3379,7 @@ bool ULearningAgentsObservations::GetArrayObservationNum(int32& OutNum, const UL
 	return GetSetObservationNum(OutNum, Object, Element, Tag);
 }
 
-bool ULearningAgentsObservations::GetArrayObservation(TArray<FLearningAgentsObservationObjectElement>& OutElements, const ULearningAgentsObservationObject* Object, const FLearningAgentsObservationObjectElement Element, const FName Tag)
+bool ULearningAgentsObservations::GetArrayObservation(TArray<FLearningAgentsObservationObjectElement>& OutElements, const ULearningAgentsObservationObject* Object, const FLearningAgentsObservationObjectElement Element, const int32 MaxNum, const FName Tag)
 {
 	int32 OutElementNum = 0;
 	if (!GetArrayObservationNum(OutElementNum, Object, Element, Tag))
@@ -3466,7 +3390,7 @@ bool ULearningAgentsObservations::GetArrayObservation(TArray<FLearningAgentsObse
 
 	OutElements.SetNumUninitialized(OutElementNum);
 
-	if (!GetArrayObservationToArrayView(OutElements, Object, Element, Tag))
+	if (!GetArrayObservationToArrayView(OutElements, Object, Element, MaxNum, Tag))
 	{
 		OutElements.Empty();
 		return false;
@@ -3475,7 +3399,7 @@ bool ULearningAgentsObservations::GetArrayObservation(TArray<FLearningAgentsObse
 	return true;
 }
 
-bool ULearningAgentsObservations::GetArrayObservationToArrayView(TArrayView<FLearningAgentsObservationObjectElement> OutElements, const ULearningAgentsObservationObject* Object, const FLearningAgentsObservationObjectElement Element, const FName Tag)
+bool ULearningAgentsObservations::GetArrayObservationToArrayView(TArrayView<FLearningAgentsObservationObjectElement> OutElements, const ULearningAgentsObservationObject* Object, const FLearningAgentsObservationObjectElement Element, const int32 MaxNum, const FName Tag)
 {
 	TArray<FLearningAgentsObservationObjectElement, TInlineAllocator<16>> Pairs;
 	Pairs.SetNumUninitialized(OutElements.Num());
@@ -3495,11 +3419,12 @@ bool ULearningAgentsObservations::GetArrayObservationToArrayView(TArrayView<FLea
 		}
 
 		int32 Index = INDEX_NONE;
-		if (!GetIndexObservation(Index, Object, Key))
+		if (!GetCountObservation(Index, Object, Key, MaxNum))
 		{
 			UE::Learning::Array::Set<1, FLearningAgentsObservationObjectElement>(OutElements, FLearningAgentsObservationObjectElement());
 			return false;
 		}
+		Index = FMath::Clamp(Index, 0, OutElements.Num() - 1);
 
 		OutElements[Index] = Value;
 	}
diff --git a/Engine/Plugins/Experimental/LearningAgents/Source/LearningAgents/Private/LearningAgentsPolicy.cpp b/Engine/Plugins/Experimental/LearningAgents/Source/LearningAgents/Private/LearningAgentsPolicy.cpp
--- a/Engine/Plugins/Experimental/LearningAgents/Source/LearningAgents/Private/LearningAgentsPolicy.cpp
+++ b/Engine/Plugins/Experimental/LearningAgents/Source/LearningAgents/Private/LearningAgentsPolicy.cpp
@@ -611,10 +611,10 @@ void ULearningAgentsPolicy::EvaluatePolicy()
 		return;
 	}
 
-	if (PolicyObject->GetNeuralNetwork()->GetOutputSize() != ActionDistributionVectors.Num<1>() + MemoryState.Num<1>())
+	if (PolicyObject->GetNeuralNetwork()->GetOutputSize() != ActionVectorsEncoded.Num<1>() + MemoryState.Num<1>())
 	{
 		UE_LOG(LogLearning, Error, TEXT("%s: Policy Network Output size don't match. Network output size is %i but Policy expects %i."), *GetName(),
-			PolicyObject->GetNeuralNetwork()->GetOutputSize(), ActionDistributionVectors.Num<1>() + MemoryState.Num<1>());
+			PolicyObject->GetNeuralNetwork()->GetOutputSize(), ActionVectorsEncoded.Num<1>() + MemoryState.Num<1>());
 		return;
 	}
 
diff --git a/Engine/Plugins/Experimental/LearningAgents/Source/LearningAgents/Public/LearningAgentsObservations.h b/Engine/Plugins/Experimental/LearningAgents/Source/LearningAgents/Public/LearningAgentsObservations.h
--- a/Engine/Plugins/Experimental/LearningAgents/Source/LearningAgents/Public/LearningAgentsObservations.h
+++ b/Engine/Plugins/Experimental/LearningAgents/Source/LearningAgents/Public/LearningAgentsObservations.h
@@ -187,18 +187,7 @@ public:
 	static FLearningAgentsObservationSchemaElement SpecifyInclusiveDiscreteObservation(ULearningAgentsObservationSchema* Schema, const int32 Size, const FName Tag = TEXT("InclusiveDiscreteObservation"));
 
 	/**
-	 * Specifies a new index observation. This represents an index into something such as an array.
-	 *
-	 * @param Schema The Observation Schema
-	 * @param Size The size of the object being indexed.
-	 * @param Tag The tag of this new observation. Used during observation object validation and debugging.
-	 * @return The newly created observation schema element.
-	 */
-	UFUNCTION(BlueprintPure, Category = "LearningAgents", meta = (AdvancedDisplay = 2))
-	static FLearningAgentsObservationSchemaElement SpecifyIndexObservation(ULearningAgentsObservationSchema* Schema, const int32 Size, const FName Tag = TEXT("IndexObservation"));
-
-	/**
-	 * Specifies a new count observation. This represents a count of something such as the size of an array.
+	 * Specifies a new count observation. This represents a count of something such as the size of, or index into, an array.
 	 *
 	 * @param Schema The Observation Schema
 	 * @param Tag The tag of this new observation. Used during observation object validation and debugging.
@@ -741,32 +730,6 @@ public:
 		const FVector VisualLoggerLocation = FVector::ZeroVector,
 		const FLinearColor VisualLoggerColor = FLinearColor::Red);
 
-	/**
-	 * Make a new index observation.
-	 *
-	 * @param Object The Observation Object
-	 * @param Index The index. Value must be smaller than the given Size.
-	 * @param Size The size of the discrete observation. Must be equal to the size given during Specify.
-	 * @param Tag The tag of the corresponding observation. Must match the tag given during Specify.
-	 * @param bVisualLoggerEnabled When true, debug data will be sent to the visual logger.
-	 * @param VisualLoggerListener The listener object which is making this observation. This must be set to use logging.
-	 * @param VisualLoggerAgentId The agent id associated with this observation.
-	 * @param VisualLoggerLocation A location for the visual logger information in the world.
-	 * @param VisualLoggerColor The color for the visual logger display.
-	 * @return The newly created observation object element.
-	 */
-	UFUNCTION(BlueprintPure, Category = "LearningAgents", meta = (AdvancedDisplay = 3, DefaultToSelf = "VisualLoggerListener"))
-	static FLearningAgentsObservationObjectElement MakeIndexObservation(
-		ULearningAgentsObservationObject* Object, 
-		const int32 Index, 
-		const int32 Size, 
-		const FName Tag = TEXT("IndexObservation"),
-		const bool bVisualLoggerEnabled = false,
-		ULearningAgentsManagerListener* VisualLoggerListener = nullptr,
-		const int32 VisualLoggerAgentId = -1,
-		const FVector VisualLoggerLocation = FVector::ZeroVector,
-		const FLinearColor VisualLoggerColor = FLinearColor::Red);
-
 	/**
 	 * Make a new count observation.
 	 *
@@ -847,7 +810,7 @@ public:
 	 * @param Tag The tag of the corresponding observation. Must match the tag given during Specify.
 	 * @return The newly created observation object element.
 	 */
-	UFUNCTION(BlueprintPure, Category = "LearningAgents", meta = (AdvancedDisplay = 2))
+	UFUNCTION(BlueprintPure, Category = "LearningAgents", meta = (AdvancedDisplay = 2, AutoCreateRefTerm = "Elements"))
 	static FLearningAgentsObservationObjectElement MakeInclusiveUnionObservation(ULearningAgentsObservationObject* Object, const TMap<FName, FLearningAgentsObservationObjectElement>& Elements, const FName Tag = TEXT("InclusiveUnionObservation"));
 
 	/**
@@ -859,7 +822,7 @@ public:
 	 * @param Tag The tag of the corresponding observation. Must match the tag given during Specify.
 	 * @return The newly created observation object element.
 	 */
-	UFUNCTION(BlueprintPure, Category = "LearningAgents", meta = (AdvancedDisplay = 3))
+	UFUNCTION(BlueprintPure, Category = "LearningAgents", meta = (AdvancedDisplay = 3, AutoCreateRefTerm = "ElementNames,Elements"))
 	static FLearningAgentsObservationObjectElement MakeInclusiveUnionObservationFromArrays(ULearningAgentsObservationObject* Object, const TArray<FName>& ElementNames, const TArray<FLearningAgentsObservationObjectElement>& Elements, const FName Tag = TEXT("InclusiveUnionObservation"));
 	
 	/**
@@ -902,7 +865,7 @@ public:
 	 * @param Tag The tag of the corresponding observation. Must match the tag given during Specify.
 	 * @return The newly created observation object element.
 	 */
-	UFUNCTION(BlueprintPure, Category = "LearningAgents", meta = (AdvancedDisplay = 2))
+	UFUNCTION(BlueprintPure, Category = "LearningAgents", meta = (AdvancedDisplay = 2, AutoCreateRefTerm = "Elements"))
 	static FLearningAgentsObservationObjectElement MakeSetObservation(ULearningAgentsObservationObject* Object, const TSet<FLearningAgentsObservationObjectElement>& Elements, const FName Tag = TEXT("SetObservation"));
 
 	/**
@@ -913,7 +876,7 @@ public:
 	 * @param Tag The tag of the corresponding observation. Must match the tag given during Specify.
 	 * @return The newly created observation object element.
 	 */
-	UFUNCTION(BlueprintPure, Category = "LearningAgents", meta = (AdvancedDisplay = 2))
+	UFUNCTION(BlueprintPure, Category = "LearningAgents", meta = (AdvancedDisplay = 2, AutoCreateRefTerm = "Elements"))
 	static FLearningAgentsObservationObjectElement MakeSetObservationFromArray(ULearningAgentsObservationObject* Object, const TArray<FLearningAgentsObservationObjectElement>& Elements, const FName Tag = TEXT("SetObservation"));
 
 	/**
@@ -943,21 +906,23 @@ public:
 	 *
 	 * @param Object The Observation Object
 	 * @param Elements The sub-observations. The number of elements here must be less than or equal to the maximum that was given during Specify.
+	 * @param MaxNum The maximum number of elements possible for this observation. Must match what was given during Specify.
 	 * @param Tag The tag of the corresponding observation. Must match the tag given during Specify.
 	 * @return The newly created observation object element.
 	 */
-	UFUNCTION(BlueprintPure, Category = "LearningAgents", meta = (AdvancedDisplay = 2))
-	static FLearningAgentsObservationObjectElement MakeArrayObservation(ULearningAgentsObservationObject* Object, const TArray<FLearningAgentsObservationObjectElement>& Elements, const FName Tag = TEXT("ArrayObservation"));
+	UFUNCTION(BlueprintPure, Category = "LearningAgents", meta = (AdvancedDisplay = 3, AutoCreateRefTerm = "Elements"))
+	static FLearningAgentsObservationObjectElement MakeArrayObservation(ULearningAgentsObservationObject* Object, const TArray<FLearningAgentsObservationObjectElement>& Elements, const int32 MaxNum, const FName Tag = TEXT("ArrayObservation"));
 
 	/**
 	 * Make a new array observation.
 	 *
 	 * @param Object The Observation Object
 	 * @param Elements The sub-observations. The number of elements here must be less than or equal to the maximum that was given during Specify.
+	 * @param MaxNum The maximum number of elements possible for this observation. Must match what was given during Specify.
 	 * @param Tag The tag of the corresponding observation. Must match the tag given during Specify.
 	 * @return The newly created observation object element.
 	 */
-	static FLearningAgentsObservationObjectElement MakeArrayObservationFromArrayView(ULearningAgentsObservationObject* Object, const TArrayView<const FLearningAgentsObservationObjectElement> Elements, const FName Tag = TEXT("ArrayObservation"));
+	static FLearningAgentsObservationObjectElement MakeArrayObservationFromArrayView(ULearningAgentsObservationObject* Object, const TArrayView<const FLearningAgentsObservationObjectElement> Elements, const int32 MaxNum, const FName Tag = TEXT("ArrayObservation"));
 
 	/**
 	 * Make a new map observation.
@@ -967,7 +932,7 @@ public:
 	 * @param Tag The tag of the corresponding observation. Must match the tag given during Specify.
 	 * @return The newly created observation object element.
 	 */
-	UFUNCTION(BlueprintPure, Category = "LearningAgents", meta = (AdvancedDisplay = 2))
+	UFUNCTION(BlueprintPure, Category = "LearningAgents", meta = (AdvancedDisplay = 2, AutoCreateRefTerm = "Map"))
 	static FLearningAgentsObservationObjectElement MakeMapObservation(ULearningAgentsObservationObject* Object, const TMap<FLearningAgentsObservationObjectElement, FLearningAgentsObservationObjectElement>& Map, const FName Tag = TEXT("MapObservation"));
 
 	/**
@@ -978,7 +943,7 @@ public:
 	 * @param Tag The tag of the corresponding observation. Must match the tag given during Specify.
 	 * @return The newly created observation object element.
 	 */
-	UFUNCTION(BlueprintPure, Category = "LearningAgents", meta = (AdvancedDisplay = 3))
+	UFUNCTION(BlueprintPure, Category = "LearningAgents", meta = (AdvancedDisplay = 3, AutoCreateRefTerm = "Keys,Values"))
 	static FLearningAgentsObservationObjectElement MakeMapObservationFromArrays(ULearningAgentsObservationObject* Object, const TArray<FLearningAgentsObservationObjectElement>& Keys, const TArray<FLearningAgentsObservationObjectElement>& Values, const FName Tag = TEXT("MapObservation"));
 	
 	/**
@@ -1562,9 +1527,6 @@ public:
 	static bool GetInclusiveDiscreteObservation(TArray<int32>& OutIndices, const ULearningAgentsObservationObject* Object, const FLearningAgentsObservationObjectElement Element, const FName Tag = TEXT("InclusiveDiscreteObservation"));
 	static bool GetInclusiveDiscreteObservationToArrayView(TArrayView<int32> OutIndices, const ULearningAgentsObservationObject* Object, const FLearningAgentsObservationObjectElement Element, const FName Tag = TEXT("InclusiveDiscreteObservation"));
 
-	UFUNCTION(BlueprintPure = false, Category = "LearningAgents", meta = (AdvancedDisplay = 3, ReturnDisplayName = "Success"))
-	static bool GetIndexObservation(int32& OutIndex, const ULearningAgentsObservationObject* Object, const FLearningAgentsObservationObjectElement Element, const FName Tag = TEXT("IndexObservation"));
-
 	UFUNCTION(BlueprintPure = false, Category = "LearningAgents", meta = (AdvancedDisplay = 3, ReturnDisplayName = "Success"))
 	static bool GetCountObservation(int32& OutNum, const ULearningAgentsObservationObject* Object, const FLearningAgentsObservationObjectElement Element, const int32 MaxNum, const FName Tag = TEXT("CountObservation"));
 
@@ -1615,8 +1577,8 @@ public:
 	static bool GetArrayObservationNum(int32& OutNum, const ULearningAgentsObservationObject* Object, const FLearningAgentsObservationObjectElement Element, const FName Tag = TEXT("ArrayObservation"));
 
 	UFUNCTION(BlueprintPure = false, Category = "LearningAgents", meta = (AdvancedDisplay = 3, ReturnDisplayName = "Success"))
-	static bool GetArrayObservation(TArray<FLearningAgentsObservationObjectElement>& OutElements, const ULearningAgentsObservationObject* Object, const FLearningAgentsObservationObjectElement Element, const FName Tag = TEXT("ArrayObservation"));
-	static bool GetArrayObservationToArrayView(TArrayView<FLearningAgentsObservationObjectElement> OutElements, const ULearningAgentsObservationObject* Object, const FLearningAgentsObservationObjectElement Element, const FName Tag = TEXT("ArrayObservation"));
+	static bool GetArrayObservation(TArray<FLearningAgentsObservationObjectElement>& OutElements, const ULearningAgentsObservationObject* Object, const FLearningAgentsObservationObjectElement Element, const int32 MaxNum, const FName Tag = TEXT("ArrayObservation"));
+	static bool GetArrayObservationToArrayView(TArrayView<FLearningAgentsObservationObjectElement> OutElements, const ULearningAgentsObservationObject* Object, const FLearningAgentsObservationObjectElement Element, const int32 MaxNum, const FName Tag = TEXT("ArrayObservation"));
 
 	UFUNCTION(BlueprintPure = false, Category = "LearningAgents", meta = (AdvancedDisplay = 3, ReturnDisplayName = "Success"))
 	static bool GetMapObservationNum(int32& OutNum, const ULearningAgentsObservationObject* Object, const FLearningAgentsObservationObjectElement Element, const FName Tag = TEXT("MapObservation"));
