LearningAgents: Added option to disable MaxFPS console variable and editor vsync while training.

diff --git a/Engine/Plugins/Experimental/LearningAgents/Source/LearningAgentsTraining/Private/LearningAgentsTrainer.cpp b/Engine/Plugins/Experimental/LearningAgents/Source/LearningAgentsTraining/Private/LearningAgentsTrainer.cpp
--- a/Engine/Plugins/Experimental/LearningAgents/Source/LearningAgentsTraining/Private/LearningAgentsTrainer.cpp
+++ b/Engine/Plugins/Experimental/LearningAgents/Source/LearningAgentsTraining/Private/LearningAgentsTrainer.cpp
@@ -435,6 +435,12 @@ void ULearningAgentsTrainer::BeginTraining(
 		MaxPhysicsStep = PhysicsSettings->MaxPhysicsDeltaTime;
 	}
 
+	IConsoleVariable* MaxFPSCVar = IConsoleManager::Get().FindConsoleVariable(TEXT("t.MaxFPS"));
+	if (MaxFPSCVar)
+	{
+		MaxFPS = MaxFPSCVar->GetInt();
+	}
+
 	UGameViewportClient* ViewportClient = GetWorld() ? GetWorld()->GetGameViewport() : nullptr;
 	if (ViewportClient)
 	{
@@ -446,6 +452,7 @@ void ULearningAgentsTrainer::BeginTraining(
 	if (EditorPerformanceSettings)
 	{
 		bUseLessCPUInTheBackground = EditorPerformanceSettings->bThrottleCPUWhenNotForeground;
+		bEditorVSyncEnabled = EditorPerformanceSettings->bEnableVSync;
 	}
 #endif
 
@@ -466,6 +473,11 @@ void ULearningAgentsTrainer::BeginTraining(
 		UE_LOG(LogLearning, Warning, TEXT("%s: Provided invalid FixedTimeStepFrequency: %0.5f"), *GetName(), TrainerGameSettings.FixedTimeStepFrequency);
 	}
 
+	if (TrainerGameSettings.bDisableMaxFPS && MaxFPSCVar)
+	{
+		MaxFPSCVar->Set(0);
+	}
+
 	if (TrainerGameSettings.bDisableVSync && GameSettings)
 	{
 		GameSettings->SetVSyncEnabled(false);
@@ -480,11 +492,14 @@ void ULearningAgentsTrainer::BeginTraining(
 #if WITH_EDITOR
 	if (TrainerGameSettings.bDisableUseLessCPUInTheBackground && EditorPerformanceSettings)
 	{
-		if (TrainerGameSettings.bDisableUseLessCPUInTheBackground)
-		{
-			EditorPerformanceSettings->bThrottleCPUWhenNotForeground = false;
-			EditorPerformanceSettings->PostEditChange();
-		}
+		EditorPerformanceSettings->bThrottleCPUWhenNotForeground = false;
+		EditorPerformanceSettings->PostEditChange();
+	}
+
+	if (TrainerGameSettings.bDisableEditorVSync && EditorPerformanceSettings)
+	{
+		EditorPerformanceSettings->bEnableVSync = false;
+		EditorPerformanceSettings->PostEditChange();
 	}
 #endif
 
@@ -615,6 +630,12 @@ void ULearningAgentsTrainer::DoneTraining()
 			PhysicsSettings->MaxPhysicsDeltaTime = MaxPhysicsStep;
 		}
 
+		IConsoleVariable* MaxFPSCVar = IConsoleManager::Get().FindConsoleVariable(TEXT("t.MaxFPS"));
+		if (MaxFPSCVar)
+		{
+			MaxFPSCVar->Set(MaxFPS);
+		}
+
 		UGameViewportClient* ViewportClient = GetWorld() ? GetWorld()->GetGameViewport() : nullptr;
 		if (ViewportClient)
 		{
@@ -626,6 +647,7 @@ void ULearningAgentsTrainer::DoneTraining()
 		if (EditorPerformanceSettings)
 		{
 			EditorPerformanceSettings->bThrottleCPUWhenNotForeground = bUseLessCPUInTheBackground;
+			EditorPerformanceSettings->bEnableVSync = bEditorVSyncEnabled;
 			EditorPerformanceSettings->PostEditChange();
 		}
 #endif
diff --git a/Engine/Plugins/Experimental/LearningAgents/Source/LearningAgentsTraining/Public/LearningAgentsTrainer.h b/Engine/Plugins/Experimental/LearningAgents/Source/LearningAgentsTraining/Public/LearningAgentsTrainer.h
--- a/Engine/Plugins/Experimental/LearningAgents/Source/LearningAgentsTraining/Public/LearningAgentsTrainer.h
+++ b/Engine/Plugins/Experimental/LearningAgents/Source/LearningAgentsTraining/Public/LearningAgentsTrainer.h
@@ -92,6 +92,10 @@ public:
 	UPROPERTY(EditAnywhere, Category = "LearningAgents")
 	bool bSetMaxPhysicsStepToFixedTimeStep = true;
 
+	/** If true, the MaxFPS console variable will be set to a negative number during training; Otherwise, it will not. */
+	UPROPERTY(EditAnywhere, Category = "LearningAgents")
+	bool bDisableMaxFPS = true;
+
 	/** If true, VSync will be disabled; Otherwise, it will not. Disabling VSync can speed up the game simulation. */
 	UPROPERTY(EditAnywhere, Category = "LearningAgents")
 	bool bDisableVSync = true;
@@ -106,6 +110,10 @@ public:
 	UPROPERTY(EditAnywhere, Category = "LearningAgents")
 	bool bDisableUseLessCPUInTheBackground = true;
 
+	/** If true, Editor VSync will be disabled; Otherwise, it will not. Disabling Editor VSync can speed up the game simulation. */
+	UPROPERTY(EditAnywhere, Category = "LearningAgents")
+	bool bDisableEditorVSync = true;
+
 #endif
 };
 
@@ -689,10 +697,13 @@ private:
 // ----- Private Recording of GameSettings ----- 
 private:
 
-	bool bUseLessCPUInTheBackground = true;
 	bool bFixedTimestepUsed = false;
 	float FixedTimeStepDeltaTime = -1.0f;
-	bool bVSyncEnabled = true;
 	float MaxPhysicsStep = -1.0f;
+	int32 MaxFPS = 120;
+	bool bVSyncEnabled = true;
 	int32 ViewModeIndex = -1;
+	
+	bool bUseLessCPUInTheBackground = true;
+	bool bEditorVSyncEnabled = true;
 };
