LearningAgents: Changed LayerNum parameter everywhere to HiddenLayerNum

diff --git a/Engine/Plugins/Experimental/LearningAgents/Source/LearningAgents/Private/LearningAgentsActions.cpp b/Engine/Plugins/Experimental/LearningAgents/Source/LearningAgents/Private/LearningAgentsActions.cpp
--- a/Engine/Plugins/Experimental/LearningAgents/Source/LearningAgents/Private/LearningAgentsActions.cpp
+++ b/Engine/Plugins/Experimental/LearningAgents/Source/LearningAgents/Private/LearningAgentsActions.cpp
@@ -1067,7 +1067,7 @@ FLearningAgentsActionSchemaElement ULearningAgentsActions::SpecifyEitherAction(U
 	return SpecifyExclusiveUnionActionFromArrayViews(Schema, { TEXT("A"), TEXT("B") }, { A, B }, { 1.0f - PriorProbabilityOfA, PriorProbabilityOfA }, Tag);
 }
 
-FLearningAgentsActionSchemaElement ULearningAgentsActions::SpecifyEncodingAction(ULearningAgentsActionSchema* Schema, const FLearningAgentsActionSchemaElement Element, const int32 EncodingSize, const int32 LayerNum, const ELearningAgentsActivationFunction ActivationFunction, const FName Tag)
+FLearningAgentsActionSchemaElement ULearningAgentsActions::SpecifyEncodingAction(ULearningAgentsActionSchema* Schema, const FLearningAgentsActionSchemaElement Element, const int32 EncodingSize, const int32 HiddenLayerNum, const ELearningAgentsActivationFunction ActivationFunction, const FName Tag)
 {
 	if (!Schema)
 	{
@@ -1081,9 +1081,9 @@ FLearningAgentsActionSchemaElement ULearningAgentsActions::SpecifyEncodingAction
 		return FLearningAgentsActionSchemaElement();
 	}
 
-	if (LayerNum < 1)
+	if (HiddenLayerNum < 1)
 	{
-		UE_LOG(LogLearning, Error, TEXT("SpecifyEncodingAction: Invalid Action LayerNum '%i' - must be greater than zero."), LayerNum);
+		UE_LOG(LogLearning, Error, TEXT("SpecifyEncodingAction: Invalid Action HiddenLayerNum '%i' - must be greater than zero."), HiddenLayerNum);
 		return FLearningAgentsActionSchemaElement();
 	}
 
@@ -1093,7 +1093,7 @@ FLearningAgentsActionSchemaElement ULearningAgentsActions::SpecifyEncodingAction
 		return FLearningAgentsActionSchemaElement();
 	}
 
-	return { Schema->ActionSchema.CreateEncoding({ Element.SchemaElement, EncodingSize, LayerNum, UE::Learning::Agents::Action::Private::GetEncodingActivationFunction(ActivationFunction) }, Tag)};
+	return { Schema->ActionSchema.CreateEncoding({ Element.SchemaElement, EncodingSize, HiddenLayerNum, UE::Learning::Agents::Action::Private::GetEncodingActivationFunction(ActivationFunction) }, Tag)};
 }
 
 FLearningAgentsActionSchemaElement ULearningAgentsActions::SpecifyBoolAction(ULearningAgentsActionSchema* Schema, const float PriorProbability, const FName Tag)
diff --git a/Engine/Plugins/Experimental/LearningAgents/Source/LearningAgents/Private/LearningAgentsCritic.cpp b/Engine/Plugins/Experimental/LearningAgents/Source/LearningAgents/Private/LearningAgentsCritic.cpp
--- a/Engine/Plugins/Experimental/LearningAgents/Source/LearningAgents/Private/LearningAgentsCritic.cpp
+++ b/Engine/Plugins/Experimental/LearningAgents/Source/LearningAgents/Private/LearningAgentsCritic.cpp
@@ -183,7 +183,7 @@ void ULearningAgentsCritic::SetupCritic(
 				ObservationEncodedVectorSize + MemoryStateSize,
 				1,
 				CriticSettings.HiddenLayerSize,
-				FMath::Max(CriticSettings.LayerNum, 2),
+				CriticSettings.HiddenLayerNum + 2, // Add 2 to account for input and output layers
 				UE::Learning::Agents::Critic::Private::GetBuilderActivationFunction(CriticSettings.ActivationFunction)));
 
 		UE_LEARNING_CHECK(CriticInputSize == ObservationEncodedVectorSize + MemoryStateSize);
diff --git a/Engine/Plugins/Experimental/LearningAgents/Source/LearningAgents/Private/LearningAgentsObservations.cpp b/Engine/Plugins/Experimental/LearningAgents/Source/LearningAgents/Private/LearningAgentsObservations.cpp
--- a/Engine/Plugins/Experimental/LearningAgents/Source/LearningAgents/Private/LearningAgentsObservations.cpp
+++ b/Engine/Plugins/Experimental/LearningAgents/Source/LearningAgents/Private/LearningAgentsObservations.cpp
@@ -943,7 +943,7 @@ FLearningAgentsObservationSchemaElement ULearningAgentsObservations::SpecifyEith
 	return SpecifyExclusiveUnionObservationFromArrayViews(Schema, { TEXT("A"), TEXT("B") }, { A, B }, EncodingSize, Tag);
 }
 
-FLearningAgentsObservationSchemaElement ULearningAgentsObservations::SpecifyEncodingObservation(ULearningAgentsObservationSchema* Schema, const FLearningAgentsObservationSchemaElement Element, const int32 EncodingSize, const int32 LayerNum, const ELearningAgentsActivationFunction ActivationFunction, const FName Tag)
+FLearningAgentsObservationSchemaElement ULearningAgentsObservations::SpecifyEncodingObservation(ULearningAgentsObservationSchema* Schema, const FLearningAgentsObservationSchemaElement Element, const int32 EncodingSize, const int32 HiddenLayerNum, const ELearningAgentsActivationFunction ActivationFunction, const FName Tag)
 {
 	if (!Schema)
 	{
@@ -957,9 +957,9 @@ FLearningAgentsObservationSchemaElement ULearningAgentsObservations::SpecifyEnco
 		return FLearningAgentsObservationSchemaElement();
 	}
 
-	if (LayerNum < 1)
+	if (HiddenLayerNum < 1)
 	{
-		UE_LOG(LogLearning, Error, TEXT("SpecifyEncodingObservation: Invalid Observation LayerNum '%i' - must be greater than zero."), LayerNum);
+		UE_LOG(LogLearning, Error, TEXT("SpecifyEncodingObservation: Invalid Observation HiddenLayerNum '%i' - must be greater than zero."), HiddenLayerNum);
 		return FLearningAgentsObservationSchemaElement();
 	}
 
@@ -969,7 +969,7 @@ FLearningAgentsObservationSchemaElement ULearningAgentsObservations::SpecifyEnco
 		return FLearningAgentsObservationSchemaElement();
 	}
 
-	return { Schema->ObservationSchema.CreateEncoding({ Element.SchemaElement, EncodingSize, LayerNum, UE::Learning::Agents::Observation::Private::GetEncodingActivationFunction(ActivationFunction) }, Tag)};
+	return { Schema->ObservationSchema.CreateEncoding({ Element.SchemaElement, EncodingSize, HiddenLayerNum, UE::Learning::Agents::Observation::Private::GetEncodingActivationFunction(ActivationFunction) }, Tag)};
 }
 
 FLearningAgentsObservationSchemaElement ULearningAgentsObservations::SpecifyBoolObservation(ULearningAgentsObservationSchema* Schema, const FName Tag)
diff --git a/Engine/Plugins/Experimental/LearningAgents/Source/LearningAgents/Private/LearningAgentsPolicy.cpp b/Engine/Plugins/Experimental/LearningAgents/Source/LearningAgents/Private/LearningAgentsPolicy.cpp
--- a/Engine/Plugins/Experimental/LearningAgents/Source/LearningAgents/Private/LearningAgentsPolicy.cpp
+++ b/Engine/Plugins/Experimental/LearningAgents/Source/LearningAgents/Private/LearningAgentsPolicy.cpp
@@ -255,7 +255,7 @@ void ULearningAgentsPolicy::SetupPolicy(
 		UE::NNE::RuntimeBasic::FModelBuilder Builder(UE::Learning::Random::Int(Seed ^ 0x69315bf9));
 
 		const int32 PolicyHiddenLayerSize = PolicySettings.HiddenLayerSize;
-		const int32 PolicyLayerNum = PolicySettings.LayerNum;
+		const int32 PolicyHiddenLayerNum = PolicySettings.HiddenLayerNum;
 		const ELearningAgentsActivationFunction PolicyActivationFunction = PolicySettings.ActivationFunction;
 		const float PolicyInitialEncodedActionScale = PolicySettings.InitialEncodedActionScale;
 
@@ -268,7 +268,7 @@ void ULearningAgentsPolicy::SetupPolicy(
 					ObservationEncodedVectorSize,
 					PolicyHiddenLayerSize,
 					PolicyHiddenLayerSize,
-					FMath::Max(PolicyLayerNum / 2 + 1, 2),
+					PolicyHiddenLayerNum / 2 + 2, // Add 2 to account for input and output layers
 					UE::Learning::Agents::Policy::Private::GetBuilderActivationFunction(PolicyActivationFunction),
 					true),
 				Builder.MakeMemoryCellWithLinearRandomKaimingWeights(
@@ -281,7 +281,7 @@ void ULearningAgentsPolicy::SetupPolicy(
 						PolicyHiddenLayerSize,
 						ActionEncodedVectorSize,
 						PolicyHiddenLayerSize,
-						FMath::Max(PolicyLayerNum / 2 + 1, 2),
+						PolicyHiddenLayerNum / 2 + 2, // Add 2 to account for input and output layers
 						UE::Learning::Agents::Policy::Private::GetBuilderActivationFunction(PolicyActivationFunction)),
 					Builder.MakeDenormalize(
 						ActionEncodedVectorSize,
diff --git a/Engine/Plugins/Experimental/LearningAgents/Source/LearningAgents/Public/LearningAgentsActions.h b/Engine/Plugins/Experimental/LearningAgents/Source/LearningAgents/Public/LearningAgentsActions.h
--- a/Engine/Plugins/Experimental/LearningAgents/Source/LearningAgents/Public/LearningAgentsActions.h
+++ b/Engine/Plugins/Experimental/LearningAgents/Source/LearningAgents/Public/LearningAgentsActions.h
@@ -444,13 +444,13 @@ public:
 	 * @param Schema The Action Schema
 	 * @param Element The sub-action.
 	 * @param EncodingSize The encoding size used to decode this sub-action.
-	 * @param LayerNum The number of layers used to decode this sub-action.
+	 * @param HiddenLayerNum The number of hidden layers used to decode this sub-action.
 	 * @param ActivationFunction The activation function used to decode this sub-action.
 	 * @param Tag The tag of this new action. Used during action object validation and debugging.
 	 * @return The newly created action schema element.
 	 */
 	UFUNCTION(BlueprintPure, Category = "LearningAgents", meta = (AdvancedDisplay = 2))
-	static FLearningAgentsActionSchemaElement SpecifyEncodingAction(ULearningAgentsActionSchema* Schema, const FLearningAgentsActionSchemaElement Element, const int32 EncodingSize = 128, const int32 LayerNum = 1, const ELearningAgentsActivationFunction ActivationFunction = ELearningAgentsActivationFunction::ELU, const FName Tag = TEXT("EncodingAction"));
+	static FLearningAgentsActionSchemaElement SpecifyEncodingAction(ULearningAgentsActionSchema* Schema, const FLearningAgentsActionSchemaElement Element, const int32 EncodingSize = 128, const int32 HiddenLayerNum = 1, const ELearningAgentsActivationFunction ActivationFunction = ELearningAgentsActivationFunction::ELU, const FName Tag = TEXT("EncodingAction"));
 
 	/**
 	 * Specifies a new bool action. This represents an action which is either true or false.
diff --git a/Engine/Plugins/Experimental/LearningAgents/Source/LearningAgents/Public/LearningAgentsCritic.h b/Engine/Plugins/Experimental/LearningAgents/Source/LearningAgents/Public/LearningAgentsCritic.h
--- a/Engine/Plugins/Experimental/LearningAgents/Source/LearningAgents/Public/LearningAgentsCritic.h
+++ b/Engine/Plugins/Experimental/LearningAgents/Source/LearningAgents/Public/LearningAgentsCritic.h
@@ -27,9 +27,9 @@ struct LEARNINGAGENTS_API FLearningAgentsCriticSettings
 
 public:
 
-	/** Total layers for critic network including input, hidden, and output layers */
-	UPROPERTY(EditAnywhere, Category = "LearningAgents", meta = (ClampMin = "2", UIMin = "2"))
-	int32 LayerNum = 3;
+	/** Number of hidden layers for critic network. */
+	UPROPERTY(EditAnywhere, Category = "LearningAgents", meta = (ClampMin = "0", UIMin = "0"))
+	int32 HiddenLayerNum = 1;
 
 	/** Number of neurons in each hidden layer of the critic network */
 	UPROPERTY(EditAnywhere, Category = "LearningAgents", meta = (ClampMin = "1", UIMin = "1"))
diff --git a/Engine/Plugins/Experimental/LearningAgents/Source/LearningAgents/Public/LearningAgentsObservations.h b/Engine/Plugins/Experimental/LearningAgents/Source/LearningAgents/Public/LearningAgentsObservations.h
--- a/Engine/Plugins/Experimental/LearningAgents/Source/LearningAgents/Public/LearningAgentsObservations.h
+++ b/Engine/Plugins/Experimental/LearningAgents/Source/LearningAgents/Public/LearningAgentsObservations.h
@@ -450,12 +450,12 @@ public:
 	 * @param Schema The Observation Schema
 	 * @param Element The sub-observation to be encoded.
 	 * @param EncodingSize The encoding size used to encode this sub-observation.
-	 * @param LayerNum The number of layers used to encode this sub-observation.
+	 * @param HiddenLayerNum The number of hidden layers used to encode this sub-observation.
 	 * @param ActivationFunction The activation function used to encode this sub-observation.
 	 * @return The newly created observation schema element.
 	 */
 	UFUNCTION(BlueprintPure, Category = "LearningAgents", meta = (AdvancedDisplay = 2))
-	static FLearningAgentsObservationSchemaElement SpecifyEncodingObservation(ULearningAgentsObservationSchema* Schema, const FLearningAgentsObservationSchemaElement Element, const int32 EncodingSize = 128, const int32 LayerNum = 1, const ELearningAgentsActivationFunction ActivationFunction = ELearningAgentsActivationFunction::ELU, const FName Tag = TEXT("EncodingObservation"));
+	static FLearningAgentsObservationSchemaElement SpecifyEncodingObservation(ULearningAgentsObservationSchema* Schema, const FLearningAgentsObservationSchemaElement Element, const int32 EncodingSize = 128, const int32 HiddenLayerNum = 1, const ELearningAgentsActivationFunction ActivationFunction = ELearningAgentsActivationFunction::ELU, const FName Tag = TEXT("EncodingObservation"));
 
 	/**
 	 * Specifies a new bool observation. A true or false observation.
diff --git a/Engine/Plugins/Experimental/LearningAgents/Source/LearningAgents/Public/LearningAgentsPolicy.h b/Engine/Plugins/Experimental/LearningAgents/Source/LearningAgents/Public/LearningAgentsPolicy.h
--- a/Engine/Plugins/Experimental/LearningAgents/Source/LearningAgents/Public/LearningAgentsPolicy.h
+++ b/Engine/Plugins/Experimental/LearningAgents/Source/LearningAgents/Public/LearningAgentsPolicy.h
@@ -26,9 +26,9 @@ struct LEARNINGAGENTS_API FLearningAgentsPolicySettings
 
 public:
 
-	/** Total layers for policy network including input, hidden, and output layers */
-	UPROPERTY(EditAnywhere, Category = "LearningAgents", meta = (ClampMin = "2", UIMin = "2"))
-	int32 LayerNum = 3;
+	/** Number of hidden layers for policy network. */
+	UPROPERTY(EditAnywhere, Category = "LearningAgents", meta = (ClampMin = "0", UIMin = "0"))
+	int32 HiddenLayerNum = 1;
 
 	/** Number of neurons in each hidden layer of the policy network */
 	UPROPERTY(EditAnywhere, Category = "LearningAgents", meta = (ClampMin = "1", UIMin = "1"))
