LearningAgents: Fixed some more minor bugs.

diff --git a/Engine/Plugins/Experimental/LearningAgents/Source/LearningAgents/Private/LearningAgentsObservations.cpp b/Engine/Plugins/Experimental/LearningAgents/Source/LearningAgents/Private/LearningAgentsObservations.cpp
--- a/Engine/Plugins/Experimental/LearningAgents/Source/LearningAgents/Private/LearningAgentsObservations.cpp
+++ b/Engine/Plugins/Experimental/LearningAgents/Source/LearningAgents/Private/LearningAgentsObservations.cpp
@@ -897,12 +897,12 @@ FLearningAgentsObservationSchemaElement ULearningAgentsObservations::SpecifyPair
 
 FLearningAgentsObservationSchemaElement ULearningAgentsObservations::SpecifyArrayObservation(ULearningAgentsObservationSchema* Schema, const FLearningAgentsObservationSchemaElement Element, const int32 MaxNum, const int32 AttentionEncodingSize, const int32 AttentionHeadNum, const int32 ValueEncodingSize, const FName Tag)
 {
-	return SpecifySetObservation(Schema, SpecifyPairObservation(Schema, SpecifyCountObservation(Schema), Element), MaxNum, AttentionEncodingSize, AttentionHeadNum, ValueEncodingSize);
+	return SpecifySetObservation(Schema, SpecifyPairObservation(Schema, SpecifyCountObservation(Schema), Element), MaxNum, AttentionEncodingSize, AttentionHeadNum, ValueEncodingSize, Tag);
 }
 
 FLearningAgentsObservationSchemaElement ULearningAgentsObservations::SpecifyMapObservation(ULearningAgentsObservationSchema* Schema, const FLearningAgentsObservationSchemaElement KeyElement, const FLearningAgentsObservationSchemaElement ValueElement, const int32 MaxNum, const int32 AttentionEncodingSize, const int32 AttentionHeadNum, const int32 ValueEncodingSize, const FName Tag)
 {
-	return SpecifySetObservation(Schema, SpecifyPairObservation(Schema, KeyElement, ValueElement), MaxNum, AttentionEncodingSize, AttentionHeadNum, ValueEncodingSize);
+	return SpecifySetObservation(Schema, SpecifyPairObservation(Schema, KeyElement, ValueElement), MaxNum, AttentionEncodingSize, AttentionHeadNum, ValueEncodingSize, Tag);
 }
 
 FLearningAgentsObservationSchemaElement ULearningAgentsObservations::SpecifyEnumObservation(ULearningAgentsObservationSchema* Schema, const UEnum* Enum, const FName Tag)
@@ -1549,7 +1549,7 @@ FLearningAgentsObservationObjectElement ULearningAgentsObservations::MakeArrayOb
 		SubElements.Add(MakePairObservation(Object, MakeCountObservation(Object, ElementIdx, MaxNum), Element));
 	}
 
-	return MakeSetObservationFromArrayView(Object, SubElements);
+	return MakeSetObservationFromArrayView(Object, SubElements, Tag);
 }
 
 FLearningAgentsObservationObjectElement ULearningAgentsObservations::MakeMapObservation(ULearningAgentsObservationObject* Object, const TMap<FLearningAgentsObservationObjectElement, FLearningAgentsObservationObjectElement>& Map, const FName Tag)
@@ -1574,7 +1574,7 @@ FLearningAgentsObservationObjectElement ULearningAgentsObservations::MakeMapObse
 		SubElements.Add(MakePairObservation(Object, Item.Key, Item.Value));
 	}
 
-	return MakeSetObservationFromArrayView(Object, SubElements);
+	return MakeSetObservationFromArrayView(Object, SubElements, Tag);
 }
 
 FLearningAgentsObservationObjectElement ULearningAgentsObservations::MakeMapObservationFromArrays(ULearningAgentsObservationObject* Object, const TArray<FLearningAgentsObservationObjectElement>& Keys, const TArray<FLearningAgentsObservationObjectElement>& Values, const FName Tag)
@@ -1611,7 +1611,7 @@ FLearningAgentsObservationObjectElement ULearningAgentsObservations::MakeMapObse
 		SubElements.Add(MakePairObservation(Object, Keys[ElementIdx], Values[ElementIdx]));
 	}
 
-	return MakeSetObservationFromArrayView(Object, SubElements);
+	return MakeSetObservationFromArrayView(Object, SubElements, Tag);
 }
 
 FLearningAgentsObservationObjectElement ULearningAgentsObservations::MakeEnumObservation(
