LearningAgents: Reverted to package python foundation site-packages location.

diff --git a/Engine/Plugins/Experimental/LearningAgents/Source/LearningAgentsTraining/Private/LearningAgentsImitationTrainer.cpp b/Engine/Plugins/Experimental/LearningAgents/Source/LearningAgentsTraining/Private/LearningAgentsImitationTrainer.cpp
--- a/Engine/Plugins/Experimental/LearningAgents/Source/LearningAgentsTraining/Private/LearningAgentsImitationTrainer.cpp
+++ b/Engine/Plugins/Experimental/LearningAgents/Source/LearningAgentsTraining/Private/LearningAgentsImitationTrainer.cpp
@@ -166,7 +166,7 @@ void ULearningAgentsImitationTrainer::BeginTraining(
 		return;
 	}
 
-	const FString SitePackagesPath = UE::Learning::Trainer::GetSitePackagesPath(ImitationTrainerPathSettings.GetIntermediatePath());
+	const FString SitePackagesPath = UE::Learning::Trainer::GetSitePackagesPath(ImitationTrainerPathSettings.GetEditorEnginePath());
 
 	if (!FPaths::DirectoryExists(SitePackagesPath))
 	{
diff --git a/Engine/Plugins/Experimental/LearningAgents/Source/LearningAgentsTraining/Private/LearningAgentsTrainer.cpp b/Engine/Plugins/Experimental/LearningAgents/Source/LearningAgentsTraining/Private/LearningAgentsTrainer.cpp
--- a/Engine/Plugins/Experimental/LearningAgents/Source/LearningAgentsTraining/Private/LearningAgentsTrainer.cpp
+++ b/Engine/Plugins/Experimental/LearningAgents/Source/LearningAgentsTraining/Private/LearningAgentsTrainer.cpp
@@ -408,7 +408,7 @@ void ULearningAgentsTrainer::BeginTraining(
 		return;
 	}
 
-	const FString SitePackagesPath = UE::Learning::Trainer::GetSitePackagesPath(TrainerPathSettings.GetIntermediatePath());
+	const FString SitePackagesPath = UE::Learning::Trainer::GetSitePackagesPath(TrainerPathSettings.GetEditorEnginePath());
 
 	if (!FPaths::DirectoryExists(SitePackagesPath))
 	{
diff --git a/Engine/Plugins/Experimental/LearningAgents/Source/LearningTraining/Private/LearningImitationTrainer.cpp b/Engine/Plugins/Experimental/LearningAgents/Source/LearningTraining/Private/LearningImitationTrainer.cpp
--- a/Engine/Plugins/Experimental/LearningAgents/Source/LearningTraining/Private/LearningImitationTrainer.cpp
+++ b/Engine/Plugins/Experimental/LearningAgents/Source/LearningTraining/Private/LearningImitationTrainer.cpp
@@ -47,7 +47,7 @@ int32 ULearningSocketImitationTrainerServerCommandlet::Main(const FString& Comma
 	const FString* LogSettingsParam = Params.Find(TEXT("LogSettings"));
 
 	const FString PythonExecutablePath = PythonExecutablePathParam ? *PythonExecutablePathParam : UE::Learning::Trainer::GetPythonExecutablePath(FPaths::EngineDir());
-	const FString SitePackagesPath = SitePackagesPathParam ? *SitePackagesPathParam : UE::Learning::Trainer::GetSitePackagesPath(FPaths::ProjectIntermediateDir());
+	const FString SitePackagesPath = SitePackagesPathParam ? *SitePackagesPathParam : UE::Learning::Trainer::GetSitePackagesPath(FPaths::EngineDir());
 	const FString PythonContentPath = PythonContentPathParam ? *PythonContentPathParam : UE::Learning::Trainer::GetPythonContentPath(FPaths::EngineDir());
 	const FString IntermediatePath = IntermediatePathParam ? *IntermediatePathParam : UE::Learning::Trainer::GetIntermediatePath(FPaths::ProjectIntermediateDir());
 
diff --git a/Engine/Plugins/Experimental/LearningAgents/Source/LearningTraining/Private/LearningPPOTrainer.cpp b/Engine/Plugins/Experimental/LearningAgents/Source/LearningTraining/Private/LearningPPOTrainer.cpp
--- a/Engine/Plugins/Experimental/LearningAgents/Source/LearningTraining/Private/LearningPPOTrainer.cpp
+++ b/Engine/Plugins/Experimental/LearningAgents/Source/LearningTraining/Private/LearningPPOTrainer.cpp
@@ -52,7 +52,7 @@ int32 ULearningSocketPPOTrainerServerCommandlet::Main(const FString& Commandline
 	const FString* LogSettingsParam = Params.Find(TEXT("LogSettings"));
 
 	const FString PythonExecutiblePath = PythonExecutiblePathParam ? *PythonExecutiblePathParam : UE::Learning::Trainer::GetPythonExecutablePath(FPaths::EngineDir());
-	const FString SitePackagesPath = SitePackagesPathParam ? *SitePackagesPathParam : UE::Learning::Trainer::GetSitePackagesPath(FPaths::ProjectIntermediateDir());
+	const FString SitePackagesPath = SitePackagesPathParam ? *SitePackagesPathParam : UE::Learning::Trainer::GetSitePackagesPath(FPaths::EngineDir());
 	const FString PythonContentPath = PythonContentPathParam ? *PythonContentPathParam : UE::Learning::Trainer::GetPythonContentPath(FPaths::EngineDir());
 	const FString IntermediatePath = IntermediatePathParam ? *IntermediatePathParam : UE::Learning::Trainer::GetIntermediatePath(FPaths::ProjectIntermediateDir());
 
diff --git a/Engine/Plugins/Experimental/LearningAgents/Source/LearningTraining/Private/LearningTrainer.cpp b/Engine/Plugins/Experimental/LearningAgents/Source/LearningTraining/Private/LearningTrainer.cpp
--- a/Engine/Plugins/Experimental/LearningAgents/Source/LearningTraining/Private/LearningTrainer.cpp
+++ b/Engine/Plugins/Experimental/LearningAgents/Source/LearningTraining/Private/LearningTrainer.cpp
@@ -296,14 +296,14 @@ namespace UE::Learning::Trainer
 	{
 		UE_LEARNING_CHECKF(PLATFORM_WINDOWS || PLATFORM_MAC || PLATFORM_LINUX, TEXT("Python only supported on Windows, Mac, and Linux."));
 
-		return EngineDir / TEXT("Binaries/ThirdParty/Python3") / FPlatformMisc::GetUBTPlatform() / (PLATFORM_WINDOWS ? TEXT("python.exe") : TEXT("bin/python3"));
+		return EngineDir / TEXT("Binaries/ThirdParty/Python3") / FPlatformMisc::GetUBTPlatform() / (PLATFORM_WINDOWS ? TEXT("python.exe") : TEXT("bin/python"));
 	}
 
-	FString GetSitePackagesPath(const FString& IntermediateDir)
+	FString GetSitePackagesPath(const FString& EngineDir)
 	{
 		UE_LEARNING_CHECKF(PLATFORM_WINDOWS || PLATFORM_MAC || PLATFORM_LINUX, TEXT("Python only supported on Windows, Mac, and Linux."));
 
-		return PLATFORM_WINDOWS ? IntermediateDir / TEXT("PipInstall/Lib/site-packages") : IntermediateDir / TEXT("PipInstall/lib/python3.9/site-packages");
+		return EngineDir / TEXT("Plugins/Experimental/PythonFoundationPackages/Content/Python/Lib") / FPlatformMisc::GetUBTPlatform() / TEXT("site-packages");
 	}
 
 	FString GetPythonContentPath(const FString& EngineDir)
diff --git a/Engine/Plugins/Experimental/LearningAgents/Source/LearningTraining/Public/LearningTrainer.h b/Engine/Plugins/Experimental/LearningAgents/Source/LearningTraining/Public/LearningTrainer.h
--- a/Engine/Plugins/Experimental/LearningAgents/Source/LearningTraining/Public/LearningTrainer.h
+++ b/Engine/Plugins/Experimental/LearningAgents/Source/LearningTraining/Public/LearningTrainer.h
@@ -123,9 +123,9 @@ namespace UE::Learning
 		LEARNINGTRAINING_API FString GetPythonExecutablePath(const FString& EngineDir);
 
 		/**
-		* Gets the Project site-packages path from the intermediate directory.
+		* Gets the PythonFoundationPackages site-packages path from the engine directory.
 		*/
-		LEARNINGTRAINING_API FString GetSitePackagesPath(const FString& IntermediateDir);
+		LEARNINGTRAINING_API FString GetSitePackagesPath(const FString& EngineDir);
 
 		/**
 		* Gets the LearningAgents Content path from the engine directory.
