LearningAgents: Fixed bug in prior probabilities.

diff --git a/Engine/Plugins/Experimental/LearningAgents/Source/Learning/Private/LearningAction.cpp b/Engine/Plugins/Experimental/LearningAgents/Source/Learning/Private/LearningAction.cpp
--- a/Engine/Plugins/Experimental/LearningAgents/Source/Learning/Private/LearningAction.cpp
+++ b/Engine/Plugins/Experimental/LearningAgents/Source/Learning/Private/LearningAction.cpp
@@ -870,7 +870,7 @@ namespace UE::Learning::Action
 
 				ReturnElement = Builder.MakeDenormalize(
 					Parameters.Num,
-					LogPriorProbabilities,
+					Builder.MakeWeightsCopy(LogPriorProbabilities),
 					Builder.MakeWeightsConstant(Parameters.Num, 1.0f));
 
 				break;
@@ -888,7 +888,7 @@ namespace UE::Learning::Action
 				}
 
 				ReturnElement = Builder.MakeDenormalize(Parameters.Num,
-					LogPriorProbabilities,
+					Builder.MakeWeightsCopy(LogPriorProbabilities),
 					Builder.MakeWeightsConstant(Parameters.Num, 1.0f));
 				break;
 			}
@@ -929,7 +929,7 @@ namespace UE::Learning::Action
 
 				BuilderLayers.Emplace(Builder.MakeDenormalize(
 					LogPriorProbabilities.Num(),
-					LogPriorProbabilities,
+					Builder.MakeWeightsCopy(LogPriorProbabilities),
 					Builder.MakeWeightsConstant(LogPriorProbabilities.Num(), 1.0f)));
 
 				ReturnElement = Builder.MakeConcat(BuilderLayers);
@@ -956,7 +956,7 @@ namespace UE::Learning::Action
 
 				BuilderLayers.Emplace(Builder.MakeDenormalize(
 					LogPriorProbabilities.Num(),
-					LogPriorProbabilities,
+					Builder.MakeWeightsCopy(LogPriorProbabilities),
 					Builder.MakeWeightsConstant(LogPriorProbabilities.Num(), 1.0f)));
 
 				ReturnElement = Builder.MakeConcat(BuilderLayers);
diff --git a/Engine/Plugins/Experimental/NNERuntimeBasicCpu/Source/NNERuntimeBasicCpu/Private/NNERuntimeBasicCpuModel.cpp b/Engine/Plugins/Experimental/NNERuntimeBasicCpu/Source/NNERuntimeBasicCpu/Private/NNERuntimeBasicCpuModel.cpp
--- a/Engine/Plugins/Experimental/NNERuntimeBasicCpu/Source/NNERuntimeBasicCpu/Private/NNERuntimeBasicCpuModel.cpp
+++ b/Engine/Plugins/Experimental/NNERuntimeBasicCpu/Source/NNERuntimeBasicCpu/Private/NNERuntimeBasicCpuModel.cpp
@@ -4128,6 +4128,13 @@ namespace UE::NNE::RuntimeBasic
 		WriteFileDataAndReset(MakeArrayView(FileData), OutInputSize, OutOutputSize, Element);
 	}
 
+	TArrayView<float> FModelBuilder::MakeWeightsCopy(const TConstArrayView<float> Weights)
+	{
+		TArray<float>& Values = WeightsPool.AddDefaulted_GetRef();
+		Values = Weights;
+		return Values;
+	}
+
 	TArrayView<float> FModelBuilder::MakeWeightsZero(const uint32 Size)
 	{
 		TArray<float>& Values = WeightsPool.AddDefaulted_GetRef();
diff --git a/Engine/Plugins/Experimental/NNERuntimeBasicCpu/Source/NNERuntimeBasicCpu/Public/NNERuntimeBasicCpuBuilder.h b/Engine/Plugins/Experimental/NNERuntimeBasicCpu/Source/NNERuntimeBasicCpu/Public/NNERuntimeBasicCpuBuilder.h
--- a/Engine/Plugins/Experimental/NNERuntimeBasicCpu/Source/NNERuntimeBasicCpu/Public/NNERuntimeBasicCpuBuilder.h
+++ b/Engine/Plugins/Experimental/NNERuntimeBasicCpu/Source/NNERuntimeBasicCpu/Public/NNERuntimeBasicCpuBuilder.h
@@ -268,6 +268,9 @@ namespace UE::NNE::RuntimeBasic
 
 	public:
 
+		/** Creates a array of weights from a copy of the given array view */
+		TArrayView<float> MakeWeightsCopy(const TConstArrayView<float> Weights);
+
 		/** Creates a array of weights set to zero of the given size */
 		TArrayView<float> MakeWeightsZero(const uint32 Size);
 
