LearningAgents: Missing Tag Passthrough

diff --git a/Engine/Plugins/Experimental/LearningAgents/Source/LearningAgents/Private/LearningAgentsActions.cpp b/Engine/Plugins/Experimental/LearningAgents/Source/LearningAgents/Private/LearningAgentsActions.cpp
--- a/Engine/Plugins/Experimental/LearningAgents/Source/LearningAgents/Private/LearningAgentsActions.cpp
+++ b/Engine/Plugins/Experimental/LearningAgents/Source/LearningAgents/Private/LearningAgentsActions.cpp
@@ -1265,7 +1265,7 @@ FLearningAgentsActionObjectElement ULearningAgentsActions::MakeStructAction(ULea
 		SubElementNames.Add(Element.Key);
 	}
 
-	return MakeStructActionFromArrayViews(Object, SubElementNames, SubElements);
+	return MakeStructActionFromArrayViews(Object, SubElementNames, SubElements, Tag);
 }
 
 FLearningAgentsActionObjectElement ULearningAgentsActions::MakeStructActionFromArrays(ULearningAgentsActionObject* Object, const TArray<FName>& ElementNames, const TArray<FLearningAgentsActionObjectElement>& Elements, const FName Tag)
