LearningAgents: Exposed functions for making ModelBuilderElements for Observation and Action Schemas so that they can be used in the construction of larger models.

diff --git a/Engine/Plugins/Experimental/LearningAgents/Source/Learning/Private/LearningAction.cpp b/Engine/Plugins/Experimental/LearningAgents/Source/Learning/Private/LearningAction.cpp
--- a/Engine/Plugins/Experimental/LearningAgents/Source/Learning/Private/LearningAction.cpp
+++ b/Engine/Plugins/Experimental/LearningAgents/Source/Learning/Private/LearningAction.cpp
@@ -827,187 +827,6 @@ namespace UE::Learning::Action
 			}
 		}
 
-		NNE::RuntimeBasic::FModelBuilderElement MakeDecoderNetworkFromSchema(
-			NNE::RuntimeBasic::FModelBuilder& Builder,
-			const FSchema& Schema,
-			const FSchemaElement SchemaElement)
-		{
-			const EType SchemaElementType = Schema.GetType(SchemaElement);
-
-			NNE::RuntimeBasic::FModelBuilderElement ReturnElement;
-
-			switch (SchemaElementType)
-			{
-
-			case EType::Null:
-			{
-				ReturnElement = Builder.MakeCopy(0);
-				break;
-			}
-
-			case EType::Continuous:
-			{
-				const int32 ValueNum = Schema.GetContinuous(SchemaElement).Num * 2;
-
-				ReturnElement = Builder.MakeDenormalize(
-					ValueNum,
-					Builder.MakeWeightsZero(ValueNum),
-					Builder.MakeWeightsConstant(ValueNum, 1.0f));
-				break;
-			}
-
-			case EType::DiscreteExclusive:
-			{
-				const FSchemaDiscreteExclusiveParameters Parameters = Schema.GetDiscreteExclusive(SchemaElement);
-
-				TArray<float, TInlineAllocator<16>> LogPriorProbabilities;
-				LogPriorProbabilities.Append(Parameters.PriorProbabilities);
-				for (int32 Idx = 0; Idx < Parameters.Num; Idx++)
-				{
-					// Clamp zero probabilities to the smallest (positive) float. This is approximately equal to a probability of 1:1e38
-					LogPriorProbabilities[Idx] = FMath::Loge(FMath::Max(LogPriorProbabilities[Idx], FLT_MIN));
-				}
-
-				ReturnElement = Builder.MakeDenormalize(
-					Parameters.Num,
-					Builder.MakeWeightsCopy(LogPriorProbabilities),
-					Builder.MakeWeightsConstant(Parameters.Num, 1.0f));
-
-				break;
-			}
-
-			case EType::DiscreteInclusive:
-			{
-				const FSchemaDiscreteInclusiveParameters Parameters = Schema.GetDiscreteInclusive(SchemaElement);
-
-				TArray<float, TInlineAllocator<16>> LogPriorProbabilities;
-				LogPriorProbabilities.Append(Parameters.PriorProbabilities);
-				for (int32 Idx = 0; Idx < Parameters.Num; Idx++)
-				{
-					LogPriorProbabilities[Idx] = Private::Logit(LogPriorProbabilities[Idx]);
-				}
-
-				ReturnElement = Builder.MakeDenormalize(Parameters.Num,
-					Builder.MakeWeightsCopy(LogPriorProbabilities),
-					Builder.MakeWeightsConstant(Parameters.Num, 1.0f));
-				break;
-			}
-					
-			case EType::And:
-			{
-				const FSchemaAndParameters Parameters = Schema.GetAnd(SchemaElement);
-
-				TArray<NNE::RuntimeBasic::FModelBuilderElement, TInlineAllocator<8>> BuilderLayers;
-				BuilderLayers.Reserve(Parameters.Elements.Num());
-				for (const FSchemaElement SubElement : Parameters.Elements)
-				{
-					BuilderLayers.Emplace(MakeDecoderNetworkFromSchema(Builder, Schema, SubElement));
-				}
-
-				ReturnElement = Builder.MakeConcat(BuilderLayers);
-				break;
-			}
-
-			case EType::OrExclusive:
-			{
-				const FSchemaOrExclusiveParameters Parameters = Schema.GetOrExclusive(SchemaElement);
-
-				TArray<NNE::RuntimeBasic::FModelBuilderElement, TInlineAllocator<8>> BuilderLayers;
-				BuilderLayers.Reserve(Parameters.Elements.Num() + 1);
-				for (const FSchemaElement SubElement : Parameters.Elements)
-				{
-					BuilderLayers.Emplace(MakeDecoderNetworkFromSchema(Builder, Schema, SubElement));
-				}
-
-				TArray<float, TInlineAllocator<16>> LogPriorProbabilities;
-				LogPriorProbabilities.Append(Parameters.PriorProbabilities);
-				for (int32 Idx = 0; Idx < Parameters.PriorProbabilities.Num(); Idx++)
-				{
-					// Clamp zero probabilities to the smallest (positive) float. This is approximately equal to a probability of 1:1e38
-					LogPriorProbabilities[Idx] = FMath::Loge(FMath::Max(LogPriorProbabilities[Idx], FLT_MIN));
-				}
-
-				BuilderLayers.Emplace(Builder.MakeDenormalize(
-					LogPriorProbabilities.Num(),
-					Builder.MakeWeightsCopy(LogPriorProbabilities),
-					Builder.MakeWeightsConstant(LogPriorProbabilities.Num(), 1.0f)));
-
-				ReturnElement = Builder.MakeConcat(BuilderLayers);
-				break;
-			}
-
-			case EType::OrInclusive:
-			{
-				const FSchemaOrInclusiveParameters Parameters = Schema.GetOrInclusive(SchemaElement);
-
-				TArray<NNE::RuntimeBasic::FModelBuilderElement, TInlineAllocator<8>> BuilderLayers;
-				BuilderLayers.Reserve(Parameters.Elements.Num() + 1);
-				for (const FSchemaElement SubElement : Parameters.Elements)
-				{
-					BuilderLayers.Emplace(MakeDecoderNetworkFromSchema(Builder, Schema, SubElement));
-				}
-
-				TArray<float, TInlineAllocator<16>> LogPriorProbabilities;
-				LogPriorProbabilities.Append(Parameters.PriorProbabilities);
-				for (int32 Idx = 0; Idx < Parameters.PriorProbabilities.Num(); Idx++)
-				{
-					LogPriorProbabilities[Idx] = Private::Logit(LogPriorProbabilities[Idx]);
-				}
-
-				BuilderLayers.Emplace(Builder.MakeDenormalize(
-					LogPriorProbabilities.Num(),
-					Builder.MakeWeightsCopy(LogPriorProbabilities),
-					Builder.MakeWeightsConstant(LogPriorProbabilities.Num(), 1.0f)));
-
-				ReturnElement = Builder.MakeConcat(BuilderLayers);
-				break;
-			}
-
-			case EType::Array:
-			{
-				const FSchemaArrayParameters Parameters = Schema.GetArray(SchemaElement);
-
-				ReturnElement = Builder.MakeArray(Parameters.Num, MakeDecoderNetworkFromSchema(Builder, Schema, Parameters.Element));
-				break;
-			}
-
-			case EType::Encoding:
-			{
-				const FSchemaEncodingParameters Parameters = Schema.GetEncoding(SchemaElement);
-
-				const int32 SubElementEncodedSize = Schema.GetEncodedVectorSize(Parameters.Element);
-
-				ReturnElement = Builder.MakeSequence({
-					Builder.MakeActivation(Parameters.EncodingSize, GetNNEActivationFunction(Parameters.ActivationFunction)),
-					Builder.MakeMLPWithRandomKaimingWeights(
-						Parameters.EncodingSize, 
-						SubElementEncodedSize, 
-						Parameters.EncodingSize, 
-						Parameters.LayerNum + 1,  // Add 1 to account for input layer 
-						GetNNEActivationFunction(Parameters.ActivationFunction),
-						false),
-					MakeDecoderNetworkFromSchema(Builder, Schema, Parameters.Element),
-				});
-				break;
-			}
-
-			default:
-			{
-				UE_LEARNING_NOT_IMPLEMENTED();
-			}
-			}
-
-			UE_LEARNING_CHECKF(ReturnElement.GetInputSize() == Schema.GetEncodedVectorSize(SchemaElement),
-				TEXT("Decoder Network Input unexpected size. Got %i, expected %i according to Schema."),
-				ReturnElement.GetInputSize(), Schema.GetEncodedVectorSize(SchemaElement));
-
-			UE_LEARNING_CHECKF(ReturnElement.GetOutputSize() == Schema.GetActionDistributionVectorSize(SchemaElement),
-				TEXT("Decoder Network Output unexpected size. Got %i, expected %i according to Schema."),
-				ReturnElement.GetOutputSize(), Schema.GetActionDistributionVectorSize(SchemaElement));
-
-			return ReturnElement;
-		}
-
 		static inline int32 HashFNameStable(const FName Name)
 		{
 			const FString NameString = Name.ToString().ToLower();
@@ -1213,7 +1032,196 @@ namespace UE::Learning::Action
 		}
 	}
 
-	LEARNING_API void GenerateDecoderNetworkFileDataFromSchema(
+	void MakeDecoderNetworkModelBuilderElementFromSchema(
+		NNE::RuntimeBasic::FModelBuilderElement& OutElement,
+		NNE::RuntimeBasic::FModelBuilder& Builder,
+		const FSchema& Schema,
+		const FSchemaElement SchemaElement)
+	{
+		const EType SchemaElementType = Schema.GetType(SchemaElement);
+
+		switch (SchemaElementType)
+		{
+
+		case EType::Null:
+		{
+			OutElement = Builder.MakeCopy(0);
+			break;
+		}
+
+		case EType::Continuous:
+		{
+			const int32 ValueNum = Schema.GetContinuous(SchemaElement).Num * 2;
+
+			OutElement = Builder.MakeDenormalize(
+				ValueNum,
+				Builder.MakeWeightsZero(ValueNum),
+				Builder.MakeWeightsConstant(ValueNum, 1.0f));
+			break;
+		}
+
+		case EType::DiscreteExclusive:
+		{
+			const FSchemaDiscreteExclusiveParameters Parameters = Schema.GetDiscreteExclusive(SchemaElement);
+
+			TArray<float, TInlineAllocator<16>> LogPriorProbabilities;
+			LogPriorProbabilities.Append(Parameters.PriorProbabilities);
+			for (int32 Idx = 0; Idx < Parameters.Num; Idx++)
+			{
+				// Clamp zero probabilities to the smallest (positive) float. This is approximately equal to a probability of 1:1e38
+				LogPriorProbabilities[Idx] = FMath::Loge(FMath::Max(LogPriorProbabilities[Idx], FLT_MIN));
+			}
+
+			OutElement = Builder.MakeDenormalize(
+				Parameters.Num,
+				Builder.MakeWeightsCopy(LogPriorProbabilities),
+				Builder.MakeWeightsConstant(Parameters.Num, 1.0f));
+
+			break;
+		}
+
+		case EType::DiscreteInclusive:
+		{
+			const FSchemaDiscreteInclusiveParameters Parameters = Schema.GetDiscreteInclusive(SchemaElement);
+
+			TArray<float, TInlineAllocator<16>> LogPriorProbabilities;
+			LogPriorProbabilities.Append(Parameters.PriorProbabilities);
+			for (int32 Idx = 0; Idx < Parameters.Num; Idx++)
+			{
+				LogPriorProbabilities[Idx] = Private::Logit(LogPriorProbabilities[Idx]);
+			}
+
+			OutElement = Builder.MakeDenormalize(Parameters.Num,
+				Builder.MakeWeightsCopy(LogPriorProbabilities),
+				Builder.MakeWeightsConstant(Parameters.Num, 1.0f));
+			break;
+		}
+
+		case EType::And:
+		{
+			const FSchemaAndParameters Parameters = Schema.GetAnd(SchemaElement);
+
+			TArray<NNE::RuntimeBasic::FModelBuilderElement, TInlineAllocator<8>> BuilderLayers;
+			BuilderLayers.Reserve(Parameters.Elements.Num());
+			for (const FSchemaElement SubElement : Parameters.Elements)
+			{
+				NNE::RuntimeBasic::FModelBuilderElement BuilderSubElement;
+				MakeDecoderNetworkModelBuilderElementFromSchema(BuilderSubElement, Builder, Schema, SubElement);
+				BuilderLayers.Emplace(BuilderSubElement);
+			}
+
+			OutElement = Builder.MakeConcat(BuilderLayers);
+			break;
+		}
+
+		case EType::OrExclusive:
+		{
+			const FSchemaOrExclusiveParameters Parameters = Schema.GetOrExclusive(SchemaElement);
+
+			TArray<NNE::RuntimeBasic::FModelBuilderElement, TInlineAllocator<8>> BuilderLayers;
+			BuilderLayers.Reserve(Parameters.Elements.Num() + 1);
+			for (const FSchemaElement SubElement : Parameters.Elements)
+			{
+				NNE::RuntimeBasic::FModelBuilderElement BuilderSubElement;
+				MakeDecoderNetworkModelBuilderElementFromSchema(BuilderSubElement, Builder, Schema, SubElement);
+				BuilderLayers.Emplace(BuilderSubElement);
+			}
+
+			TArray<float, TInlineAllocator<16>> LogPriorProbabilities;
+			LogPriorProbabilities.Append(Parameters.PriorProbabilities);
+			for (int32 Idx = 0; Idx < Parameters.PriorProbabilities.Num(); Idx++)
+			{
+				// Clamp zero probabilities to the smallest (positive) float. This is approximately equal to a probability of 1:1e38
+				LogPriorProbabilities[Idx] = FMath::Loge(FMath::Max(LogPriorProbabilities[Idx], FLT_MIN));
+			}
+
+			BuilderLayers.Emplace(Builder.MakeDenormalize(
+				LogPriorProbabilities.Num(),
+				Builder.MakeWeightsCopy(LogPriorProbabilities),
+				Builder.MakeWeightsConstant(LogPriorProbabilities.Num(), 1.0f)));
+
+			OutElement = Builder.MakeConcat(BuilderLayers);
+			break;
+		}
+
+		case EType::OrInclusive:
+		{
+			const FSchemaOrInclusiveParameters Parameters = Schema.GetOrInclusive(SchemaElement);
+
+			TArray<NNE::RuntimeBasic::FModelBuilderElement, TInlineAllocator<8>> BuilderLayers;
+			BuilderLayers.Reserve(Parameters.Elements.Num() + 1);
+			for (const FSchemaElement SubElement : Parameters.Elements)
+			{
+				NNE::RuntimeBasic::FModelBuilderElement BuilderSubElement;
+				MakeDecoderNetworkModelBuilderElementFromSchema(BuilderSubElement, Builder, Schema, SubElement);
+				BuilderLayers.Emplace(BuilderSubElement);
+			}
+
+			TArray<float, TInlineAllocator<16>> LogPriorProbabilities;
+			LogPriorProbabilities.Append(Parameters.PriorProbabilities);
+			for (int32 Idx = 0; Idx < Parameters.PriorProbabilities.Num(); Idx++)
+			{
+				LogPriorProbabilities[Idx] = Private::Logit(LogPriorProbabilities[Idx]);
+			}
+
+			BuilderLayers.Emplace(Builder.MakeDenormalize(
+				LogPriorProbabilities.Num(),
+				Builder.MakeWeightsCopy(LogPriorProbabilities),
+				Builder.MakeWeightsConstant(LogPriorProbabilities.Num(), 1.0f)));
+
+			OutElement = Builder.MakeConcat(BuilderLayers);
+			break;
+		}
+
+		case EType::Array:
+		{
+			const FSchemaArrayParameters Parameters = Schema.GetArray(SchemaElement);
+
+			NNE::RuntimeBasic::FModelBuilderElement BuilderSubElement;
+			MakeDecoderNetworkModelBuilderElementFromSchema(BuilderSubElement, Builder, Schema, Parameters.Element);
+			OutElement = Builder.MakeArray(Parameters.Num, BuilderSubElement);
+			break;
+		}
+
+		case EType::Encoding:
+		{
+			const FSchemaEncodingParameters Parameters = Schema.GetEncoding(SchemaElement);
+
+			const int32 SubElementEncodedSize = Schema.GetEncodedVectorSize(Parameters.Element);
+
+			NNE::RuntimeBasic::FModelBuilderElement BuilderSubElement;
+			MakeDecoderNetworkModelBuilderElementFromSchema(BuilderSubElement, Builder, Schema, Parameters.Element);
+
+			OutElement = Builder.MakeSequence({
+				Builder.MakeActivation(Parameters.EncodingSize, Private::GetNNEActivationFunction(Parameters.ActivationFunction)),
+				Builder.MakeMLPWithRandomKaimingWeights(
+					Parameters.EncodingSize,
+					SubElementEncodedSize,
+					Parameters.EncodingSize,
+					Parameters.LayerNum + 1,  // Add 1 to account for input layer 
+					Private::GetNNEActivationFunction(Parameters.ActivationFunction),
+					false),
+				BuilderSubElement,
+				});
+			break;
+		}
+
+		default:
+		{
+			UE_LEARNING_NOT_IMPLEMENTED();
+		}
+		}
+
+		UE_LEARNING_CHECKF(OutElement.GetInputSize() == Schema.GetEncodedVectorSize(SchemaElement),
+			TEXT("Decoder Network Input unexpected size. Got %i, expected %i according to Schema."),
+			OutElement.GetInputSize(), Schema.GetEncodedVectorSize(SchemaElement));
+
+		UE_LEARNING_CHECKF(OutElement.GetOutputSize() == Schema.GetActionDistributionVectorSize(SchemaElement),
+			TEXT("Decoder Network Output unexpected size. Got %i, expected %i according to Schema."),
+			OutElement.GetOutputSize(), Schema.GetActionDistributionVectorSize(SchemaElement));
+	}
+
+	void GenerateDecoderNetworkFileDataFromSchema(
 		TArray<uint8>& OutFileData,
 		uint32& OutInputSize,
 		uint32& OutOutputSize,
@@ -1224,7 +1232,9 @@ namespace UE::Learning::Action
 		UE_LEARNING_CHECK(Schema.IsValid(SchemaElement));
 
 		NNE::RuntimeBasic::FModelBuilder Builder(Seed);
-		Builder.WriteFileDataAndReset(OutFileData, OutInputSize, OutOutputSize, Private::MakeDecoderNetworkFromSchema(Builder, Schema, SchemaElement));
+		NNE::RuntimeBasic::FModelBuilderElement Element;
+		MakeDecoderNetworkModelBuilderElementFromSchema(Element, Builder, Schema, SchemaElement);
+		Builder.WriteFileDataAndReset(OutFileData, OutInputSize, OutOutputSize, Element);
 	}
 
 	void SampleVectorFromDistributionVector(
diff --git a/Engine/Plugins/Experimental/LearningAgents/Source/Learning/Private/LearningObservation.cpp b/Engine/Plugins/Experimental/LearningAgents/Source/Learning/Private/LearningObservation.cpp
--- a/Engine/Plugins/Experimental/LearningAgents/Source/Learning/Private/LearningObservation.cpp
+++ b/Engine/Plugins/Experimental/LearningAgents/Source/Learning/Private/LearningObservation.cpp
@@ -680,162 +680,6 @@ namespace UE::Learning::Observation
 			}
 		}
 
-		NNE::RuntimeBasic::FModelBuilderElement MakeEncoderNetworkFromSchema(
-			NNE::RuntimeBasic::FModelBuilder& Builder,
-			const FSchema& Schema,
-			const FSchemaElement SchemaElement)
-		{
-			const EType SchemaElementType = Schema.GetType(SchemaElement);
-
-			NNE::RuntimeBasic::FModelBuilderElement ReturnElement;
-
-			switch (SchemaElementType)
-			{
-			case EType::Null:
-			{
-				ReturnElement = Builder.MakeCopy(0);
-				break;
-			}
-
-			case EType::Continuous:
-			{
-				const int32 ValueNum = Schema.GetContinuous(SchemaElement).Num;
-
-				ReturnElement = Builder.MakeDenormalize(
-					ValueNum,
-					Builder.MakeWeightsZero(ValueNum),
-					Builder.MakeWeightsConstant(ValueNum, 1.0f));
-				break;
-			}
-
-			case EType::And:
-			{
-				const FSchemaAndParameters Parameters = Schema.GetAnd(SchemaElement);
-
-				TArray<NNE::RuntimeBasic::FModelBuilderElement, TInlineAllocator<8>> BuilderLayers;
-				BuilderLayers.Reserve(Parameters.Elements.Num());
-				for (const FSchemaElement SubElement : Parameters.Elements)
-				{
-					BuilderLayers.Emplace(MakeEncoderNetworkFromSchema(Builder, Schema, SubElement));
-				}
-
-				ReturnElement = Builder.MakeConcat(BuilderLayers);
-				break;
-			}
-
-			case EType::OrExclusive:
-			{
-				const FSchemaOrExclusiveParameters Parameters = Schema.GetOrExclusive(SchemaElement);
-
-				TArray<NNE::RuntimeBasic::FModelBuilderElement, TInlineAllocator<8>> BuilderSubLayers;
-				TArray<NNE::RuntimeBasic::FModelBuilderElement, TInlineAllocator<8>> BuilderEncoders;
-				BuilderSubLayers.Reserve(Parameters.Elements.Num());
-				BuilderEncoders.Reserve(Parameters.Elements.Num());
-				for (const FSchemaElement SubElement : Parameters.Elements)
-				{
-					const int32 SubElementEncodedSize = Schema.GetEncodedVectorSize(SubElement);
-					BuilderSubLayers.Emplace(MakeEncoderNetworkFromSchema(Builder, Schema, SubElement));
-					BuilderEncoders.Emplace(Builder.MakeLinearWithRandomKaimingWeights(SubElementEncodedSize, Parameters.EncodingSize));
-				}
-
-				ReturnElement = Builder.MakeAggregateOrExclusive(Parameters.EncodingSize, BuilderSubLayers, BuilderEncoders);
-				break;
-			}
-
-			case EType::OrInclusive:
-			{
-				const FSchemaOrInclusiveParameters Parameters = Schema.GetOrInclusive(SchemaElement);
-
-				TArray<NNE::RuntimeBasic::FModelBuilderElement, TInlineAllocator<8>> BuilderSubLayers;
-				TArray<NNE::RuntimeBasic::FModelBuilderElement, TInlineAllocator<8>> BuilderQueryLayers;
-				TArray<NNE::RuntimeBasic::FModelBuilderElement, TInlineAllocator<8>> BuilderKeyLayers;
-				TArray<NNE::RuntimeBasic::FModelBuilderElement, TInlineAllocator<8>> BuilderValueLayers;
-				BuilderSubLayers.Reserve(Parameters.Elements.Num());
-				BuilderQueryLayers.Reserve(Parameters.Elements.Num());
-				BuilderValueLayers.Reserve(Parameters.Elements.Num());
-				for (const FSchemaElement SubElement : Parameters.Elements)
-				{
-					const int32 SubElementEncodedSize = Schema.GetEncodedVectorSize(SubElement);
-					BuilderSubLayers.Emplace(MakeEncoderNetworkFromSchema(Builder, Schema, SubElement));
-					BuilderQueryLayers.Emplace(Builder.MakeLinearWithRandomKaimingWeights(SubElementEncodedSize, Parameters.AttentionHeadNum * Parameters.AttentionEncodingSize));
-					BuilderKeyLayers.Emplace(Builder.MakeLinearWithRandomKaimingWeights(SubElementEncodedSize, Parameters.AttentionHeadNum * Parameters.AttentionEncodingSize));
-					BuilderValueLayers.Emplace(Builder.MakeLinearWithRandomKaimingWeights(SubElementEncodedSize, Parameters.AttentionHeadNum * Parameters.ValueEncodingSize));
-				}
-
-				ReturnElement = Builder.MakeAggregateOrInclusive(
-					Parameters.ValueEncodingSize,
-					Parameters.AttentionEncodingSize,
-					Parameters.AttentionHeadNum,
-					BuilderSubLayers,
-					BuilderQueryLayers,
-					BuilderKeyLayers,
-					BuilderValueLayers);
-
-				break;
-			}
-
-			case EType::Array:
-			{
-				const FSchemaArrayParameters Parameters = Schema.GetArray(SchemaElement);
-
-				ReturnElement = Builder.MakeArray(Parameters.Num, MakeEncoderNetworkFromSchema(Builder, Schema, Parameters.Element));
-				break;
-			}
-
-			case EType::Set:
-			{
-				const FSchemaSetParameters Parameters = Schema.GetSet(SchemaElement);
-
-				const int32 SubElementEncodedSize = Schema.GetEncodedVectorSize(Parameters.Element);
-
-				ReturnElement = Builder.MakeAggregateSet(
-					Parameters.MaxNum,
-					Parameters.ValueEncodingSize,
-					Parameters.AttentionEncodingSize,
-					Parameters.AttentionHeadNum,
-					MakeEncoderNetworkFromSchema(Builder, Schema, Parameters.Element),
-					Builder.MakeLinearWithRandomKaimingWeights(SubElementEncodedSize, Parameters.AttentionHeadNum * Parameters.AttentionEncodingSize),
-					Builder.MakeLinearWithRandomKaimingWeights(SubElementEncodedSize, Parameters.AttentionHeadNum * Parameters.AttentionEncodingSize),
-					Builder.MakeLinearWithRandomKaimingWeights(SubElementEncodedSize, Parameters.AttentionHeadNum * Parameters.ValueEncodingSize));
-				break;
-			}
-
-			case EType::Encoding:
-			{
-				const FSchemaEncodingParameters Parameters = Schema.GetEncoding(SchemaElement);
-
-				const int32 SubElementEncodedSize = Schema.GetEncodedVectorSize(Parameters.Element);
-
-				ReturnElement = Builder.MakeSequence({
-					MakeEncoderNetworkFromSchema(Builder, Schema, Parameters.Element),
-					Builder.MakeMLPWithRandomKaimingWeights(
-						SubElementEncodedSize,
-						Parameters.EncodingSize,
-						Parameters.EncodingSize,
-						Parameters.LayerNum + 1, // Add 1 to account for input layer
-						GetNNEActivationFunction(Parameters.ActivationFunction),
-						true)
-					});
-				break;
-			}
-
-			default:
-			{
-				UE_LEARNING_NOT_IMPLEMENTED();
-			}
-			}
-
-			UE_LEARNING_CHECKF(ReturnElement.GetInputSize() == Schema.GetObservationVectorSize(SchemaElement),
-				TEXT("Encoder Network Input unexpected size for %s. Got %i, expected %i according to Schema."),
-				*Schema.GetTag(SchemaElement).ToString(), ReturnElement.GetInputSize(), Schema.GetObservationVectorSize(SchemaElement));
-
-			UE_LEARNING_CHECKF(ReturnElement.GetOutputSize() == Schema.GetEncodedVectorSize(SchemaElement),
-				TEXT("Encoder Network Output unexpected size for %s. Got %i, expected %i according to Schema."),
-				*Schema.GetTag(SchemaElement).ToString(), ReturnElement.GetOutputSize(), Schema.GetEncodedVectorSize(SchemaElement));
-
-			return ReturnElement;
-		}
-
 		static inline int32 HashFNameStable(const FName Name)
 		{
 			const FString NameString = Name.ToString().ToLower();
@@ -1047,6 +891,173 @@ namespace UE::Learning::Observation
 		}
 	}
 
+	void MakeEncoderNetworkModelBuilderElementFromSchema(
+		NNE::RuntimeBasic::FModelBuilderElement& OutElement,
+		NNE::RuntimeBasic::FModelBuilder& Builder,
+		const FSchema& Schema,
+		const FSchemaElement SchemaElement)
+	{
+		const EType SchemaElementType = Schema.GetType(SchemaElement);
+
+		switch (SchemaElementType)
+		{
+		case EType::Null:
+		{
+			OutElement = Builder.MakeCopy(0);
+			break;
+		}
+
+		case EType::Continuous:
+		{
+			const int32 ValueNum = Schema.GetContinuous(SchemaElement).Num;
+
+			OutElement = Builder.MakeDenormalize(
+				ValueNum,
+				Builder.MakeWeightsZero(ValueNum),
+				Builder.MakeWeightsConstant(ValueNum, 1.0f));
+			break;
+		}
+
+		case EType::And:
+		{
+			const FSchemaAndParameters Parameters = Schema.GetAnd(SchemaElement);
+
+			TArray<NNE::RuntimeBasic::FModelBuilderElement, TInlineAllocator<8>> BuilderLayers;
+			BuilderLayers.Reserve(Parameters.Elements.Num());
+			for (const FSchemaElement SubElement : Parameters.Elements)
+			{
+				NNE::RuntimeBasic::FModelBuilderElement BuilderSubElement;
+				MakeEncoderNetworkModelBuilderElementFromSchema(BuilderSubElement, Builder, Schema, SubElement);
+				BuilderLayers.Emplace(BuilderSubElement);
+			}
+
+			OutElement = Builder.MakeConcat(BuilderLayers);
+			break;
+		}
+
+		case EType::OrExclusive:
+		{
+			const FSchemaOrExclusiveParameters Parameters = Schema.GetOrExclusive(SchemaElement);
+
+			TArray<NNE::RuntimeBasic::FModelBuilderElement, TInlineAllocator<8>> BuilderSubLayers;
+			TArray<NNE::RuntimeBasic::FModelBuilderElement, TInlineAllocator<8>> BuilderEncoders;
+			BuilderSubLayers.Reserve(Parameters.Elements.Num());
+			BuilderEncoders.Reserve(Parameters.Elements.Num());
+			for (const FSchemaElement SubElement : Parameters.Elements)
+			{
+				const int32 SubElementEncodedSize = Schema.GetEncodedVectorSize(SubElement);
+				NNE::RuntimeBasic::FModelBuilderElement BuilderSubElement;
+				MakeEncoderNetworkModelBuilderElementFromSchema(BuilderSubElement, Builder, Schema, SubElement);
+				BuilderSubLayers.Emplace(BuilderSubElement);
+				BuilderEncoders.Emplace(Builder.MakeLinearWithRandomKaimingWeights(SubElementEncodedSize, Parameters.EncodingSize));
+			}
+
+			OutElement = Builder.MakeAggregateOrExclusive(Parameters.EncodingSize, BuilderSubLayers, BuilderEncoders);
+			break;
+		}
+
+		case EType::OrInclusive:
+		{
+			const FSchemaOrInclusiveParameters Parameters = Schema.GetOrInclusive(SchemaElement);
+
+			TArray<NNE::RuntimeBasic::FModelBuilderElement, TInlineAllocator<8>> BuilderSubLayers;
+			TArray<NNE::RuntimeBasic::FModelBuilderElement, TInlineAllocator<8>> BuilderQueryLayers;
+			TArray<NNE::RuntimeBasic::FModelBuilderElement, TInlineAllocator<8>> BuilderKeyLayers;
+			TArray<NNE::RuntimeBasic::FModelBuilderElement, TInlineAllocator<8>> BuilderValueLayers;
+			BuilderSubLayers.Reserve(Parameters.Elements.Num());
+			BuilderQueryLayers.Reserve(Parameters.Elements.Num());
+			BuilderValueLayers.Reserve(Parameters.Elements.Num());
+			for (const FSchemaElement SubElement : Parameters.Elements)
+			{
+				const int32 SubElementEncodedSize = Schema.GetEncodedVectorSize(SubElement);
+				NNE::RuntimeBasic::FModelBuilderElement BuilderSubElement;
+				MakeEncoderNetworkModelBuilderElementFromSchema(BuilderSubElement, Builder, Schema, SubElement);
+				BuilderSubLayers.Emplace(BuilderSubElement);
+				BuilderQueryLayers.Emplace(Builder.MakeLinearWithRandomKaimingWeights(SubElementEncodedSize, Parameters.AttentionHeadNum * Parameters.AttentionEncodingSize));
+				BuilderKeyLayers.Emplace(Builder.MakeLinearWithRandomKaimingWeights(SubElementEncodedSize, Parameters.AttentionHeadNum * Parameters.AttentionEncodingSize));
+				BuilderValueLayers.Emplace(Builder.MakeLinearWithRandomKaimingWeights(SubElementEncodedSize, Parameters.AttentionHeadNum * Parameters.ValueEncodingSize));
+			}
+
+			OutElement = Builder.MakeAggregateOrInclusive(
+				Parameters.ValueEncodingSize,
+				Parameters.AttentionEncodingSize,
+				Parameters.AttentionHeadNum,
+				BuilderSubLayers,
+				BuilderQueryLayers,
+				BuilderKeyLayers,
+				BuilderValueLayers);
+
+			break;
+		}
+
+		case EType::Array:
+		{
+			const FSchemaArrayParameters Parameters = Schema.GetArray(SchemaElement);
+
+			NNE::RuntimeBasic::FModelBuilderElement BuilderSubElement;
+			MakeEncoderNetworkModelBuilderElementFromSchema(BuilderSubElement, Builder, Schema, Parameters.Element);
+			OutElement = Builder.MakeArray(Parameters.Num, BuilderSubElement);
+			break;
+		}
+
+		case EType::Set:
+		{
+			const FSchemaSetParameters Parameters = Schema.GetSet(SchemaElement);
+
+			const int32 SubElementEncodedSize = Schema.GetEncodedVectorSize(Parameters.Element);
+
+			NNE::RuntimeBasic::FModelBuilderElement BuilderSubElement;
+			MakeEncoderNetworkModelBuilderElementFromSchema(BuilderSubElement, Builder, Schema, Parameters.Element);
+
+			OutElement = Builder.MakeAggregateSet(
+				Parameters.MaxNum,
+				Parameters.ValueEncodingSize,
+				Parameters.AttentionEncodingSize,
+				Parameters.AttentionHeadNum,
+				BuilderSubElement,
+				Builder.MakeLinearWithRandomKaimingWeights(SubElementEncodedSize, Parameters.AttentionHeadNum * Parameters.AttentionEncodingSize),
+				Builder.MakeLinearWithRandomKaimingWeights(SubElementEncodedSize, Parameters.AttentionHeadNum * Parameters.AttentionEncodingSize),
+				Builder.MakeLinearWithRandomKaimingWeights(SubElementEncodedSize, Parameters.AttentionHeadNum * Parameters.ValueEncodingSize));
+			break;
+		}
+
+		case EType::Encoding:
+		{
+			const FSchemaEncodingParameters Parameters = Schema.GetEncoding(SchemaElement);
+
+			const int32 SubElementEncodedSize = Schema.GetEncodedVectorSize(Parameters.Element);
+
+			NNE::RuntimeBasic::FModelBuilderElement BuilderSubElement;
+			MakeEncoderNetworkModelBuilderElementFromSchema(BuilderSubElement, Builder, Schema, Parameters.Element);
+
+			OutElement = Builder.MakeSequence({
+				BuilderSubElement,
+				Builder.MakeMLPWithRandomKaimingWeights(
+					SubElementEncodedSize,
+					Parameters.EncodingSize,
+					Parameters.EncodingSize,
+					Parameters.LayerNum + 1, // Add 1 to account for input layer
+					Private::GetNNEActivationFunction(Parameters.ActivationFunction),
+					true)
+				});
+			break;
+		}
+
+		default:
+		{
+			UE_LEARNING_NOT_IMPLEMENTED();
+		}
+		}
+
+		UE_LEARNING_CHECKF(OutElement.GetInputSize() == Schema.GetObservationVectorSize(SchemaElement),
+			TEXT("Encoder Network Input unexpected size for %s. Got %i, expected %i according to Schema."),
+			*Schema.GetTag(SchemaElement).ToString(), OutElement.GetInputSize(), Schema.GetObservationVectorSize(SchemaElement));
+
+		UE_LEARNING_CHECKF(OutElement.GetOutputSize() == Schema.GetEncodedVectorSize(SchemaElement),
+			TEXT("Encoder Network Output unexpected size for %s. Got %i, expected %i according to Schema."),
+			*Schema.GetTag(SchemaElement).ToString(), OutElement.GetOutputSize(), Schema.GetEncodedVectorSize(SchemaElement));
+	}
+
 	void GenerateEncoderNetworkFileDataFromSchema(
 		TArray<uint8>& OutFileData,
 		uint32& OutInputSize,
@@ -1058,7 +1069,9 @@ namespace UE::Learning::Observation
 		UE_LEARNING_CHECK(Schema.IsValid(SchemaElement));
 
 		NNE::RuntimeBasic::FModelBuilder Builder(Seed);
-		Builder.WriteFileDataAndReset(OutFileData, OutInputSize, OutOutputSize, Private::MakeEncoderNetworkFromSchema(Builder, Schema, SchemaElement));
+		NNE::RuntimeBasic::FModelBuilderElement Element;
+		MakeEncoderNetworkModelBuilderElementFromSchema(Element, Builder, Schema, SchemaElement);
+		Builder.WriteFileDataAndReset(OutFileData, OutInputSize, OutOutputSize, Element);
 	}
 
 	void SetVectorFromObject(
diff --git a/Engine/Plugins/Experimental/LearningAgents/Source/Learning/Public/LearningAction.h b/Engine/Plugins/Experimental/LearningAgents/Source/Learning/Public/LearningAction.h
--- a/Engine/Plugins/Experimental/LearningAgents/Source/Learning/Public/LearningAction.h
+++ b/Engine/Plugins/Experimental/LearningAgents/Source/Learning/Public/LearningAction.h
@@ -6,6 +6,12 @@
 
 #include "UObject/NameTypes.h"
 
+namespace UE::NNE::RuntimeBasic
+{
+	class FModelBuilderElement;
+	class FModelBuilder;
+}
+
 namespace UE::Learning::Action
 {
 	/**
@@ -435,6 +441,21 @@ namespace UE::Learning::Action
 		const FSchema& SchemaB,
 		const FSchemaElement SchemaElementB);
 
+	/**
+	 * Make a NNE::RuntimeBasic::FModelBuilderElement for the given Schema. This can be used if you want to plug the Decoder generated by this 
+	 * schema as a part of a larger model you are making with a NNE::RuntimeBasic::FModelBuilder.
+	 *
+	 * @param OutElement				Output Builder Element
+	 * @param Builder					Model Builder
+	 * @param Schema					Action Schema
+	 * @param SchemaElement				Action Schema Element
+	 */
+	LEARNING_API void MakeDecoderNetworkModelBuilderElementFromSchema(
+		NNE::RuntimeBasic::FModelBuilderElement& OutElement,
+		NNE::RuntimeBasic::FModelBuilder& Builder,
+		const FSchema& Schema,
+		const FSchemaElement SchemaElement);
+
 	/**
 	 * Generate FileData for a Neural Network that can act as a Decoder for the given schema. This network will take as input a vector of size
 	 * EncodedActionSize and produce a vector of size ActionDistributionVectorSize.
diff --git a/Engine/Plugins/Experimental/LearningAgents/Source/Learning/Public/LearningObservation.h b/Engine/Plugins/Experimental/LearningAgents/Source/Learning/Public/LearningObservation.h
--- a/Engine/Plugins/Experimental/LearningAgents/Source/Learning/Public/LearningObservation.h
+++ b/Engine/Plugins/Experimental/LearningAgents/Source/Learning/Public/LearningObservation.h
@@ -6,6 +6,12 @@
 
 #include "UObject/NameTypes.h"
 
+namespace UE::NNE::RuntimeBasic
+{
+	class FModelBuilderElement;
+	class FModelBuilder;
+}
+
 namespace UE::Learning::Observation
 {
 	/**
@@ -417,6 +423,21 @@ namespace UE::Learning::Observation
 		const FSchema& SchemaB,
 		const FSchemaElement SchemaElementB);
 
+	/**
+	 * Make a NNE::RuntimeBasic::FModelBuilderElement for the given Schema. This can be used if you want to plug the Encoder generated by this
+	 * schema as a part of a larger model you are making with a NNE::RuntimeBasic::FModelBuilder.
+	 *
+	 * @param OutElement				Output Builder Element
+	 * @param Builder					Model Builder
+	 * @param Schema					Observation Schema
+	 * @param SchemaElement				Observation Schema Element
+	 */
+	LEARNING_API void MakeEncoderNetworkModelBuilderElementFromSchema(
+		NNE::RuntimeBasic::FModelBuilderElement& OutElement,
+		NNE::RuntimeBasic::FModelBuilder& Builder,
+		const FSchema& Schema,
+		const FSchemaElement SchemaElement);
+
 	/**
 	 * Generate FileData for a Neural Network that can act as a Encoder for the given schema. This network will take as input a vector of size
 	 * ObservationVectorSize and produce a vector of size EncodedVectorSize.
