LearningAgents: Added Random Gaussian Vector Function. Fixed issue with Neural Network being updated with new network of different size.

diff --git a/Engine/Plugins/Experimental/LearningAgents/Source/Learning/Private/LearningNeuralNetwork.cpp b/Engine/Plugins/Experimental/LearningAgents/Source/Learning/Private/LearningNeuralNetwork.cpp
--- a/Engine/Plugins/Experimental/LearningAgents/Source/Learning/Private/LearningNeuralNetwork.cpp
+++ b/Engine/Plugins/Experimental/LearningAgents/Source/Learning/Private/LearningNeuralNetwork.cpp
@@ -214,7 +214,7 @@ namespace UE::Learning
 			{
 				if (TSharedPtr<UE::Learning::FNeuralNetworkInference> InferenceObjectPtr = InferenceObject.Pin())
 				{
-					InferenceObjectPtr->ReloadModelInstances(*Model);
+					InferenceObjectPtr->ReloadModelInstances(*Model, InputSize, OutputSize);
 				}
 			}
 
@@ -236,17 +236,18 @@ namespace UE::Learning
 		const int32 InOutputSize,
 		const FNeuralNetworkInferenceSettings& InSettings)
 		: MaxBatchSize(InMaxBatchSize)
-		, InputSize(InInputSize)
-		, OutputSize(InOutputSize)
 		, Settings(InSettings)
 	{
-		ReloadModelInstances(InModel);
+		ReloadModelInstances(InModel, InInputSize, InOutputSize);
 	}
 
-	void FNeuralNetworkInference::ReloadModelInstances(NNE::IModelCPU& InModel)
+	void FNeuralNetworkInference::ReloadModelInstances(NNE::IModelCPU& InModel, const int32 InInputSize, const int32 InOutputSize)
 	{
 		UE_LEARNING_TRACE_CPUPROFILER_EVENT_SCOPE(Learning::FNeuralNetworkInference::ReloadModelInstances);
 
+		InputSize = InInputSize;
+		OutputSize = InOutputSize;
+
 		if (Settings.bParallelEvaluation)
 		{
 			const int32 IdealInferenceInstanceNum = ParallelForImpl::GetNumberOfThreadTasks(MaxBatchSize, Settings.MinParallelBatchSize, EParallelForFlags::None);
diff --git a/Engine/Plugins/Experimental/LearningAgents/Source/Learning/Private/LearningRandom.cpp b/Engine/Plugins/Experimental/LearningAgents/Source/Learning/Private/LearningRandom.cpp
--- a/Engine/Plugins/Experimental/LearningAgents/Source/Learning/Private/LearningRandom.cpp
+++ b/Engine/Plugins/Experimental/LearningAgents/Source/Learning/Private/LearningRandom.cpp
@@ -61,6 +61,17 @@ namespace UE::Learning::Random
 			Float(State ^ 0x0eedb850)) + Mean;
 	}
 
+	FVector VectorGaussian(
+		const uint32 State,
+		const FVector Mean,
+		const FVector Std)
+	{
+		return FVector(
+			Gaussian(State ^ 0x0c7eaaf3, Mean.X, Std.X),
+			Gaussian(State ^ 0x109be32d, Mean.Y, Std.Y),
+			Gaussian(State ^ 0xed839d93, Mean.Z, Std.Z));
+	}
+
 	float ClippedGaussian(
 		const uint32 State,
 		const float Mean,
@@ -434,6 +445,15 @@ namespace UE::Learning::Random
 		return Gaussian(State ^ 0x5df36815, Mean, Std);
 	}
 
+	FVector SampleVectorGaussian(
+		uint32& State,
+		const FVector Mean,
+		const FVector Std)
+	{
+		State = Int(State ^ 0x9efa29b7);
+		return VectorGaussian(State ^ 0x57167b59, Mean, Std);
+	}
+
 	float SampleClippedGaussian(
 		uint32& State,
 		const float Mean,
diff --git a/Engine/Plugins/Experimental/LearningAgents/Source/Learning/Public/LearningNeuralNetwork.h b/Engine/Plugins/Experimental/LearningAgents/Source/Learning/Public/LearningNeuralNetwork.h
--- a/Engine/Plugins/Experimental/LearningAgents/Source/Learning/Public/LearningNeuralNetwork.h
+++ b/Engine/Plugins/Experimental/LearningAgents/Source/Learning/Public/LearningNeuralNetwork.h
@@ -181,7 +181,7 @@ namespace UE::Learning
 		void Evaluate(TLearningArrayView<2, float> Output, const TLearningArrayView<2, const float> Input);
 
 		// This function will re-build the internal Model Instances used for multi-threading. It should be called whenever the given Model is updated. 
-		void ReloadModelInstances(NNE::IModelCPU& Model);
+		void ReloadModelInstances(NNE::IModelCPU& Model, const int32 InInputSize, const int32 InOutputSize);
 
 		int32 GetMaxBatchSize() const;
 		int32 GetInputSize() const;
diff --git a/Engine/Plugins/Experimental/LearningAgents/Source/Learning/Public/LearningRandom.h b/Engine/Plugins/Experimental/LearningAgents/Source/Learning/Public/LearningRandom.h
--- a/Engine/Plugins/Experimental/LearningAgents/Source/Learning/Public/LearningRandom.h
+++ b/Engine/Plugins/Experimental/LearningAgents/Source/Learning/Public/LearningRandom.h
@@ -71,6 +71,11 @@ namespace UE::Learning::Random
 		const float Mean = 0.0f,
 		const float Std = 1.0f);
 
+	LEARNING_API FVector VectorGaussian(
+		const uint32 State,
+		const FVector Mean = FVector::ZeroVector,
+		const FVector Std = FVector::OneVector);
+
 	LEARNING_API float ClippedGaussian(
 		const uint32 State,
 		const float Mean = 0.0f,
@@ -196,6 +201,11 @@ namespace UE::Learning::Random
 		const float Mean = 0.0f,
 		const float Std = 1.0f);
 
+	LEARNING_API FVector SampleVectorGaussian(
+		uint32& State,
+		const FVector Mean = FVector::ZeroVector,
+		const FVector Std = FVector::OneVector);
+
 	LEARNING_API float SampleClippedGaussian(
 		uint32& State,
 		const float Mean = 0.0f,
