LearningAgents: Added new reward function for distance between locations.

#virtualized

diff --git a/Engine/Plugins/Experimental/LearningAgents/Source/LearningAgentsTraining/Private/LearningAgentsRewards.cpp b/Engine/Plugins/Experimental/LearningAgents/Source/LearningAgentsTraining/Private/LearningAgentsRewards.cpp
--- a/Engine/Plugins/Experimental/LearningAgents/Source/LearningAgentsTraining/Private/LearningAgentsRewards.cpp
+++ b/Engine/Plugins/Experimental/LearningAgents/Source/LearningAgentsTraining/Private/LearningAgentsRewards.cpp
@@ -73,6 +73,64 @@ float ULearningAgentsRewards::MakeRewardOnCondition(
 	return Reward;
 }
 
+float ULearningAgentsRewards::MakeRewardFromLocationDifference(
+	const FVector LocationA,
+	const FVector LocationB,
+	const float LocationScale,
+	const float RewardScale,
+	const FName Tag,
+	const bool bVisualLoggerEnabled,
+	ULearningAgentsManagerListener* VisualLoggerListener,
+	const int32 VisualLoggerAgentId,
+	const FVector VisualLoggerLocation,
+	const FLinearColor VisualLoggerColor)
+{
+	const float Distance = FVector::Dist(LocationA, LocationB);
+	const float Difference = Distance / FMath::Max(LocationScale, UE_SMALL_NUMBER);
+	const float Reward = Difference * RewardScale;
+
+#if UE_LEARNING_AGENTS_ENABLE_VISUAL_LOG
+	if (bVisualLoggerEnabled && VisualLoggerListener)
+	{
+		const ULearningAgentsVisualLoggerObject* VisualLoggerObject = VisualLoggerListener->GetOrAddVisualLoggerObject(Tag);
+
+		UE_LEARNING_AGENTS_VLOG_LOCATION(VisualLoggerObject, LogLearning, Display,
+			LocationA,
+			10,
+			VisualLoggerColor.ToFColor(true),
+			TEXT(""));
+
+		UE_LEARNING_AGENTS_VLOG_SEGMENT(VisualLoggerObject, LogLearning, Display,
+			LocationA,
+			LocationB,
+			VisualLoggerColor.ToFColor(true),
+			TEXT(""));
+
+		UE_LEARNING_AGENTS_VLOG_LOCATION(VisualLoggerObject, LogLearning, Display,
+			LocationB,
+			10,
+			VisualLoggerColor.ToFColor(true),
+			TEXT(""));
+
+		UE_LEARNING_AGENTS_VLOG_STRING(VisualLoggerObject, LogLearning, Display, VisualLoggerLocation,
+			VisualLoggerColor.ToFColor(true),
+			TEXT("Listener: %s\nTag: %s\nAgent Id: % 3i\nLocationA: [% 6.1f % 6.1f % 6.1f]\nLocationB: [% 6.1f % 6.1f % 6.1f]\nDistance: [% 6.2f]\nLocationScale: [% 6.2f]\nDifference: [% 6.2f]\nScale: [% 6.2f]\nReward: [% 6.2f]"),
+			*VisualLoggerListener->GetName(),
+			*Tag.ToString(),
+			VisualLoggerAgentId,
+			LocationA.X, LocationA.Y, LocationA.Z,
+			LocationB.X, LocationB.Y, LocationB.Z,
+			Distance,
+			LocationScale,
+			Difference,
+			RewardScale,
+			Reward);
+	}
+#endif
+
+	return Reward;
+}
+
 float ULearningAgentsRewards::MakeRewardOnLocationDifferenceBelowThreshold(
 	const FVector LocationA, 
 	const FVector LocationB, 
diff --git a/Engine/Plugins/Experimental/LearningAgents/Source/LearningAgentsTraining/Public/LearningAgentsRewards.h b/Engine/Plugins/Experimental/LearningAgents/Source/LearningAgentsTraining/Public/LearningAgentsRewards.h
--- a/Engine/Plugins/Experimental/LearningAgents/Source/LearningAgentsTraining/Public/LearningAgentsRewards.h
+++ b/Engine/Plugins/Experimental/LearningAgents/Source/LearningAgentsTraining/Public/LearningAgentsRewards.h
@@ -64,6 +64,34 @@ public:
 		const FVector VisualLoggerLocation = FVector::ZeroVector,
 		const FLinearColor VisualLoggerColor = FLinearColor::Green);
 
+	/**
+	 * Make a reward based on the distance between two locations. Gives a larger reward the further the distance between the locations.
+	 *
+	 * @param LocationA The first location.
+	 * @param LocationB The second location.
+	 * @param LocationScale The expected scale for the distance between locations.
+	 * @param RewardScale The scale of the reward. Use a negative scale to create a penalty.
+	 * @param Tag The tag for the reward. Used for debugging.
+	 * @param bVisualLoggerEnabled When true, debug data will be sent to the visual logger.
+	 * @param VisualLoggerListener The listener object which is making this reward. This must be set to use logging.
+	 * @param VisualLoggerAgentId The agent id associated with this reward.
+	 * @param VisualLoggerLocation A location for the visual logger information in the world.
+	 * @param VisualLoggerColor The color for the visual logger display.
+	 * @return The resulting reward value.
+	 */
+	UFUNCTION(BlueprintPure, Category = "LearningAgents", meta = (AdvancedDisplay = 4, DefaultToSelf = "VisualLoggerListener"))
+	static float MakeRewardFromLocationDifference(
+		const FVector LocationA,
+		const FVector LocationB,
+		const float LocationScale = 100.0f,
+		const float RewardScale = 1.0f,
+		const FName Tag = TEXT("LocationDifferenceReward"),
+		const bool bVisualLoggerEnabled = false,
+		ULearningAgentsManagerListener* VisualLoggerListener = nullptr,
+		const int32 VisualLoggerAgentId = -1,
+		const FVector VisualLoggerLocation = FVector::ZeroVector,
+		const FLinearColor VisualLoggerColor = FLinearColor::Green);
+
 	/**
 	 * Make a reward when the distance between two locations is below a threshold, otherwise returns zero.
 	 *
