LearningAgents: Changed scaling factors to be given in Specify stage.

#virtualized

diff --git a/Engine/Plugins/Experimental/LearningAgents/Source/Learning/Private/LearningAction.cpp b/Engine/Plugins/Experimental/LearningAgents/Source/Learning/Private/LearningAction.cpp
--- a/Engine/Plugins/Experimental/LearningAgents/Source/Learning/Private/LearningAction.cpp
+++ b/Engine/Plugins/Experimental/LearningAgents/Source/Learning/Private/LearningAction.cpp
@@ -131,8 +131,12 @@ namespace UE::Learning::Action
 
 	FSchemaElement FSchema::CreateContinuous(const FSchemaContinuousParameters Parameters, const FName Tag)
 	{
+		UE_LEARNING_CHECK(Parameters.Num >= 0);
+		UE_LEARNING_CHECK(Parameters.Scale >= 0.0f);
+
 		FContinuousData ElementData;
 		ElementData.Num = Parameters.Num;
+		ElementData.Scale = Parameters.Scale;
 
 		const int32 Index = Types.Add(EType::Continuous);
 		Tags.Add(Tag);
@@ -346,9 +350,11 @@ namespace UE::Learning::Action
 	FSchemaContinuousParameters FSchema::GetContinuous(const FSchemaElement Element) const
 	{
 		UE_LEARNING_CHECK(IsValid(Element) && GetType(Element) == EType::Continuous);
+		const FContinuousData& ElementData = ContinuousData[TypeDataIndices[Element.Index]];
 
 		FSchemaContinuousParameters Parameters;
-		Parameters.Num = ContinuousData[TypeDataIndices[Element.Index]].Num;
+		Parameters.Num = ElementData.Num;
+		Parameters.Scale = ElementData.Scale;
 		return Parameters;
 	}
 
@@ -1491,13 +1497,23 @@ namespace UE::Learning::Action
 		{
 			// Check the input sizes match
 
+			const FSchemaContinuousParameters SchemaParameters = Schema.GetContinuous(SchemaElement);
+
 			TArrayView<const float> ActionValues = Object.GetContinuous(ObjectElement).Values;
 			UE_LEARNING_CHECK(Schema.GetActionVectorSize(SchemaElement) == ActionValues.Num());
 			UE_LEARNING_CHECK(Schema.GetActionVectorSize(SchemaElement) == OutActionVector.Num());
+			UE_LEARNING_CHECK(Schema.GetActionVectorSize(SchemaElement) == SchemaParameters.Num);
+
+			// Copy in and scale the values from the action object
+
+			const int32 ValueNum = SchemaParameters.Num;
+			const float ValueScale = FMath::Max(SchemaParameters.Scale, UE_SMALL_NUMBER);
 
-			// Copy in the values from the action object
+			for (int32 ValueIdx = 0; ValueIdx < ValueNum; ValueIdx++)
+			{
+				OutActionVector[ValueIdx] = ActionValues[ValueIdx] / ValueScale;
+			}
 
-			Array::Copy<1, float>(OutActionVector, ActionValues);
 			return;
 		}
 
@@ -1719,9 +1735,20 @@ namespace UE::Learning::Action
 
 		case EType::Continuous:
 		{
-			UE_LEARNING_CHECK(ActionVectorSize == Schema.GetContinuous(SchemaElement).Num);
+			const FSchemaContinuousParameters SchemaParameters = Schema.GetContinuous(SchemaElement);
+			UE_LEARNING_CHECK(ActionVectorSize == SchemaParameters.Num);
+
+			const int32 ValueNum = SchemaParameters.Num;
+			const float ValueScale = FMath::Max(SchemaParameters.Scale, UE_SMALL_NUMBER);
+
+			TLearningArray<1, float, TInlineAllocator<32>> ActionValues;
+			ActionValues.SetNumUninitialized({ ValueNum });
+			for (int32 ValueIdx = 0; ValueIdx < ValueNum; ValueIdx++)
+			{
+				ActionValues[ValueIdx] = ValueScale * ActionVector[ValueIdx];
+			}
 
-			OutObjectElement = OutObject.CreateContinuous({ MakeArrayView(ActionVector.GetData(), ActionVector.Num()) }, SchemaElementTag);
+			OutObjectElement = OutObject.CreateContinuous({ MakeArrayView(ActionValues.GetData(), ActionValues.Num()) }, SchemaElementTag);
 			return;
 		}
 
diff --git a/Engine/Plugins/Experimental/LearningAgents/Source/Learning/Private/LearningObservation.cpp b/Engine/Plugins/Experimental/LearningAgents/Source/Learning/Private/LearningObservation.cpp
--- a/Engine/Plugins/Experimental/LearningAgents/Source/Learning/Private/LearningObservation.cpp
+++ b/Engine/Plugins/Experimental/LearningAgents/Source/Learning/Private/LearningObservation.cpp
@@ -78,9 +78,11 @@ namespace UE::Learning::Observation
 	FSchemaElement FSchema::CreateContinuous(const FSchemaContinuousParameters Parameters, const FName Tag)
 	{
 		UE_LEARNING_CHECK(Parameters.Num >= 0);
+		UE_LEARNING_CHECK(Parameters.Scale >= 0.0f);
 
 		FContinuousData ElementData;
 		ElementData.Num = Parameters.Num;
+		ElementData.Scale = Parameters.Scale;
 
 		const int32 Index = Types.Add(EType::Continuous);
 		Tags.Add(Tag);
@@ -259,9 +261,11 @@ namespace UE::Learning::Observation
 	FSchemaContinuousParameters FSchema::GetContinuous(const FSchemaElement Element) const
 	{
 		UE_LEARNING_CHECK(IsValid(Element) && GetType(Element) == EType::Continuous);
+		const FContinuousData& ElementData = ContinuousData[TypeDataIndices[Element.Index]];
 
 		FSchemaContinuousParameters Parameters;
-		Parameters.Num = ContinuousData[TypeDataIndices[Element.Index]].Num;
+		Parameters.Num = ElementData.Num;
+		Parameters.Scale = ElementData.Scale;
 		return Parameters;
 	}
 
@@ -1106,13 +1110,23 @@ namespace UE::Learning::Observation
 		{
 			// Check the input sizes match
 
+			const FSchemaContinuousParameters SchemaParameters = Schema.GetContinuous(SchemaElement);
+
 			const TArrayView<const float> ObservationValues = Object.GetContinuous(ObjectElement).Values;
 			UE_LEARNING_CHECK(Schema.GetObservationVectorSize(SchemaElement) == ObservationValues.Num());
 			UE_LEARNING_CHECK(Schema.GetObservationVectorSize(SchemaElement) == OutObservationVector.Num());
+			UE_LEARNING_CHECK(Schema.GetObservationVectorSize(SchemaElement) == SchemaParameters.Num);
+
+			// Copy in and scale the values from the observation object
 
-			// Copy in the values from the observation object
+			const int32 ValueNum = SchemaParameters.Num;
+			const float ValueScale = FMath::Max(SchemaParameters.Scale, UE_SMALL_NUMBER);
+
+			for (int32 ValueIdx = 0; ValueIdx < ValueNum; ValueIdx++)
+			{
+				OutObservationVector[ValueIdx] = ObservationValues[ValueIdx] / ValueScale;
+			}
 
-			Array::Copy<1, float>(OutObservationVector, ObservationValues);
 			return;
 		}
 
@@ -1342,9 +1356,20 @@ namespace UE::Learning::Observation
 
 		case EType::Continuous:
 		{
-			UE_LEARNING_CHECK(ObservationVectorSize == Schema.GetContinuous(SchemaElement).Num);
+			const FSchemaContinuousParameters SchemaParameters = Schema.GetContinuous(SchemaElement);
+			UE_LEARNING_CHECK(ObservationVectorSize == SchemaParameters.Num);
+
+			const int32 ValueNum = SchemaParameters.Num;
+			const float ValueScale = FMath::Max(SchemaParameters.Scale, UE_SMALL_NUMBER);
+
+			TLearningArray<1, float, TInlineAllocator<32>> ObservationValues;
+			ObservationValues.SetNumUninitialized({ ValueNum });
+			for (int32 ValueIdx = 0; ValueIdx < ValueNum; ValueIdx++)
+			{
+				ObservationValues[ValueIdx] = ValueScale * ObservationVector[ValueIdx];
+			}
 
-			OutObjectElement = OutObject.CreateContinuous({ MakeArrayView(ObservationVector.GetData(), ObservationVector.Num()) }, SchemaElementTag);
+			OutObjectElement = OutObject.CreateContinuous({ MakeArrayView(ObservationValues.GetData(), ObservationValues.Num()) }, SchemaElementTag);
 			return;
 		}
 
diff --git a/Engine/Plugins/Experimental/LearningAgents/Source/Learning/Public/LearningAction.h b/Engine/Plugins/Experimental/LearningAgents/Source/Learning/Public/LearningAction.h
--- a/Engine/Plugins/Experimental/LearningAgents/Source/Learning/Public/LearningAction.h
+++ b/Engine/Plugins/Experimental/LearningAgents/Source/Learning/Public/LearningAction.h
@@ -61,6 +61,9 @@ namespace UE::Learning::Action
 	{
 		// Number of values in the continuous action
 		int32 Num = 0;
+
+		// Scale factor for the continuous action
+		float Scale = 1.0f;
 	};
 
 	struct FSchemaDiscreteExclusiveParameters
@@ -206,6 +209,7 @@ namespace UE::Learning::Action
 		struct FContinuousData
 		{
 			int32 Num = INDEX_NONE;
+			float Scale = 0.0f;
 		};
 
 		struct FDiscreteExclusiveData
diff --git a/Engine/Plugins/Experimental/LearningAgents/Source/Learning/Public/LearningObservation.h b/Engine/Plugins/Experimental/LearningAgents/Source/Learning/Public/LearningObservation.h
--- a/Engine/Plugins/Experimental/LearningAgents/Source/Learning/Public/LearningObservation.h
+++ b/Engine/Plugins/Experimental/LearningAgents/Source/Learning/Public/LearningObservation.h
@@ -58,6 +58,9 @@ namespace UE::Learning::Observation
 	{
 		// Number of values in the continuous observation
 		int32 Num = 0;
+
+		// Scale factor for the continuous observation
+		float Scale = 1.0f;
 	};
 
 	struct FSchemaAndParameters
@@ -204,7 +207,7 @@ namespace UE::Learning::Observation
 		struct FContinuousData
 		{
 			int32 Num = INDEX_NONE;
-			float Scale = 1.0f;
+			float Scale = 0.0f;
 		};
 
 		struct FAndData
diff --git a/Engine/Plugins/Experimental/LearningAgents/Source/LearningAgents/Private/LearningAgentsActions.cpp b/Engine/Plugins/Experimental/LearningAgents/Source/LearningAgents/Private/LearningAgentsActions.cpp
--- a/Engine/Plugins/Experimental/LearningAgents/Source/LearningAgents/Private/LearningAgentsActions.cpp
+++ b/Engine/Plugins/Experimental/LearningAgents/Source/LearningAgents/Private/LearningAgentsActions.cpp
@@ -538,7 +538,7 @@ FLearningAgentsActionSchemaElement ULearningAgentsActions::SpecifyNullAction(ULe
 	return { Schema->ActionSchema.CreateNull(Tag) };
 }
 
-FLearningAgentsActionSchemaElement ULearningAgentsActions::SpecifyContinuousAction(ULearningAgentsActionSchema* Schema, const int32 Size, const FName Tag)
+FLearningAgentsActionSchemaElement ULearningAgentsActions::SpecifyContinuousAction(ULearningAgentsActionSchema* Schema, const int32 Size, const float Scale, const FName Tag)
 {
 	if (!Schema)
 	{
@@ -552,12 +552,18 @@ FLearningAgentsActionSchemaElement ULearningAgentsActions::SpecifyContinuousActi
 		return FLearningAgentsActionSchemaElement();
 	}
 
+	if (Scale < UE_SMALL_NUMBER)
+	{
+		UE_LOG(LogLearning, Error, TEXT("SpecifyContinuousAction: Invalid Scale for Continuous Action '%f', must be greater than '%f'."), Scale, UE_SMALL_NUMBER);
+		return FLearningAgentsActionSchemaElement();
+	}
+
 	if (Size == 0)
 	{
 		UE_LOG(LogLearning, Warning, TEXT("SpecifyContinuousAction: Specifying zero-sized Continuous Action."));
 	}
 
-	return { Schema->ActionSchema.CreateContinuous({ Size }, Tag) };
+	return { Schema->ActionSchema.CreateContinuous({ Size, Scale }, Tag) };
 }
 
 FLearningAgentsActionSchemaElement ULearningAgentsActions::SpecifyExclusiveDiscreteAction(ULearningAgentsActionSchema* Schema, const int32 Size, const TArray<float>& PriorProbabilities, const FName Tag)
@@ -1101,27 +1107,32 @@ FLearningAgentsActionSchemaElement ULearningAgentsActions::SpecifyBoolAction(ULe
 	return SpecifyExclusiveDiscreteActionFromArrayView(Schema, 2, { 1.0f - PriorProbability, PriorProbability }, Tag);
 }
 
-FLearningAgentsActionSchemaElement ULearningAgentsActions::SpecifyFloatAction(ULearningAgentsActionSchema* Schema, const FName Tag)
+FLearningAgentsActionSchemaElement ULearningAgentsActions::SpecifyFloatAction(ULearningAgentsActionSchema* Schema, const float FloatScale, const FName Tag)
 {
-	return SpecifyContinuousAction(Schema, 1, Tag);
+	return SpecifyContinuousAction(Schema, 1, FloatScale, Tag);
 }
 
-FLearningAgentsActionSchemaElement ULearningAgentsActions::SpecifyLocationAction(ULearningAgentsActionSchema* Schema, const FName Tag)
+FLearningAgentsActionSchemaElement ULearningAgentsActions::SpecifyLocationAction(ULearningAgentsActionSchema* Schema, const float LocationScale, const FName Tag)
 {
-	return SpecifyContinuousAction(Schema, 3, Tag);
+	return SpecifyContinuousAction(Schema, 3, LocationScale, Tag);
 }
 
-FLearningAgentsActionSchemaElement ULearningAgentsActions::SpecifyRotationAction(ULearningAgentsActionSchema* Schema, const FName Tag)
+FLearningAgentsActionSchemaElement ULearningAgentsActions::SpecifyRotationAction(ULearningAgentsActionSchema* Schema, const float RotationScale, const FName Tag)
 {
-	return SpecifyContinuousAction(Schema, 3, Tag);
+	return SpecifyContinuousAction(Schema, 3, RotationScale, Tag);
 }
 
-FLearningAgentsActionSchemaElement ULearningAgentsActions::SpecifyScaleAction(ULearningAgentsActionSchema* Schema, const FName Tag)
+FLearningAgentsActionSchemaElement ULearningAgentsActions::SpecifyScaleAction(ULearningAgentsActionSchema* Schema, const float ScaleScale, const FName Tag)
 {
-	return SpecifyContinuousAction(Schema, 3, Tag);
+	return SpecifyContinuousAction(Schema, 3, ScaleScale, Tag);
 }
 
-FLearningAgentsActionSchemaElement ULearningAgentsActions::SpecifyTransformAction(ULearningAgentsActionSchema* Schema, const FName Tag)
+FLearningAgentsActionSchemaElement ULearningAgentsActions::SpecifyTransformAction(
+	ULearningAgentsActionSchema* Schema, 
+	const float LocationScale ,
+	const float RotationScale,
+	const float ScaleScale,
+	const FName Tag)
 {
 	return SpecifyStructActionFromArrayViews(Schema,
 		{
@@ -1130,26 +1141,26 @@ FLearningAgentsActionSchemaElement ULearningAgentsActions::SpecifyTransformActio
 			TEXT("Scale")
 		},
 		{
-			SpecifyLocationAction(Schema),
-			SpecifyRotationAction(Schema),
-			SpecifyScaleAction(Schema)
+			SpecifyLocationAction(Schema, LocationScale),
+			SpecifyRotationAction(Schema, RotationScale),
+			SpecifyScaleAction(Schema, ScaleScale)
 		}, 
 		Tag);
 }
 
-FLearningAgentsActionSchemaElement ULearningAgentsActions::SpecifyAngleAction(ULearningAgentsActionSchema* Schema, const FName Tag)
+FLearningAgentsActionSchemaElement ULearningAgentsActions::SpecifyAngleAction(ULearningAgentsActionSchema* Schema, const float AngleScale, const FName Tag)
 {
-	return SpecifyContinuousAction(Schema, 1, Tag);
+	return SpecifyContinuousAction(Schema, 1, AngleScale, Tag);
 }
 
-FLearningAgentsActionSchemaElement ULearningAgentsActions::SpecifyVelocityAction(ULearningAgentsActionSchema* Schema, const FName Tag)
+FLearningAgentsActionSchemaElement ULearningAgentsActions::SpecifyVelocityAction(ULearningAgentsActionSchema* Schema, const float VelocityScale, const FName Tag)
 {
-	return SpecifyContinuousAction(Schema, 3, Tag);
+	return SpecifyContinuousAction(Schema, 3, VelocityScale, Tag);
 }
 
 FLearningAgentsActionSchemaElement ULearningAgentsActions::SpecifyDirectionAction(ULearningAgentsActionSchema* Schema, const FName Tag)
 {
-	return SpecifyContinuousAction(Schema, 3, Tag);
+	return SpecifyContinuousAction(Schema, 3, 1.0f, Tag);
 }
 
 void ULearningAgentsActions::LogAction(const ULearningAgentsActionObject* Object, const FLearningAgentsActionObjectElement Element)
@@ -1534,36 +1545,36 @@ FLearningAgentsActionObjectElement ULearningAgentsActions::MakeBoolAction(ULearn
 	return MakeExclusiveDiscreteAction(Object, bValue ? 1 : 0, Tag);
 }
 
-FLearningAgentsActionObjectElement ULearningAgentsActions::MakeFloatAction(ULearningAgentsActionObject* Object, const float Value, const float FloatScale, const FName Tag)
+FLearningAgentsActionObjectElement ULearningAgentsActions::MakeFloatAction(ULearningAgentsActionObject* Object, const float Value, const FName Tag)
 {
-	return MakeContinuousActionFromArrayView(Object, { Value / FMath::Max(FloatScale, UE_SMALL_NUMBER) }, Tag);
+	return MakeContinuousActionFromArrayView(Object, { Value }, Tag);
 }
 
-FLearningAgentsActionObjectElement ULearningAgentsActions::MakeLocationAction(ULearningAgentsActionObject* Object, const FVector Location, const FTransform RelativeTransform, const float LocationScale, const FName Tag)
+FLearningAgentsActionObjectElement ULearningAgentsActions::MakeLocationAction(ULearningAgentsActionObject* Object, const FVector Location, const FTransform RelativeTransform, const FName Tag)
 {
 	const FVector LocalLocation = RelativeTransform.InverseTransformPosition(Location);
 
 	return MakeContinuousActionFromArrayView(Object, {
-		(float)LocalLocation.X / FMath::Max(LocationScale, UE_SMALL_NUMBER),
-		(float)LocalLocation.Y / FMath::Max(LocationScale, UE_SMALL_NUMBER),
-		(float)LocalLocation.Z / FMath::Max(LocationScale, UE_SMALL_NUMBER) }, Tag);
+		(float)LocalLocation.X,
+		(float)LocalLocation.Y,
+		(float)LocalLocation.Z }, Tag);
 }
 
-FLearningAgentsActionObjectElement ULearningAgentsActions::MakeRotationAction(ULearningAgentsActionObject* Object, const FRotator Rotation, const FRotator RelativeRotation, const float RotationScale, const FName Tag)
+FLearningAgentsActionObjectElement ULearningAgentsActions::MakeRotationAction(ULearningAgentsActionObject* Object, const FRotator Rotation, const FRotator RelativeRotation, const FName Tag)
 {
-	return MakeRotationActionFromQuat(Object, FQuat::MakeFromRotator(Rotation), FQuat::MakeFromRotator(RelativeRotation), RotationScale, Tag);
+	return MakeRotationActionFromQuat(Object, FQuat::MakeFromRotator(Rotation), FQuat::MakeFromRotator(RelativeRotation), Tag);
 }
 
-FLearningAgentsActionObjectElement ULearningAgentsActions::MakeRotationActionFromQuat(ULearningAgentsActionObject* Object, const FQuat Rotation, const FQuat RelativeRotation, const float RotationScale, const FName Tag)
+FLearningAgentsActionObjectElement ULearningAgentsActions::MakeRotationActionFromQuat(ULearningAgentsActionObject* Object, const FQuat Rotation, const FQuat RelativeRotation, const FName Tag)
 {
 	FQuat LocalRotation = RelativeRotation.Inverse() * Rotation;
 	LocalRotation.EnforceShortestArcWith(FQuat::Identity);
 	const FVector RotationVector = LocalRotation.ToRotationVector();
 
 	return MakeContinuousActionFromArrayView(Object, {
-		(float)RotationVector.X / FMath::Max(FMath::DegreesToRadians(RotationScale), UE_SMALL_NUMBER),
-		(float)RotationVector.Y / FMath::Max(FMath::DegreesToRadians(RotationScale), UE_SMALL_NUMBER),
-		(float)RotationVector.Z / FMath::Max(FMath::DegreesToRadians(RotationScale), UE_SMALL_NUMBER),
+		(float)FMath::RadiansToDegrees(RotationVector.X),
+		(float)FMath::RadiansToDegrees(RotationVector.Y),
+		(float)FMath::RadiansToDegrees(RotationVector.Z),
 		}, Tag);
 }
 
@@ -1580,10 +1591,8 @@ FLearningAgentsActionObjectElement ULearningAgentsActions::MakeScaleAction(ULear
 		}, Tag);
 }
 
-FLearningAgentsActionObjectElement ULearningAgentsActions::MakeTransformAction(ULearningAgentsActionObject* Object, const FTransform Transform, const FTransform RelativeTransform, const float LocationScale, const FName Tag)
+FLearningAgentsActionObjectElement ULearningAgentsActions::MakeTransformAction(ULearningAgentsActionObject* Object, const FTransform Transform, const FTransform RelativeTransform, const FName Tag)
 {
-	const FTransform LocalTransform = Transform * RelativeTransform.Inverse();
-
 	return MakeStructActionFromArrayViews(Object,
 		{
 			TEXT("Location"),
@@ -1591,32 +1600,31 @@ FLearningAgentsActionObjectElement ULearningAgentsActions::MakeTransformAction(U
 			TEXT("Scale")
 		},
 		{
-			MakeLocationAction(Object, LocalTransform.GetLocation(), FTransform::Identity, LocationScale),
-			MakeRotationActionFromQuat(Object, LocalTransform.GetRotation(), FQuat::Identity),
-			MakeScaleAction(Object, LocalTransform.GetScale3D(), FVector::OneVector)
+			MakeLocationAction(Object, Transform.GetLocation(), RelativeTransform),
+			MakeRotationActionFromQuat(Object, Transform.GetRotation(), RelativeTransform.GetRotation()),
+			MakeScaleAction(Object, Transform.GetScale3D(), RelativeTransform.GetScale3D())
 		},
 		Tag);
 }
 
-FLearningAgentsActionObjectElement ULearningAgentsActions::MakeAngleAction(ULearningAgentsActionObject* Object, const float Angle, const float RelativeAngle, const float AngleScale, const FName Tag)
+FLearningAgentsActionObjectElement ULearningAgentsActions::MakeAngleAction(ULearningAgentsActionObject* Object, const float Angle, const float RelativeAngle, const FName Tag)
 {
-	return MakeAngleActionRadians(Object, FMath::DegreesToRadians(Angle), FMath::DegreesToRadians(RelativeAngle), FMath::DegreesToRadians(AngleScale), Tag);
+	return MakeContinuousActionFromArrayView(Object, { FMath::FindDeltaAngleDegrees(RelativeAngle, Angle) }, Tag);
 }
 
-FLearningAgentsActionObjectElement ULearningAgentsActions::MakeAngleActionRadians(ULearningAgentsActionObject* Object, const float Angle, const float RelativeAngle, const float AngleScale, const FName Tag)
+FLearningAgentsActionObjectElement ULearningAgentsActions::MakeAngleActionRadians(ULearningAgentsActionObject* Object, const float Angle, const float RelativeAngle, const FName Tag)
 {
-	const float LocalAngle = FMath::FindDeltaAngleRadians(RelativeAngle, Angle);
-	return MakeContinuousActionFromArrayView(Object, { LocalAngle / FMath::Max(AngleScale, UE_SMALL_NUMBER) }, Tag);
+	return MakeAngleAction(Object, FMath::RadiansToDegrees(Angle), FMath::RadiansToDegrees(RelativeAngle), Tag);
 }
 
-FLearningAgentsActionObjectElement ULearningAgentsActions::MakeVelocityAction(ULearningAgentsActionObject* Object, const FVector Velocity, const FTransform RelativeTransform, const float VelocityScale, const FName Tag)
+FLearningAgentsActionObjectElement ULearningAgentsActions::MakeVelocityAction(ULearningAgentsActionObject* Object, const FVector Velocity, const FTransform RelativeTransform, const FName Tag)
 {
 	const FVector LocalVelocity = RelativeTransform.InverseTransformVectorNoScale(Velocity);
 
 	return MakeContinuousActionFromArrayView(Object, {
-		(float)LocalVelocity.X / FMath::Max(VelocityScale, UE_SMALL_NUMBER),
-		(float)LocalVelocity.Y / FMath::Max(VelocityScale, UE_SMALL_NUMBER),
-		(float)LocalVelocity.Z / FMath::Max(VelocityScale, UE_SMALL_NUMBER),
+		(float)LocalVelocity.X,
+		(float)LocalVelocity.Y,
+		(float)LocalVelocity.Z,
 		}, Tag);
 }
 
@@ -2736,7 +2744,6 @@ bool ULearningAgentsActions::GetFloatAction(
 	float& OutValue, 
 	const ULearningAgentsActionObject* Object, 
 	const FLearningAgentsActionObjectElement Element, 
-	const float FloatScale, 
 	const FName Tag,
 	const bool bVisualLoggerEnabled,
 	ULearningAgentsManagerListener* VisualLoggerListener,
@@ -2744,15 +2751,12 @@ bool ULearningAgentsActions::GetFloatAction(
 	const FVector VisualLoggerLocation,
 	const FLinearColor VisualLoggerColor)
 {
-	float OutValuesData;
-	if (!GetContinuousActionToArrayView(MakeArrayView(&OutValuesData, 1), Object, Element, Tag))
+	if (!GetContinuousActionToArrayView(MakeArrayView(&OutValue, 1), Object, Element, Tag))
 	{
 		OutValue = 0.0f;
 		return false;
 	}
 
-	OutValue = OutValuesData * FloatScale;
-
 #if UE_LEARNING_AGENTS_ENABLE_VISUAL_LOG
 	if (bVisualLoggerEnabled && VisualLoggerListener)
 	{
@@ -2760,12 +2764,10 @@ bool ULearningAgentsActions::GetFloatAction(
 
 		UE_LEARNING_AGENTS_VLOG_STRING(VisualLoggerObject, LogLearning, Display, VisualLoggerLocation,
 			VisualLoggerColor.ToFColor(true),
-			TEXT("Listener: %s\nTag: %s\nAgent Id: % 3i\nEncoded: [% 6.2f]\nScale: [% 6.2f]\nValue: [% 6.1f]"),
+			TEXT("Listener: %s\nTag: %s\nAgent Id: % 3i\nValue: [% 6.1f]"),
 			*VisualLoggerListener->GetName(),
 			*Tag.ToString(),
 			VisualLoggerAgentId,
-			OutValuesData,
-			FloatScale,
 			OutValue);
 	}
 #endif
@@ -2778,7 +2780,6 @@ bool ULearningAgentsActions::GetLocationAction(
 	const ULearningAgentsActionObject* Object, 
 	const FLearningAgentsActionObjectElement Element, 
 	const FTransform RelativeTransform, 
-	const float LocationScale, 
 	const FName Tag,
 	const bool bVisualLoggerEnabled,
 	ULearningAgentsManagerListener* VisualLoggerListener,
@@ -2793,7 +2794,7 @@ bool ULearningAgentsActions::GetLocationAction(
 		return false;
 	}
 
-	const FVector LocalLocation = LocationScale * FVector(OutValues[0], OutValues[1], OutValues[2]);
+	const FVector LocalLocation = FVector(OutValues[0], OutValues[1], OutValues[2]);
 	OutLocation = RelativeTransform.TransformPosition(LocalLocation);
 	
 #if UE_LEARNING_AGENTS_ENABLE_VISUAL_LOG
@@ -2821,12 +2822,10 @@ bool ULearningAgentsActions::GetLocationAction(
 
 		UE_LEARNING_AGENTS_VLOG_STRING(VisualLoggerObject, LogLearning, Display, VisualLoggerLocation,
 			VisualLoggerColor.ToFColor(true),
-			TEXT("Listener: %s\nTag: %s\nAgent Id: % 3i\nEncoded: [% 6.2f % 6.2f % 6.2f]\nScale: [% 6.2f]\nLocal Location: [% 6.1f % 6.1f % 6.1f]\nLocation: [% 6.1f % 6.1f % 6.1f]"),
+			TEXT("Listener: %s\nTag: %s\nAgent Id: % 3i\nLocal Location: [% 6.1f % 6.1f % 6.1f]\nLocation: [% 6.1f % 6.1f % 6.1f]"),
 			*VisualLoggerListener->GetName(),
 			*Tag.ToString(),
 			VisualLoggerAgentId,
-			OutValues[0], OutValues[1], OutValues[2],
-			LocationScale,
 			LocalLocation.X, LocalLocation.Y, LocalLocation.Z,
 			OutLocation.X, OutLocation.Y, OutLocation.Z);
 	}
@@ -2840,7 +2839,6 @@ bool ULearningAgentsActions::GetRotationAction(
 	const ULearningAgentsActionObject* Object, 
 	const FLearningAgentsActionObjectElement Element, 
 	const FRotator RelativeRotation, 
-	const float RotationScale, 
 	const FName Tag,
 	const bool bVisualLoggerEnabled,
 	ULearningAgentsManagerListener* VisualLoggerListener,
@@ -2850,7 +2848,7 @@ bool ULearningAgentsActions::GetRotationAction(
 	const FLinearColor VisualLoggerColor)
 {
 	FQuat OutRotationQuat;
-	if (!GetRotationActionAsQuat(OutRotationQuat, Object, Element, FQuat::MakeFromRotator(RelativeRotation), RotationScale, Tag, bVisualLoggerEnabled, VisualLoggerListener, VisualLoggerAgentId, VisualLoggerRotationLocation, VisualLoggerLocation, VisualLoggerColor))
+	if (!GetRotationActionAsQuat(OutRotationQuat, Object, Element, FQuat::MakeFromRotator(RelativeRotation), Tag, bVisualLoggerEnabled, VisualLoggerListener, VisualLoggerAgentId, VisualLoggerRotationLocation, VisualLoggerLocation, VisualLoggerColor))
 	{
 		OutRotation = FRotator::ZeroRotator;
 		return false;
@@ -2865,7 +2863,6 @@ bool ULearningAgentsActions::GetRotationActionAsQuat(
 	const ULearningAgentsActionObject* Object, 
 	const FLearningAgentsActionObjectElement Element, 
 	const FQuat RelativeRotation, 
-	const float RotationScale, 
 	const FName Tag,
 	const bool bVisualLoggerEnabled,
 	ULearningAgentsManagerListener* VisualLoggerListener,
@@ -2881,7 +2878,7 @@ bool ULearningAgentsActions::GetRotationActionAsQuat(
 		return false;
 	}
 
-	const FVector LocalRotationVector = FMath::DegreesToRadians(RotationScale) * FVector(OutValues[0], OutValues[1], OutValues[2]);
+	const FVector LocalRotationVector = FVector(FMath::DegreesToRadians(OutValues[0]), FMath::DegreesToRadians(OutValues[1]), FMath::DegreesToRadians(OutValues[2]));
 	const FQuat LocalRotation = FQuat::MakeFromRotationVector(LocalRotationVector);
 	OutRotation = RelativeRotation * LocalRotation;
 
@@ -2898,12 +2895,10 @@ bool ULearningAgentsActions::GetRotationActionAsQuat(
 
 		UE_LEARNING_AGENTS_VLOG_STRING(VisualLoggerObject, LogLearning, Display, VisualLoggerLocation,
 			VisualLoggerColor.ToFColor(true),
-			TEXT("Listener: %s\nTag: %s\nAgent Id: % 3i\nEncoded: [% 6.2f % 6.2f % 6.2f]\nScale: [% 6.2f]\nLocal Rotation Vector: [% 6.1f % 6.1f % 6.1f]\nLocal Rotation: [% 6.1f % 6.1f % 6.1f % 6.1f]\nRotation: [% 6.1f % 6.1f % 6.1f % 6.1f]"),
+			TEXT("Listener: %s\nTag: %s\nAgent Id: % 3i\nLocal Rotation Vector: [% 6.1f % 6.1f % 6.1f]\nLocal Rotation: [% 6.1f % 6.1f % 6.1f % 6.1f]\nRotation: [% 6.1f % 6.1f % 6.1f % 6.1f]"),
 			*VisualLoggerListener->GetName(),
 			*Tag.ToString(),
 			VisualLoggerAgentId,
-			OutValues[0], OutValues[1], OutValues[2],
-			RotationScale,
 			LocalRotationVector.X, LocalRotationVector.Y, LocalRotationVector.Z,
 			LocalRotation.X, LocalRotation.Y, LocalRotation.Z, LocalRotation.W,
 			OutRotation.X, OutRotation.Y, OutRotation.Z, OutRotation.W);
@@ -2918,7 +2913,6 @@ bool ULearningAgentsActions::GetScaleAction(
 	const ULearningAgentsActionObject* Object, 
 	const FLearningAgentsActionObjectElement Element, 
 	const FVector RelativeScale, 
-	const float Scale, 
 	const FName Tag,
 	const bool bVisualLoggerEnabled,
 	ULearningAgentsManagerListener* VisualLoggerListener,
@@ -2933,7 +2927,7 @@ bool ULearningAgentsActions::GetScaleAction(
 		return false;
 	}
 
-	const FVector LocalScaleVector = UE::Learning::Agents::Action::Private::VectorExp(Scale * FVector(OutValues[0], OutValues[1], OutValues[2]));
+	const FVector LocalScaleVector = UE::Learning::Agents::Action::Private::VectorExp(FVector(OutValues[0], OutValues[1], OutValues[2]));
 	OutScale = RelativeScale * LocalScaleVector;
 
 #if UE_LEARNING_AGENTS_ENABLE_VISUAL_LOG
@@ -2943,12 +2937,10 @@ bool ULearningAgentsActions::GetScaleAction(
 
 		UE_LEARNING_AGENTS_VLOG_STRING(VisualLoggerObject, LogLearning, Display, VisualLoggerLocation,
 			VisualLoggerColor.ToFColor(true),
-			TEXT("Listener: %s\nTag: %s\nAgent Id: % 3i\nEncoded: [% 6.2f % 6.2f % 6.2f]\nScale: [% 6.2f]\nLocal Scale: [% 6.1f % 6.1f % 6.1f]\nScale: [% 6.1f % 6.1f % 6.1f]"),
+			TEXT("Listener: %s\nTag: %s\nAgent Id: % 3i\nLocal Scale: [% 6.1f % 6.1f % 6.1f]\nScale: [% 6.1f % 6.1f % 6.1f]"),
 			*VisualLoggerListener->GetName(),
 			*Tag.ToString(),
 			VisualLoggerAgentId,
-			OutValues[0], OutValues[1], OutValues[2],
-			Scale,
 			LocalScaleVector.X, LocalScaleVector.Y, LocalScaleVector.Z,
 			OutScale.X, OutScale.Y, OutScale.Z);
 	}
@@ -2962,9 +2954,6 @@ bool ULearningAgentsActions::GetTransformAction(
 	const ULearningAgentsActionObject* Object, 
 	const FLearningAgentsActionObjectElement Element, 
 	const FTransform RelativeTransform, 
-	const float LocationScale, 
-	const float RotationScale, 
-	const float ScaleScale, 
 	const FName Tag,
 	const bool bVisualLoggerEnabled,
 	ULearningAgentsManagerListener* VisualLoggerListener,
@@ -2981,21 +2970,21 @@ bool ULearningAgentsActions::GetTransformAction(
 	}
 
 	FVector OutLocation;
-	if (!GetLocationAction(OutLocation, Object, OutElements[MakeArrayView(OutElementNames).Find(TEXT("Location"))], RelativeTransform, LocationScale))
+	if (!GetLocationAction(OutLocation, Object, OutElements[MakeArrayView(OutElementNames).Find(TEXT("Location"))], RelativeTransform))
 	{
 		OutTransform = FTransform::Identity;
 		return false;
 	}
 
 	FQuat OutRotation;
-	if (!GetRotationActionAsQuat(OutRotation, Object, OutElements[MakeArrayView(OutElementNames).Find(TEXT("Rotation"))], RelativeTransform.GetRotation(), RotationScale))
+	if (!GetRotationActionAsQuat(OutRotation, Object, OutElements[MakeArrayView(OutElementNames).Find(TEXT("Rotation"))], RelativeTransform.GetRotation()))
 	{
 		OutTransform = FTransform::Identity;
 		return false;
 	}
 
 	FVector OutScale;
-	if (!GetScaleAction(OutScale, Object, OutElements[MakeArrayView(OutElementNames).Find(TEXT("Scale"))], RelativeTransform.GetScale3D(), ScaleScale))
+	if (!GetScaleAction(OutScale, Object, OutElements[MakeArrayView(OutElementNames).Find(TEXT("Scale"))], RelativeTransform.GetScale3D()))
 	{
 		OutTransform = FTransform::Identity;
 		return false;
@@ -3016,13 +3005,10 @@ bool ULearningAgentsActions::GetTransformAction(
 
 		UE_LEARNING_AGENTS_VLOG_STRING(VisualLoggerObject, LogLearning, Display, VisualLoggerLocation,
 			VisualLoggerColor.ToFColor(true),
-			TEXT("Listener: %s\nTag: %s\nAgent Id: % 3i\nLocation Scale: [% 6.1f]\nRotation Scale: [% 6.1f]\nScale Scale: [% 6.1f]\nLocation: [% 6.1f % 6.1f % 6.1f]\nRotation: [% 6.1f % 6.1f % 6.1f % 6.1f]\nScale: [% 6.1f % 6.1f % 6.1f]"),
+			TEXT("Listener: %s\nTag: %s\nAgent Id: % 3i\nLocation: [% 6.1f % 6.1f % 6.1f]\nRotation: [% 6.1f % 6.1f % 6.1f % 6.1f]\nScale: [% 6.1f % 6.1f % 6.1f]"),
 			*VisualLoggerListener->GetName(),
 			*Tag.ToString(),
 			VisualLoggerAgentId,
-			LocationScale,
-			RotationScale,
-			ScaleScale,
 			OutLocation.X, OutLocation.Y, OutLocation.Z,
 			OutRotation.X, OutRotation.Y, OutRotation.Z, OutRotation.W,
 			OutScale.X, OutScale.Y, OutScale.Z);
@@ -3037,7 +3023,6 @@ bool ULearningAgentsActions::GetAngleAction(
 	const ULearningAgentsActionObject* Object, 
 	const FLearningAgentsActionObjectElement Element, 
 	const float RelativeAngle, 
-	const float AngleScale, 
 	const FName Tag,
 	const bool bVisualLoggerEnabled,
 	ULearningAgentsManagerListener* VisualLoggerListener,
@@ -3046,14 +3031,13 @@ bool ULearningAgentsActions::GetAngleAction(
 	const FVector VisualLoggerLocation,
 	const FLinearColor VisualLoggerColor)
 {
-	float OutValue = 0.0f;
-	if (!GetContinuousActionToArrayView(MakeArrayView(&OutValue, 1), Object, Element, Tag))
+	float LocalAngle = 0.0f;
+	if (!GetContinuousActionToArrayView(MakeArrayView(&LocalAngle, 1), Object, Element, Tag))
 	{
 		OutAngle = 0.0f;
 		return false;
 	}
 
-	const float LocalAngle = AngleScale * OutValue;
 	OutAngle = RelativeAngle + LocalAngle;
 
 #if UE_LEARNING_AGENTS_ENABLE_VISUAL_LOG
@@ -3071,12 +3055,10 @@ bool ULearningAgentsActions::GetAngleAction(
 
 		UE_LEARNING_AGENTS_VLOG_STRING(VisualLoggerObject, LogLearning, Display, VisualLoggerLocation,
 			VisualLoggerColor.ToFColor(true),
-			TEXT("Listener: %s\nTag: %s\nAgent Id: % 3i\nEncoded: [% 6.2f]\nScale: [% 6.2f]\nLocal Angle: [% 6.1f]\nAngle: [% 6.1f]"),
+			TEXT("Listener: %s\nTag: %s\nAgent Id: % 3i\nLocal Angle: [% 6.1f]\nAngle: [% 6.1f]"),
 			*VisualLoggerListener->GetName(),
 			*Tag.ToString(),
 			VisualLoggerAgentId,
-			OutValue,
-			AngleScale,
 			LocalAngle,
 			OutAngle);
 	}
@@ -3090,7 +3072,6 @@ bool ULearningAgentsActions::GetAngleActionRadians(
 	const ULearningAgentsActionObject* Object, 
 	const FLearningAgentsActionObjectElement Element, 
 	const float RelativeAngle, 
-	const float AngleScale, 
 	const FName Tag,
 	const bool bVisualLoggerEnabled,
 	ULearningAgentsManagerListener* VisualLoggerListener,
@@ -3099,7 +3080,7 @@ bool ULearningAgentsActions::GetAngleActionRadians(
 	const FVector VisualLoggerLocation,
 	const FLinearColor VisualLoggerColor)
 {
-	if (!GetAngleAction(OutAngle, Object, Element, FMath::RadiansToDegrees(RelativeAngle), FMath::RadiansToDegrees(AngleScale), Tag, bVisualLoggerEnabled, VisualLoggerListener, VisualLoggerAgentId, VisualLoggerAngleLocation, VisualLoggerLocation, VisualLoggerColor))
+	if (!GetAngleAction(OutAngle, Object, Element, FMath::RadiansToDegrees(RelativeAngle), Tag, bVisualLoggerEnabled, VisualLoggerListener, VisualLoggerAgentId, VisualLoggerAngleLocation, VisualLoggerLocation, VisualLoggerColor))
 	{
 		OutAngle = 0.0f;
 		return false;
@@ -3114,7 +3095,6 @@ bool ULearningAgentsActions::GetVelocityAction(
 	const ULearningAgentsActionObject* Object, 
 	const FLearningAgentsActionObjectElement Element, 
 	const FTransform RelativeTransform, 
-	const float VelocityScale, 
 	const FName Tag,
 	const bool bVisualLoggerEnabled,
 	ULearningAgentsManagerListener* VisualLoggerListener,
@@ -3130,7 +3110,7 @@ bool ULearningAgentsActions::GetVelocityAction(
 		return false;
 	}
 
-	const FVector LocalVelocity = VelocityScale * FVector(OutValues[0], OutValues[1], OutValues[2]);
+	const FVector LocalVelocity = FVector(OutValues[0], OutValues[1], OutValues[2]);
 	OutVelocity = RelativeTransform.TransformVector(LocalVelocity);
 
 #if UE_LEARNING_AGENTS_ENABLE_VISUAL_LOG
@@ -3152,12 +3132,10 @@ bool ULearningAgentsActions::GetVelocityAction(
 
 		UE_LEARNING_AGENTS_VLOG_STRING(VisualLoggerObject, LogLearning, Display, VisualLoggerLocation,
 			VisualLoggerColor.ToFColor(true),
-			TEXT("Listener: %s\nTag: %s\nAgent Id: % 3i\nEncoded: [% 6.2f % 6.2f % 6.2f]\nScale: [% 6.2f]\nLocal Velocity: [% 6.1f % 6.1f % 6.1f]\nVelocity: [% 6.1f % 6.1f % 6.1f]"),
+			TEXT("Listener: %s\nTag: %s\nAgent Id: % 3i\nLocal Velocity: [% 6.1f % 6.1f % 6.1f]\nVelocity: [% 6.1f % 6.1f % 6.1f]"),
 			*VisualLoggerListener->GetName(),
 			*Tag.ToString(),
 			VisualLoggerAgentId,
-			OutValues[0], OutValues[1], OutValues[2],
-			VelocityScale,
 			LocalVelocity.X, LocalVelocity.Y, LocalVelocity.Z,
 			OutVelocity.X, OutVelocity.Y, OutVelocity.Z);
 	}
diff --git a/Engine/Plugins/Experimental/LearningAgents/Source/LearningAgents/Private/LearningAgentsObservations.cpp b/Engine/Plugins/Experimental/LearningAgents/Source/LearningAgents/Private/LearningAgentsObservations.cpp
--- a/Engine/Plugins/Experimental/LearningAgents/Source/LearningAgents/Private/LearningAgentsObservations.cpp
+++ b/Engine/Plugins/Experimental/LearningAgents/Source/LearningAgents/Private/LearningAgentsObservations.cpp
@@ -491,7 +491,7 @@ FLearningAgentsObservationSchemaElement ULearningAgentsObservations::SpecifyNull
 	return { Schema->ObservationSchema.CreateNull(Tag)};
 }
 
-FLearningAgentsObservationSchemaElement ULearningAgentsObservations::SpecifyContinuousObservation(ULearningAgentsObservationSchema* Schema, const int32 Size, const FName Tag)
+FLearningAgentsObservationSchemaElement ULearningAgentsObservations::SpecifyContinuousObservation(ULearningAgentsObservationSchema* Schema, const int32 Size, const float Scale, const FName Tag)
 {
 	if (!Schema)
 	{
@@ -505,27 +505,33 @@ FLearningAgentsObservationSchemaElement ULearningAgentsObservations::SpecifyCont
 		return FLearningAgentsObservationSchemaElement();
 	}
 	
+	if (Scale < UE_SMALL_NUMBER)
+	{
+		UE_LOG(LogLearning, Error, TEXT("SpecifyContinuousObservation: Invalid Scale for Continuous Observation '%f', must be greater than '%f'."), Scale, UE_SMALL_NUMBER);
+		return FLearningAgentsObservationSchemaElement();
+	}
+
 	if (Size == 0)
 	{
 		UE_LOG(LogLearning, Warning, TEXT("SpecifyContinuousObservation: Specifying zero-sized Continuous Observation."));
 	}
 
-	return { Schema->ObservationSchema.CreateContinuous({ Size }, Tag) };
+	return { Schema->ObservationSchema.CreateContinuous({ Size, Scale }, Tag) };
 }
 
 FLearningAgentsObservationSchemaElement ULearningAgentsObservations::SpecifyExclusiveDiscreteObservation(ULearningAgentsObservationSchema* Schema, const int32 Size, const FName Tag)
 {
-	return SpecifyContinuousObservation(Schema, Size, Tag);
+	return SpecifyContinuousObservation(Schema, Size, 1.0f, Tag);
 }
 
 FLearningAgentsObservationSchemaElement ULearningAgentsObservations::SpecifyInclusiveDiscreteObservation(ULearningAgentsObservationSchema* Schema, const int32 Size, const FName Tag)
 {
-	return SpecifyContinuousObservation(Schema, Size, Tag);
+	return SpecifyContinuousObservation(Schema, Size, 1.0f, Tag);
 }
 
 FLearningAgentsObservationSchemaElement ULearningAgentsObservations::SpecifyCountObservation(ULearningAgentsObservationSchema* Schema, const FName Tag)
 {
-	return SpecifyContinuousObservation(Schema, 1, Tag);
+	return SpecifyContinuousObservation(Schema, 1, 1.0f, Tag);
 }
 
 FLearningAgentsObservationSchemaElement ULearningAgentsObservations::SpecifyStructObservation(ULearningAgentsObservationSchema* Schema, const TMap<FName, FLearningAgentsObservationSchemaElement>& Elements, const FName Tag)
@@ -913,7 +919,7 @@ FLearningAgentsObservationSchemaElement ULearningAgentsObservations::SpecifyEnum
 		return FLearningAgentsObservationSchemaElement();
 	}
 
-	return SpecifyContinuousObservation(Schema, Enum->NumEnums() - 1, Tag);
+	return SpecifyContinuousObservation(Schema, Enum->NumEnums() - 1, 1.0f, Tag);
 }
 
 FLearningAgentsObservationSchemaElement ULearningAgentsObservations::SpecifyBitmaskObservation(ULearningAgentsObservationSchema* Schema, const UEnum* Enum, const FName Tag)
@@ -930,7 +936,7 @@ FLearningAgentsObservationSchemaElement ULearningAgentsObservations::SpecifyBitm
 		return FLearningAgentsObservationSchemaElement();
 	}
 
-	return SpecifyContinuousObservation(Schema, Enum->NumEnums() - 1, Tag);
+	return SpecifyContinuousObservation(Schema, Enum->NumEnums() - 1, 1.0f, Tag);
 }
 
 FLearningAgentsObservationSchemaElement ULearningAgentsObservations::SpecifyOptionalObservation(ULearningAgentsObservationSchema* Schema, const FLearningAgentsObservationSchemaElement Element, const int32 EncodingSize, const FName Tag)
@@ -974,30 +980,30 @@ FLearningAgentsObservationSchemaElement ULearningAgentsObservations::SpecifyEnco
 
 FLearningAgentsObservationSchemaElement ULearningAgentsObservations::SpecifyBoolObservation(ULearningAgentsObservationSchema* Schema, const FName Tag)
 {
-	return SpecifyContinuousObservation(Schema, 1, Tag);
+	return SpecifyContinuousObservation(Schema, 1, 1.0f, Tag);
 }
 
-FLearningAgentsObservationSchemaElement ULearningAgentsObservations::SpecifyFloatObservation(ULearningAgentsObservationSchema* Schema, const FName Tag)
+FLearningAgentsObservationSchemaElement ULearningAgentsObservations::SpecifyFloatObservation(ULearningAgentsObservationSchema* Schema, const float FloatScale, const FName Tag)
 {
-	return SpecifyContinuousObservation(Schema, 1, Tag);
+	return SpecifyContinuousObservation(Schema, 1, FloatScale, Tag);
 }
 
-FLearningAgentsObservationSchemaElement ULearningAgentsObservations::SpecifyLocationObservation(ULearningAgentsObservationSchema* Schema, const FName Tag)
+FLearningAgentsObservationSchemaElement ULearningAgentsObservations::SpecifyLocationObservation(ULearningAgentsObservationSchema* Schema, const float LocationScale, const FName Tag)
 {
-	return SpecifyContinuousObservation(Schema, 3, Tag);
+	return SpecifyContinuousObservation(Schema, 3, LocationScale, Tag);
 }
 
 FLearningAgentsObservationSchemaElement ULearningAgentsObservations::SpecifyRotationObservation(ULearningAgentsObservationSchema* Schema, const FName Tag)
 {
-	return SpecifyContinuousObservation(Schema, 6, Tag);
+	return SpecifyContinuousObservation(Schema, 6, 1.0f, Tag);
 }
 
 FLearningAgentsObservationSchemaElement ULearningAgentsObservations::SpecifyScaleObservation(ULearningAgentsObservationSchema* Schema, const FName Tag)
 {
-	return SpecifyContinuousObservation(Schema, 3, Tag);
+	return SpecifyContinuousObservation(Schema, 3, 1.0f, Tag);
 }
 
-FLearningAgentsObservationSchemaElement ULearningAgentsObservations::SpecifyTransformObservation(ULearningAgentsObservationSchema* Schema, const FName Tag)
+FLearningAgentsObservationSchemaElement ULearningAgentsObservations::SpecifyTransformObservation(ULearningAgentsObservationSchema* Schema, const float LocationScale, const FName Tag)
 {
 	return SpecifyStructObservationFromArrayViews(Schema,
 		{
@@ -1006,7 +1012,7 @@ FLearningAgentsObservationSchemaElement ULearningAgentsObservations::SpecifyTran
 			TEXT("Scale")
 		},
 		{
-			SpecifyLocationObservation(Schema),
+			SpecifyLocationObservation(Schema, LocationScale),
 			SpecifyRotationObservation(Schema),
 			SpecifyScaleObservation(Schema)
 		}, 
@@ -1015,22 +1021,22 @@ FLearningAgentsObservationSchemaElement ULearningAgentsObservations::SpecifyTran
 
 FLearningAgentsObservationSchemaElement ULearningAgentsObservations::SpecifyAngleObservation(ULearningAgentsObservationSchema* Schema, const FName Tag)
 {
-	return SpecifyContinuousObservation(Schema, 2, Tag);
+	return SpecifyContinuousObservation(Schema, 2, 1.0f, Tag);
 }
 
-FLearningAgentsObservationSchemaElement ULearningAgentsObservations::SpecifyVelocityObservation(ULearningAgentsObservationSchema* Schema, const FName Tag)
+FLearningAgentsObservationSchemaElement ULearningAgentsObservations::SpecifyVelocityObservation(ULearningAgentsObservationSchema* Schema, const float VelocityScale, const FName Tag)
 {
-	return SpecifyContinuousObservation(Schema, 3, Tag);
+	return SpecifyContinuousObservation(Schema, 3, VelocityScale, Tag);
 }
 
 FLearningAgentsObservationSchemaElement ULearningAgentsObservations::SpecifyDirectionObservation(ULearningAgentsObservationSchema* Schema, const FName Tag)
 {
-	return SpecifyContinuousObservation(Schema, 3, Tag);
+	return SpecifyContinuousObservation(Schema, 3, 1.0f, Tag);
 }
 
-FLearningAgentsObservationSchemaElement ULearningAgentsObservations::SpecifyLocationAlongSplineObservation(ULearningAgentsObservationSchema* Schema, const FName Tag)
+FLearningAgentsObservationSchemaElement ULearningAgentsObservations::SpecifyLocationAlongSplineObservation(ULearningAgentsObservationSchema* Schema, const float LocationScale, const FName Tag)
 {
-	return SpecifyLocationObservation(Schema, Tag);
+	return SpecifyLocationObservation(Schema, LocationScale, Tag);
 }
 
 FLearningAgentsObservationSchemaElement ULearningAgentsObservations::SpecifyProportionAlongSplineObservation(ULearningAgentsObservationSchema* Schema, const FName Tag)
@@ -1055,7 +1061,7 @@ FLearningAgentsObservationSchemaElement ULearningAgentsObservations::SpecifyDire
 
 FLearningAgentsObservationSchemaElement ULearningAgentsObservations::SpecifyProportionAlongRayObservation(ULearningAgentsObservationSchema* Schema, const FName Tag)
 {
-	return SpecifyFloatObservation(Schema, Tag);
+	return SpecifyFloatObservation(Schema, 1.0f, Tag);
 }
 
 void ULearningAgentsObservations::LogObservation(const ULearningAgentsObservationObject* Object, const FLearningAgentsObservationObjectElement Element)
@@ -1121,12 +1127,11 @@ FLearningAgentsObservationObjectElement ULearningAgentsObservations::MakeContinu
 
 		UE_LEARNING_AGENTS_VLOG_STRING(VisualLoggerObject, LogLearning, Display, VisualLoggerLocation,
 			VisualLoggerColor.ToFColor(true),
-			TEXT("Listener: %s\nTag: %s\nAgent Id: % 3i\nValues: %s\nEncoded: %s"),
+			TEXT("Listener: %s\nTag: %s\nAgent Id: % 3i\nValues: %s"),
 			*VisualLoggerListener->GetName(),
 			*Tag.ToString(),
 			VisualLoggerAgentId,
-			*UE::Learning::Array::FormatFloat(Values),
-			*UE::Learning::Array::FormatFloat(Values)); // Encoded is identical to provided values
+			*UE::Learning::Array::FormatFloat(Values));
 	}
 #endif
 
@@ -1821,7 +1826,6 @@ FLearningAgentsObservationObjectElement ULearningAgentsObservations::MakeBoolObs
 FLearningAgentsObservationObjectElement ULearningAgentsObservations::MakeFloatObservation(
 	ULearningAgentsObservationObject* Object,
 	const float Value,
-	const float FloatScale,
 	const FName Tag,
 	const bool bVisualLoggerEnabled,
 	ULearningAgentsManagerListener* VisualLoggerListener,
@@ -1829,8 +1833,6 @@ FLearningAgentsObservationObjectElement ULearningAgentsObservations::MakeFloatOb
 	const FVector VisualLoggerLocation,
 	const FLinearColor VisualLoggerColor)
 {
-	const float EncodedValue = Value / FMath::Max(FloatScale, UE_SMALL_NUMBER);
-
 #if UE_LEARNING_AGENTS_ENABLE_VISUAL_LOG
 	if (bVisualLoggerEnabled && VisualLoggerListener)
 	{
@@ -1838,24 +1840,21 @@ FLearningAgentsObservationObjectElement ULearningAgentsObservations::MakeFloatOb
 
 		UE_LEARNING_AGENTS_VLOG_STRING(VisualLoggerObject, LogLearning, Display, VisualLoggerLocation,
 			VisualLoggerColor.ToFColor(true),
-			TEXT("Listener: %s\nTag: %s\nAgent Id: % 3i\nValue: [% 6.1f]\nScale: [% 6.2f]\nEncoded: [% 6.2f]"),
+			TEXT("Listener: %s\nTag: %s\nAgent Id: % 3i\nValue: [% 6.1f]"),
 			*VisualLoggerListener->GetName(),
 			*Tag.ToString(),
 			VisualLoggerAgentId,
-			Value,
-			FloatScale,
-			EncodedValue);
+			Value);
 	}
 #endif
 
-	return MakeContinuousObservationFromArrayView(Object, { EncodedValue }, Tag);
+	return MakeContinuousObservationFromArrayView(Object, { Value }, Tag);
 }
 
 FLearningAgentsObservationObjectElement ULearningAgentsObservations::MakeLocationObservation(
 	ULearningAgentsObservationObject* Object,
 	const FVector Location,
 	const FTransform RelativeTransform,
-	const float LocationScale,
 	const FName Tag,
 	const bool bVisualLoggerEnabled,
 	ULearningAgentsManagerListener* VisualLoggerListener,
@@ -1864,10 +1863,6 @@ FLearningAgentsObservationObjectElement ULearningAgentsObservations::MakeLocatio
 	const FLinearColor VisualLoggerColor)
 {
 	const FVector LocalLocation = RelativeTransform.InverseTransformPosition(Location);
-	const FVector EncodedLocation = FVector(
-		LocalLocation.X / FMath::Max(LocationScale, UE_SMALL_NUMBER),
-		LocalLocation.Y / FMath::Max(LocationScale, UE_SMALL_NUMBER),
-		LocalLocation.Z / FMath::Max(LocationScale, UE_SMALL_NUMBER));
 
 #if UE_LEARNING_AGENTS_ENABLE_VISUAL_LOG
 	if (bVisualLoggerEnabled && VisualLoggerListener)
@@ -1894,21 +1889,19 @@ FLearningAgentsObservationObjectElement ULearningAgentsObservations::MakeLocatio
 
 		UE_LEARNING_AGENTS_VLOG_STRING(VisualLoggerObject, LogLearning, Display, VisualLoggerLocation,
 			VisualLoggerColor.ToFColor(true),
-			TEXT("Listener: %s\nTag: %s\nAgent Id: % 3i\nLocation: [% 6.1f % 6.1f % 6.1f]\nLocal Location: [% 6.1f % 6.1f % 6.1f]\nScale: [% 6.2f]\nEncoded: [% 6.2f % 6.2f % 6.2f]"),
+			TEXT("Listener: %s\nTag: %s\nAgent Id: % 3i\nLocation: [% 6.1f % 6.1f % 6.1f]\nLocal Location: [% 6.1f % 6.1f % 6.1f]"),
 			*VisualLoggerListener->GetName(),
 			*Tag.ToString(),
 			VisualLoggerAgentId,
 			Location.X, Location.Y, Location.Z,
-			LocalLocation.X, LocalLocation.Y, LocalLocation.Z,
-			LocationScale,
-			EncodedLocation.X, EncodedLocation.Y, EncodedLocation.Z);
+			LocalLocation.X, LocalLocation.Y, LocalLocation.Z);
 	}
 #endif
 
 	return MakeContinuousObservationFromArrayView(Object, {
-		(float)EncodedLocation.X,
-		(float)EncodedLocation.Y,
-		(float)EncodedLocation.Z,
+		(float)LocalLocation.X,
+		(float)LocalLocation.Y,
+		(float)LocalLocation.Z,
 		}, Tag);
 }
 
@@ -2025,7 +2018,6 @@ FLearningAgentsObservationObjectElement ULearningAgentsObservations::MakeTransfo
 	ULearningAgentsObservationObject* Object,
 	const FTransform Transform,
 	const FTransform RelativeTransform,
-	const float LocationScale,
 	const FName Tag,
 	const bool bVisualLoggerEnabled,
 	ULearningAgentsManagerListener* VisualLoggerListener,
@@ -2056,7 +2048,7 @@ FLearningAgentsObservationObjectElement ULearningAgentsObservations::MakeTransfo
 
 		UE_LEARNING_AGENTS_VLOG_STRING(VisualLoggerObject, LogLearning, Display, VisualLoggerLocation,
 			VisualLoggerColor.ToFColor(true),
-			TEXT("Listener: %s\nTag: %s\nAgent Id: % 3i\nLocation: [% 6.1f % 6.1f % 6.1f]\nLocal Location: [% 6.1f % 6.1f % 6.1f]\nRotation: [% 6.1f % 6.1f % 6.1f]\nLocal Rotation: [% 6.1f % 6.1f % 6.1f]\nScale: [% 6.1f % 6.1f % 6.1f]\nLocal Scale: [% 6.1f % 6.1f % 6.1f]\nLocation Scale: [% 6.1f]"),
+			TEXT("Listener: %s\nTag: %s\nAgent Id: % 3i\nLocation: [% 6.1f % 6.1f % 6.1f]\nLocal Location: [% 6.1f % 6.1f % 6.1f]\nRotation: [% 6.1f % 6.1f % 6.1f]\nLocal Rotation: [% 6.1f % 6.1f % 6.1f]\nScale: [% 6.1f % 6.1f % 6.1f]\nLocal Scale: [% 6.1f % 6.1f % 6.1f]"),
 			*VisualLoggerListener->GetName(),
 			*Tag.ToString(),
 			VisualLoggerAgentId,
@@ -2065,8 +2057,7 @@ FLearningAgentsObservationObjectElement ULearningAgentsObservations::MakeTransfo
 			Rotation.Roll, Rotation.Pitch, Rotation.Yaw,
 			LocalRotation.Roll, LocalRotation.Pitch, LocalRotation.Yaw,
 			Scale.X, Scale.Y, Scale.Z,
-			LocalScale.X, LocalScale.Y, LocalScale.Z,
-			LocationScale);
+			LocalScale.X, LocalScale.Y, LocalScale.Z);
 	}
 #endif
 
@@ -2077,7 +2068,7 @@ FLearningAgentsObservationObjectElement ULearningAgentsObservations::MakeTransfo
 			TEXT("Scale")
 		},
 		{
-			MakeLocationObservation(Object, LocalTransform.GetLocation(), FTransform::Identity, LocationScale),
+			MakeLocationObservation(Object, LocalTransform.GetLocation(), FTransform::Identity),
 			MakeRotationObservationFromQuat(Object, LocalTransform.GetRotation(), FQuat::Identity),
 			MakeScaleObservation(Object, LocalTransform.GetScale3D(), FVector::OneVector)
 		},
@@ -2155,7 +2146,6 @@ FLearningAgentsObservationObjectElement ULearningAgentsObservations::MakeVelocit
 	ULearningAgentsObservationObject* Object, 
 	const FVector Velocity, 
 	const FTransform RelativeTransform, 
-	const float VelocityScale, 
 	const FName Tag,
 	const bool bVisualLoggerEnabled,
 	ULearningAgentsManagerListener* VisualLoggerListener,
@@ -2165,10 +2155,6 @@ FLearningAgentsObservationObjectElement ULearningAgentsObservations::MakeVelocit
 	const FLinearColor VisualLoggerColor)
 {
 	const FVector LocalVelocity = RelativeTransform.InverseTransformVectorNoScale(Velocity);
-	const FVector EncodedVelocity = FVector(
-		LocalVelocity.X / FMath::Max(VelocityScale, UE_SMALL_NUMBER),
-		LocalVelocity.Y / FMath::Max(VelocityScale, UE_SMALL_NUMBER),
-		LocalVelocity.Z / FMath::Max(VelocityScale, UE_SMALL_NUMBER));
 
 #if UE_LEARNING_AGENTS_ENABLE_VISUAL_LOG
 	if (bVisualLoggerEnabled && VisualLoggerListener)
@@ -2189,17 +2175,16 @@ FLearningAgentsObservationObjectElement ULearningAgentsObservations::MakeVelocit
 
 		UE_LEARNING_AGENTS_VLOG_STRING(VisualLoggerObject, LogLearning, Display, VisualLoggerLocation,
 			VisualLoggerColor.ToFColor(true),
-			TEXT("Listener: %s\nTag: %s\nAgent Id: % 3i\nVelocity: [% 6.1f % 6.1f % 6.1f]\nLocal Velocity: [% 6.1f % 6.1f % 6.1f]\nEncoded: [% 6.2f % 6.2f % 6.2f]"),
+			TEXT("Listener: %s\nTag: %s\nAgent Id: % 3i\nVelocity: [% 6.1f % 6.1f % 6.1f]\nLocal Velocity: [% 6.1f % 6.1f % 6.1f]"),
 			*VisualLoggerListener->GetName(),
 			*Tag.ToString(),
 			VisualLoggerAgentId,
 			Velocity.X, Velocity.Y, Velocity.Z,
-			LocalVelocity.X, LocalVelocity.Y, LocalVelocity.Z,
-			EncodedVelocity.X, EncodedVelocity.Y, EncodedVelocity.Z);
+			LocalVelocity.X, LocalVelocity.Y, LocalVelocity.Z);
 	}
 #endif
 
-	return MakeContinuousObservationFromArrayView(Object, { (float)EncodedVelocity.X, (float)EncodedVelocity.Y, (float)EncodedVelocity.Z }, Tag);
+	return MakeContinuousObservationFromArrayView(Object, { (float)LocalVelocity.X, (float)LocalVelocity.Y, (float)LocalVelocity.Z }, Tag);
 }
 
 FLearningAgentsObservationObjectElement ULearningAgentsObservations::MakeDirectionObservation(
@@ -2236,13 +2221,12 @@ FLearningAgentsObservationObjectElement ULearningAgentsObservations::MakeDirecti
 
 		UE_LEARNING_AGENTS_VLOG_STRING(VisualLoggerObject, LogLearning, Display, VisualLoggerLocation,
 			VisualLoggerColor.ToFColor(true),
-			TEXT("Listener: %s\nTag: %s\nAgent Id: % 3i\nDirection: [% 6.1f % 6.1f % 6.1f]\nLocal Direction: [% 6.1f % 6.1f % 6.1f]\nEncoded: [% 6.2f % 6.2f % 6.2f]"),
+			TEXT("Listener: %s\nTag: %s\nAgent Id: % 3i\nDirection: [% 6.1f % 6.1f % 6.1f]\nLocal Direction: [% 6.1f % 6.1f % 6.1f]"),
 			*VisualLoggerListener->GetName(),
 			*Tag.ToString(),
 			VisualLoggerAgentId,
 			Direction.X, Direction.Y, Direction.Z,
-			LocalDirection.X, LocalDirection.Y, LocalDirection.Z,
-			LocalDirection.X, LocalDirection.Y, LocalDirection.Z); // We don't have a scaling so the encoded value is the same
+			LocalDirection.X, LocalDirection.Y, LocalDirection.Z);
 	}
 #endif
 
@@ -2258,7 +2242,6 @@ FLearningAgentsObservationObjectElement ULearningAgentsObservations::MakeLocatio
 	const USplineComponent* SplineComponent,
 	const float DistanceAlongSpline,
 	const FTransform RelativeTransform,
-	const float LocationScale,
 	const FName Tag,
 	const bool bVisualLoggerEnabled,
 	ULearningAgentsManagerListener* VisualLoggerListener,
@@ -2279,7 +2262,6 @@ FLearningAgentsObservationObjectElement ULearningAgentsObservations::MakeLocatio
 		Object, 
 		SplineComponent->GetLocationAtDistanceAlongSpline(LoopedDistance, ESplineCoordinateSpace::World), 
 		RelativeTransform, 
-		LocationScale, 
 		Tag, 
 		bVisualLoggerEnabled, 
 		VisualLoggerListener, 
@@ -2330,7 +2312,6 @@ FLearningAgentsObservationObjectElement ULearningAgentsObservations::MakeProport
 		return MakeExclusiveUnionObservation(Object, TEXT("Proportion"), MakeFloatObservation(
 			Object, 
 			Proportion, 
-			1.0f, 
 			TEXT("FloatObservation"),
 			bVisualLoggerEnabled, 
 			VisualLoggerListener, 
@@ -2452,7 +2433,7 @@ FLearningAgentsObservationObjectElement ULearningAgentsObservations::MakeProport
 	}
 #endif
 
-	return MakeFloatObservation(Object, Encoded, 1.0f, Tag);
+	return MakeFloatObservation(Object, Encoded, Tag);
 }
 
 
@@ -3696,7 +3677,7 @@ bool ULearningAgentsObservations::GetEncodingObservation(FLearningAgentsObservat
 bool ULearningAgentsObservations::GetBoolObservation(bool& bOutValue, const ULearningAgentsObservationObject* Object, const FLearningAgentsObservationObjectElement Element, const FName Tag)
 {
 	float OutValue = 0.0f;
-	if (!GetFloatObservation(OutValue, Object, Element, 1.0f, Tag))
+	if (!GetFloatObservation(OutValue, Object, Element, Tag))
 	{
 		bOutValue = false;
 		return false;
@@ -3706,20 +3687,18 @@ bool ULearningAgentsObservations::GetBoolObservation(bool& bOutValue, const ULea
 	return true;
 }
 
-bool ULearningAgentsObservations::GetFloatObservation(float& OutValue, const ULearningAgentsObservationObject* Object, const FLearningAgentsObservationObjectElement Element, const float FloatScale, const FName Tag)
+bool ULearningAgentsObservations::GetFloatObservation(float& OutValue, const ULearningAgentsObservationObject* Object, const FLearningAgentsObservationObjectElement Element, const FName Tag)
 {
-	float OutValuesData;
-	if (!GetContinuousObservationToArrayView(MakeArrayView(&OutValuesData, 1), Object, Element, Tag))
+	if (!GetContinuousObservationToArrayView(MakeArrayView(&OutValue, 1), Object, Element, Tag))
 	{
 		OutValue = 0.0f;
 		return false;
 	}
 
-	OutValue = OutValuesData * FloatScale;
 	return true;
 }
 
-bool ULearningAgentsObservations::GetLocationObservation(FVector& OutLocation, const ULearningAgentsObservationObject* Object, const FLearningAgentsObservationObjectElement Element, const FTransform RelativeTransform, const float LocationScale, const FName Tag)
+bool ULearningAgentsObservations::GetLocationObservation(FVector& OutLocation, const ULearningAgentsObservationObject* Object, const FLearningAgentsObservationObjectElement Element, const FTransform RelativeTransform, const FName Tag)
 {
 	TStaticArray<float, 3> OutValues;
 	if (!GetContinuousObservationToArrayView(OutValues, Object, Element, Tag))
@@ -3728,7 +3707,7 @@ bool ULearningAgentsObservations::GetLocationObservation(FVector& OutLocation, c
 		return false;
 	}
 
-	OutLocation = RelativeTransform.TransformPosition(LocationScale * FVector(OutValues[0], OutValues[1], OutValues[2]));
+	OutLocation = RelativeTransform.TransformPosition(FVector(OutValues[0], OutValues[1], OutValues[2]));
 	return true;
 }
 
@@ -3782,7 +3761,7 @@ bool ULearningAgentsObservations::GetScaleObservation(FVector& OutScale, const U
 	return true;
 }
 
-bool ULearningAgentsObservations::GetTransformObservation(FTransform& OutTransform, const ULearningAgentsObservationObject* Object, const FLearningAgentsObservationObjectElement Element, const FTransform RelativeTransform, const float LocationScale, const FName Tag)
+bool ULearningAgentsObservations::GetTransformObservation(FTransform& OutTransform, const ULearningAgentsObservationObject* Object, const FLearningAgentsObservationObjectElement Element, const FTransform RelativeTransform, const FName Tag)
 {
 	TStaticArray<FName, 3> OutElementNames;
 	TStaticArray<FLearningAgentsObservationObjectElement, 3> OutElements;
@@ -3794,7 +3773,7 @@ bool ULearningAgentsObservations::GetTransformObservation(FTransform& OutTransfo
 
 	const int32 LocationElement = MakeArrayView(OutElementNames).Find(TEXT("Location"));
 	FVector OutLocation;
-	if (LocationElement == INDEX_NONE || !GetLocationObservation(OutLocation, Object, OutElements[LocationElement], RelativeTransform, LocationScale))
+	if (LocationElement == INDEX_NONE || !GetLocationObservation(OutLocation, Object, OutElements[LocationElement], RelativeTransform))
 	{
 		OutTransform = FTransform::Identity;
 		return false;
@@ -3845,7 +3824,7 @@ bool ULearningAgentsObservations::GetAngleObservation(float& OutAngle, const ULe
 	return true;
 }
 
-bool ULearningAgentsObservations::GetVelocityObservation(FVector& OutVelocity, const ULearningAgentsObservationObject* Object, const FLearningAgentsObservationObjectElement Element, const FTransform RelativeTransform, const float VelocityScale, const FName Tag)
+bool ULearningAgentsObservations::GetVelocityObservation(FVector& OutVelocity, const ULearningAgentsObservationObject* Object, const FLearningAgentsObservationObjectElement Element, const FTransform RelativeTransform, const FName Tag)
 {
 	TStaticArray<float, 3> OutValues;
 	if (!GetContinuousObservationToArrayView(OutValues, Object, Element, Tag))
@@ -3854,7 +3833,7 @@ bool ULearningAgentsObservations::GetVelocityObservation(FVector& OutVelocity, c
 		return false;
 	}
 
-	OutVelocity = RelativeTransform.TransformVectorNoScale(VelocityScale * FVector(OutValues[0], OutValues[1], OutValues[2]));
+	OutVelocity = RelativeTransform.TransformVectorNoScale(FVector(OutValues[0], OutValues[1], OutValues[2]));
 	return true;
 }
 
@@ -3871,9 +3850,9 @@ bool ULearningAgentsObservations::GetDirectionObservation(FVector& OutDirection,
 	return true;
 }
 
-bool ULearningAgentsObservations::GetLocationAlongSplineObservation(FVector& OutLocation, const ULearningAgentsObservationObject* Object, const FLearningAgentsObservationObjectElement Element, const FTransform RelativeTransform, const float LocationScale, const FName Tag)
+bool ULearningAgentsObservations::GetLocationAlongSplineObservation(FVector& OutLocation, const ULearningAgentsObservationObject* Object, const FLearningAgentsObservationObjectElement Element, const FTransform RelativeTransform, const FName Tag)
 {
-	return GetLocationObservation(OutLocation, Object, Element, RelativeTransform, LocationScale, Tag);
+	return GetLocationObservation(OutLocation, Object, Element, RelativeTransform, Tag);
 }
 
 bool ULearningAgentsObservations::GetProportionAlongSplineObservation(bool& bOutIsClosedLoop, float& OutAngle, float& OutPropotion, const ULearningAgentsObservationObject* Object, const FLearningAgentsObservationObjectElement Element, const FName Tag)
@@ -3909,7 +3888,7 @@ bool ULearningAgentsObservations::GetDirectionAlongSplineObservation(FVector& Ou
 
 bool ULearningAgentsObservations::GetProportionAlongRayObservation(float& OutProportion, const ULearningAgentsObservationObject* Object, const FLearningAgentsObservationObjectElement Element, const FName Tag)
 {
-	if (!GetFloatObservation(OutProportion, Object, Element, 1.0f, Tag))
+	if (!GetFloatObservation(OutProportion, Object, Element, Tag))
 	{
 		OutProportion = 0.0f;
 		return false;
diff --git a/Engine/Plugins/Experimental/LearningAgents/Source/LearningAgents/Public/LearningAgentsActions.h b/Engine/Plugins/Experimental/LearningAgents/Source/LearningAgents/Public/LearningAgentsActions.h
--- a/Engine/Plugins/Experimental/LearningAgents/Source/LearningAgents/Public/LearningAgentsActions.h
+++ b/Engine/Plugins/Experimental/LearningAgents/Source/LearningAgents/Public/LearningAgentsActions.h
@@ -146,11 +146,12 @@ public:
 	 *
 	 * @param Schema The Action Schema
 	 * @param Size The number of float values in the action.
+	 * @param Scale The scale used for this action.
 	 * @param Tag The tag of this new action. Used during action object validation and debugging.
 	 * @return The newly created action schema element.
 	 */
-	UFUNCTION(BlueprintPure, Category = "LearningAgents", meta = (AdvancedDisplay = 2))
-	static FLearningAgentsActionSchemaElement SpecifyContinuousAction(ULearningAgentsActionSchema* Schema, const int32 Size, const FName Tag = TEXT("ContinuousAction"));
+	UFUNCTION(BlueprintPure, Category = "LearningAgents", meta = (AdvancedDisplay = 3))
+	static FLearningAgentsActionSchemaElement SpecifyContinuousAction(ULearningAgentsActionSchema* Schema, const int32 Size, const float Scale = 1.0f, const FName Tag = TEXT("ContinuousAction"));
 
 	/**
 	 * Specifies a new exclusive discrete action. This represents an action which is an exclusive choice from a number of discrete options, sampled 
@@ -468,71 +469,85 @@ public:
 	 * catch-all for situations where a type-specific action does not exist.
 	 *
 	 * @param Schema The Action Schema
+	 * @param FloatScale The scale used for this action.
 	 * @param Tag The tag of this new action. Used during action object validation and debugging.
 	 * @return The newly created action schema element.
 	 */
-	UFUNCTION(BlueprintPure, Category = "LearningAgents", meta = (AdvancedDisplay = 1))
-	static FLearningAgentsActionSchemaElement SpecifyFloatAction(ULearningAgentsActionSchema* Schema, const FName Tag = TEXT("FloatAction"));
+	UFUNCTION(BlueprintPure, Category = "LearningAgents", meta = (AdvancedDisplay = 2))
+	static FLearningAgentsActionSchemaElement SpecifyFloatAction(ULearningAgentsActionSchema* Schema, const float FloatScale = 1.0f, const FName Tag = TEXT("FloatAction"));
 
 	/**
 	 * Specifies a new location action. This represents an action which is a location sampled from a Gaussian distribution.
 	 *
 	 * @param Schema The Action Schema
+	 * @param LocationScale The scale used for this action in cm.
 	 * @param Tag The tag of this new action. Used during action object validation and debugging.
 	 * @return The newly created action schema element.
 	 */
-	UFUNCTION(BlueprintPure, Category = "LearningAgents", meta = (AdvancedDisplay = 1))
-	static FLearningAgentsActionSchemaElement SpecifyLocationAction(ULearningAgentsActionSchema* Schema, const FName Tag = TEXT("LocationAction"));
+	UFUNCTION(BlueprintPure, Category = "LearningAgents", meta = (AdvancedDisplay = 2))
+	static FLearningAgentsActionSchemaElement SpecifyLocationAction(ULearningAgentsActionSchema* Schema, const float LocationScale = 100.0f, const FName Tag = TEXT("LocationAction"));
 
 	/**
 	 * Specifies a new rotation action. This represents an action which is a rotation sampled from a Gaussian distribution in the angle-axis space.
 	 *
 	 * @param Schema The Action Schema
+	 * @param RotationScale The scale used for this action in degrees.
 	 * @param Tag The tag of this new action. Used during action object validation and debugging.
 	 * @return The newly created action schema element.
 	 */
-	UFUNCTION(BlueprintPure, Category = "LearningAgents", meta = (AdvancedDisplay = 1))
-	static FLearningAgentsActionSchemaElement SpecifyRotationAction(ULearningAgentsActionSchema* Schema, const FName Tag = TEXT("RotationAction"));
+	UFUNCTION(BlueprintPure, Category = "LearningAgents", meta = (AdvancedDisplay = 2))
+	static FLearningAgentsActionSchemaElement SpecifyRotationAction(ULearningAgentsActionSchema* Schema, const float RotationScale = 90.0f, const FName Tag = TEXT("RotationAction"));
 
 	/**
 	 * Specifies a new scale action. This represents an action which is a scale sampled from a Gaussian distribution in the log space.
 	 *
 	 * @param Schema The Action Schema
+	 * @param ScaleScale The scale used for this action.
 	 * @param Tag The tag of this new action. Used during action object validation and debugging.
 	 * @return The newly created action schema element.
 	 */
-	UFUNCTION(BlueprintPure, Category = "LearningAgents", meta = (AdvancedDisplay = 1))
-	static FLearningAgentsActionSchemaElement SpecifyScaleAction(ULearningAgentsActionSchema* Schema, const FName Tag = TEXT("ScaleAction"));
+	UFUNCTION(BlueprintPure, Category = "LearningAgents", meta = (AdvancedDisplay = 2))
+	static FLearningAgentsActionSchemaElement SpecifyScaleAction(ULearningAgentsActionSchema* Schema, const float ScaleScale = 1.0f, const FName Tag = TEXT("ScaleAction"));
 
 	/**
 	 * Specifies a new transform action.
 	 *
 	 * @param Schema The Action Schema
+	 * @param LocationScale The scale used for the Location part of the transform in this action in cm.
+	 * @param RotationScale The scale used for the Rotation part of the transform in this action in degrees.
+	 * @param ScaleScale The scale used for the Scale part of the transform in this action.
 	 * @param Tag The tag of this new action. Used during action object validation and debugging.
 	 * @return The newly created action schema element.
 	 */
-	UFUNCTION(BlueprintPure, Category = "LearningAgents", meta = (AdvancedDisplay = 1))
-	static FLearningAgentsActionSchemaElement SpecifyTransformAction(ULearningAgentsActionSchema* Schema, const FName Tag = TEXT("TransformAction"));
+	UFUNCTION(BlueprintPure, Category = "LearningAgents", meta = (AdvancedDisplay = 4))
+	static FLearningAgentsActionSchemaElement SpecifyTransformAction(
+		ULearningAgentsActionSchema* Schema, 
+		const float LocationScale = 100.0f, 
+		const float RotationScale = 90.0f, 
+		const float ScaleScale = 1.0f, 
+		const FName Tag = TEXT("TransformAction"));
 
 	/**
 	 * Specifies a new angle action. This represents an action which is an angle sampled from a Gaussian distribution centered around zero.
 	 *
 	 * @param Schema The Action Schema
+	 * @param AngleScale The scale used for this action in degrees.
 	 * @param Tag The tag of this new action. Used during action object validation and debugging.
 	 * @return The newly created action schema element.
 	 */
-	UFUNCTION(BlueprintPure, Category = "LearningAgents", meta = (AdvancedDisplay = 1))
-	static FLearningAgentsActionSchemaElement SpecifyAngleAction(ULearningAgentsActionSchema* Schema, const FName Tag = TEXT("AngleAction"));
+	UFUNCTION(BlueprintPure, Category = "LearningAgents", meta = (AdvancedDisplay = 2))
+	static FLearningAgentsActionSchemaElement SpecifyAngleAction(ULearningAgentsActionSchema* Schema, const float AngleScale = 90.0f, const FName Tag = TEXT("AngleAction"));
 
 	/**
 	 * Specifies a new velocity action. This represents an action which is a velocity sampled from a Gaussian distribution.
 	 *
 	 * @param Schema The Action Schema
+	 * @param VelocityScale The scale used for this action in cm/s.
 	 * @param Tag The tag of this new action. Used during action object validation and debugging.
 	 * @return The newly created action schema element.
 	 */
-	UFUNCTION(BlueprintPure, Category = "LearningAgents", meta = (AdvancedDisplay = 1))
-	static FLearningAgentsActionSchemaElement SpecifyVelocityAction(ULearningAgentsActionSchema* Schema, const FName Tag = TEXT("VelocityAction"));
+	UFUNCTION(BlueprintPure, Category = "LearningAgents", meta = (AdvancedDisplay = 2))
+	static FLearningAgentsActionSchemaElement SpecifyVelocityAction(ULearningAgentsActionSchema* Schema, const float VelocityScale = 200.0f, const FName Tag = TEXT("VelocityAction"));
 
 	/**
 	 * Specifies a new direction action. This represents an action which is a direction sampled from a Gaussian distribution and normalized.
@@ -614,32 +629,32 @@ public:
 	UFUNCTION(BlueprintPure, Category = "LearningAgents", meta = (AdvancedDisplay = 2))
 	static FLearningAgentsActionObjectElement MakeBoolAction(ULearningAgentsActionObject* Object, const bool bValue, const FName Tag = TEXT("BoolAction"));
 
-	UFUNCTION(BlueprintPure, Category = "LearningAgents", meta = (AdvancedDisplay = 3))
-	static FLearningAgentsActionObjectElement MakeFloatAction(ULearningAgentsActionObject* Object, const float Value, const float FloatScale = 1.0f, const FName Tag = TEXT("FloatAction"));
+	UFUNCTION(BlueprintPure, Category = "LearningAgents", meta = (AdvancedDisplay = 2))
+	static FLearningAgentsActionObjectElement MakeFloatAction(ULearningAgentsActionObject* Object, const float Value, const FName Tag = TEXT("FloatAction"));
 
-	UFUNCTION(BlueprintPure, Category = "LearningAgents", meta = (AdvancedDisplay = 4))
-	static FLearningAgentsActionObjectElement MakeLocationAction(ULearningAgentsActionObject* Object, const FVector Location, const FTransform RelativeTransform = FTransform(), const float LocationScale = 100.0f, const FName Tag = TEXT("LocationAction"));
+	UFUNCTION(BlueprintPure, Category = "LearningAgents", meta = (AdvancedDisplay = 3))
+	static FLearningAgentsActionObjectElement MakeLocationAction(ULearningAgentsActionObject* Object, const FVector Location, const FTransform RelativeTransform = FTransform(), const FName Tag = TEXT("LocationAction"));
 
-	UFUNCTION(BlueprintPure, Category = "LearningAgents", meta = (AdvancedDisplay = 4))
-	static FLearningAgentsActionObjectElement MakeRotationAction(ULearningAgentsActionObject* Object, const FRotator Rotation, const FRotator RelativeRotation = FRotator::ZeroRotator, const float RotationScale = 90.0f, const FName Tag = TEXT("RotationAction"));
+	UFUNCTION(BlueprintPure, Category = "LearningAgents", meta = (AdvancedDisplay = 3))
+	static FLearningAgentsActionObjectElement MakeRotationAction(ULearningAgentsActionObject* Object, const FRotator Rotation, const FRotator RelativeRotation = FRotator::ZeroRotator, const FName Tag = TEXT("RotationAction"));
 
-	UFUNCTION(BlueprintPure, Category = "LearningAgents", meta = (AdvancedDisplay = 4))
-	static FLearningAgentsActionObjectElement MakeRotationActionFromQuat(ULearningAgentsActionObject* Object, const FQuat Rotation, const FQuat RelativeRotation, const float RotationScale = 90.0f, const FName Tag = TEXT("RotationAction"));
+	UFUNCTION(BlueprintPure, Category = "LearningAgents", meta = (AdvancedDisplay = 3))
+	static FLearningAgentsActionObjectElement MakeRotationActionFromQuat(ULearningAgentsActionObject* Object, const FQuat Rotation, const FQuat RelativeRotation, const FName Tag = TEXT("RotationAction"));
 
 	UFUNCTION(BlueprintPure, Category = "LearningAgents", meta = (AdvancedDisplay = 3))
 	static FLearningAgentsActionObjectElement MakeScaleAction(ULearningAgentsActionObject* Object, const FVector Scale, const FVector RelativeScale = FVector(1, 1, 1), const FName Tag = TEXT("ScaleAction"));
 
-	UFUNCTION(BlueprintPure, Category = "LearningAgents", meta = (AdvancedDisplay = 4))
-	static FLearningAgentsActionObjectElement MakeTransformAction(ULearningAgentsActionObject* Object, const FTransform Transform, const FTransform RelativeTransform = FTransform(), const float LocationScale = 100.0f, const FName Tag = TEXT("TransformAction"));
+	UFUNCTION(BlueprintPure, Category = "LearningAgents", meta = (AdvancedDisplay = 3))
+	static FLearningAgentsActionObjectElement MakeTransformAction(ULearningAgentsActionObject* Object, const FTransform Transform, const FTransform RelativeTransform = FTransform(), const FName Tag = TEXT("TransformAction"));
 
-	UFUNCTION(BlueprintPure, Category = "LearningAgents", meta = (AdvancedDisplay = 4))
-	static FLearningAgentsActionObjectElement MakeAngleAction(ULearningAgentsActionObject* Object, const float Angle, const float RelativeAngle = 0.0f, const float AngleScale = 90.0f, const FName Tag = TEXT("AngleAction"));
+	UFUNCTION(BlueprintPure, Category = "LearningAgents", meta = (AdvancedDisplay = 3))
+	static FLearningAgentsActionObjectElement MakeAngleAction(ULearningAgentsActionObject* Object, const float Angle, const float RelativeAngle = 0.0f, const FName Tag = TEXT("AngleAction"));
 
-	UFUNCTION(BlueprintPure, Category = "LearningAgents", meta = (AdvancedDisplay = 4))
-	static FLearningAgentsActionObjectElement MakeAngleActionRadians(ULearningAgentsActionObject* Object, const float Angle, const float RelativeAngle = 0.0f, const float AngleScale = 1.57079632679f, const FName Tag = TEXT("AngleAction"));
+	UFUNCTION(BlueprintPure, Category = "LearningAgents", meta = (AdvancedDisplay = 3))
+	static FLearningAgentsActionObjectElement MakeAngleActionRadians(ULearningAgentsActionObject* Object, const float Angle, const float RelativeAngle = 0.0f, const FName Tag = TEXT("AngleAction"));
 
-	UFUNCTION(BlueprintPure, Category = "LearningAgents", meta = (AdvancedDisplay = 4))
-	static FLearningAgentsActionObjectElement MakeVelocityAction(ULearningAgentsActionObject* Object, const FVector Velocity, const FTransform RelativeTransform = FTransform(), const float VelocityScale = 200.0f, const FName Tag = TEXT("VelocityAction"));
+	UFUNCTION(BlueprintPure, Category = "LearningAgents", meta = (AdvancedDisplay = 3))
+	static FLearningAgentsActionObjectElement MakeVelocityAction(ULearningAgentsActionObject* Object, const FVector Velocity, const FTransform RelativeTransform = FTransform(), const FName Tag = TEXT("VelocityAction"));
 
 	UFUNCTION(BlueprintPure, Category = "LearningAgents", meta = (AdvancedDisplay = 3))
 	static FLearningAgentsActionObjectElement MakeDirectionAction(ULearningAgentsActionObject* Object, const FVector Direction, const FTransform RelativeTransform = FTransform(), const FName Tag = TEXT("DirectionAction"));
@@ -1093,7 +1108,6 @@ public:
 	 * @param OutValue The output float value.
 	 * @param Object The Action Object
 	 * @param Element The Action Object Element
-	 * @param FloatScale The scale used to control the overall magnitude of the outputted float action.
 	 * @param Tag The tag of the corresponding action. Must match the tag given during Specify.
 	 * @param bVisualLoggerEnabled When true, debug data will be sent to the visual logger.
 	 * @param VisualLoggerListener The listener object which is making this action. This must be set to use logging.
@@ -1102,12 +1116,11 @@ public:
 	 * @param VisualLoggerColor The color for the visual logger display.
 	 * @return true if the provided Element is the correct type, otherwise false.
 	 */
-	UFUNCTION(BlueprintPure = false, Category = "LearningAgents", meta = (AdvancedDisplay = 4, ReturnDisplayName = "Success", DefaultToSelf = "VisualLoggerListener"))
+	UFUNCTION(BlueprintPure = false, Category = "LearningAgents", meta = (AdvancedDisplay = 3, ReturnDisplayName = "Success", DefaultToSelf = "VisualLoggerListener"))
 	static bool GetFloatAction(
 		float& OutValue, 
 		const ULearningAgentsActionObject* Object, 
 		const FLearningAgentsActionObjectElement Element, 
-		const float FloatScale = 1.0f, 
 		const FName Tag = TEXT("FloatAction"),
 		const bool bVisualLoggerEnabled = false,
 		ULearningAgentsManagerListener* VisualLoggerListener = nullptr,
@@ -1122,7 +1135,6 @@ public:
 	 * @param Object The Action Object
 	 * @param Element The Action Object Element
 	 * @param RelativeTransform The relative transform to transform the location by.
-	 * @param LocationScale The scale used to control the overall magnitude of the outputted location action.
 	 * @param Tag The tag of the corresponding action. Must match the tag given during Specify.
 	 * @param bVisualLoggerEnabled When true, debug data will be sent to the visual logger.
 	 * @param VisualLoggerListener The listener object which is making this action. This must be set to use logging.
@@ -1131,13 +1143,12 @@ public:
 	 * @param VisualLoggerColor The color for the visual logger display.
 	 * @return true if the provided Element is the correct type, otherwise false.
 	 */
-	UFUNCTION(BlueprintPure = false, Category = "LearningAgents", meta = (AdvancedDisplay = 5, ReturnDisplayName = "Success", DefaultToSelf = "VisualLoggerListener"))
+	UFUNCTION(BlueprintPure = false, Category = "LearningAgents", meta = (AdvancedDisplay = 4, ReturnDisplayName = "Success", DefaultToSelf = "VisualLoggerListener"))
 	static bool GetLocationAction(
 		FVector& OutLocation, 
 		const ULearningAgentsActionObject* Object, 
 		const FLearningAgentsActionObjectElement Element, 
 		const FTransform RelativeTransform = FTransform(), 
-		const float LocationScale = 100.0f, 
 		const FName Tag = TEXT("LocationAction"),
 		const bool bVisualLoggerEnabled = false,
 		ULearningAgentsManagerListener* VisualLoggerListener = nullptr,
@@ -1152,7 +1163,6 @@ public:
 	 * @param Object The Action Object
 	 * @param Element The Action Object Element
 	 * @param RelativeRotation The relative rotation to transform the rotation by.
-	 * @param RotationScale The scale used to control the overall magnitude of the outputted rotation action.
 	 * @param Tag The tag of the corresponding action. Must match the tag given during Specify.
 	 * @param bVisualLoggerEnabled When true, debug data will be sent to the visual logger.
 	 * @param VisualLoggerListener The listener object which is making this action. This must be set to use logging.
@@ -1162,13 +1172,12 @@ public:
 	 * @param VisualLoggerColor The color for the visual logger display.
 	 * @return true if the provided Element is the correct type, otherwise false.
 	 */
-	UFUNCTION(BlueprintPure = false, Category = "LearningAgents", meta = (AdvancedDisplay = 5, ReturnDisplayName = "Success", DefaultToSelf = "VisualLoggerListener"))
+	UFUNCTION(BlueprintPure = false, Category = "LearningAgents", meta = (AdvancedDisplay = 4, ReturnDisplayName = "Success", DefaultToSelf = "VisualLoggerListener"))
 	static bool GetRotationAction(
 		FRotator& OutRotation, 
 		const ULearningAgentsActionObject* Object, 
 		const FLearningAgentsActionObjectElement Element, 
 		const FRotator RelativeRotation = FRotator::ZeroRotator, 
-		const float RotationScale = 90.0f, 
 		const FName Tag = TEXT("RotationAction"),
 		const bool bVisualLoggerEnabled = false,
 		ULearningAgentsManagerListener* VisualLoggerListener = nullptr,
@@ -1184,7 +1193,6 @@ public:
 	 * @param Object The Action Object
 	 * @param Element The Action Object Element
 	 * @param RelativeRotation The relative rotation to transform the rotation by.
-	 * @param RotationScale The scale used to control the overall magnitude of the outputted rotation action.
 	 * @param Tag The tag of the corresponding action. Must match the tag given during Specify.
 	 * @param bVisualLoggerEnabled When true, debug data will be sent to the visual logger.
 	 * @param VisualLoggerListener The listener object which is making this action. This must be set to use logging.
@@ -1194,13 +1202,12 @@ public:
 	 * @param VisualLoggerColor The color for the visual logger display.
 	 * @return true if the provided Element is the correct type, otherwise false.
 	 */
-	UFUNCTION(BlueprintPure = false, Category = "LearningAgents", meta = (AdvancedDisplay = 5, ReturnDisplayName = "Success", DefaultToSelf = "VisualLoggerListener"))
+	UFUNCTION(BlueprintPure = false, Category = "LearningAgents", meta = (AdvancedDisplay = 4, ReturnDisplayName = "Success", DefaultToSelf = "VisualLoggerListener"))
 	static bool GetRotationActionAsQuat(
 		FQuat& OutRotation, 
 		const ULearningAgentsActionObject* Object, 
 		const FLearningAgentsActionObjectElement Element, 
 		const FQuat RelativeRotation, 
-		const float RotationScale = 90.0f, 
 		const FName Tag = TEXT("RotationAction"),
 		const bool bVisualLoggerEnabled = false,
 		ULearningAgentsManagerListener* VisualLoggerListener = nullptr,
@@ -1216,7 +1223,6 @@ public:
 	 * @param Object The Action Object
 	 * @param Element The Action Object Element
 	 * @param RelativeScale The relative scale to transform the scale by.
-	 * @param Scale The scale used to control the overall magnitude of the outputted scale action.
 	 * @param Tag The tag of the corresponding action. Must match the tag given during Specify.
 	 * @param bVisualLoggerEnabled When true, debug data will be sent to the visual logger.
 	 * @param VisualLoggerListener The listener object which is making this action. This must be set to use logging.
@@ -1225,13 +1231,12 @@ public:
 	 * @param VisualLoggerColor The color for the visual logger display.
 	 * @return true if the provided Element is the correct type, otherwise false.
 	 */
-	UFUNCTION(BlueprintPure = false, Category = "LearningAgents", meta = (AdvancedDisplay = 5, ReturnDisplayName = "Success", DefaultToSelf = "VisualLoggerListener"))
+	UFUNCTION(BlueprintPure = false, Category = "LearningAgents", meta = (AdvancedDisplay = 4, ReturnDisplayName = "Success", DefaultToSelf = "VisualLoggerListener"))
 	static bool GetScaleAction(
 		FVector& OutScale, 
 		const ULearningAgentsActionObject* Object, 
 		const FLearningAgentsActionObjectElement Element, 
 		const FVector RelativeScale = FVector(1, 1, 1), 
-		const float Scale = 1.0f, 
 		const FName Tag = TEXT("ScaleAction"),
 		const bool bVisualLoggerEnabled = false,
 		ULearningAgentsManagerListener* VisualLoggerListener = nullptr,
@@ -1246,9 +1251,6 @@ public:
 	 * @param Object The Action Object
 	 * @param Element The Action Object Element
 	 * @param RelativeTransform The relative transform.
-	 * @param LocationScale The scale used to control the overall magnitude of the outputted location action.
-	 * @param RotationScale The scale used to control the overall magnitude of the outputted rotation action.
-	 * @param ScaleScale The scale used to control the overall magnitude of the outputted scale action.
 	 * @param Tag The tag of the corresponding action. Must match the tag given during Specify.
 	 * @param bVisualLoggerEnabled When true, debug data will be sent to the visual logger.
 	 * @param VisualLoggerListener The listener object which is making this action. This must be set to use logging.
@@ -1257,15 +1259,12 @@ public:
 	 * @param VisualLoggerColor The color for the visual logger display.
 	 * @return true if the provided Element is the correct type, otherwise false.
 	 */
-	UFUNCTION(BlueprintPure = false, Category = "LearningAgents", meta = (AdvancedDisplay = 7, ReturnDisplayName = "Success", DefaultToSelf = "VisualLoggerListener"))
+	UFUNCTION(BlueprintPure = false, Category = "LearningAgents", meta = (AdvancedDisplay = 6, ReturnDisplayName = "Success", DefaultToSelf = "VisualLoggerListener"))
 	static bool GetTransformAction(
 		FTransform& OutTransform, 
 		const ULearningAgentsActionObject* Object, 
 		const FLearningAgentsActionObjectElement Element, 
 		const FTransform RelativeTransform = FTransform(), 
-		const float LocationScale = 100.0f, 
-		const float RotationScale = 1.0f, 
-		const float ScaleScale = 1.0f, 
 		const FName Tag = TEXT("TransformAction"),
 		const bool bVisualLoggerEnabled = false,
 		ULearningAgentsManagerListener* VisualLoggerListener = nullptr,
@@ -1280,7 +1279,6 @@ public:
 	 * @param Object The Action Object
 	 * @param Element The Action Object Element
 	 * @param RelativeAngle The relative angle to transform the angle by.
-	 * @param AngleScale The scale used to control the overall magnitude of the outputted scale action.
 	 * @param Tag The tag of the corresponding action. Must match the tag given during Specify.
 	 * @param bVisualLoggerEnabled When true, debug data will be sent to the visual logger.
 	 * @param VisualLoggerListener The listener object which is making this action. This must be set to use logging.
@@ -1290,13 +1288,12 @@ public:
 	 * @param VisualLoggerColor The color for the visual logger display.
 	 * @return true if the provided Element is the correct type, otherwise false.
 	 */
-	UFUNCTION(BlueprintPure = false, Category = "LearningAgents", meta = (AdvancedDisplay = 5, ReturnDisplayName = "Success", DefaultToSelf = "VisualLoggerListener"))
+	UFUNCTION(BlueprintPure = false, Category = "LearningAgents", meta = (AdvancedDisplay = 4, ReturnDisplayName = "Success", DefaultToSelf = "VisualLoggerListener"))
 	static bool GetAngleAction(
 		float& OutAngle, 
 		const ULearningAgentsActionObject* Object, 
 		const FLearningAgentsActionObjectElement Element, 
 		const float RelativeAngle = 0.0f, 
-		const float AngleScale = 90.0f, 
 		const FName Tag = TEXT("AngleAction"),
 		const bool bVisualLoggerEnabled = false,
 		ULearningAgentsManagerListener* VisualLoggerListener = nullptr,
@@ -1312,7 +1309,6 @@ public:
 	 * @param Object The Action Object
 	 * @param Element The Action Object Element
 	 * @param RelativeAngle The relative angle to transform the angle by.
-	 * @param AngleScale The scale used to control the overall magnitude of the outputted scale action.
 	 * @param Tag The tag of the corresponding action. Must match the tag given during Specify.
 	 * @param bVisualLoggerEnabled When true, debug data will be sent to the visual logger.
 	 * @param VisualLoggerListener The listener object which is making this action. This must be set to use logging.
@@ -1322,13 +1318,12 @@ public:
 	 * @param VisualLoggerColor The color for the visual logger display.
 	 * @return true if the provided Element is the correct type, otherwise false.
 	 */
-	UFUNCTION(BlueprintPure = false, Category = "LearningAgents", meta = (AdvancedDisplay = 5, ReturnDisplayName = "Success", DefaultToSelf = "VisualLoggerListener"))
+	UFUNCTION(BlueprintPure = false, Category = "LearningAgents", meta = (AdvancedDisplay = 4, ReturnDisplayName = "Success", DefaultToSelf = "VisualLoggerListener"))
 	static bool GetAngleActionRadians(
 		float& OutAngle, 
 		const ULearningAgentsActionObject* Object,
 		const FLearningAgentsActionObjectElement Element, 
 		const float RelativeAngle = 0.0f,
-		const float AngleScale = 1.57079632679f, 
 		const FName Tag = TEXT("AngleAction"),
 		const bool bVisualLoggerEnabled = false,
 		ULearningAgentsManagerListener* VisualLoggerListener = nullptr,
@@ -1344,7 +1339,6 @@ public:
 	 * @param Object The Action Object
 	 * @param Element The Action Object Element
 	 * @param RelativeTransform The relative transform to transform the velocity by.
-	 * @param VelocityScale The scale used to control the overall magnitude of the outputted velocity action.
 	 * @param Tag The tag of the corresponding action. Must match the tag given during Specify.
 	 * @param bVisualLoggerEnabled When true, debug data will be sent to the visual logger.
 	 * @param VisualLoggerListener The listener object which is making this action. This must be set to use logging.
@@ -1354,13 +1348,12 @@ public:
 	 * @param VisualLoggerColor The color for the visual logger display.
 	 * @return true if the provided Element is the correct type, otherwise false.
 	 */
-	UFUNCTION(BlueprintPure = false, Category = "LearningAgents", meta = (AdvancedDisplay = 5, ReturnDisplayName = "Success", DefaultToSelf = "VisualLoggerListener"))
+	UFUNCTION(BlueprintPure = false, Category = "LearningAgents", meta = (AdvancedDisplay = 4, ReturnDisplayName = "Success", DefaultToSelf = "VisualLoggerListener"))
 	static bool GetVelocityAction(
 		FVector& OutVelocity, 
 		const ULearningAgentsActionObject* Object, 
 		const FLearningAgentsActionObjectElement Element, 
 		const FTransform RelativeTransform = FTransform(), 
-		const float VelocityScale = 200.0f, 
 		const FName Tag = TEXT("VelocityAction"),
 		const bool bVisualLoggerEnabled = false,
 		ULearningAgentsManagerListener* VisualLoggerListener = nullptr,
diff --git a/Engine/Plugins/Experimental/LearningAgents/Source/LearningAgents/Public/LearningAgentsObservations.h b/Engine/Plugins/Experimental/LearningAgents/Source/LearningAgents/Public/LearningAgentsObservations.h
--- a/Engine/Plugins/Experimental/LearningAgents/Source/LearningAgents/Public/LearningAgentsObservations.h
+++ b/Engine/Plugins/Experimental/LearningAgents/Source/LearningAgents/Public/LearningAgentsObservations.h
@@ -158,11 +158,12 @@ public:
 	 *
 	 * @param Schema The Observation Schema
 	 * @param Size The number of float values in the observation.
+	 * @param Scale The scale used to normalize this observation.
 	 * @param Tag The tag of this new observation. Used during observation object validation and debugging.
 	 * @return The newly created observation schema element.
 	 */
-	UFUNCTION(BlueprintPure, Category = "LearningAgents", meta = (AdvancedDisplay = 2))
-	static FLearningAgentsObservationSchemaElement SpecifyContinuousObservation(ULearningAgentsObservationSchema* Schema, const int32 Size, const FName Tag = TEXT("ContinuousObservation"));
+	UFUNCTION(BlueprintPure, Category = "LearningAgents", meta = (AdvancedDisplay = 3))
+	static FLearningAgentsObservationSchemaElement SpecifyContinuousObservation(ULearningAgentsObservationSchema* Schema, const int32 Size, const float Scale = 1.0f, const FName Tag = TEXT("ContinuousObservation"));
 
 	/**
 	 * Specifies a new exclusive discrete observation. This represents a discrete observation which is an exclusive selection from multiple choices.
@@ -472,21 +473,23 @@ public:
 	 * type-specific observation does not exist.
 	 *
 	 * @param Schema The Observation Schema
+	 * @param FloatScale The scale used to normalize this observation.
 	 * @param Tag The tag of this new observation. Used during observation object validation and debugging.
 	 * @return The newly created observation schema element.
 	 */
-	UFUNCTION(BlueprintPure, Category = "LearningAgents", meta = (AdvancedDisplay = 1))
-	static FLearningAgentsObservationSchemaElement SpecifyFloatObservation(ULearningAgentsObservationSchema* Schema, const FName Tag = TEXT("FloatObservation"));
+	UFUNCTION(BlueprintPure, Category = "LearningAgents", meta = (AdvancedDisplay = 2))
+	static FLearningAgentsObservationSchemaElement SpecifyFloatObservation(ULearningAgentsObservationSchema* Schema, const float FloatScale = 1.0f, const FName Tag = TEXT("FloatObservation"));
 
 	/**
 	 * Specifies a new location observation. Allows an agent to observe the location of some entity.
 	 *
 	 * @param Schema The Observation Schema
+	 * @param LocationScale The scale used to normalize this observation in cm.
 	 * @param Tag The tag of this new observation. Used during observation object validation and debugging.
 	 * @return The newly created observation schema element.
 	 */
-	UFUNCTION(BlueprintPure, Category = "LearningAgents", meta = (AdvancedDisplay = 1))
-	static FLearningAgentsObservationSchemaElement SpecifyLocationObservation(ULearningAgentsObservationSchema* Schema, const FName Tag = TEXT("LocationObservation"));
+	UFUNCTION(BlueprintPure, Category = "LearningAgents", meta = (AdvancedDisplay = 2))
+	static FLearningAgentsObservationSchemaElement SpecifyLocationObservation(ULearningAgentsObservationSchema* Schema, const float LocationScale = 100.0f, const FName Tag = TEXT("LocationObservation"));
 
 	/**
 	 * Specifies a new rotation observation. Allows an agent to observe the rotation of some entity. Rotations are encoded as two columns of the 
@@ -514,11 +517,12 @@ public:
 	 * Specifies a new transform observation. Allows an agent to observe the transform of some entity.
 	 *
 	 * @param Schema The Observation Schema
+	 * @param LocationScale The scale used to normalize the Location part of this transform observation in cm.
 	 * @param Tag The tag of this new observation. Used during observation object validation and debugging.
 	 * @return The newly created observation schema element.
 	 */
-	UFUNCTION(BlueprintPure, Category = "LearningAgents", meta = (AdvancedDisplay = 1))
-	static FLearningAgentsObservationSchemaElement SpecifyTransformObservation(ULearningAgentsObservationSchema* Schema, const FName Tag = TEXT("TransformObservation"));
+	UFUNCTION(BlueprintPure, Category = "LearningAgents", meta = (AdvancedDisplay = 2))
+	static FLearningAgentsObservationSchemaElement SpecifyTransformObservation(ULearningAgentsObservationSchema* Schema, const float LocationScale = 100.0f, const FName Tag = TEXT("TransformObservation"));
 
 	/**
 	 * Specifies a new angle observation. This will be encoded as a 2-dimension Cartesian vector so that 0 and 350 are close to each other in the 
@@ -535,11 +539,12 @@ public:
 	 * Specifies a new velocity observation. Allows an agent to observe the velocity of some entity.
 	 *
 	 * @param Schema The Observation Schema
+	 * @param VelocityScale The scale used to normalize this observation in cm/s.
 	 * @param Tag The tag of this new observation. Used during observation object validation and debugging.
 	 * @return The newly created observation schema element.
 	 */
-	UFUNCTION(BlueprintPure, Category = "LearningAgents", meta = (AdvancedDisplay = 1))
-	static FLearningAgentsObservationSchemaElement SpecifyVelocityObservation(ULearningAgentsObservationSchema* Schema, const FName Tag = TEXT("VelocityObservation"));
+	UFUNCTION(BlueprintPure, Category = "LearningAgents", meta = (AdvancedDisplay = 2))
+	static FLearningAgentsObservationSchemaElement SpecifyVelocityObservation(ULearningAgentsObservationSchema* Schema, const float VelocityScale = 200.0f, const FName Tag = TEXT("VelocityObservation"));
 
 	/**
 	 * Specifies a new direction observation. Allows an agent to observe the direction of some entity.
@@ -555,11 +560,12 @@ public:
 	 * Specifies a new location along spline observation. This observes the location of the spline at the given distance along that spline.
 	 *
 	 * @param Schema The Observation Schema
+	 * @param LocationScale The scale used to normalize this observation in cm.
 	 * @param Tag The tag of this new observation. Used during observation object validation and debugging.
 	 * @return The newly created observation schema element.
 	 */
-	UFUNCTION(BlueprintPure, Category = "LearningAgents", meta = (AdvancedDisplay = 1))
-	static FLearningAgentsObservationSchemaElement SpecifyLocationAlongSplineObservation(ULearningAgentsObservationSchema* Schema, const FName Tag = TEXT("LocationAlongSplineObservation"));
+	UFUNCTION(BlueprintPure, Category = "LearningAgents", meta = (AdvancedDisplay = 2))
+	static FLearningAgentsObservationSchemaElement SpecifyLocationAlongSplineObservation(ULearningAgentsObservationSchema* Schema, const float LocationScale = 100.0f, const FName Tag = TEXT("LocationAlongSplineObservation"));
 
 	/**
 	 * Specifies a new proportion along spline observation. This observes the proportion along a spline at the given distance. For looped splines 
@@ -1114,7 +1120,6 @@ public:
 	 *
 	 * @param Object The Observation Object
 	 * @param Value The new value of this observation.
-	 * @param FloatScale Used to normalize the data for this observation.
 	 * @param Tag The tag of the corresponding observation. Must match the tag given during Specify.
 	 * @param bVisualLoggerEnabled When true, debug data will be sent to the visual logger.
 	 * @param VisualLoggerListener The listener object which is making this observation. This must be set to use logging.
@@ -1123,11 +1128,10 @@ public:
 	 * @param VisualLoggerColor The color for the visual logger display.
 	 * @return The newly created observation object element.
 	 */
-	UFUNCTION(BlueprintPure, Category = "LearningAgents", meta = (AdvancedDisplay = 3, DefaultToSelf = "VisualLoggerListener"))
+	UFUNCTION(BlueprintPure, Category = "LearningAgents", meta = (AdvancedDisplay = 2, DefaultToSelf = "VisualLoggerListener"))
 	static FLearningAgentsObservationObjectElement MakeFloatObservation(
 		ULearningAgentsObservationObject* Object,
 		const float Value,
-		const float FloatScale = 1.0f,
 		const FName Tag = TEXT("FloatObservation"),
 		const bool bVisualLoggerEnabled = false,
 		ULearningAgentsManagerListener* VisualLoggerListener = nullptr,
@@ -1141,7 +1145,6 @@ public:
 	 * @param Object The Observation Object
 	 * @param Location The location of interest to the agent.
 	 * @param RelativeTransform The transform the provided location should be encoded relative to.
-	 * @param LocationScale Used to normalize the data for this observation.
 	 * @param Tag The tag of the corresponding observation. Must match the tag given during Specify.
 	 * @param bVisualLoggerEnabled When true, debug data will be sent to the visual logger.
 	 * @param VisualLoggerListener The listener object which is making this observation. This must be set to use logging.
@@ -1150,12 +1153,11 @@ public:
 	 * @param VisualLoggerColor The color for the visual logger display.
 	 * @return The newly created observation object element.
 	 */
-	UFUNCTION(BlueprintPure, Category = "LearningAgents", meta = (AdvancedDisplay = 4, DefaultToSelf = "VisualLoggerListener"))
+	UFUNCTION(BlueprintPure, Category = "LearningAgents", meta = (AdvancedDisplay = 3, DefaultToSelf = "VisualLoggerListener"))
 	static FLearningAgentsObservationObjectElement MakeLocationObservation(
 		ULearningAgentsObservationObject* Object,
 		const FVector Location,
 		const FTransform RelativeTransform = FTransform(),
-		const float LocationScale = 100.0f,
 		const FName Tag = TEXT("LocationObservation"),
 		const bool bVisualLoggerEnabled = false,
 		ULearningAgentsManagerListener* VisualLoggerListener = nullptr,
@@ -1253,7 +1255,6 @@ public:
 	 * @param Object The Observation Object
 	 * @param Transform The transform of interest to the agent.
 	 * @param RelativeTransform The transform the provided transform should be encoded relative to.
-	 * @param LocationScale Used to normalize the transform's location for this observation.
 	 * @param Tag The tag of the corresponding observation. Must match the tag given during Specify.
 	 * @param bVisualLoggerEnabled When true, debug data will be sent to the visual logger.
 	 * @param VisualLoggerListener The listener object which is making this observation. This must be set to use logging.
@@ -1262,12 +1263,11 @@ public:
 	 * @param VisualLoggerColor The color for the visual logger display.
 	 * @return The newly created observation object element.
 	 */
-	UFUNCTION(BlueprintPure, Category = "LearningAgents", meta = (AdvancedDisplay = 4, DefaultToSelf = "VisualLoggerListener"))
+	UFUNCTION(BlueprintPure, Category = "LearningAgents", meta = (AdvancedDisplay = 3, DefaultToSelf = "VisualLoggerListener"))
 	static FLearningAgentsObservationObjectElement MakeTransformObservation(
 		ULearningAgentsObservationObject* Object,
 		const FTransform Transform,
 		const FTransform RelativeTransform = FTransform(),
-		const float LocationScale = 100.0f,
 		const FName Tag = TEXT("TransformObservation"),
 		const bool bVisualLoggerEnabled = false,
 		ULearningAgentsManagerListener* VisualLoggerListener = nullptr,
@@ -1333,7 +1333,6 @@ public:
 	 * @param Object The Observation Object
 	 * @param Velocity The velocity of interest to the agent.
 	 * @param RelativeTransform The transform the provided velocity should be encoded relative to.
-	 * @param VelocityScale Used to normalize the data for this observation.
 	 * @param Tag The tag of the corresponding observation. Must match the tag given during Specify.
 	 * @param bVisualLoggerEnabled When true, debug data will be sent to the visual logger.
 	 * @param VisualLoggerListener The listener object which is making this observation. This must be set to use logging.
@@ -1343,12 +1342,11 @@ public:
 	 * @param VisualLoggerColor The color for the visual logger display.
 	 * @return The newly created observation object element.
 	 */
-	UFUNCTION(BlueprintPure, Category = "LearningAgents", meta = (AdvancedDisplay = 4, DefaultToSelf = "VisualLoggerListener"))
+	UFUNCTION(BlueprintPure, Category = "LearningAgents", meta = (AdvancedDisplay = 3, DefaultToSelf = "VisualLoggerListener"))
 	static FLearningAgentsObservationObjectElement MakeVelocityObservation(
 		ULearningAgentsObservationObject* Object, 
 		const FVector Velocity, 
 		const FTransform RelativeTransform = FTransform(), 
-		const float VelocityScale = 200.0f, 
 		const FName Tag = TEXT("VelocityObservation"),
 		const bool bVisualLoggerEnabled = false,
 		ULearningAgentsManagerListener* VisualLoggerListener = nullptr,
@@ -1394,7 +1392,6 @@ public:
 	 * @param SplineComponent The spline to observe.
 	 * @param DistanceAlongSpline The distance along that spline.
 	 * @param RelativeTransform The transform the provided location should be encoded relative to.
-	 * @param LocationScale Used to normalize the transform's location for this observation.
 	 * @param Tag The tag of the corresponding observation. Must match the tag given during Specify.
 	 * @param bVisualLoggerEnabled When true, debug data will be sent to the visual logger.
 	 * @param VisualLoggerListener The listener object which is making this observation. This must be set to use logging.
@@ -1403,13 +1400,12 @@ public:
 	 * @param VisualLoggerColor The color for the visual logger display.
 	 * @return The newly created observation object element.
 	 */
-	UFUNCTION(BlueprintPure, Category = "LearningAgents", meta = (AdvancedDisplay = 5, DefaultToSelf = "VisualLoggerListener"))
+	UFUNCTION(BlueprintPure, Category = "LearningAgents", meta = (AdvancedDisplay = 4, DefaultToSelf = "VisualLoggerListener"))
 	static FLearningAgentsObservationObjectElement MakeLocationAlongSplineObservation(
 		ULearningAgentsObservationObject* Object,
 		const USplineComponent* SplineComponent,
 		const float DistanceAlongSpline,
 		const FTransform RelativeTransform = FTransform(),
-		const float LocationScale = 100.0f,
 		const FName Tag = TEXT("LocationAlongSplineObservation"),
 		const bool bVisualLoggerEnabled = false,
 		ULearningAgentsManagerListener* VisualLoggerListener = nullptr,
@@ -1608,11 +1604,11 @@ public:
 	UFUNCTION(BlueprintPure = false, Category = "LearningAgents", meta = (AdvancedDisplay = 3, ReturnDisplayName = "Success"))
 	static bool GetBoolObservation(bool& bOutValue, const ULearningAgentsObservationObject* Object, const FLearningAgentsObservationObjectElement Element, const FName Tag = TEXT("BoolObservation"));
 
-	UFUNCTION(BlueprintPure = false, Category = "LearningAgents", meta = (AdvancedDisplay = 4, ReturnDisplayName = "Success"))
-	static bool GetFloatObservation(float& OutValue, const ULearningAgentsObservationObject* Object, const FLearningAgentsObservationObjectElement Element, const float FloatScale = 1.0f, const FName Tag = TEXT("FloatObservation"));
+	UFUNCTION(BlueprintPure = false, Category = "LearningAgents", meta = (AdvancedDisplay = 3, ReturnDisplayName = "Success"))
+	static bool GetFloatObservation(float& OutValue, const ULearningAgentsObservationObject* Object, const FLearningAgentsObservationObjectElement Element, const FName Tag = TEXT("FloatObservation"));
 
-	UFUNCTION(BlueprintPure = false, Category = "LearningAgents", meta = (AdvancedDisplay = 5, ReturnDisplayName = "Success"))
-	static bool GetLocationObservation(FVector& OutLocation, const ULearningAgentsObservationObject* Object, const FLearningAgentsObservationObjectElement Element, const FTransform RelativeTransform = FTransform(), const float LocationScale = 100.0f, const FName Tag = TEXT("LocationObservation"));
+	UFUNCTION(BlueprintPure = false, Category = "LearningAgents", meta = (AdvancedDisplay = 4, ReturnDisplayName = "Success"))
+	static bool GetLocationObservation(FVector& OutLocation, const ULearningAgentsObservationObject* Object, const FLearningAgentsObservationObjectElement Element, const FTransform RelativeTransform = FTransform(), const FName Tag = TEXT("LocationObservation"));
 
 	UFUNCTION(BlueprintPure = false, Category = "LearningAgents", meta = (AdvancedDisplay = 4, ReturnDisplayName = "Success"))
 	static bool GetRotationObservation(FRotator& OutRotation, const ULearningAgentsObservationObject* Object, const FLearningAgentsObservationObjectElement Element, const FRotator RelativeRotation = FRotator::ZeroRotator, const FName Tag = TEXT("RotationObservation"));
@@ -1623,8 +1619,8 @@ public:
 	UFUNCTION(BlueprintPure = false, Category = "LearningAgents", meta = (AdvancedDisplay = 4, ReturnDisplayName = "Success"))
 	static bool GetScaleObservation(FVector& OutScale, const ULearningAgentsObservationObject* Object, const FLearningAgentsObservationObjectElement Element, const FVector RelativeScale = FVector(1, 1, 1), const FName Tag = TEXT("ScaleObservation"));
 
-	UFUNCTION(BlueprintPure = false, Category = "LearningAgents", meta = (AdvancedDisplay = 5, ReturnDisplayName = "Success"))
-	static bool GetTransformObservation(FTransform& OutTransform, const ULearningAgentsObservationObject* Object, const FLearningAgentsObservationObjectElement Element, const FTransform RelativeTransform = FTransform(), const float LocationScale = 100.0f, const FName Tag = TEXT("TransformObservation"));
+	UFUNCTION(BlueprintPure = false, Category = "LearningAgents", meta = (AdvancedDisplay = 4, ReturnDisplayName = "Success"))
+	static bool GetTransformObservation(FTransform& OutTransform, const ULearningAgentsObservationObject* Object, const FLearningAgentsObservationObjectElement Element, const FTransform RelativeTransform = FTransform(), const FName Tag = TEXT("TransformObservation"));
 
 	UFUNCTION(BlueprintPure = false, Category = "LearningAgents", meta = (AdvancedDisplay = 4, ReturnDisplayName = "Success"))
 	static bool GetAngleObservation(float& OutAngle, const ULearningAgentsObservationObject* Object, const FLearningAgentsObservationObjectElement Element, const float RelativeAngle = 0.0f, const FName Tag = TEXT("AngleObservation"));
@@ -1632,16 +1628,16 @@ public:
 	UFUNCTION(BlueprintPure = false, Category = "LearningAgents", meta = (AdvancedDisplay = 4, ReturnDisplayName = "Success"))
 	static bool GetAngleObservationRadians(float& OutAngle, const ULearningAgentsObservationObject* Object, const FLearningAgentsObservationObjectElement Element, const float RelativeAngle = 0.0f, const FName Tag = TEXT("AngleObservation"));
 
-	UFUNCTION(BlueprintPure = false, Category = "LearningAgents", meta = (AdvancedDisplay = 5, ReturnDisplayName = "Success"))
-	static bool GetVelocityObservation(FVector& OutVelocity, const ULearningAgentsObservationObject* Object, const FLearningAgentsObservationObjectElement Element, const FTransform RelativeTransform = FTransform(), const float VelocityScale = 200.0f, const FName Tag = TEXT("VelocityObservation"));
+	UFUNCTION(BlueprintPure = false, Category = "LearningAgents", meta = (AdvancedDisplay = 4, ReturnDisplayName = "Success"))
+	static bool GetVelocityObservation(FVector& OutVelocity, const ULearningAgentsObservationObject* Object, const FLearningAgentsObservationObjectElement Element, const FTransform RelativeTransform = FTransform(), const FName Tag = TEXT("VelocityObservation"));
 
 	UFUNCTION(BlueprintPure = false, Category = "LearningAgents", meta = (AdvancedDisplay = 4, ReturnDisplayName = "Success"))
 	static bool GetDirectionObservation(FVector& OutDirection, const ULearningAgentsObservationObject* Object, const FLearningAgentsObservationObjectElement Element, const FTransform RelativeTransform = FTransform(), const FName Tag = TEXT("DirectionObservation"));
 
 	// Get Spline Observations
 
-	UFUNCTION(BlueprintPure = false, Category = "LearningAgents", meta = (AdvancedDisplay = 5, ReturnDisplayName = "Success"))
-	static bool GetLocationAlongSplineObservation(FVector& OutLocation, const ULearningAgentsObservationObject* Object, const FLearningAgentsObservationObjectElement Element, const FTransform RelativeTransform = FTransform(), const float LocationScale = 100.0f, const FName Tag = TEXT("LocationAlongSplineObservation"));
+	UFUNCTION(BlueprintPure = false, Category = "LearningAgents", meta = (AdvancedDisplay = 4, ReturnDisplayName = "Success"))
+	static bool GetLocationAlongSplineObservation(FVector& OutLocation, const ULearningAgentsObservationObject* Object, const FLearningAgentsObservationObjectElement Element, const FTransform RelativeTransform = FTransform(), const FName Tag = TEXT("LocationAlongSplineObservation"));
 
 	UFUNCTION(BlueprintPure = false, Category = "LearningAgents", meta = (AdvancedDisplay = 5, ReturnDisplayName = "Success"))
 	static bool GetProportionAlongSplineObservation(bool& bOutIsClosedLoop, float& OutAngle, float& OutPropotion, const ULearningAgentsObservationObject* Object, const FLearningAgentsObservationObjectElement Element, const FName Tag = TEXT("ProportionAlongSplineObservation"));
