LearningAgents: Minor fix to LearningAgents MultiArray.

diff --git a/Engine/Plugins/Experimental/LearningAgents/Source/Learning/Public/MultiArray.h b/Engine/Plugins/Experimental/LearningAgents/Source/Learning/Public/MultiArray.h
--- a/Engine/Plugins/Experimental/LearningAgents/Source/Learning/Public/MultiArray.h
+++ b/Engine/Plugins/Experimental/LearningAgents/Source/Learning/Public/MultiArray.h
@@ -917,7 +917,7 @@ public:
 	/** Implicit cast to TArrayView. */
 	FORCEINLINE operator TArrayView<ElementType, SizeType>() const
 	{
-		return TArrayView(GetData(), Num());
+		return TArrayView<ElementType, SizeType>(GetData(), Num());
 	}
 
 public:
