LearningAgents: Added ability to use compressed linear layers in generated networks.

diff --git a/Engine/Plugins/Experimental/LearningAgents/Source/Learning/Private/LearningAction.cpp b/Engine/Plugins/Experimental/LearningAgents/Source/Learning/Private/LearningAction.cpp
--- a/Engine/Plugins/Experimental/LearningAgents/Source/Learning/Private/LearningAction.cpp
+++ b/Engine/Plugins/Experimental/LearningAgents/Source/Learning/Private/LearningAction.cpp
@@ -1042,7 +1042,8 @@ namespace UE::Learning::Action
 		NNE::RuntimeBasic::FModelBuilderElement& OutElement,
 		NNE::RuntimeBasic::FModelBuilder& Builder,
 		const FSchema& Schema,
-		const FSchemaElement SchemaElement)
+		const FSchemaElement SchemaElement,
+		const FNetworkSettings& NetworkSettings)
 	{
 		const EType SchemaElementType = Schema.GetType(SchemaElement);
 
@@ -1112,7 +1113,7 @@ namespace UE::Learning::Action
 			for (const FSchemaElement SubElement : Parameters.Elements)
 			{
 				NNE::RuntimeBasic::FModelBuilderElement BuilderSubElement;
-				MakeDecoderNetworkModelBuilderElementFromSchema(BuilderSubElement, Builder, Schema, SubElement);
+				MakeDecoderNetworkModelBuilderElementFromSchema(BuilderSubElement, Builder, Schema, SubElement, NetworkSettings);
 				BuilderLayers.Emplace(BuilderSubElement);
 			}
 
@@ -1129,7 +1130,7 @@ namespace UE::Learning::Action
 			for (const FSchemaElement SubElement : Parameters.Elements)
 			{
 				NNE::RuntimeBasic::FModelBuilderElement BuilderSubElement;
-				MakeDecoderNetworkModelBuilderElementFromSchema(BuilderSubElement, Builder, Schema, SubElement);
+				MakeDecoderNetworkModelBuilderElementFromSchema(BuilderSubElement, Builder, Schema, SubElement, NetworkSettings);
 				BuilderLayers.Emplace(BuilderSubElement);
 			}
 
@@ -1159,7 +1160,7 @@ namespace UE::Learning::Action
 			for (const FSchemaElement SubElement : Parameters.Elements)
 			{
 				NNE::RuntimeBasic::FModelBuilderElement BuilderSubElement;
-				MakeDecoderNetworkModelBuilderElementFromSchema(BuilderSubElement, Builder, Schema, SubElement);
+				MakeDecoderNetworkModelBuilderElementFromSchema(BuilderSubElement, Builder, Schema, SubElement, NetworkSettings);
 				BuilderLayers.Emplace(BuilderSubElement);
 			}
 
@@ -1184,7 +1185,7 @@ namespace UE::Learning::Action
 			const FSchemaArrayParameters Parameters = Schema.GetArray(SchemaElement);
 
 			NNE::RuntimeBasic::FModelBuilderElement BuilderSubElement;
-			MakeDecoderNetworkModelBuilderElementFromSchema(BuilderSubElement, Builder, Schema, Parameters.Element);
+			MakeDecoderNetworkModelBuilderElementFromSchema(BuilderSubElement, Builder, Schema, Parameters.Element, NetworkSettings);
 			OutElement = Builder.MakeArray(Parameters.Num, BuilderSubElement);
 			break;
 		}
@@ -1196,19 +1197,37 @@ namespace UE::Learning::Action
 			const int32 SubElementEncodedSize = Schema.GetEncodedVectorSize(Parameters.Element);
 
 			NNE::RuntimeBasic::FModelBuilderElement BuilderSubElement;
-			MakeDecoderNetworkModelBuilderElementFromSchema(BuilderSubElement, Builder, Schema, Parameters.Element);
-
-			OutElement = Builder.MakeSequence({
-				Builder.MakeActivation(Parameters.EncodingSize, Private::GetNNEActivationFunction(Parameters.ActivationFunction)),
-				Builder.MakeMLPWithRandomKaimingWeights(
-					Parameters.EncodingSize,
-					SubElementEncodedSize,
-					Parameters.EncodingSize,
-					Parameters.LayerNum + 1,  // Add 1 to account for input layer 
-					Private::GetNNEActivationFunction(Parameters.ActivationFunction),
-					false),
-				BuilderSubElement,
-				});
+			MakeDecoderNetworkModelBuilderElementFromSchema(BuilderSubElement, Builder, Schema, Parameters.Element, NetworkSettings);
+
+			if (NetworkSettings.bUseCompressedLinearLayers)
+			{
+				OutElement = Builder.MakeSequence({
+					Builder.MakeActivation(Parameters.EncodingSize, Private::GetNNEActivationFunction(Parameters.ActivationFunction)),
+					Builder.MakeCompressedMLPWithRandomKaimingWeights(
+						Parameters.EncodingSize,
+						SubElementEncodedSize,
+						Parameters.EncodingSize,
+						Parameters.LayerNum + 1,  // Add 1 to account for input layer 
+						Private::GetNNEActivationFunction(Parameters.ActivationFunction),
+						false),
+					BuilderSubElement,
+					});
+			}
+			else
+			{
+				OutElement = Builder.MakeSequence({
+					Builder.MakeActivation(Parameters.EncodingSize, Private::GetNNEActivationFunction(Parameters.ActivationFunction)),
+					Builder.MakeMLPWithRandomKaimingWeights(
+						Parameters.EncodingSize,
+						SubElementEncodedSize,
+						Parameters.EncodingSize,
+						Parameters.LayerNum + 1,  // Add 1 to account for input layer 
+						Private::GetNNEActivationFunction(Parameters.ActivationFunction),
+						false),
+					BuilderSubElement,
+					});
+			}
+
 			break;
 		}
 
@@ -1233,13 +1252,14 @@ namespace UE::Learning::Action
 		uint32& OutOutputSize,
 		const FSchema& Schema,
 		const FSchemaElement SchemaElement,
+		const FNetworkSettings& NetworkSettings,
 		const uint32 Seed)
 	{
 		UE_LEARNING_CHECK(Schema.IsValid(SchemaElement));
 
 		NNE::RuntimeBasic::FModelBuilder Builder(Seed);
 		NNE::RuntimeBasic::FModelBuilderElement Element;
-		MakeDecoderNetworkModelBuilderElementFromSchema(Element, Builder, Schema, SchemaElement);
+		MakeDecoderNetworkModelBuilderElementFromSchema(Element, Builder, Schema, SchemaElement, NetworkSettings);
 		Builder.WriteFileDataAndReset(OutFileData, OutInputSize, OutOutputSize, Element);
 	}
 
diff --git a/Engine/Plugins/Experimental/LearningAgents/Source/Learning/Private/LearningObservation.cpp b/Engine/Plugins/Experimental/LearningAgents/Source/Learning/Private/LearningObservation.cpp
--- a/Engine/Plugins/Experimental/LearningAgents/Source/Learning/Private/LearningObservation.cpp
+++ b/Engine/Plugins/Experimental/LearningAgents/Source/Learning/Private/LearningObservation.cpp
@@ -899,7 +899,8 @@ namespace UE::Learning::Observation
 		NNE::RuntimeBasic::FModelBuilderElement& OutElement,
 		NNE::RuntimeBasic::FModelBuilder& Builder,
 		const FSchema& Schema,
-		const FSchemaElement SchemaElement)
+		const FSchemaElement SchemaElement,
+		const FNetworkSettings& NetworkSettings)
 	{
 		const EType SchemaElementType = Schema.GetType(SchemaElement);
 
@@ -931,7 +932,7 @@ namespace UE::Learning::Observation
 			for (const FSchemaElement SubElement : Parameters.Elements)
 			{
 				NNE::RuntimeBasic::FModelBuilderElement BuilderSubElement;
-				MakeEncoderNetworkModelBuilderElementFromSchema(BuilderSubElement, Builder, Schema, SubElement);
+				MakeEncoderNetworkModelBuilderElementFromSchema(BuilderSubElement, Builder, Schema, SubElement, NetworkSettings);
 				BuilderLayers.Emplace(BuilderSubElement);
 			}
 
@@ -951,9 +952,16 @@ namespace UE::Learning::Observation
 			{
 				const int32 SubElementEncodedSize = Schema.GetEncodedVectorSize(SubElement);
 				NNE::RuntimeBasic::FModelBuilderElement BuilderSubElement;
-				MakeEncoderNetworkModelBuilderElementFromSchema(BuilderSubElement, Builder, Schema, SubElement);
+				MakeEncoderNetworkModelBuilderElementFromSchema(BuilderSubElement, Builder, Schema, SubElement, NetworkSettings);
 				BuilderSubLayers.Emplace(BuilderSubElement);
-				BuilderEncoders.Emplace(Builder.MakeLinearWithRandomKaimingWeights(SubElementEncodedSize, Parameters.EncodingSize));
+				if (NetworkSettings.bUseCompressedLinearLayers)
+				{
+					BuilderEncoders.Emplace(Builder.MakeCompressedLinearWithRandomKaimingWeights(SubElementEncodedSize, Parameters.EncodingSize));
+				}
+				else
+				{
+					BuilderEncoders.Emplace(Builder.MakeLinearWithRandomKaimingWeights(SubElementEncodedSize, Parameters.EncodingSize));
+				}
 			}
 
 			OutElement = Builder.MakeAggregateOrExclusive(Parameters.EncodingSize, BuilderSubLayers, BuilderEncoders);
@@ -975,11 +983,20 @@ namespace UE::Learning::Observation
 			{
 				const int32 SubElementEncodedSize = Schema.GetEncodedVectorSize(SubElement);
 				NNE::RuntimeBasic::FModelBuilderElement BuilderSubElement;
-				MakeEncoderNetworkModelBuilderElementFromSchema(BuilderSubElement, Builder, Schema, SubElement);
+				MakeEncoderNetworkModelBuilderElementFromSchema(BuilderSubElement, Builder, Schema, SubElement, NetworkSettings);
 				BuilderSubLayers.Emplace(BuilderSubElement);
-				BuilderQueryLayers.Emplace(Builder.MakeLinearWithRandomKaimingWeights(SubElementEncodedSize, Parameters.AttentionHeadNum * Parameters.AttentionEncodingSize));
-				BuilderKeyLayers.Emplace(Builder.MakeLinearWithRandomKaimingWeights(SubElementEncodedSize, Parameters.AttentionHeadNum * Parameters.AttentionEncodingSize));
-				BuilderValueLayers.Emplace(Builder.MakeLinearWithRandomKaimingWeights(SubElementEncodedSize, Parameters.AttentionHeadNum * Parameters.ValueEncodingSize));
+				if (NetworkSettings.bUseCompressedLinearLayers)
+				{
+					BuilderQueryLayers.Emplace(Builder.MakeCompressedLinearWithRandomKaimingWeights(SubElementEncodedSize, Parameters.AttentionHeadNum * Parameters.AttentionEncodingSize));
+					BuilderKeyLayers.Emplace(Builder.MakeCompressedLinearWithRandomKaimingWeights(SubElementEncodedSize, Parameters.AttentionHeadNum * Parameters.AttentionEncodingSize));
+					BuilderValueLayers.Emplace(Builder.MakeCompressedLinearWithRandomKaimingWeights(SubElementEncodedSize, Parameters.AttentionHeadNum * Parameters.ValueEncodingSize));
+				}
+				else
+				{
+					BuilderQueryLayers.Emplace(Builder.MakeLinearWithRandomKaimingWeights(SubElementEncodedSize, Parameters.AttentionHeadNum * Parameters.AttentionEncodingSize));
+					BuilderKeyLayers.Emplace(Builder.MakeLinearWithRandomKaimingWeights(SubElementEncodedSize, Parameters.AttentionHeadNum * Parameters.AttentionEncodingSize));
+					BuilderValueLayers.Emplace(Builder.MakeLinearWithRandomKaimingWeights(SubElementEncodedSize, Parameters.AttentionHeadNum * Parameters.ValueEncodingSize));
+				}
 			}
 
 			OutElement = Builder.MakeAggregateOrInclusive(
@@ -999,7 +1016,7 @@ namespace UE::Learning::Observation
 			const FSchemaArrayParameters Parameters = Schema.GetArray(SchemaElement);
 
 			NNE::RuntimeBasic::FModelBuilderElement BuilderSubElement;
-			MakeEncoderNetworkModelBuilderElementFromSchema(BuilderSubElement, Builder, Schema, Parameters.Element);
+			MakeEncoderNetworkModelBuilderElementFromSchema(BuilderSubElement, Builder, Schema, Parameters.Element, NetworkSettings);
 			OutElement = Builder.MakeArray(Parameters.Num, BuilderSubElement);
 			break;
 		}
@@ -1011,17 +1028,33 @@ namespace UE::Learning::Observation
 			const int32 SubElementEncodedSize = Schema.GetEncodedVectorSize(Parameters.Element);
 
 			NNE::RuntimeBasic::FModelBuilderElement BuilderSubElement;
-			MakeEncoderNetworkModelBuilderElementFromSchema(BuilderSubElement, Builder, Schema, Parameters.Element);
+			MakeEncoderNetworkModelBuilderElementFromSchema(BuilderSubElement, Builder, Schema, Parameters.Element, NetworkSettings);
+
+			if (NetworkSettings.bUseCompressedLinearLayers)
+			{
+				OutElement = Builder.MakeAggregateSet(
+					Parameters.MaxNum,
+					Parameters.ValueEncodingSize,
+					Parameters.AttentionEncodingSize,
+					Parameters.AttentionHeadNum,
+					BuilderSubElement,
+					Builder.MakeCompressedLinearWithRandomKaimingWeights(SubElementEncodedSize, Parameters.AttentionHeadNum * Parameters.AttentionEncodingSize),
+					Builder.MakeCompressedLinearWithRandomKaimingWeights(SubElementEncodedSize, Parameters.AttentionHeadNum * Parameters.AttentionEncodingSize),
+					Builder.MakeCompressedLinearWithRandomKaimingWeights(SubElementEncodedSize, Parameters.AttentionHeadNum * Parameters.ValueEncodingSize));
+			}
+			else
+			{
+				OutElement = Builder.MakeAggregateSet(
+					Parameters.MaxNum,
+					Parameters.ValueEncodingSize,
+					Parameters.AttentionEncodingSize,
+					Parameters.AttentionHeadNum,
+					BuilderSubElement,
+					Builder.MakeLinearWithRandomKaimingWeights(SubElementEncodedSize, Parameters.AttentionHeadNum * Parameters.AttentionEncodingSize),
+					Builder.MakeLinearWithRandomKaimingWeights(SubElementEncodedSize, Parameters.AttentionHeadNum * Parameters.AttentionEncodingSize),
+					Builder.MakeLinearWithRandomKaimingWeights(SubElementEncodedSize, Parameters.AttentionHeadNum * Parameters.ValueEncodingSize));
+			}
 
-			OutElement = Builder.MakeAggregateSet(
-				Parameters.MaxNum,
-				Parameters.ValueEncodingSize,
-				Parameters.AttentionEncodingSize,
-				Parameters.AttentionHeadNum,
-				BuilderSubElement,
-				Builder.MakeLinearWithRandomKaimingWeights(SubElementEncodedSize, Parameters.AttentionHeadNum * Parameters.AttentionEncodingSize),
-				Builder.MakeLinearWithRandomKaimingWeights(SubElementEncodedSize, Parameters.AttentionHeadNum * Parameters.AttentionEncodingSize),
-				Builder.MakeLinearWithRandomKaimingWeights(SubElementEncodedSize, Parameters.AttentionHeadNum * Parameters.ValueEncodingSize));
 			break;
 		}
 
@@ -1032,18 +1065,35 @@ namespace UE::Learning::Observation
 			const int32 SubElementEncodedSize = Schema.GetEncodedVectorSize(Parameters.Element);
 
 			NNE::RuntimeBasic::FModelBuilderElement BuilderSubElement;
-			MakeEncoderNetworkModelBuilderElementFromSchema(BuilderSubElement, Builder, Schema, Parameters.Element);
-
-			OutElement = Builder.MakeSequence({
-				BuilderSubElement,
-				Builder.MakeMLPWithRandomKaimingWeights(
-					SubElementEncodedSize,
-					Parameters.EncodingSize,
-					Parameters.EncodingSize,
-					Parameters.LayerNum + 1, // Add 1 to account for input layer
-					Private::GetNNEActivationFunction(Parameters.ActivationFunction),
-					true)
-				});
+			MakeEncoderNetworkModelBuilderElementFromSchema(BuilderSubElement, Builder, Schema, Parameters.Element, NetworkSettings);
+
+			if (NetworkSettings.bUseCompressedLinearLayers)
+			{
+				OutElement = Builder.MakeSequence({
+					BuilderSubElement,
+					Builder.MakeCompressedMLPWithRandomKaimingWeights(
+						SubElementEncodedSize,
+						Parameters.EncodingSize,
+						Parameters.EncodingSize,
+						Parameters.LayerNum + 1, // Add 1 to account for input layer
+						Private::GetNNEActivationFunction(Parameters.ActivationFunction),
+						true)
+					});
+			}
+			else
+			{
+				OutElement = Builder.MakeSequence({
+					BuilderSubElement,
+					Builder.MakeMLPWithRandomKaimingWeights(
+						SubElementEncodedSize,
+						Parameters.EncodingSize,
+						Parameters.EncodingSize,
+						Parameters.LayerNum + 1, // Add 1 to account for input layer
+						Private::GetNNEActivationFunction(Parameters.ActivationFunction),
+						true)
+					});
+			}
+
 			break;
 		}
 
@@ -1068,13 +1118,14 @@ namespace UE::Learning::Observation
 		uint32& OutOutputSize,
 		const FSchema& Schema,
 		const FSchemaElement SchemaElement,
+		const FNetworkSettings& NetworkSettings,
 		const uint32 Seed)
 	{
 		UE_LEARNING_CHECK(Schema.IsValid(SchemaElement));
 
 		NNE::RuntimeBasic::FModelBuilder Builder(Seed);
 		NNE::RuntimeBasic::FModelBuilderElement Element;
-		MakeEncoderNetworkModelBuilderElementFromSchema(Element, Builder, Schema, SchemaElement);
+		MakeEncoderNetworkModelBuilderElementFromSchema(Element, Builder, Schema, SchemaElement, NetworkSettings);
 		Builder.WriteFileDataAndReset(OutFileData, OutInputSize, OutOutputSize, Element);
 	}
 
diff --git a/Engine/Plugins/Experimental/LearningAgents/Source/Learning/Public/LearningAction.h b/Engine/Plugins/Experimental/LearningAgents/Source/Learning/Public/LearningAction.h
--- a/Engine/Plugins/Experimental/LearningAgents/Source/Learning/Public/LearningAction.h
+++ b/Engine/Plugins/Experimental/LearningAgents/Source/Learning/Public/LearningAction.h
@@ -445,6 +445,15 @@ namespace UE::Learning::Action
 		const FSchema& SchemaB,
 		const FSchemaElement SchemaElementB);
 
+	/**
+	 * Settings for building a network from a schema.
+	 */
+	struct LEARNING_API FNetworkSettings
+	{
+		// If to use compressed linear layers. This reduces the memory usage by half at the cost of some small impact on quality and evaluation speed.
+		bool bUseCompressedLinearLayers = false;
+	};
+
 	/**
 	 * Make a NNE::RuntimeBasic::FModelBuilderElement for the given Schema. This can be used if you want to plug the Decoder generated by this 
 	 * schema as a part of a larger model you are making with a NNE::RuntimeBasic::FModelBuilder.
@@ -453,12 +462,14 @@ namespace UE::Learning::Action
 	 * @param Builder					Model Builder
 	 * @param Schema					Action Schema
 	 * @param SchemaElement				Action Schema Element
+	 * @param NetworkSettings			Network Generation Settings
 	 */
 	LEARNING_API void MakeDecoderNetworkModelBuilderElementFromSchema(
 		NNE::RuntimeBasic::FModelBuilderElement& OutElement,
 		NNE::RuntimeBasic::FModelBuilder& Builder,
 		const FSchema& Schema,
-		const FSchemaElement SchemaElement);
+		const FSchemaElement SchemaElement,
+		const FNetworkSettings& NetworkSettings = FNetworkSettings());
 
 	/**
 	 * Generate FileData for a Neural Network that can act as a Decoder for the given schema. This network will take as input a vector of size
@@ -469,6 +480,7 @@ namespace UE::Learning::Action
 	 * @param OutOutputSize				Size of the vector this network takes as output
 	 * @param Schema					Action Schema
 	 * @param SchemaElement				Action Schema Element
+	 * @param NetworkSettings			Network Generation Settings
 	 * @param Seed						The random seed used in initializing the network weights
 	 */
 	LEARNING_API void GenerateDecoderNetworkFileDataFromSchema(
@@ -477,6 +489,7 @@ namespace UE::Learning::Action
 		uint32& OutOutputSize,
 		const FSchema& Schema,
 		const FSchemaElement SchemaElement,
+		const FNetworkSettings& NetworkSettings = FNetworkSettings(),
 		const uint32 Seed = 0x05dbd1f5);
 
 	/**
diff --git a/Engine/Plugins/Experimental/LearningAgents/Source/Learning/Public/LearningObservation.h b/Engine/Plugins/Experimental/LearningAgents/Source/Learning/Public/LearningObservation.h
--- a/Engine/Plugins/Experimental/LearningAgents/Source/Learning/Public/LearningObservation.h
+++ b/Engine/Plugins/Experimental/LearningAgents/Source/Learning/Public/LearningObservation.h
@@ -426,6 +426,15 @@ namespace UE::Learning::Observation
 		const FSchema& SchemaB,
 		const FSchemaElement SchemaElementB);
 
+	/**
+	 * Settings for building a network from a schema.
+	 */
+	struct LEARNING_API FNetworkSettings
+	{
+		// If to use compressed linear layers. This reduces the memory usage by half at the cost of some small impact on quality and evaluation speed.
+		bool bUseCompressedLinearLayers = false;
+	};
+
 	/**
 	 * Make a NNE::RuntimeBasic::FModelBuilderElement for the given Schema. This can be used if you want to plug the Encoder generated by this
 	 * schema as a part of a larger model you are making with a NNE::RuntimeBasic::FModelBuilder.
@@ -434,12 +443,14 @@ namespace UE::Learning::Observation
 	 * @param Builder					Model Builder
 	 * @param Schema					Observation Schema
 	 * @param SchemaElement				Observation Schema Element
+	 * @param NetworkSettings			Network Generation Settings
 	 */
 	LEARNING_API void MakeEncoderNetworkModelBuilderElementFromSchema(
 		NNE::RuntimeBasic::FModelBuilderElement& OutElement,
 		NNE::RuntimeBasic::FModelBuilder& Builder,
 		const FSchema& Schema,
-		const FSchemaElement SchemaElement);
+		const FSchemaElement SchemaElement,
+		const FNetworkSettings& NetworkSettings = FNetworkSettings());
 
 	/**
 	 * Generate FileData for a Neural Network that can act as a Encoder for the given schema. This network will take as input a vector of size
@@ -450,6 +461,7 @@ namespace UE::Learning::Observation
 	 * @param OutOutputSize				Size of the vector this network takes as output
 	 * @param Schema					Observation Schema
 	 * @param SchemaElement				Observation Schema Element
+	 * @param NetworkSettings			Network Generation Settings
 	 * @param Seed						The random seed used in initializing the network weights
 	 */
 	LEARNING_API void GenerateEncoderNetworkFileDataFromSchema(
@@ -458,6 +470,7 @@ namespace UE::Learning::Observation
 		uint32& OutOutputSize,
 		const FSchema& Schema,
 		const FSchemaElement SchemaElement,
+		const FNetworkSettings& NetworkSettings = FNetworkSettings(),
 		const uint32 Seed = 0x08ab1c49);
 
 	/**
diff --git a/Engine/Plugins/Experimental/LearningAgents/Source/LearningAgents/Private/LearningAgentsPolicy.cpp b/Engine/Plugins/Experimental/LearningAgents/Source/LearningAgents/Private/LearningAgentsPolicy.cpp
--- a/Engine/Plugins/Experimental/LearningAgents/Source/LearningAgents/Private/LearningAgentsPolicy.cpp
+++ b/Engine/Plugins/Experimental/LearningAgents/Source/LearningAgents/Private/LearningAgentsPolicy.cpp
@@ -203,6 +203,7 @@ void ULearningAgentsPolicy::SetupPolicy(
 			EncoderOutputSize,
 			Interactor->GetObservationSchema(),
 			Interactor->GetObservationSchemaElement(),
+			UE::Learning::Observation::FNetworkSettings(),
 			UE::Learning::Random::Int(Seed ^ 0x658868dd));
 
 		UE_LEARNING_CHECK(EncoderInputSize == ObservationVectorSize);
@@ -345,6 +346,7 @@ void ULearningAgentsPolicy::SetupPolicy(
 			DecoderOutputSize,
 			Interactor->GetActionSchema(),
 			Interactor->GetActionSchemaElement(),
+			UE::Learning::Action::FNetworkSettings(),
 			UE::Learning::Random::Int(Seed ^ 0xfa88bb7f));
 
 		UE_LEARNING_CHECK(DecoderInputSize == ActionEncodedVectorSize);
diff --git a/Engine/Plugins/Experimental/NNERuntimeBasicCpu/Source/NNERuntimeBasicCpu/Private/NNERuntimeBasicCpuModel.cpp b/Engine/Plugins/Experimental/NNERuntimeBasicCpu/Source/NNERuntimeBasicCpu/Private/NNERuntimeBasicCpuModel.cpp
--- a/Engine/Plugins/Experimental/NNERuntimeBasicCpu/Source/NNERuntimeBasicCpu/Private/NNERuntimeBasicCpuModel.cpp
+++ b/Engine/Plugins/Experimental/NNERuntimeBasicCpu/Source/NNERuntimeBasicCpu/Private/NNERuntimeBasicCpuModel.cpp
@@ -4045,9 +4045,30 @@ namespace UE::NNE::RuntimeBasic
 
 	namespace Private
 	{
-		static inline float UniformToGaussian(const float R1, const float R2)
+		static inline uint32 RngInt(const uint32 State)
 		{
-			return FMath::Sqrt(-2.0f * FMath::Loge(FMath::Max(R1, UE_SMALL_NUMBER))) * FMath::Cos(R2 * UE_TWO_PI);
+			uint32 X = State ^ 0xb74eaecf;
+			X = ((X >> 16) ^ X) * 0x45d9f3b;
+			X = ((X >> 16) ^ X) * 0x45d9f3b;
+			return (X >> 16) ^ X;
+		}
+
+		static inline float RngUniform(const uint32 State)
+		{
+			// Same approach as used in FRandomStream
+			float Output;
+			*((uint32*)(&Output)) = 0x3F800000U | (RngInt(State ^ 0x1c89a74a) >> 9);
+			return Output - 1.0f;
+		}
+
+		static inline float RngGaussian(const uint32 State)
+		{
+			return FMath::Sqrt(-2.0f * FMath::Loge(FMath::Max(RngUniform(State ^ 0xe427d90b), UE_SMALL_NUMBER))) * FMath::Cos(RngUniform(State ^ 0xd5444566) * UE_TWO_PI);
+		}
+
+		static inline void RngUpdate(uint32& State)
+		{
+			State = RngInt(State ^ 0x0c32dd74);
 		}
 	}
 
@@ -4393,6 +4414,23 @@ namespace UE::NNE::RuntimeBasic
 			MakeLinearWithRandomKaimingWeights(MemoryNum, OutputNum, WeightScale));
 	}
 
+	FModelBuilderElement FModelBuilder::MakeMemoryCellWithCompressedLinearRandomKaimingWeights(
+		const uint32 InputNum,
+		const uint32 OutputNum,
+		const uint32 MemoryNum,
+		const float WeightScale)
+	{
+		return MakeMemoryCell(
+			InputNum,
+			OutputNum,
+			MemoryNum,
+			MakeCompressedLinearWithRandomKaimingWeights(InputNum + MemoryNum, MemoryNum, WeightScale),
+			MakeCompressedLinearWithRandomKaimingWeights(InputNum + MemoryNum, OutputNum, WeightScale),
+			MakeCompressedLinearWithRandomKaimingWeights(InputNum + MemoryNum, MemoryNum, WeightScale),
+			MakeCompressedLinearWithRandomKaimingWeights(InputNum + MemoryNum, OutputNum, WeightScale),
+			MakeCompressedLinearWithRandomKaimingWeights(MemoryNum, OutputNum, WeightScale));
+	}
+
 	FModelBuilderElement FModelBuilder::MakeMemoryBackbone(
 		const uint32 MemoryNum,
 		const FModelBuilderElement& Prefix,
@@ -4589,7 +4627,7 @@ namespace UE::NNE::RuntimeBasic
 
 	void FModelBuilder::Reset()
 	{
-		Rng.Reset();
+		Rng = RngInitialState;
 		WeightsPool.Empty();
 		CompressedWeightsPool.Empty();
 		SizesPool.Empty();
@@ -4671,7 +4709,8 @@ namespace UE::NNE::RuntimeBasic
 
 		for (uint32 Idx = 0; Idx < InputSize * OutputSize; Idx++)
 		{
-			Values[Idx] = Std * Private::UniformToGaussian(Rng.FRand(), Rng.FRand());
+			Values[Idx] = Std * Private::RngGaussian(Rng);
+			Private::RngUpdate(Rng);
 		}
 		
 		return Values;
@@ -4694,7 +4733,8 @@ namespace UE::NNE::RuntimeBasic
 
 		for (uint32 Idx = 0; Idx < InputSize * OutputSize; Idx++)
 		{
-			Values[Idx] = Std * Private::UniformToGaussian(Rng.FRand(), Rng.FRand());
+			Values[Idx] = Std * Private::RngGaussian(Rng);
+			Private::RngUpdate(Rng);
 		}
 
 		// Find Min and Max
diff --git a/Engine/Plugins/Experimental/NNERuntimeBasicCpu/Source/NNERuntimeBasicCpu/Public/NNERuntimeBasicCpuBuilder.h b/Engine/Plugins/Experimental/NNERuntimeBasicCpu/Source/NNERuntimeBasicCpu/Public/NNERuntimeBasicCpuBuilder.h
--- a/Engine/Plugins/Experimental/NNERuntimeBasicCpu/Source/NNERuntimeBasicCpu/Public/NNERuntimeBasicCpuBuilder.h
+++ b/Engine/Plugins/Experimental/NNERuntimeBasicCpu/Source/NNERuntimeBasicCpu/Public/NNERuntimeBasicCpuBuilder.h
@@ -257,6 +257,20 @@ namespace UE::NNE::RuntimeBasic
 			const uint32 MemoryNum,
 			const float WeightScale = 1.0f);
 
+		/**
+		 * Make a new Memory Cell layer with compressed weights randomly initialized using the Kaiming method.
+		 *
+		 * @param InputNum			Number of normal inputs to the model
+		 * @param OutputNum			Number of normal outputs from the model
+		 * @param MemoryNum			The size of the memory vector used by the model
+		 * @param WeightScale		Scaling factor for the weight creation.
+		 */
+		FModelBuilderElement MakeMemoryCellWithCompressedLinearRandomKaimingWeights(
+			const uint32 InputNum,
+			const uint32 OutputNum,
+			const uint32 MemoryNum,
+			const float WeightScale = 1.0f);
+
 		/**
 		 * Make a new Memory Backbone layer.
 		 *
@@ -418,8 +432,10 @@ namespace UE::NNE::RuntimeBasic
 
 	private:
 
-		/** Random Number Stream for generating random weights */
-		FRandomStream Rng;
+		static constexpr uint32 RngInitialState = 0xafcc2b45;
+
+		/** Random Number State for generating random weights */
+		uint32 Rng = RngInitialState;
 
 		/** Pool of all weights data used by the `MakeWeights` functions. */
 		TArray<TArray<float>> WeightsPool;
