LearningAgents: Added function to add noise to observation vectors.

diff --git a/Engine/Plugins/Experimental/LearningAgents/Source/Learning/Private/LearningAction.cpp b/Engine/Plugins/Experimental/LearningAgents/Source/Learning/Private/LearningAction.cpp
--- a/Engine/Plugins/Experimental/LearningAgents/Source/Learning/Private/LearningAction.cpp
+++ b/Engine/Plugins/Experimental/LearningAgents/Source/Learning/Private/LearningAction.cpp
@@ -1264,7 +1264,7 @@ namespace UE::Learning::Action
 	}
 
 	void SampleVectorFromDistributionVector(
-		uint32& InOutSeed,
+		uint32& InOutRandomState,
 		TLearningArrayView<1, float> OutActionVector,
 		const TLearningArrayView<1, const float> ActionDistributionVector,
 		const FSchema& Schema,
@@ -1287,7 +1287,7 @@ namespace UE::Learning::Action
 
 			Random::SampleDistributionIndependantNormal(
 				OutActionVector,
-				InOutSeed,
+				InOutRandomState,
 				ActionDistributionVector.Slice(0, ValueNum),
 				ActionDistributionVector.Slice(ValueNum, ValueNum),
 				ActionNoiseScale);
@@ -1303,7 +1303,7 @@ namespace UE::Learning::Action
 
 			Random::SampleDistributionMultinoulli(
 				OutActionVector,
-				InOutSeed,
+				InOutRandomState,
 				ActionDistributionVector,
 				ActionNoiseScale);
 
@@ -1318,7 +1318,7 @@ namespace UE::Learning::Action
 
 			Random::SampleDistributionBernoulli(
 				OutActionVector,
-				InOutSeed,
+				InOutRandomState,
 				ActionDistributionVector,
 				ActionNoiseScale);
 
@@ -1338,7 +1338,7 @@ namespace UE::Learning::Action
 				const int32 SubElementActionDistributionVectorSize = Schema.GetActionDistributionVectorSize(SubElement);
 
 				SampleVectorFromDistributionVector(
-					InOutSeed,
+					InOutRandomState,
 					OutActionVector.Slice(SubElementActionVectorOffset, SubElementActionVectorSize),
 					ActionDistributionVector.Slice(SubElementActionDistributionVectorOffset, SubElementActionDistributionVectorSize),
 					Schema,
@@ -1368,7 +1368,7 @@ namespace UE::Learning::Action
 			// Sample which sub-element to generate
 			Random::SampleDistributionMultinoulli(
 				OutActionVector.Slice(SubElementActionVectorMax, Parameters.Elements.Num()),
-				InOutSeed,
+				InOutRandomState,
 				ActionDistributionVector.Slice(SubElementActionDistributionVectorTotal, Parameters.Elements.Num()),
 				ActionNoiseScale);
 
@@ -1384,7 +1384,7 @@ namespace UE::Learning::Action
 
 					// Sample Sub-Element
 					SampleVectorFromDistributionVector(
-						InOutSeed,
+						InOutRandomState,
 						OutActionVector.Slice(0, SubElementActionVectorSize),
 						ActionDistributionVector.Slice(SubElementActionDistributionVectorOffset, SubElementActionDistributionVectorSize),
 						Schema,
@@ -1414,7 +1414,7 @@ namespace UE::Learning::Action
 			// Sample which sub-elements to generate
 			Random::SampleDistributionBernoulli(
 				OutActionVector.Slice(SubElementActionVectorTotal, Parameters.Elements.Num()),
-				InOutSeed,
+				InOutRandomState,
 				ActionDistributionVector.Slice(SubElementActionDistributionVectorTotal, Parameters.Elements.Num()),
 				ActionNoiseScale);
 
@@ -1431,7 +1431,7 @@ namespace UE::Learning::Action
 
 					// Sample sub-elements
 					SampleVectorFromDistributionVector(
-						InOutSeed,
+						InOutRandomState,
 						OutActionVector.Slice(SubElementActionVectorOffset, SubElementActionVectorSize),
 						ActionDistributionVector.Slice(SubElementActionDistributionVectorOffset, SubElementActionDistributionVectorSize),
 						Schema,
@@ -1456,7 +1456,7 @@ namespace UE::Learning::Action
 			for (int32 ElementIdx = 0; ElementIdx < Parameters.Num; ElementIdx++)
 			{
 				SampleVectorFromDistributionVector(
-					InOutSeed,
+					InOutRandomState,
 					OutActionVector.Slice(ElementIdx * SubElementActionVectorSize, SubElementActionVectorSize),
 					ActionDistributionVector.Slice(ElementIdx * SubElementActionDistributionVectorSize, SubElementActionDistributionVectorSize),
 					Schema,
@@ -1472,7 +1472,7 @@ namespace UE::Learning::Action
 			const FSchemaEncodingParameters Parameters = Schema.GetEncoding(SchemaElement);
 
 			SampleVectorFromDistributionVector(
-				InOutSeed,
+				InOutRandomState,
 				OutActionVector,
 				ActionDistributionVector,
 				Schema,
diff --git a/Engine/Plugins/Experimental/LearningAgents/Source/Learning/Private/LearningObservation.cpp b/Engine/Plugins/Experimental/LearningAgents/Source/Learning/Private/LearningObservation.cpp
--- a/Engine/Plugins/Experimental/LearningAgents/Source/Learning/Private/LearningObservation.cpp
+++ b/Engine/Plugins/Experimental/LearningAgents/Source/Learning/Private/LearningObservation.cpp
@@ -1,6 +1,7 @@
 // Copyright Epic Games, Inc. All Rights Reserved.
 
 #include "LearningObservation.h"
+#include "LearningRandom.h"
 
 #include "NNERuntimeBasicCpuBuilder.h"
 
@@ -1623,4 +1624,190 @@ namespace UE::Learning::Observation
 		}
 	}
 
+	void AddGaussianNoiseToVector(
+		uint32& InOutRandomState,
+		TLearningArrayView<1, float> InOutObservationVector,
+		const FSchema& Schema,
+		const FSchemaElement SchemaElement,
+		const float NoiseScale)
+	{
+		UE_LEARNING_CHECK(Schema.IsValid(SchemaElement));
+
+		const EType SchemaElementType = Schema.GetType(SchemaElement);
+
+		const int32 ObservationVectorSize = InOutObservationVector.Num();
+		UE_LEARNING_CHECK(ObservationVectorSize == Schema.GetObservationVectorSize(SchemaElement));
+
+		switch (SchemaElementType)
+		{
+
+		case EType::Null:
+		{
+			return;
+		}
+
+		case EType::Continuous:
+		{
+			const FSchemaContinuousParameters SchemaParameters = Schema.GetContinuous(SchemaElement);
+			UE_LEARNING_CHECK(ObservationVectorSize == SchemaParameters.Num);
+
+			TLearningArray<1, float, TInlineAllocator<32>> NoiseValues;
+			NoiseValues.SetNumUninitialized({ SchemaParameters.Num });
+			Random::SampleGaussianArray(NoiseValues, InOutRandomState, 0.0f, NoiseScale);
+			for (int32 ValueIdx = 0; ValueIdx < SchemaParameters.Num; ValueIdx++)
+			{
+				InOutObservationVector[ValueIdx] += NoiseValues[ValueIdx];
+			}
+
+			return;
+		}
+
+		case EType::And:
+		{
+			const FSchemaAndParameters Parameters = Schema.GetAnd(SchemaElement);
+
+			int32 SubElementOffset = 0;
+			for (int32 SchemaElementIdx = 0; SchemaElementIdx < Parameters.Elements.Num(); SchemaElementIdx++)
+			{
+				const int32 SubElementSize = Schema.GetObservationVectorSize(Parameters.Elements[SchemaElementIdx]);
+
+				AddGaussianNoiseToVector(
+					InOutRandomState,
+					InOutObservationVector.Slice(SubElementOffset, SubElementSize),
+					Schema,
+					Parameters.Elements[SchemaElementIdx],
+					NoiseScale);
+
+				SubElementOffset += SubElementSize;
+			}
+			UE_LEARNING_CHECK(SubElementOffset == ObservationVectorSize);
+
+			return;
+		}
+
+		case EType::OrExclusive:
+		{
+			const FSchemaOrExclusiveParameters Parameters = Schema.GetOrExclusive(SchemaElement);
+			const int32 MaxSubElementSize = Private::GetMaxObservationVectorSize(Schema, Parameters.Elements);
+
+			int32 SchemaElementIndex = INDEX_NONE;
+			for (int32 SubElementIdx = 0; SubElementIdx < Parameters.Elements.Num(); SubElementIdx++)
+			{
+				UE_LEARNING_CHECK(InOutObservationVector[MaxSubElementSize + SubElementIdx] == 0.0f || InOutObservationVector[MaxSubElementSize + SubElementIdx] == 1.0f);
+				if (InOutObservationVector[MaxSubElementSize + SubElementIdx])
+				{
+					SchemaElementIndex = SubElementIdx;
+					break;
+				}
+			}
+			UE_LEARNING_CHECK(SchemaElementIndex != INDEX_NONE);
+
+			const int32 SubElementSize = Schema.GetObservationVectorSize(Parameters.Elements[SchemaElementIndex]);
+
+			AddGaussianNoiseToVector(
+				InOutRandomState,
+				InOutObservationVector.Slice(0, SubElementSize),
+				Schema,
+				Parameters.Elements[SchemaElementIndex],
+				NoiseScale);
+
+			return;
+		}
+
+		case EType::OrInclusive:
+		{
+			const FSchemaOrInclusiveParameters Parameters = Schema.GetOrInclusive(SchemaElement);
+			const int32 TotalSubElementSize = Private::GetTotalObservationVectorSize(Schema, Parameters.Elements);
+
+			int32 SubElementOffset = 0;
+			for (int32 SubElementIdx = 0; SubElementIdx < Parameters.Elements.Num(); SubElementIdx++)
+			{
+				const int32 SubElementSize = Schema.GetObservationVectorSize(Parameters.Elements[SubElementIdx]);
+
+				UE_LEARNING_CHECK(
+					InOutObservationVector[TotalSubElementSize + SubElementIdx] == 0.0f ||
+					InOutObservationVector[TotalSubElementSize + SubElementIdx] == 1.0f);
+
+				if (InOutObservationVector[TotalSubElementSize + SubElementIdx] == 1.0f)
+				{
+					AddGaussianNoiseToVector(
+						InOutRandomState,
+						InOutObservationVector.Slice(SubElementOffset, SubElementSize),
+						Schema,
+						Parameters.Elements[SubElementIdx],
+						NoiseScale);
+				}
+
+				SubElementOffset += SubElementSize;
+			}
+			UE_LEARNING_CHECK(SubElementOffset + Parameters.Elements.Num() == ObservationVectorSize);
+
+			return;
+		}
+
+		case EType::Array:
+		{
+			const FSchemaArrayParameters Parameters = Schema.GetArray(SchemaElement);
+			const int32 SubElementSize = Schema.GetObservationVectorSize(Parameters.Element);
+
+			for (int32 ElementIdx = 0; ElementIdx < Parameters.Num; ElementIdx++)
+			{
+				AddGaussianNoiseToVector(
+					InOutRandomState,
+					InOutObservationVector.Slice(ElementIdx* SubElementSize, SubElementSize),
+					Schema,
+					Parameters.Element,
+					NoiseScale);
+			}
+
+			return;
+		}
+
+		case EType::Set:
+		{
+			const FSchemaSetParameters Parameters = Schema.GetSet(SchemaElement);
+			const int32 SubElementSize = Schema.GetObservationVectorSize(Parameters.Element);
+
+			for (int32 SubElementIdx = 0; SubElementIdx < Parameters.MaxNum; SubElementIdx++)
+			{
+				UE_LEARNING_CHECK(
+					InOutObservationVector[SubElementSize * Parameters.MaxNum + SubElementIdx] == 0.0f ||
+					InOutObservationVector[SubElementSize * Parameters.MaxNum + SubElementIdx] == 1.0f);
+
+				if (InOutObservationVector[SubElementSize * Parameters.MaxNum + SubElementIdx] == 0.0f)
+				{
+					break;
+				}
+
+				AddGaussianNoiseToVector(
+					InOutRandomState,
+					InOutObservationVector.Slice(SubElementIdx * SubElementSize, SubElementSize),
+					Schema,
+					Parameters.Element,
+					NoiseScale);
+			}
+
+			return;
+		}
+
+		case EType::Encoding:
+		{
+			AddGaussianNoiseToVector(
+				InOutRandomState,
+				InOutObservationVector,
+				Schema,
+				Schema.GetEncoding(SchemaElement).Element,
+				NoiseScale);
+
+			return;
+		}
+
+		default:
+		{
+			UE_LEARNING_NOT_IMPLEMENTED();
+			return;
+		}
+		}
+	}
+
 }
\ No newline at end of file
diff --git a/Engine/Plugins/Experimental/LearningAgents/Source/Learning/Public/LearningAction.h b/Engine/Plugins/Experimental/LearningAgents/Source/Learning/Public/LearningAction.h
--- a/Engine/Plugins/Experimental/LearningAgents/Source/Learning/Public/LearningAction.h
+++ b/Engine/Plugins/Experimental/LearningAgents/Source/Learning/Public/LearningAction.h
@@ -495,7 +495,7 @@ namespace UE::Learning::Action
 	/**
 	 * Sample an action vector from a vector representing the action distribution.
 	 *
-	 * @param InOutSeed					Seed to use when sampling the action vector
+	 * @param InOutRandomState			Random State/Seed to use when sampling the action vector
 	 * @param OutActionVector			Output action vector
 	 * @param ActionDistributionVector	Input Action Distribution vector
 	 * @param Schema					Action Schema
@@ -503,7 +503,7 @@ namespace UE::Learning::Action
 	 * @param ActionNoiseScale			Scale factor for the noise. Set this to zero to always choose the mean (expected) action.
 	 */
 	LEARNING_API void SampleVectorFromDistributionVector(
-		uint32& InOutSeed,
+		uint32& InOutRandomState,
 		TLearningArrayView<1, float> OutActionVector,
 		const TLearningArrayView<1, const float> ActionDistributionVector,
 		const FSchema& Schema,
diff --git a/Engine/Plugins/Experimental/LearningAgents/Source/Learning/Public/LearningObservation.h b/Engine/Plugins/Experimental/LearningAgents/Source/Learning/Public/LearningObservation.h
--- a/Engine/Plugins/Experimental/LearningAgents/Source/Learning/Public/LearningObservation.h
+++ b/Engine/Plugins/Experimental/LearningAgents/Source/Learning/Public/LearningObservation.h
@@ -504,4 +504,21 @@ namespace UE::Learning::Observation
 		const FSchema& Schema,
 		const FSchemaElement SchemaElement,
 		const TLearningArrayView<1, const float> ObservationVector);
+
+	/**
+	 * Add Gaussian noise to all the continuous parts of an observation vector.
+	 *
+	 * @param InOutRandomState			Random State/Seed to use when sampling the noise
+	 * @param InOutObservationVector	Observation vector to add the noise to
+	 * @param Schema					Observation Schema
+	 * @param SchemaElement				Observation Schema Element
+	 * @param NoiseScale				Scale of the noise to apply
+	 */
+	LEARNING_API void AddGaussianNoiseToVector(
+		uint32& InOutRandomState,
+		TLearningArrayView<1, float> InOutObservationVector,
+		const FSchema& Schema,
+		const FSchemaElement SchemaElement,
+		const float NoiseScale = 1.0f);
+
 }
\ No newline at end of file
