Learning Agents: Fix MaxAgentNum bug with manager component not working on spatially loaded actors

diff --git a/Engine/Plugins/Experimental/LearningAgents/Source/LearningAgents/Private/LearningAgentsManager.cpp b/Engine/Plugins/Experimental/LearningAgents/Source/LearningAgents/Private/LearningAgentsManager.cpp
--- a/Engine/Plugins/Experimental/LearningAgents/Source/LearningAgents/Private/LearningAgentsManager.cpp
+++ b/Engine/Plugins/Experimental/LearningAgents/Source/LearningAgents/Private/LearningAgentsManager.cpp
@@ -13,14 +13,9 @@ ULearningAgentsManager::ULearningAgentsManager() : Super(FObjectInitializer::Get
 ULearningAgentsManager::ULearningAgentsManager(FVTableHelper& Helper) : Super(Helper) {}
 ULearningAgentsManager::~ULearningAgentsManager() = default;
 
-void ULearningAgentsManager::PostInitProperties()
+void ULearningAgentsManager::OnRegister()
 {
-	Super::PostInitProperties();
-
-	if (HasAnyFlags(RF_ClassDefaultObject))
-	{
-		return;
-	}
+	Super::OnRegister();
 
 	// Pre-populate the vacant ids
 	OnEventAgentIds.Reserve(MaxAgentNum);
diff --git a/Engine/Plugins/Experimental/LearningAgents/Source/LearningAgents/Public/LearningAgentsManager.h b/Engine/Plugins/Experimental/LearningAgents/Source/LearningAgents/Public/LearningAgentsManager.h
--- a/Engine/Plugins/Experimental/LearningAgents/Source/LearningAgents/Public/LearningAgentsManager.h
+++ b/Engine/Plugins/Experimental/LearningAgents/Source/LearningAgents/Public/LearningAgentsManager.h
@@ -36,7 +36,7 @@ public:
 	ULearningAgentsManager(FVTableHelper& Helper);
 	virtual ~ULearningAgentsManager();
 
-	virtual void PostInitProperties() override;
+	virtual void OnRegister() override;
 
 	virtual void TickComponent(float DeltaTime, enum ELevelTick TickType, FActorComponentTickFunction* ThisTickFunction) override;
 
