LearningAgents: Fixed constness on implicit conversion to TArrayView on TMultiArray and TMultiArrayView and added explicit conversion.

diff --git a/Engine/Plugins/Experimental/LearningAgents/Source/Learning/Public/MultiArray.h b/Engine/Plugins/Experimental/LearningAgents/Source/Learning/Public/MultiArray.h
--- a/Engine/Plugins/Experimental/LearningAgents/Source/Learning/Public/MultiArray.h
+++ b/Engine/Plugins/Experimental/LearningAgents/Source/Learning/Public/MultiArray.h
@@ -914,12 +914,30 @@ public:
 		return TMultiArrayView<1, const ElementType, bIsChecked, bIsRestrict, SizeType>(GetData(), Num());
 	}
 
+	/** Explicit conversion to TArrayView. */
+	FORCEINLINE TArrayView<ElementType, SizeType> ArrayView()
+	{
+		return TArrayView<ElementType, SizeType>(GetData(), Num());
+	}
+
+	/** Explicit conversion to TArrayView. */
+	FORCEINLINE TArrayView<const ElementType, SizeType> ArrayView() const
+	{
+		return TArrayView<const ElementType, SizeType>(GetData(), Num());
+	}
+
 	/** Implicit cast to TArrayView. */
-	FORCEINLINE operator TArrayView<ElementType, SizeType>() const
+	FORCEINLINE operator TArrayView<ElementType, SizeType>()
 	{
 		return TArrayView<ElementType, SizeType>(GetData(), Num());
 	}
 
+	/** Implicit cast to TArrayView. */
+	FORCEINLINE operator TArrayView<const ElementType, SizeType>() const
+	{
+		return TArrayView<const ElementType, SizeType>(GetData(), Num());
+	}
+
 public:
 
 	/**
diff --git a/Engine/Plugins/Experimental/LearningAgents/Source/Learning/Public/MultiArrayView.h b/Engine/Plugins/Experimental/LearningAgents/Source/Learning/Public/MultiArrayView.h
--- a/Engine/Plugins/Experimental/LearningAgents/Source/Learning/Public/MultiArrayView.h
+++ b/Engine/Plugins/Experimental/LearningAgents/Source/Learning/Public/MultiArrayView.h
@@ -710,10 +710,28 @@ public:
 		return DataPtr[Index];
 	}
 
+	/** Explicit conversion to TArrayView. */
+	FORCEINLINE TArrayView<ElementType, SizeType> ArrayView()
+	{
+		return TArrayView<ElementType, SizeType>(GetData(), Num());
+	}
+
+	/** Explicit conversion to TArrayView. */
+	FORCEINLINE TArrayView<const ElementType, SizeType> ArrayView() const
+	{
+		return TArrayView<const ElementType, SizeType>(GetData(), Num());
+	}
+
+	/** Implicit cast to TArrayView. */
+	FORCEINLINE operator TArrayView<ElementType, SizeType>()
+	{
+		return TArrayView<ElementType, SizeType>(DataPtr, ArrayShape[0]);
+	}
+
 	/** Implicit cast to TArrayView. */
-	FORCEINLINE operator TArrayView<ElementType, SizeType>() const
+	FORCEINLINE operator TArrayView<const ElementType, SizeType>() const
 	{
-		return TArrayView(DataPtr, ArrayShape[0]);
+		return TArrayView<const ElementType, SizeType>(DataPtr, ArrayShape[0]);
 	}
 
 public:
