LearningAgents: Swapped out FMonitoredProcess for something custom due to issues on linux.

diff --git a/Engine/Plugins/Experimental/LearningAgents/Source/LearningAgentsTraining/Private/LearningAgentsImitationTrainer.cpp b/Engine/Plugins/Experimental/LearningAgents/Source/LearningAgentsTraining/Private/LearningAgentsImitationTrainer.cpp
--- a/Engine/Plugins/Experimental/LearningAgents/Source/LearningAgentsTraining/Private/LearningAgentsImitationTrainer.cpp
+++ b/Engine/Plugins/Experimental/LearningAgents/Source/LearningAgentsTraining/Private/LearningAgentsImitationTrainer.cpp
@@ -115,12 +115,6 @@ void ULearningAgentsImitationTrainer::BeginTraining(
 	const FLearningAgentsImitationTrainerTrainingSettings& ImitationTrainerTrainingSettings,
 	const FLearningAgentsTrainerPathSettings& ImitationTrainerPathSettings)
 {
-	if (!PLATFORM_WINDOWS)
-	{
-		UE_LOG(LogLearning, Error, TEXT("%s: Training currently only supported on Windows."), *GetName());
-		return;
-	}
-
 	if (!IsSetup())
 	{
 		UE_LOG(LogLearning, Error, TEXT("%s: Setup not complete."), *GetName());
diff --git a/Engine/Plugins/Experimental/LearningAgents/Source/LearningAgentsTraining/Private/LearningAgentsTrainer.cpp b/Engine/Plugins/Experimental/LearningAgents/Source/LearningAgentsTraining/Private/LearningAgentsTrainer.cpp
--- a/Engine/Plugins/Experimental/LearningAgents/Source/LearningAgentsTraining/Private/LearningAgentsTrainer.cpp
+++ b/Engine/Plugins/Experimental/LearningAgents/Source/LearningAgentsTraining/Private/LearningAgentsTrainer.cpp
@@ -372,12 +372,6 @@ void ULearningAgentsTrainer::BeginTraining(
 	const FLearningAgentsTrainerPathSettings& TrainerPathSettings,
 	const bool bResetAgentsOnBegin)
 {
-	if (!PLATFORM_WINDOWS)
-	{
-		UE_LOG(LogLearning, Error, TEXT("%s: Training currently only supported on Windows."), *GetName());
-		return;
-	}
-
 	if (!IsSetup())
 	{
 		UE_LOG(LogLearning, Error, TEXT("%s: Setup not complete."), *GetName());
diff --git a/Engine/Plugins/Experimental/LearningAgents/Source/LearningTraining/Private/LearningImitationTrainer.cpp b/Engine/Plugins/Experimental/LearningAgents/Source/LearningTraining/Private/LearningImitationTrainer.cpp
--- a/Engine/Plugins/Experimental/LearningAgents/Source/LearningTraining/Private/LearningImitationTrainer.cpp
+++ b/Engine/Plugins/Experimental/LearningAgents/Source/LearningTraining/Private/LearningImitationTrainer.cpp
@@ -10,7 +10,6 @@
 #include "LearningSharedMemoryTraining.h"
 #include "LearningSocketTraining.h"
 
-#include "Misc/MonitoredProcess.h"
 #include "Misc/Guid.h"
 #include "Misc/FileHelper.h"
 #include "Misc/Paths.h"
@@ -212,20 +211,10 @@ namespace UE::Learning
 			*FileManager.ConvertToAbsolutePathForExternalAppForRead(*(PythonContentPath / TEXT("train_imitation.py"))),
 			*FileManager.ConvertToAbsolutePathForExternalAppForRead(*ConfigPath));
 
-		TrainingProcess = MakeShared<FMonitoredProcess>(
+		TrainingProcess.Launch(
 			FileManager.ConvertToAbsolutePathForExternalAppForRead(*PythonExecutablePath),
 			CommandLineArguments,
-			!(TrainingProcessFlags & ESubprocessFlags::ShowWindow),
-			!(TrainingProcessFlags & ESubprocessFlags::NoRedirectOutput));
-
-		if (!(TrainingProcessFlags & ESubprocessFlags::NoRedirectOutput))
-		{
-			TrainingProcess->OnCanceled().BindRaw(this, &FSharedMemoryImitationTrainer::HandleTrainingProcessCanceled);
-			TrainingProcess->OnCompleted().BindRaw(this, &FSharedMemoryImitationTrainer::HandleTrainingProcessCompleted);
-			TrainingProcess->OnOutput().BindStatic(&FSharedMemoryImitationTrainer::HandleTrainingProcessOutput);
-		}
-
-		TrainingProcess->Launch();
+			TrainingProcessFlags);
 	}
 
 	FSharedMemoryImitationTrainer::~FSharedMemoryImitationTrainer()
@@ -250,9 +239,9 @@ namespace UE::Learning
 		const ELogSetting LogSettings)
 	{
 		return SharedMemoryTraining::RecvNetwork(
-			TrainingProcess.Get(),
 			Controls.View,
 			OutNetwork,
+			TrainingProcess,
 			SharedMemoryTraining::EControls::PolicySignal,
 			Policy.View,
 			Timeout,
@@ -267,9 +256,9 @@ namespace UE::Learning
 		const ELogSetting LogSettings)
 	{
 		return SharedMemoryTraining::RecvNetwork(
-			TrainingProcess.Get(),
 			Controls.View,
 			OutNetwork,
+			TrainingProcess,
 			SharedMemoryTraining::EControls::EncoderSignal,
 			Encoder.View,
 			Timeout,
@@ -284,9 +273,9 @@ namespace UE::Learning
 		const ELogSetting LogSettings)
 	{
 		return SharedMemoryTraining::RecvNetwork(
-			TrainingProcess.Get(),
 			Controls.View,
 			OutNetwork,
+			TrainingProcess,
 			SharedMemoryTraining::EControls::DecoderSignal,
 			Decoder.View,
 			Timeout,
@@ -301,9 +290,9 @@ namespace UE::Learning
 		const ELogSetting LogSettings)
 	{
 		return SharedMemoryTraining::SendNetwork(
-			TrainingProcess.Get(),
 			Controls.View,
 			Policy.View,
+			TrainingProcess,
 			SharedMemoryTraining::EControls::PolicySignal,
 			Network,
 			Timeout,
@@ -318,9 +307,9 @@ namespace UE::Learning
 		const ELogSetting LogSettings)
 	{
 		return SharedMemoryTraining::SendNetwork(
-			TrainingProcess.Get(),
 			Controls.View,
 			Encoder.View,
+			TrainingProcess,
 			SharedMemoryTraining::EControls::EncoderSignal,
 			Network,
 			Timeout,
@@ -335,9 +324,9 @@ namespace UE::Learning
 		const ELogSetting LogSettings)
 	{
 		return SharedMemoryTraining::SendNetwork(
-			TrainingProcess.Get(),
 			Controls.View,
 			Decoder.View,
+			TrainingProcess,
 			SharedMemoryTraining::EControls::DecoderSignal,
 			Network,
 			Timeout,
@@ -354,12 +343,12 @@ namespace UE::Learning
 		const ELogSetting LogSettings)
 	{
 		return SharedMemoryTraining::SendExperience(
-			TrainingProcess.Get(),
 			EpisodeStarts.View,
 			EpisodeLengths.View,
 			Observations.View,
 			Actions.View,
 			Controls.View,
+			TrainingProcess,
 			EpisodeStartsExperience,
 			EpisodeLengthsExperience,
 			ObservationsExperience,
@@ -383,7 +372,7 @@ namespace UE::Learning
 		const float SleepTime = 0.001f;
 		float WaitTime = 0.0f;
 
-		while (TrainingProcess.IsValid())
+		while (TrainingProcess.IsRunning())
 		{
 			FPlatformProcess::Sleep(SleepTime);
 			WaitTime += SleepTime;
@@ -399,12 +388,7 @@ namespace UE::Learning
 
 	void FSharedMemoryImitationTrainer::Terminate()
 	{
-		if (TrainingProcess.IsValid())
-		{
-			TrainingProcess->Cancel(true);
-		}
-
-		TrainingProcess.Reset();
+		TrainingProcess.Terminate();
 
 		if (Policy.Region)
 		{
@@ -412,32 +396,6 @@ namespace UE::Learning
 		}
 	}
 
-	void FSharedMemoryImitationTrainer::HandleTrainingProcessCanceled()
-	{
-		UE_LOG(LogLearning, Warning, TEXT("Training process canceled"));
-
-		TrainingProcess.Reset();
-	}
-
-	void FSharedMemoryImitationTrainer::HandleTrainingProcessCompleted(int32 ReturnCode)
-	{
-		if (ReturnCode != 0)
-		{
-			UE_LOG(LogLearning, Warning, TEXT("Training Process finished with warnings or errors"));
-		}
-
-		TrainingProcess.Reset();
-	}
-
-	void FSharedMemoryImitationTrainer::HandleTrainingProcessOutput(FString Output)
-	{
-		if (!Output.IsEmpty())
-		{
-			UE_LOG(LogLearning, Display, TEXT("Training Process: %s"), *Output);
-		}
-	}
-
-
 	FSocketImitationTrainerServerProcess::FSocketImitationTrainerServerProcess(
 		const FString& PythonExecutablePath,
 		const FString& ExtraSitePackagesPath,
@@ -461,20 +419,10 @@ namespace UE::Learning
 			*FileManager.ConvertToAbsolutePathForExternalAppForRead(*IntermediatePath),
 			LogSettings == ELogSetting::Normal ? 1 : 0);
 
-		TrainingProcess = MakeShared<FMonitoredProcess>(
+		TrainingProcess.Launch(
 			FileManager.ConvertToAbsolutePathForExternalAppForRead(*PythonExecutablePath),
 			CommandLineArguments,
-			!(TrainingProcessFlags & ESubprocessFlags::ShowWindow),
-			!(TrainingProcessFlags & ESubprocessFlags::NoRedirectOutput));
-
-		if (!(TrainingProcessFlags & ESubprocessFlags::NoRedirectOutput))
-		{
-			TrainingProcess->OnCanceled().BindRaw(this, &FSocketImitationTrainerServerProcess::HandleTrainingProcessCanceled);
-			TrainingProcess->OnCompleted().BindRaw(this, &FSocketImitationTrainerServerProcess::HandleTrainingProcessCompleted);
-			TrainingProcess->OnOutput().BindStatic(&FSocketImitationTrainerServerProcess::HandleTrainingProcessOutput);
-		}
-
-		TrainingProcess->Launch();
+			TrainingProcessFlags);
 	}
 
 	FSocketImitationTrainerServerProcess::~FSocketImitationTrainerServerProcess()
@@ -484,7 +432,7 @@ namespace UE::Learning
 
 	bool FSocketImitationTrainerServerProcess::IsRunning() const
 	{
-		return TrainingProcess.IsValid();
+		return TrainingProcess.IsRunning();
 	}
 
 	bool FSocketImitationTrainerServerProcess::Wait(float Timeout)
@@ -492,7 +440,7 @@ namespace UE::Learning
 		const float SleepTime = 0.001f;
 		float WaitTime = 0.0f;
 
-		while (TrainingProcess.IsValid())
+		while (TrainingProcess.IsRunning())
 		{
 			FPlatformProcess::Sleep(SleepTime);
 			WaitTime += SleepTime;
@@ -508,37 +456,7 @@ namespace UE::Learning
 
 	void FSocketImitationTrainerServerProcess::Terminate()
 	{
-		if (TrainingProcess.IsValid())
-		{
-			TrainingProcess->Cancel(true);
-		}
-
-		TrainingProcess.Reset();
-	}
-
-	void FSocketImitationTrainerServerProcess::HandleTrainingProcessCanceled()
-	{
-		UE_LOG(LogLearning, Warning, TEXT("Training process canceled"));
-
-		TrainingProcess.Reset();
-	}
-
-	void FSocketImitationTrainerServerProcess::HandleTrainingProcessCompleted(int32 ReturnCode)
-	{
-		if (ReturnCode != 0)
-		{
-			UE_LOG(LogLearning, Warning, TEXT("Training Process finished with warnings or errors"));
-		}
-
-		TrainingProcess.Reset();
-	}
-
-	void FSocketImitationTrainerServerProcess::HandleTrainingProcessOutput(FString Output)
-	{
-		if (!Output.IsEmpty())
-		{
-			UE_LOG(LogLearning, Display, TEXT("Training Process: %s"), *Output);
-		}
+		TrainingProcess.Terminate();
 	}
 
 	FSocketImitationTrainer::FSocketImitationTrainer(
diff --git a/Engine/Plugins/Experimental/LearningAgents/Source/LearningTraining/Private/LearningPPOTrainer.cpp b/Engine/Plugins/Experimental/LearningAgents/Source/LearningTraining/Private/LearningPPOTrainer.cpp
--- a/Engine/Plugins/Experimental/LearningAgents/Source/LearningTraining/Private/LearningPPOTrainer.cpp
+++ b/Engine/Plugins/Experimental/LearningAgents/Source/LearningTraining/Private/LearningPPOTrainer.cpp
@@ -13,7 +13,6 @@
 #include "LearningObservation.h"
 #include "LearningAction.h"
 
-#include "Misc/MonitoredProcess.h"
 #include "Misc/Guid.h"
 #include "Misc/FileHelper.h"
 #include "Misc/CommandLine.h"
@@ -198,24 +197,10 @@ namespace UE::Learning
 				SubprocessCommandLine += FString(TEXT(" -LearningMemoryStatesGuid ")) + MemoryStates.Guid.ToString();
 				SubprocessCommandLine += FString(TEXT(" -LearningRewardsGuid ")) + Rewards.Guid.ToString();
 
-				const TSharedPtr<FMonitoredProcess> Subprocess = MakeShared<FMonitoredProcess>(
-					FPlatformProcess::ExecutablePath(), 
-					SubprocessCommandLine, 
-					!(MultiProcessFlags & ESubprocessFlags::ShowWindow),
-					!(MultiProcessFlags & ESubprocessFlags::NoRedirectOutput));
-
-				if (!(MultiProcessFlags & ESubprocessFlags::NoRedirectOutput))
-				{
-					Subprocess->OnCanceled().BindRaw(this, &FSharedMemoryPPOTrainer::HandleSubprocessCanceled);
-					Subprocess->OnCompleted().BindRaw(this, &FSharedMemoryPPOTrainer::HandleSubprocessCompleted);
-					Subprocess->OnOutput().BindStatic(&FSharedMemoryPPOTrainer::HandleSubprocessOutput);
-				}
-
-				Subprocess->Launch();
+				FSubprocess& Subprocess = ExperienceGatheringSubprocesses.AddDefaulted_GetRef();
+				Subprocess.Launch(FPlatformProcess::ExecutablePath(), SubprocessCommandLine, MultiProcessFlags);
 
 				UE_LOG(LogLearning, Display, TEXT("Subprocess Command: %s %s"), FPlatformProcess::ExecutablePath(), *SubprocessCommandLine);
-
-				ExperienceGatheringSubprocesses.Emplace(Subprocess);
 			}
 
 			// Write Config
@@ -308,20 +293,10 @@ namespace UE::Learning
 				*FileManager.ConvertToAbsolutePathForExternalAppForRead(*(PythonContentPath / TEXT("train_ppo.py"))), 
 				*FileManager.ConvertToAbsolutePathForExternalAppForRead(*ConfigPath));
 
-			TrainingProcess = MakeShared<FMonitoredProcess>(
+			TrainingProcess.Launch(
 				FileManager.ConvertToAbsolutePathForExternalAppForRead(*PythonExecutablePath),
-				CommandLineArguments, 
-				!(TrainingProcessFlags & ESubprocessFlags::ShowWindow),
-				!(TrainingProcessFlags & ESubprocessFlags::NoRedirectOutput));
-
-			if (!(TrainingProcessFlags & ESubprocessFlags::NoRedirectOutput))
-			{
-				TrainingProcess->OnCanceled().BindRaw(this, &FSharedMemoryPPOTrainer::HandleTrainingProcessCanceled);
-				TrainingProcess->OnCompleted().BindRaw(this, &FSharedMemoryPPOTrainer::HandleTrainingProcessCompleted);
-				TrainingProcess->OnOutput().BindStatic(&FSharedMemoryPPOTrainer::HandleTrainingProcessOutput);
-			}
-
-			TrainingProcess->Launch();
+				CommandLineArguments,
+				TrainingProcessFlags);
 		}
 		else
 		{
@@ -371,8 +346,9 @@ namespace UE::Learning
 		const float SleepTime = 0.001f;
 		float WaitTime = 0.0f;
 
-		while (TrainingProcess.IsValid())
+		while (TrainingProcess.IsRunning())
 		{
+			TrainingProcess.Update();
 			FPlatformProcess::Sleep(SleepTime);
 			WaitTime += SleepTime;
 
@@ -387,12 +363,7 @@ namespace UE::Learning
 
 	void FSharedMemoryPPOTrainer::Terminate()
 	{
-		if (TrainingProcess.IsValid())
-		{
-			TrainingProcess->Cancel(true);
-		}
-
-		TrainingProcess.Reset();
+		TrainingProcess.Terminate();
 
 		if (Policy.Region)
 		{
@@ -411,10 +382,12 @@ namespace UE::Learning
 		FRWLock* NetworkLock,
 		const ELogSetting LogSettings)
 	{
+		TrainingProcess.Update();
+
 		return SharedMemoryTraining::RecvNetwork(
-			TrainingProcess.Get(),
 			Controls.View[ProcessIdx],
 			OutNetwork,
+			TrainingProcess,
 			SharedMemoryTraining::EControls::PolicySignal,
 			Policy.View,
 			Timeout,
@@ -428,10 +401,12 @@ namespace UE::Learning
 		FRWLock* NetworkLock,
 		const ELogSetting LogSettings)
 	{
+		TrainingProcess.Update();
+
 		return SharedMemoryTraining::RecvNetwork(
-			TrainingProcess.Get(),
 			Controls.View[ProcessIdx],
 			OutNetwork,
+			TrainingProcess,
 			SharedMemoryTraining::EControls::CriticSignal,
 			Critic.View,
 			Timeout,
@@ -445,10 +420,12 @@ namespace UE::Learning
 		FRWLock* NetworkLock,
 		const ELogSetting LogSettings)
 	{
+		TrainingProcess.Update();
+
 		return SharedMemoryTraining::RecvNetwork(
-			TrainingProcess.Get(),
 			Controls.View[ProcessIdx],
 			OutNetwork,
+			TrainingProcess,
 			SharedMemoryTraining::EControls::EncoderSignal,
 			Encoder.View,
 			Timeout,
@@ -462,10 +439,12 @@ namespace UE::Learning
 		FRWLock* NetworkLock,
 		const ELogSetting LogSettings)
 	{
+		TrainingProcess.Update();
+
 		return SharedMemoryTraining::RecvNetwork(
-			TrainingProcess.Get(),
 			Controls.View[ProcessIdx],
 			OutNetwork,
+			TrainingProcess,
 			SharedMemoryTraining::EControls::DecoderSignal,
 			Decoder.View,
 			Timeout,
@@ -479,10 +458,12 @@ namespace UE::Learning
 		FRWLock* NetworkLock,
 		const ELogSetting LogSettings)
 	{
+		TrainingProcess.Update();
+
 		return SharedMemoryTraining::SendNetwork(
-			TrainingProcess.Get(),
 			Controls.View[ProcessIdx],
 			Policy.View,
+			TrainingProcess,
 			SharedMemoryTraining::EControls::PolicySignal,
 			Network,
 			Timeout,
@@ -496,10 +477,12 @@ namespace UE::Learning
 		FRWLock* NetworkLock,
 		const ELogSetting LogSettings)
 	{
+		TrainingProcess.Update();
+
 		return SharedMemoryTraining::SendNetwork(
-			TrainingProcess.Get(),
 			Controls.View[ProcessIdx],
 			Critic.View,
+			TrainingProcess,
 			SharedMemoryTraining::EControls::CriticSignal,
 			Network,
 			Timeout,
@@ -513,10 +496,12 @@ namespace UE::Learning
 		FRWLock* NetworkLock,
 		const ELogSetting LogSettings)
 	{
+		TrainingProcess.Update();
+
 		return SharedMemoryTraining::SendNetwork(
-			TrainingProcess.Get(),
 			Controls.View[ProcessIdx],
 			Encoder.View,
+			TrainingProcess,
 			SharedMemoryTraining::EControls::EncoderSignal,
 			Network,
 			Timeout,
@@ -530,10 +515,12 @@ namespace UE::Learning
 		FRWLock* NetworkLock,
 		const ELogSetting LogSettings)
 	{
+		TrainingProcess.Update();
+
 		return SharedMemoryTraining::SendNetwork(
-			TrainingProcess.Get(),
 			Controls.View[ProcessIdx],
 			Decoder.View,
+			TrainingProcess,
 			SharedMemoryTraining::EControls::DecoderSignal,
 			Network,
 			Timeout,
@@ -546,8 +533,9 @@ namespace UE::Learning
 		const float Timeout,
 		const ELogSetting LogSettings)
 	{
+		TrainingProcess.Update();
+
 		return SharedMemoryTraining::SendExperience(
-			TrainingProcess.Get(),
 			EpisodeStarts.View[ProcessIdx],
 			EpisodeLengths.View[ProcessIdx],
 			EpisodeCompletionModes.View[ProcessIdx],
@@ -558,6 +546,7 @@ namespace UE::Learning
 			MemoryStates.View[ProcessIdx],
 			Rewards.View[ProcessIdx],
 			Controls.View[ProcessIdx],
+			TrainingProcess,
 			ReplayBuffer,
 			Timeout,
 			LogSettings);
@@ -584,53 +573,6 @@ namespace UE::Learning
 		}
 	}
 
-	void FSharedMemoryPPOTrainer::HandleSubprocessCanceled()
-	{
-		UE_LOG(LogLearning, Warning, TEXT("Subprocess canceled"));
-	}
-
-	void FSharedMemoryPPOTrainer::HandleSubprocessCompleted(int32 ReturnCode)
-	{
-		if (ReturnCode != 0)
-		{
-			UE_LOG(LogLearning, Warning, TEXT("Subprocess finished with warnings or errors"));
-		}
-	}
-
-	void FSharedMemoryPPOTrainer::HandleSubprocessOutput(FString Output)
-	{
-		if (!Output.IsEmpty())
-		{
-			UE_LOG(LogLearning, Display, TEXT("Subprocess: %s"), *Output);
-		}
-	}
-
-
-	void FSharedMemoryPPOTrainer::HandleTrainingProcessCanceled()
-	{
-		UE_LOG(LogLearning, Warning, TEXT("Training process canceled"));
-
-		TrainingProcess.Reset();
-	}
-
-	void FSharedMemoryPPOTrainer::HandleTrainingProcessCompleted(int32 ReturnCode)
-	{
-		if (ReturnCode != 0)
-		{
-			UE_LOG(LogLearning, Warning, TEXT("Training Process finished with warnings or errors"));
-		}
-
-		TrainingProcess.Reset();
-	}
-
-	void FSharedMemoryPPOTrainer::HandleTrainingProcessOutput(FString Output)
-	{
-		if (!Output.IsEmpty())
-		{
-			UE_LOG(LogLearning, Display, TEXT("Training Process: %s"), *Output);
-		}
-	}
-
 	FSocketPPOTrainerServerProcess::FSocketPPOTrainerServerProcess(
 		const FString& PythonExecutablePath,
 		const FString& ExtraSitePackagesPath,
@@ -653,20 +595,10 @@ namespace UE::Learning
 			*FileManager.ConvertToAbsolutePathForExternalAppForRead(*IntermediatePath), 
 			LogSettings == ELogSetting::Normal ? 1 : 0);
 
-		TrainingProcess = MakeShared<FMonitoredProcess>(
+		TrainingProcess.Launch(
 			FileManager.ConvertToAbsolutePathForExternalAppForRead(*PythonExecutablePath),
 			CommandLineArguments,
-			!(TrainingProcessFlags & ESubprocessFlags::ShowWindow),
-			!(TrainingProcessFlags & ESubprocessFlags::NoRedirectOutput));
-
-		if (!(TrainingProcessFlags & ESubprocessFlags::NoRedirectOutput))
-		{
-			TrainingProcess->OnCanceled().BindRaw(this, &FSocketPPOTrainerServerProcess::HandleTrainingProcessCanceled);
-			TrainingProcess->OnCompleted().BindRaw(this, &FSocketPPOTrainerServerProcess::HandleTrainingProcessCompleted);
-			TrainingProcess->OnOutput().BindStatic(&FSocketPPOTrainerServerProcess::HandleTrainingProcessOutput);
-		}
-
-		TrainingProcess->Launch();
+			TrainingProcessFlags);
 	}
 
 	FSocketPPOTrainerServerProcess::~FSocketPPOTrainerServerProcess()
@@ -676,7 +608,7 @@ namespace UE::Learning
 
 	bool FSocketPPOTrainerServerProcess::IsRunning() const
 	{
-		return TrainingProcess.IsValid();
+		return TrainingProcess.IsRunning();
 	}
 
 	bool FSocketPPOTrainerServerProcess::Wait(float Timeout)
@@ -684,8 +616,9 @@ namespace UE::Learning
 		const float SleepTime = 0.001f;
 		float WaitTime = 0.0f;
 
-		while (TrainingProcess.IsValid())
+		while (TrainingProcess.IsRunning())
 		{
+			TrainingProcess.Update();
 			FPlatformProcess::Sleep(SleepTime);
 			WaitTime += SleepTime;
 
@@ -700,37 +633,7 @@ namespace UE::Learning
 
 	void FSocketPPOTrainerServerProcess::Terminate()
 	{
-		if (TrainingProcess.IsValid())
-		{
-			TrainingProcess->Cancel(true);
-		}
-
-		TrainingProcess.Reset();
-	}
-
-	void FSocketPPOTrainerServerProcess::HandleTrainingProcessCanceled()
-	{
-		UE_LOG(LogLearning, Warning, TEXT("Training process canceled"));
-
-		TrainingProcess.Reset();
-	}
-
-	void FSocketPPOTrainerServerProcess::HandleTrainingProcessCompleted(int32 ReturnCode)
-	{
-		if (ReturnCode != 0)
-		{
-			UE_LOG(LogLearning, Warning, TEXT("Training Process finished with warnings or errors"));
-		}
-
-		TrainingProcess.Reset();
-	}
-
-	void FSocketPPOTrainerServerProcess::HandleTrainingProcessOutput(FString Output)
-	{
-		if (!Output.IsEmpty())
-		{
-			UE_LOG(LogLearning, Display, TEXT("Training Process: %s"), *Output);
-		}
+		TrainingProcess.Terminate();
 	}
 
 	FSocketPPOTrainer::FSocketPPOTrainer(
diff --git a/Engine/Plugins/Experimental/LearningAgents/Source/LearningTraining/Private/LearningSharedMemoryTraining.cpp b/Engine/Plugins/Experimental/LearningAgents/Source/LearningTraining/Private/LearningSharedMemoryTraining.cpp
--- a/Engine/Plugins/Experimental/LearningAgents/Source/LearningTraining/Private/LearningSharedMemoryTraining.cpp
+++ b/Engine/Plugins/Experimental/LearningAgents/Source/LearningTraining/Private/LearningSharedMemoryTraining.cpp
@@ -7,7 +7,6 @@
 #include "LearningExperience.h"
 #include "LearningCompletion.h"
 
-#include "Misc/MonitoredProcess.h"
 #include "HAL/PlatformProcess.h"
 
 namespace UE::Learning::SharedMemoryTraining
@@ -29,9 +28,9 @@ namespace UE::Learning::SharedMemoryTraining
 	}
 
 	ETrainerResponse RecvNetwork(
-		FMonitoredProcess* Process, 
 		TLearningArrayView<1, volatile int32> Controls,
 		ULearningNeuralNetworkData& OutNetwork,
+		const FSubprocess& Process,
 		const EControls Signal,
 		const TLearningArrayView<1, const uint8> NetworkData,
 		const float Timeout,
@@ -53,7 +52,7 @@ namespace UE::Learning::SharedMemoryTraining
 			}
 
 			// Check if the process has exited
-			if (!Process || !Process->Update() || Process->GetReturnCode() == 1)
+			if (!Process.IsRunning())
 			{
 				return ETrainerResponse::Unexpected;
 			}
@@ -112,9 +111,9 @@ namespace UE::Learning::SharedMemoryTraining
 	}
 
 	ETrainerResponse SendNetwork(
-		FMonitoredProcess* Process,
 		TLearningArrayView<1, volatile int32> Controls,
 		TLearningArrayView<1, uint8> NetworkData,
+		const FSubprocess& Process,
 		const EControls Signal,
 		const ULearningNeuralNetworkData& Network,
 		const float Timeout,
@@ -128,7 +127,7 @@ namespace UE::Learning::SharedMemoryTraining
 		while (!Controls[(uint8)Signal])
 		{
 			// Check if the process has exited
-			if (!Process || !Process->Update() || Process->GetReturnCode() == 1)
+			if (!Process.IsRunning())
 			{
 				return ETrainerResponse::Unexpected;
 			}
@@ -179,7 +178,6 @@ namespace UE::Learning::SharedMemoryTraining
 	}
 
 	ETrainerResponse SendExperience(
-		FMonitoredProcess* Process,
 		TLearningArrayView<1, int32> EpisodeStarts,
 		TLearningArrayView<1, int32> EpisodeLengths,
 		TLearningArrayView<1, ECompletionMode> EpisodeCompletionModes,
@@ -190,6 +188,7 @@ namespace UE::Learning::SharedMemoryTraining
 		TLearningArrayView<2, float> MemoryStates,
 		TLearningArrayView<1, float> Rewards,
 		TLearningArrayView<1, volatile int32> Controls,
+		const FSubprocess& Process,
 		const FReplayBuffer& ReplayBuffer,
 		const float Timeout,
 		const ELogSetting LogSettings)
@@ -201,7 +200,7 @@ namespace UE::Learning::SharedMemoryTraining
 		while (Controls[(uint8)EControls::ExperienceSignal])
 		{
 			// Check if the process has exited
-			if (!Process || !Process->Update() || Process->GetReturnCode() == 1)
+			if (!Process.IsRunning())
 			{
 				return ETrainerResponse::Unexpected;
 			}
@@ -252,12 +251,12 @@ namespace UE::Learning::SharedMemoryTraining
 	}
 
 	ETrainerResponse SendExperience(
-		FMonitoredProcess* Process, 
 		TLearningArrayView<1, int32> EpisodeStarts,
 		TLearningArrayView<1, int32> EpisodeLengths,
 		TLearningArrayView<2, float> Observations,
 		TLearningArrayView<2, float> Actions,
 		TLearningArrayView<1, volatile int32> Controls,
+		const FSubprocess& Process,
 		const TLearningArrayView<1, const int32> EpisodeStartsExperience,
 		const TLearningArrayView<1, const int32> EpisodeLengthsExperience,
 		const TLearningArrayView<2, const float> ObservationExperience,
@@ -272,7 +271,7 @@ namespace UE::Learning::SharedMemoryTraining
 		while (Controls[(uint8)EControls::ExperienceSignal])
 		{
 			// Check if the process has exited
-			if (!Process || !Process->Update() || Process->GetReturnCode() == 1)
+			if (!Process.IsRunning())
 			{
 				return ETrainerResponse::Unexpected;
 			}
diff --git a/Engine/Plugins/Experimental/LearningAgents/Source/LearningTraining/Private/LearningSharedMemoryTraining.h b/Engine/Plugins/Experimental/LearningAgents/Source/LearningTraining/Private/LearningSharedMemoryTraining.h
--- a/Engine/Plugins/Experimental/LearningAgents/Source/LearningTraining/Private/LearningSharedMemoryTraining.h
+++ b/Engine/Plugins/Experimental/LearningAgents/Source/LearningTraining/Private/LearningSharedMemoryTraining.h
@@ -4,7 +4,6 @@
 
 #include "LearningTrainer.h"
 
-class FMonitoredProcess;
 class ULearningNeuralNetworkData;
 
 namespace UE::Learning
@@ -33,9 +32,9 @@ namespace UE::Learning
 		LEARNINGTRAINING_API uint8 GetControlNum();
 
 		LEARNINGTRAINING_API ETrainerResponse RecvNetwork(
-			FMonitoredProcess* Process,
 			TLearningArrayView<1, volatile int32> Controls,
 			ULearningNeuralNetworkData& OutNetwork,
+			const FSubprocess& Process,
 			const EControls Signal,
 			const TLearningArrayView<1, const uint8> NetworkData,
 			const float Timeout = Trainer::DefaultTimeout,
@@ -48,9 +47,9 @@ namespace UE::Learning
 		LEARNINGTRAINING_API bool HasPolicyOrCompleted(TLearningArrayView<1, volatile int32> Controls);
 
 		LEARNINGTRAINING_API ETrainerResponse SendNetwork(
-			FMonitoredProcess* Process,
 			TLearningArrayView<1, volatile int32> Controls,
 			TLearningArrayView<1, uint8> NetworkData,
+			const FSubprocess& Process,
 			const EControls Signal,
 			const ULearningNeuralNetworkData& Network,
 			const float Timeout = Trainer::DefaultTimeout,
@@ -58,7 +57,6 @@ namespace UE::Learning
 			const ELogSetting LogSettings = Trainer::DefaultLogSettings);
 
 		LEARNINGTRAINING_API ETrainerResponse SendExperience(
-			FMonitoredProcess* Process,
 			TLearningArrayView<1, int32> EpisodeStarts,
 			TLearningArrayView<1, int32> EpisodeLengths,
 			TLearningArrayView<1, ECompletionMode> EpisodeCompletionModes,
@@ -69,17 +67,18 @@ namespace UE::Learning
 			TLearningArrayView<2, float> MemoryStates,
 			TLearningArrayView<1, float> Rewards,
 			TLearningArrayView<1, volatile int32> Controls,
+			const FSubprocess& Process,
 			const FReplayBuffer& ReplayBuffer,
 			const float Timeout = Trainer::DefaultTimeout,
 			const ELogSetting LogSettings = Trainer::DefaultLogSettings);
 
 		LEARNINGTRAINING_API ETrainerResponse SendExperience(
-			FMonitoredProcess* Process,
 			TLearningArrayView<1, int32> EpisodeStarts,
 			TLearningArrayView<1, int32> EpisodeLengths,
 			TLearningArrayView<2, float> Observations,
 			TLearningArrayView<2, float> Actions,
 			TLearningArrayView<1, volatile int32> Controls,
+			const FSubprocess& Process,
 			const TLearningArrayView<1, const int32> EpisodeStartsExperience,
 			const TLearningArrayView<1, const int32> EpisodeLengthsExperience,
 			const TLearningArrayView<2, const float> ObservationExperience,
diff --git a/Engine/Plugins/Experimental/LearningAgents/Source/LearningTraining/Private/LearningTrainer.cpp b/Engine/Plugins/Experimental/LearningAgents/Source/LearningTraining/Private/LearningTrainer.cpp
--- a/Engine/Plugins/Experimental/LearningAgents/Source/LearningTraining/Private/LearningTrainer.cpp
+++ b/Engine/Plugins/Experimental/LearningAgents/Source/LearningTraining/Private/LearningTrainer.cpp
@@ -7,9 +7,97 @@
 
 #include "HAL/Platform.h"
 #include "Dom/JsonObject.h"
+#include "Misc/Paths.h"
+
+namespace UE::Learning
+{
+	FSubprocess::~FSubprocess()
+	{
+		Terminate();
+	}
+
+	bool FSubprocess::Launch(const FString& Path, const FString& Params, const ESubprocessFlags Flags)
+	{
+		ensureMsgf(!bIsLaunched, TEXT("Subprocess already launched."));
+
+		Terminate();
+
+		const bool bCreatePipes = !(Flags & ESubprocessFlags::NoRedirectOutput);
+		const bool bHideWindow = !(Flags & ESubprocessFlags::ShowWindow);
+
+		if (bCreatePipes && !FPlatformProcess::CreatePipe(ReadPipe, WritePipe))
+		{
+			return false;
+		}
+
+		ProcessHandle = FPlatformProcess::CreateProc(*Path, *Params, false, bHideWindow, bHideWindow, nullptr, 0, *FPaths::RootDir(), WritePipe, ReadPipe);
+		bIsLaunched = true;
+		return true;
+	}
+
+	bool FSubprocess::IsRunning() const
+	{
+		return bIsLaunched && FPlatformProcess::IsProcRunning(const_cast<FProcHandle&>(ProcessHandle));
+	}
+
+	void FSubprocess::Terminate()
+	{
+		if (IsRunning())
+		{
+			UE_LOG(LogLearning, Warning, TEXT("Terminating Subprocess..."));
+
+			FPlatformProcess::TerminateProc(ProcessHandle, true);
+		}
+
+		Update();
+	}
+
+	void FSubprocess::Update()
+	{
+		// Do nothing if the process is not launched
+		if (!bIsLaunched)
+		{
+			return;
+		}
+
+		// Append the process stdout to the buffer
+		OutputBuffer += FPlatformProcess::ReadPipe(ReadPipe);
+
+		// Output all the complete lines
+		int32 LineStartIdx = 0;
+		for (int32 Idx = 0; Idx < OutputBuffer.Len(); Idx++)
+		{
+			if (OutputBuffer[Idx] == '\r' || OutputBuffer[Idx] == '\n')
+			{
+				UE_LOG(LogLearning, Display, TEXT("Subprocess: %s"), *OutputBuffer.Mid(LineStartIdx, Idx - LineStartIdx));
+
+				if (OutputBuffer[Idx] == '\r' && Idx + 1 < OutputBuffer.Len() && OutputBuffer[Idx + 1] == '\n')
+				{
+					Idx++;
+				}
+
+				LineStartIdx = Idx + 1;
+			}
+		}
+
+		// Remove all the complete lines from the buffer
+		OutputBuffer.MidInline(LineStartIdx, MAX_int32, EAllowShrinking::Yes);
+
+		// If the process is no longer running then close the pipes
+		if (!IsRunning())
+		{
+			FPlatformProcess::ClosePipe(ReadPipe, WritePipe);
+			ReadPipe = nullptr;
+			WritePipe = nullptr;
+			bIsLaunched = false;
+		}
+	}
+
+}
 
 namespace UE::Learning::Trainer
 {
+
 	TSharedPtr<FJsonObject> ConvertObservationSchemaToJSON(
 		const Observation::FSchema& ObservationSchema,
 		const Observation::FSchemaElement& ObservationSchemaElement)
diff --git a/Engine/Plugins/Experimental/LearningAgents/Source/LearningTraining/Public/LearningImitationTrainer.h b/Engine/Plugins/Experimental/LearningAgents/Source/LearningTraining/Public/LearningImitationTrainer.h
--- a/Engine/Plugins/Experimental/LearningAgents/Source/LearningTraining/Public/LearningImitationTrainer.h
+++ b/Engine/Plugins/Experimental/LearningAgents/Source/LearningTraining/Public/LearningImitationTrainer.h
@@ -13,7 +13,6 @@
 #include "LearningImitationTrainer.generated.h"
 
 class FSocket;
-class FMonitoredProcess;
 class ULearningNeuralNetworkData;
 
 UCLASS()
@@ -352,11 +351,7 @@ namespace UE::Learning
 
 		// Training Process
 
-		TSharedPtr<FMonitoredProcess> TrainingProcess;
-
-		void HandleTrainingProcessCanceled();
-		void HandleTrainingProcessCompleted(int32 ReturnCode);
-		static void HandleTrainingProcessOutput(FString Output);
+		FSubprocess TrainingProcess;
 	};
 
 	/**
@@ -409,11 +404,7 @@ namespace UE::Learning
 
 	private:
 
-		TSharedPtr<FMonitoredProcess> TrainingProcess;
-
-		void HandleTrainingProcessCanceled();
-		void HandleTrainingProcessCompleted(int32 ReturnCode);
-		static void HandleTrainingProcessOutput(FString Output);
+		FSubprocess TrainingProcess;
 	};
 
 	struct LEARNINGTRAINING_API FSocketImitationTrainer : public IImitationTrainer
diff --git a/Engine/Plugins/Experimental/LearningAgents/Source/LearningTraining/Public/LearningPPOTrainer.h b/Engine/Plugins/Experimental/LearningAgents/Source/LearningTraining/Public/LearningPPOTrainer.h
--- a/Engine/Plugins/Experimental/LearningAgents/Source/LearningTraining/Public/LearningPPOTrainer.h
+++ b/Engine/Plugins/Experimental/LearningAgents/Source/LearningTraining/Public/LearningPPOTrainer.h
@@ -13,7 +13,6 @@
 #include "LearningPPOTrainer.generated.h"
 
 class FSocket;
-class FMonitoredProcess;
 class ULearningNeuralNetworkData;
 
 UCLASS()
@@ -468,16 +467,8 @@ namespace UE::Learning
 		// Training Process
 
 		uint16 ProcessIdx = INDEX_NONE;
-		TSharedPtr<FMonitoredProcess> TrainingProcess;
-		TArray<TSharedPtr<FMonitoredProcess>, TInlineAllocator<128>> ExperienceGatheringSubprocesses;
-
-		void HandleSubprocessCanceled();
-		void HandleSubprocessCompleted(int32 ReturnCode);
-		static void HandleSubprocessOutput(FString Output);
-
-		void HandleTrainingProcessCanceled();
-		void HandleTrainingProcessCompleted(int32 ReturnCode);
-		static void HandleTrainingProcessOutput(FString Output);
+		FSubprocess TrainingProcess;
+		TArray<FSubprocess, TInlineAllocator<128>> ExperienceGatheringSubprocesses;
 	};
 
 	/**
@@ -530,11 +521,7 @@ namespace UE::Learning
 
 	private:
 
-		TSharedPtr<FMonitoredProcess> TrainingProcess;
-
-		void HandleTrainingProcessCanceled();
-		void HandleTrainingProcessCompleted(int32 ReturnCode);
-		static void HandleTrainingProcessOutput(FString Output);
+		FSubprocess TrainingProcess;
 	};
 
 	/**
diff --git a/Engine/Plugins/Experimental/LearningAgents/Source/LearningTraining/Public/LearningTrainer.h b/Engine/Plugins/Experimental/LearningAgents/Source/LearningTraining/Public/LearningTrainer.h
--- a/Engine/Plugins/Experimental/LearningAgents/Source/LearningTraining/Public/LearningTrainer.h
+++ b/Engine/Plugins/Experimental/LearningAgents/Source/LearningTraining/Public/LearningTrainer.h
@@ -6,6 +6,7 @@
 #include "LearningLog.h"
 
 #include "Templates/SharedPointer.h"
+#include "HAL/PlatformProcess.h"
 
 class FJsonObject;
 
@@ -68,6 +69,57 @@ namespace UE::Learning
 	};
 	ENUM_CLASS_FLAGS(ESubprocessFlags)
 
+	/**
+	* Simple managed subprocess similar to FMonitoredProcess
+	*/
+	struct LEARNINGTRAINING_API FSubprocess
+	{
+		/** Will Terminate the subprocess if it is running. */
+		~FSubprocess();
+
+		/**
+		 * Launches a new subprocess.
+		 *
+		 * @param Path The path of the executable to launch.
+		 * @param Params The command line parameters.
+		 * @param Flags Subprocess flags.
+		 * @returns true if the launch was successful, otherwise false
+		 */
+		bool Launch(const FString& Path, const FString& Params, const ESubprocessFlags Flags);
+
+		/** 
+		* Return true if the Subprocess is launched and running, otherwise false.
+		*/
+		bool IsRunning() const;
+
+		/**
+		* Terminates the subprocess
+		*/
+		void Terminate();
+
+		/**
+		* Outputs anything the subprocess has written to stdout to the log line-by-line, and checks if the subprocess is completed.
+		*/
+		void Update();
+
+	private:
+
+		// Buffer for the subprocess' stdout
+		FString OutputBuffer;
+
+		// If a subprocess has been launched
+		bool bIsLaunched = false;
+
+		// Subprocess handle
+		FProcHandle ProcessHandle;
+
+		// Read pipe for subprocess stdout
+		void* ReadPipe = nullptr;
+
+		// Write pipe for subprocess stdin
+		void* WritePipe = nullptr;
+	};
+
 	namespace Trainer
 	{
 		/**
