LearningAgents: Fixed issue with null action. Changed training subprocess interface to make it easier to update while we are polling for termination since error logs were not getting output.

#virtualized

diff --git a/Engine/Plugins/Experimental/LearningAgents/Source/LearningTraining/Private/LearningImitationTrainer.cpp b/Engine/Plugins/Experimental/LearningAgents/Source/LearningTraining/Private/LearningImitationTrainer.cpp
--- a/Engine/Plugins/Experimental/LearningAgents/Source/LearningTraining/Private/LearningImitationTrainer.cpp
+++ b/Engine/Plugins/Experimental/LearningAgents/Source/LearningTraining/Private/LearningImitationTrainer.cpp
@@ -372,7 +372,7 @@ namespace UE::Learning
 		const float SleepTime = 0.001f;
 		float WaitTime = 0.0f;
 
-		while (TrainingProcess.IsRunning())
+		while (TrainingProcess.Update())
 		{
 			FPlatformProcess::Sleep(SleepTime);
 			WaitTime += SleepTime;
@@ -440,7 +440,7 @@ namespace UE::Learning
 		const float SleepTime = 0.001f;
 		float WaitTime = 0.0f;
 
-		while (TrainingProcess.IsRunning())
+		while (TrainingProcess.Update())
 		{
 			FPlatformProcess::Sleep(SleepTime);
 			WaitTime += SleepTime;
diff --git a/Engine/Plugins/Experimental/LearningAgents/Source/LearningTraining/Private/LearningPPOTrainer.cpp b/Engine/Plugins/Experimental/LearningAgents/Source/LearningTraining/Private/LearningPPOTrainer.cpp
--- a/Engine/Plugins/Experimental/LearningAgents/Source/LearningTraining/Private/LearningPPOTrainer.cpp
+++ b/Engine/Plugins/Experimental/LearningAgents/Source/LearningTraining/Private/LearningPPOTrainer.cpp
@@ -346,9 +346,8 @@ namespace UE::Learning
 		const float SleepTime = 0.001f;
 		float WaitTime = 0.0f;
 
-		while (TrainingProcess.IsRunning())
+		while (TrainingProcess.Update())
 		{
-			TrainingProcess.Update();
 			FPlatformProcess::Sleep(SleepTime);
 			WaitTime += SleepTime;
 
@@ -382,8 +381,6 @@ namespace UE::Learning
 		FRWLock* NetworkLock,
 		const ELogSetting LogSettings)
 	{
-		TrainingProcess.Update();
-
 		return SharedMemoryTraining::RecvNetwork(
 			Controls.View[ProcessIdx],
 			OutNetwork,
@@ -401,8 +398,6 @@ namespace UE::Learning
 		FRWLock* NetworkLock,
 		const ELogSetting LogSettings)
 	{
-		TrainingProcess.Update();
-
 		return SharedMemoryTraining::RecvNetwork(
 			Controls.View[ProcessIdx],
 			OutNetwork,
@@ -420,8 +415,6 @@ namespace UE::Learning
 		FRWLock* NetworkLock,
 		const ELogSetting LogSettings)
 	{
-		TrainingProcess.Update();
-
 		return SharedMemoryTraining::RecvNetwork(
 			Controls.View[ProcessIdx],
 			OutNetwork,
@@ -439,8 +432,6 @@ namespace UE::Learning
 		FRWLock* NetworkLock,
 		const ELogSetting LogSettings)
 	{
-		TrainingProcess.Update();
-
 		return SharedMemoryTraining::RecvNetwork(
 			Controls.View[ProcessIdx],
 			OutNetwork,
@@ -458,8 +449,6 @@ namespace UE::Learning
 		FRWLock* NetworkLock,
 		const ELogSetting LogSettings)
 	{
-		TrainingProcess.Update();
-
 		return SharedMemoryTraining::SendNetwork(
 			Controls.View[ProcessIdx],
 			Policy.View,
@@ -477,8 +466,6 @@ namespace UE::Learning
 		FRWLock* NetworkLock,
 		const ELogSetting LogSettings)
 	{
-		TrainingProcess.Update();
-
 		return SharedMemoryTraining::SendNetwork(
 			Controls.View[ProcessIdx],
 			Critic.View,
@@ -496,8 +483,6 @@ namespace UE::Learning
 		FRWLock* NetworkLock,
 		const ELogSetting LogSettings)
 	{
-		TrainingProcess.Update();
-
 		return SharedMemoryTraining::SendNetwork(
 			Controls.View[ProcessIdx],
 			Encoder.View,
@@ -515,8 +500,6 @@ namespace UE::Learning
 		FRWLock* NetworkLock,
 		const ELogSetting LogSettings)
 	{
-		TrainingProcess.Update();
-
 		return SharedMemoryTraining::SendNetwork(
 			Controls.View[ProcessIdx],
 			Decoder.View,
@@ -533,8 +516,6 @@ namespace UE::Learning
 		const float Timeout,
 		const ELogSetting LogSettings)
 	{
-		TrainingProcess.Update();
-
 		return SharedMemoryTraining::SendExperience(
 			EpisodeStarts.View[ProcessIdx],
 			EpisodeLengths.View[ProcessIdx],
@@ -616,9 +597,8 @@ namespace UE::Learning
 		const float SleepTime = 0.001f;
 		float WaitTime = 0.0f;
 
-		while (TrainingProcess.IsRunning())
+		while (TrainingProcess.Update())
 		{
-			TrainingProcess.Update();
 			FPlatformProcess::Sleep(SleepTime);
 			WaitTime += SleepTime;
 
diff --git a/Engine/Plugins/Experimental/LearningAgents/Source/LearningTraining/Private/LearningSharedMemoryTraining.cpp b/Engine/Plugins/Experimental/LearningAgents/Source/LearningTraining/Private/LearningSharedMemoryTraining.cpp
--- a/Engine/Plugins/Experimental/LearningAgents/Source/LearningTraining/Private/LearningSharedMemoryTraining.cpp
+++ b/Engine/Plugins/Experimental/LearningAgents/Source/LearningTraining/Private/LearningSharedMemoryTraining.cpp
@@ -30,7 +30,7 @@ namespace UE::Learning::SharedMemoryTraining
 	ETrainerResponse RecvNetwork(
 		TLearningArrayView<1, volatile int32> Controls,
 		ULearningNeuralNetworkData& OutNetwork,
-		const FSubprocess& Process,
+		FSubprocess& Process,
 		const EControls Signal,
 		const TLearningArrayView<1, const uint8> NetworkData,
 		const float Timeout,
@@ -52,7 +52,7 @@ namespace UE::Learning::SharedMemoryTraining
 			}
 
 			// Check if the process has exited
-			if (!Process.IsRunning())
+			if (!Process.Update())
 			{
 				return ETrainerResponse::Unexpected;
 			}
@@ -113,7 +113,7 @@ namespace UE::Learning::SharedMemoryTraining
 	ETrainerResponse SendNetwork(
 		TLearningArrayView<1, volatile int32> Controls,
 		TLearningArrayView<1, uint8> NetworkData,
-		const FSubprocess& Process,
+		FSubprocess& Process,
 		const EControls Signal,
 		const ULearningNeuralNetworkData& Network,
 		const float Timeout,
@@ -127,7 +127,7 @@ namespace UE::Learning::SharedMemoryTraining
 		while (!Controls[(uint8)Signal])
 		{
 			// Check if the process has exited
-			if (!Process.IsRunning())
+			if (!Process.Update())
 			{
 				return ETrainerResponse::Unexpected;
 			}
@@ -188,7 +188,7 @@ namespace UE::Learning::SharedMemoryTraining
 		TLearningArrayView<2, float> MemoryStates,
 		TLearningArrayView<1, float> Rewards,
 		TLearningArrayView<1, volatile int32> Controls,
-		const FSubprocess& Process,
+		FSubprocess& Process,
 		const FReplayBuffer& ReplayBuffer,
 		const float Timeout,
 		const ELogSetting LogSettings)
@@ -200,7 +200,7 @@ namespace UE::Learning::SharedMemoryTraining
 		while (Controls[(uint8)EControls::ExperienceSignal])
 		{
 			// Check if the process has exited
-			if (!Process.IsRunning())
+			if (!Process.Update())
 			{
 				return ETrainerResponse::Unexpected;
 			}
@@ -256,7 +256,7 @@ namespace UE::Learning::SharedMemoryTraining
 		TLearningArrayView<2, float> Observations,
 		TLearningArrayView<2, float> Actions,
 		TLearningArrayView<1, volatile int32> Controls,
-		const FSubprocess& Process,
+		FSubprocess& Process,
 		const TLearningArrayView<1, const int32> EpisodeStartsExperience,
 		const TLearningArrayView<1, const int32> EpisodeLengthsExperience,
 		const TLearningArrayView<2, const float> ObservationExperience,
@@ -271,7 +271,7 @@ namespace UE::Learning::SharedMemoryTraining
 		while (Controls[(uint8)EControls::ExperienceSignal])
 		{
 			// Check if the process has exited
-			if (!Process.IsRunning())
+			if (!Process.Update())
 			{
 				return ETrainerResponse::Unexpected;
 			}
diff --git a/Engine/Plugins/Experimental/LearningAgents/Source/LearningTraining/Private/LearningSharedMemoryTraining.h b/Engine/Plugins/Experimental/LearningAgents/Source/LearningTraining/Private/LearningSharedMemoryTraining.h
--- a/Engine/Plugins/Experimental/LearningAgents/Source/LearningTraining/Private/LearningSharedMemoryTraining.h
+++ b/Engine/Plugins/Experimental/LearningAgents/Source/LearningTraining/Private/LearningSharedMemoryTraining.h
@@ -34,7 +34,7 @@ namespace UE::Learning
 		LEARNINGTRAINING_API ETrainerResponse RecvNetwork(
 			TLearningArrayView<1, volatile int32> Controls,
 			ULearningNeuralNetworkData& OutNetwork,
-			const FSubprocess& Process,
+			FSubprocess& Process,
 			const EControls Signal,
 			const TLearningArrayView<1, const uint8> NetworkData,
 			const float Timeout = Trainer::DefaultTimeout,
@@ -49,7 +49,7 @@ namespace UE::Learning
 		LEARNINGTRAINING_API ETrainerResponse SendNetwork(
 			TLearningArrayView<1, volatile int32> Controls,
 			TLearningArrayView<1, uint8> NetworkData,
-			const FSubprocess& Process,
+			FSubprocess& Process,
 			const EControls Signal,
 			const ULearningNeuralNetworkData& Network,
 			const float Timeout = Trainer::DefaultTimeout,
@@ -67,7 +67,7 @@ namespace UE::Learning
 			TLearningArrayView<2, float> MemoryStates,
 			TLearningArrayView<1, float> Rewards,
 			TLearningArrayView<1, volatile int32> Controls,
-			const FSubprocess& Process,
+			FSubprocess& Process,
 			const FReplayBuffer& ReplayBuffer,
 			const float Timeout = Trainer::DefaultTimeout,
 			const ELogSetting LogSettings = Trainer::DefaultLogSettings);
@@ -78,7 +78,7 @@ namespace UE::Learning
 			TLearningArrayView<2, float> Observations,
 			TLearningArrayView<2, float> Actions,
 			TLearningArrayView<1, volatile int32> Controls,
-			const FSubprocess& Process,
+			FSubprocess& Process,
 			const TLearningArrayView<1, const int32> EpisodeStartsExperience,
 			const TLearningArrayView<1, const int32> EpisodeLengthsExperience,
 			const TLearningArrayView<2, const float> ObservationExperience,
diff --git a/Engine/Plugins/Experimental/LearningAgents/Source/LearningTraining/Private/LearningTrainer.cpp b/Engine/Plugins/Experimental/LearningAgents/Source/LearningTraining/Private/LearningTrainer.cpp
--- a/Engine/Plugins/Experimental/LearningAgents/Source/LearningTraining/Private/LearningTrainer.cpp
+++ b/Engine/Plugins/Experimental/LearningAgents/Source/LearningTraining/Private/LearningTrainer.cpp
@@ -52,12 +52,12 @@ namespace UE::Learning
 		Update();
 	}
 
-	void FSubprocess::Update()
+	bool FSubprocess::Update()
 	{
 		// Do nothing if the process is not launched
 		if (!bIsLaunched)
 		{
-			return;
+			return false;
 		}
 
 		// Append the process stdout to the buffer
@@ -90,7 +90,10 @@ namespace UE::Learning
 			ReadPipe = nullptr;
 			WritePipe = nullptr;
 			bIsLaunched = false;
+			return false;
 		}
+		
+		return true;
 	}
 
 }
diff --git a/Engine/Plugins/Experimental/LearningAgents/Source/LearningTraining/Public/LearningTrainer.h b/Engine/Plugins/Experimental/LearningAgents/Source/LearningTraining/Public/LearningTrainer.h
--- a/Engine/Plugins/Experimental/LearningAgents/Source/LearningTraining/Public/LearningTrainer.h
+++ b/Engine/Plugins/Experimental/LearningAgents/Source/LearningTraining/Public/LearningTrainer.h
@@ -98,9 +98,9 @@ namespace UE::Learning
 		void Terminate();
 
 		/**
-		* Outputs anything the subprocess has written to stdout to the log line-by-line, and checks if the subprocess is completed.
+		* Outputs anything the subprocess has written to stdout to the log line-by-line and returns true if the subprocess is still running.
 		*/
-		void Update();
+		bool Update();
 
 	private:
 
