LearningAgents: Added GetStructObservationElement and GetStructActionElement

diff --git a/Engine/Plugins/Experimental/LearningAgents/Source/LearningAgents/Private/LearningAgentsActions.cpp b/Engine/Plugins/Experimental/LearningAgents/Source/LearningAgents/Private/LearningAgentsActions.cpp
--- a/Engine/Plugins/Experimental/LearningAgents/Source/LearningAgents/Private/LearningAgentsActions.cpp
+++ b/Engine/Plugins/Experimental/LearningAgents/Source/LearningAgents/Private/LearningAgentsActions.cpp
@@ -2067,6 +2067,52 @@ bool ULearningAgentsActions::GetStructAction(TMap<FName, FLearningAgentsActionOb
 	return true;
 }
 
+bool ULearningAgentsActions::GetStructActionElement(FLearningAgentsActionObjectElement& OutElement, const ULearningAgentsActionObject* Object, const FLearningAgentsActionObjectElement Element, const FName ElementName, const FName Tag)
+{
+	if (!Object)
+	{
+		UE_LOG(LogLearning, Error, TEXT("GetStructActionElement: Object is nullptr."));
+		OutElement = FLearningAgentsActionObjectElement();
+		return false;
+	}
+
+	if (!Object->ActionObject.IsValid(Element.ObjectElement))
+	{
+		UE_LOG(LogLearning, Error, TEXT("GetStructActionElement: Invalid Action Object."));
+		OutElement = FLearningAgentsActionObjectElement();
+		return false;
+	}
+
+	if (Object->ActionObject.GetTag(Element.ObjectElement) != Tag)
+	{
+		UE_LOG(LogLearning, Warning, TEXT("GetStructActionElement: Action tag does not match. Action is '%s' but asked for '%s'."), *Object->ActionObject.GetTag(Element.ObjectElement).ToString(), *Tag.ToString());
+	}
+
+	if (Object->ActionObject.GetType(Element.ObjectElement) != UE::Learning::Action::EType::And)
+	{
+		UE_LOG(LogLearning, Error, TEXT("GetStructActionElement: Action '%s' type does not match. Action is '%s' but asked for '%s'."),
+			*Object->ActionObject.GetTag(Element.ObjectElement).ToString(),
+			UE::Learning::Agents::Action::Private::GetActionTypeString(Object->ActionObject.GetType(Element.ObjectElement)),
+			UE::Learning::Agents::Action::Private::GetActionTypeString(UE::Learning::Action::EType::And));
+		OutElement = FLearningAgentsActionObjectElement();
+		return false;
+	}
+
+	const UE::Learning::Action::FObjectAndParameters Parameters = Object->ActionObject.GetAnd(Element.ObjectElement);
+
+	const int32 ElementIdx = Parameters.ElementNames.Find(ElementName);
+
+	if (ElementIdx == INDEX_NONE)
+	{
+		UE_LOG(LogLearning, Error, TEXT("GetStructActionElement: Element '%s' not found."), *ElementName.ToString());
+		OutElement = FLearningAgentsActionObjectElement();
+		return false;
+	}
+
+	OutElement = { Parameters.Elements[ElementIdx] };
+	return true;
+}
+
 bool ULearningAgentsActions::GetStructActionToArrays(TArray<FName>& OutElementNames, TArray<FLearningAgentsActionObjectElement>& OutElements, const ULearningAgentsActionObject* Object, const FLearningAgentsActionObjectElement Element, const FName Tag)
 {
 	int32 OutElementNum = 0;
diff --git a/Engine/Plugins/Experimental/LearningAgents/Source/LearningAgents/Private/LearningAgentsObservations.cpp b/Engine/Plugins/Experimental/LearningAgents/Source/LearningAgents/Private/LearningAgentsObservations.cpp
--- a/Engine/Plugins/Experimental/LearningAgents/Source/LearningAgents/Private/LearningAgentsObservations.cpp
+++ b/Engine/Plugins/Experimental/LearningAgents/Source/LearningAgents/Private/LearningAgentsObservations.cpp
@@ -2762,6 +2762,52 @@ bool ULearningAgentsObservations::GetStructObservation(TMap<FName, FLearningAgen
 	return true;
 }
 
+bool ULearningAgentsObservations::GetStructObservationElement(FLearningAgentsObservationObjectElement& OutElement, const ULearningAgentsObservationObject* Object, const FLearningAgentsObservationObjectElement Element, const FName ElementName, const FName Tag)
+{
+	if (!Object)
+	{
+		UE_LOG(LogLearning, Error, TEXT("GetStructObservationElement: Object is nullptr."));
+		OutElement = FLearningAgentsObservationObjectElement();
+		return false;
+	}
+
+	if (!Object->ObservationObject.IsValid(Element.ObjectElement))
+	{
+		UE_LOG(LogLearning, Error, TEXT("GetStructObservationElement: Invalid Observation Object."));
+		OutElement = FLearningAgentsObservationObjectElement();
+		return false;
+	}
+
+	if (Object->ObservationObject.GetTag(Element.ObjectElement) != Tag)
+	{
+		UE_LOG(LogLearning, Warning, TEXT("GetStructObservationElement: Observation tag does not match. Observation is '%s' but asked for '%s'."), *Object->ObservationObject.GetTag(Element.ObjectElement).ToString(), *Tag.ToString());
+	}
+
+	if (Object->ObservationObject.GetType(Element.ObjectElement) != UE::Learning::Observation::EType::And)
+	{
+		UE_LOG(LogLearning, Error, TEXT("GetStructObservationElement: Observation '%s' type does not match. Observation is '%s' but asked for '%s'."),
+			*Object->ObservationObject.GetTag(Element.ObjectElement).ToString(),
+			UE::Learning::Agents::Observation::Private::GetObservationTypeString(Object->ObservationObject.GetType(Element.ObjectElement)),
+			UE::Learning::Agents::Observation::Private::GetObservationTypeString(UE::Learning::Observation::EType::And));
+		OutElement = FLearningAgentsObservationObjectElement();
+		return false;
+	}
+
+	const UE::Learning::Observation::FObjectAndParameters Parameters = Object->ObservationObject.GetAnd(Element.ObjectElement);
+
+	const int32 ElementIdx = Parameters.ElementNames.Find(ElementName);
+
+	if (ElementIdx == INDEX_NONE)
+	{
+		UE_LOG(LogLearning, Error, TEXT("GetStructActionElement: Element '%s' not found."), *ElementName.ToString());
+		OutElement = FLearningAgentsObservationObjectElement();
+		return false;
+	}
+
+	OutElement = { Parameters.Elements[ElementIdx] };
+	return true;
+}
+
 bool ULearningAgentsObservations::GetStructObservationToArrays(TArray<FName>& OutElementNames, TArray<FLearningAgentsObservationObjectElement>& OutElements, const ULearningAgentsObservationObject* Object, const FLearningAgentsObservationObjectElement Element, const FName Tag)
 {
 	int32 OutElementNum = 0;
diff --git a/Engine/Plugins/Experimental/LearningAgents/Source/LearningAgents/Public/LearningAgentsActions.h b/Engine/Plugins/Experimental/LearningAgents/Source/LearningAgents/Public/LearningAgentsActions.h
--- a/Engine/Plugins/Experimental/LearningAgents/Source/LearningAgents/Public/LearningAgentsActions.h
+++ b/Engine/Plugins/Experimental/LearningAgents/Source/LearningAgents/Public/LearningAgentsActions.h
@@ -848,6 +848,19 @@ public:
 	UFUNCTION(BlueprintPure = false, Category = "LearningAgents", meta = (AdvancedDisplay = 3, ReturnDisplayName = "Success"))
 	static bool GetStructAction(TMap<FName, FLearningAgentsActionObjectElement>& OutElements, const ULearningAgentsActionObject* Object, const FLearningAgentsActionObjectElement Element, const FName Tag = TEXT("StructAction"));
 
+	/**
+	 * Get the sub-action given its name for a struct action.
+	 *
+	 * @param OutElement The output sub-action.
+	 * @param Object The Action Object
+	 * @param Element The Action Object Element
+	 * @param ElementName The Sub-action Name.
+	 * @param Tag The tag of the corresponding action. Must match the tag given during Specify.
+	 * @return true if the provided Element has the sub-action, otherwise false.
+	 */
+	UFUNCTION(BlueprintPure = false, Category = "LearningAgents", meta = (AdvancedDisplay = 4, ReturnDisplayName = "Success"))
+	static bool GetStructActionElement(FLearningAgentsActionObjectElement& OutElement, const ULearningAgentsActionObject* Object, const FLearningAgentsActionObjectElement Element, const FName ElementName, const FName Tag = TEXT("StructAction"));
+
 	/**
 	 * Get the sub-actions for a struct action.
 	 *
diff --git a/Engine/Plugins/Experimental/LearningAgents/Source/LearningAgents/Public/LearningAgentsObservations.h b/Engine/Plugins/Experimental/LearningAgents/Source/LearningAgents/Public/LearningAgentsObservations.h
--- a/Engine/Plugins/Experimental/LearningAgents/Source/LearningAgents/Public/LearningAgentsObservations.h
+++ b/Engine/Plugins/Experimental/LearningAgents/Source/LearningAgents/Public/LearningAgentsObservations.h
@@ -1531,6 +1531,9 @@ public:
 
 	UFUNCTION(BlueprintPure = false, Category = "LearningAgents", meta = (AdvancedDisplay = 3, ReturnDisplayName = "Success"))
 	static bool GetStructObservation(TMap<FName, FLearningAgentsObservationObjectElement>& OutElements, const ULearningAgentsObservationObject* Object, const FLearningAgentsObservationObjectElement Element, const FName Tag = TEXT("StructObservation"));
+	
+	UFUNCTION(BlueprintPure = false, Category = "LearningAgents", meta = (AdvancedDisplay = 4, ReturnDisplayName = "Success"))
+	static bool GetStructObservationElement(FLearningAgentsObservationObjectElement& OutElement, const ULearningAgentsObservationObject* Object, const FLearningAgentsObservationObjectElement Element, const FName ElementName, const FName Tag = TEXT("StructObservation"));
 
 	UFUNCTION(BlueprintPure = false, Category = "LearningAgents", meta = (AdvancedDisplay = 4, ReturnDisplayName = "Success"))
 	static bool GetStructObservationToArrays(TArray<FName>& OutElementNames, TArray<FLearningAgentsObservationObjectElement>& OutElements, const ULearningAgentsObservationObject* Object, const FLearningAgentsObservationObjectElement Element, const FName Tag = TEXT("StructObservation"));
