Learning Agents Training Rewrite: Part 1

#rnx

diff --git a/Engine/Plugins/Experimental/LearningAgents/Source/LearningAgents/Private/LearningAgentsInteractor.cpp b/Engine/Plugins/Experimental/LearningAgents/Source/LearningAgents/Private/LearningAgentsInteractor.cpp
--- a/Engine/Plugins/Experimental/LearningAgents/Source/LearningAgents/Private/LearningAgentsInteractor.cpp
+++ b/Engine/Plugins/Experimental/LearningAgents/Source/LearningAgents/Private/LearningAgentsInteractor.cpp
@@ -590,4 +590,24 @@ const UE::Learning::Action::FSchema& ULearningAgentsInteractor::GetActionSchema(
 UE::Learning::Action::FSchemaElement ULearningAgentsInteractor::GetActionSchemaElement() const
 {
 	return ActionSchemaElement.SchemaElement;
-}
\ No newline at end of file
+}
+
+TLearningArrayView<2, const float> ULearningAgentsInteractor::GetObservationVectorArrayView() const
+{
+	return ObservationVectors;
+}
+
+uint64 ULearningAgentsInteractor::GetObservationIteration(const int32 AgentId) const
+{
+	return ObservationVectorIteration[AgentId];
+}
+
+TLearningArrayView<2, const float> ULearningAgentsInteractor::GetActionVectorArrayView() const
+{
+	return ActionVectors;
+}
+
+uint64 ULearningAgentsInteractor::GetActionIteration(const int32 AgentId) const
+{
+	return ActionVectorIteration[AgentId];
+}
diff --git a/Engine/Plugins/Experimental/LearningAgents/Source/LearningAgents/Private/LearningAgentsPolicy.cpp b/Engine/Plugins/Experimental/LearningAgents/Source/LearningAgents/Private/LearningAgentsPolicy.cpp
--- a/Engine/Plugins/Experimental/LearningAgents/Source/LearningAgents/Private/LearningAgentsPolicy.cpp
+++ b/Engine/Plugins/Experimental/LearningAgents/Source/LearningAgents/Private/LearningAgentsPolicy.cpp
@@ -478,6 +478,16 @@ UE::Learning::FNeuralNetworkFunction& ULearningAgentsPolicy::GetDecoderObject()
 	return *DecoderObject;
 }
 
+TLearningArrayView<2, const float> ULearningAgentsPolicy::GetPreEvaluationMemoryState() const
+{
+	return PreEvaluationMemoryState;
+}
+
+TLearningArrayView<2, const float> ULearningAgentsPolicy::GetMemoryState() const
+{
+	return MemoryState;
+}
+
 ULearningAgentsNeuralNetwork* ULearningAgentsPolicy::GetEncoderNetworkAsset()
 {
 	if (!IsSetup())
diff --git a/Engine/Plugins/Experimental/LearningAgents/Source/LearningAgents/Public/LearningAgentsInteractor.h b/Engine/Plugins/Experimental/LearningAgents/Source/LearningAgents/Public/LearningAgentsInteractor.h
--- a/Engine/Plugins/Experimental/LearningAgents/Source/LearningAgents/Public/LearningAgentsInteractor.h
+++ b/Engine/Plugins/Experimental/LearningAgents/Source/LearningAgents/Public/LearningAgentsInteractor.h
@@ -248,6 +248,18 @@ public:
 	/** Gets the internal action schema element */
 	UE::Learning::Action::FSchemaElement GetActionSchemaElement() const;
 
+	/** Gets the observation vectors as a const array view. */
+	TLearningArrayView<2, const float> GetObservationVectorArrayView() const;
+
+	/** Gets the observation iteration value for the given agent id. */
+	uint64 GetObservationIteration(const int32 AgentId) const;
+
+	/** Gets the action vectors as a const array view. */
+	TLearningArrayView<2, const float> GetActionVectorArrayView() const;
+
+	/** Gets the action iteration value for the given agent id. */
+	uint64 GetActionIteration(const int32 AgentId) const;
+
 private:
 
 	/** Observation Schema used by this interactor */
@@ -294,7 +306,7 @@ private:
 	/** Compatibility Hash for Observation Schema */
 	int32 ObservationCompatibilityHash = 0;
 
-	/** Compatibility Hash for Actiuon Schema */
+	/** Compatibility Hash for Action Schema */
 	int32 ActionCompatibilityHash = 0;
 
 	/** Number of times observation vector has been set for all agents */
diff --git a/Engine/Plugins/Experimental/LearningAgents/Source/LearningAgents/Public/LearningAgentsPolicy.h b/Engine/Plugins/Experimental/LearningAgents/Source/LearningAgents/Public/LearningAgentsPolicy.h
--- a/Engine/Plugins/Experimental/LearningAgents/Source/LearningAgents/Public/LearningAgentsPolicy.h
+++ b/Engine/Plugins/Experimental/LearningAgents/Source/LearningAgents/Public/LearningAgentsPolicy.h
@@ -223,6 +223,12 @@ public:
 	/** Get a reference to this policy's decoder function object. */
 	UE::Learning::FNeuralNetworkFunction& GetDecoderObject();
 
+	/** Get a const view to this policy's pre-evaluation memory state. */
+	TLearningArrayView<2, const float> GetPreEvaluationMemoryState() const;
+
+	/** Get a const view to this policy's memory state. */
+	TLearningArrayView<2, const float> GetMemoryState() const;
+
 // ----- Private Data -----
 private:
 
diff --git a/Engine/Plugins/Experimental/LearningAgents/Source/LearningAgentsTraining/LearningAgentsTraining.Build.cs b/Engine/Plugins/Experimental/LearningAgents/Source/LearningAgentsTraining/LearningAgentsTraining.Build.cs
--- a/Engine/Plugins/Experimental/LearningAgents/Source/LearningAgentsTraining/LearningAgentsTraining.Build.cs
+++ b/Engine/Plugins/Experimental/LearningAgents/Source/LearningAgentsTraining/LearningAgentsTraining.Build.cs
@@ -30,6 +30,8 @@ public class LearningAgentsTraining : ModuleRules
 			{
 				"CoreUObject",
 				"Engine",
+				"Json",
+				"JsonUtilities",
 				"Learning",
 				"LearningAgents",
 				"LearningTraining",
diff --git a/Engine/Plugins/Experimental/LearningAgents/Source/LearningAgentsTraining/Private/LearningAgentsCommunicator.cpp b/Engine/Plugins/Experimental/LearningAgents/Source/LearningAgentsTraining/Private/LearningAgentsCommunicator.cpp
new file mode 100644
--- /dev/null
+++ b/Engine/Plugins/Experimental/LearningAgents/Source/LearningAgentsTraining/Private/LearningAgentsCommunicator.cpp
@@ -0,0 +1,119 @@
+// Copyright Epic Games, Inc. All Rights Reserved.
+
+#include "LearningAgentsCommunicator.h"
+
+#include "LearningAgentsInteractor.h"
+#include "LearningExternalTrainer.h"
+
+#include "Misc/Paths.h"
+
+FLearningAgentsTrainerProcess ULearningAgentsCommunicatorLibrary::SpawnSharedMemoryTrainingProcess(
+	const FLearningAgentsTrainerProcessSettings& TrainerProcessSettings,
+	const FLearningAgentsSharedMemoryCommunicatorSettings& SharedMemorySettings)
+{
+	FLearningAgentsTrainerProcess TrainerProcess;
+
+	const FString PythonExecutablePath = UE::Learning::Trainer::GetPythonExecutablePath(TrainerProcessSettings.GetIntermediatePath());
+	if (!FPaths::FileExists(PythonExecutablePath))
+	{
+		UE_LOG(LogLearning, Error, TEXT("SpawnSharedMemoryTrainingProcess: Can't find Python executable \"%s\"."), *PythonExecutablePath);
+		return TrainerProcess;
+	}
+
+	const FString PythonContentPath = UE::Learning::Trainer::GetPythonContentPath(TrainerProcessSettings.GetEditorEnginePath());
+	if (!FPaths::DirectoryExists(PythonContentPath))
+	{
+		UE_LOG(LogLearning, Error, TEXT("SpawnSharedMemoryTrainingProcess: Can't find LearningAgents plugin Content \"%s\"."), *PythonContentPath);
+		return TrainerProcess;
+	}
+
+	const FString IntermediatePath = UE::Learning::Trainer::GetIntermediatePath(TrainerProcessSettings.GetIntermediatePath());
+
+	TrainerProcess.TrainerProcess = MakeShared<UE::Learning::FSharedMemoryTrainerServerProcess>(
+		SharedMemorySettings.TaskName,
+		PythonExecutablePath,
+		PythonContentPath,
+		IntermediatePath,
+		SharedMemorySettings.ProcessNum,
+		SharedMemorySettings.Timeout);
+
+	return TrainerProcess;
+}
+
+FLearningAgentsCommunicator ULearningAgentsCommunicatorLibrary::MakeSharedMemoryCommunicator(
+	const FLearningAgentsTrainerProcess& TrainerProcess,
+	const FLearningAgentsTrainerProcessSettings& TrainerProcessSettings,
+	const FLearningAgentsSharedMemoryCommunicatorSettings& SharedMemorySettings)
+{
+	FLearningAgentsCommunicator Communicator;
+
+	if (!TrainerProcess.TrainerProcess)
+	{
+		UE_LOG(LogLearning, Error, TEXT("MakeSharedMemoryCommunicator: TrainerProcess is nullptr"));
+		return Communicator;
+	}
+
+	Communicator.Trainer = MakeShared<UE::Learning::FSharedMemoryTrainer>(
+		SharedMemorySettings.TaskName,
+		SharedMemorySettings.ProcessNum,
+		TrainerProcess.TrainerProcess,
+		SharedMemorySettings.Timeout);
+
+	return Communicator;
+}
+
+FLearningAgentsTrainerProcess ULearningAgentsCommunicatorLibrary::SpawnSocketTrainingProcess(
+	const FLearningAgentsTrainerProcessSettings& TrainerProcessSettings,
+	const FLearningAgentsSocketCommunicatorSettings& SocketSettings)
+{
+	FLearningAgentsTrainerProcess TrainerProcess;
+
+	const FString PythonExecutablePath = UE::Learning::Trainer::GetPythonExecutablePath(TrainerProcessSettings.GetIntermediatePath());
+	if (!FPaths::FileExists(PythonExecutablePath))
+	{
+		UE_LOG(LogLearning, Error, TEXT("SpawnSocketTrainingProcess: Can't find Python executable \"%s\"."), *PythonExecutablePath);
+		return TrainerProcess;
+	}
+
+	const FString PythonContentPath = UE::Learning::Trainer::GetPythonContentPath(TrainerProcessSettings.GetEditorEnginePath());
+	if (!FPaths::DirectoryExists(PythonContentPath))
+	{
+		UE_LOG(LogLearning, Error, TEXT("SpawnSocketTrainingProcess: Can't find LearningAgents plugin Content \"%s\"."), *PythonContentPath);
+		return TrainerProcess;
+	}
+
+	const FString IntermediatePath = UE::Learning::Trainer::GetIntermediatePath(TrainerProcessSettings.GetIntermediatePath());
+
+	TrainerProcess.TrainerProcess = MakeShared<UE::Learning::FSocketTrainerServerProcess>(
+		PythonExecutablePath,
+		PythonContentPath,
+		IntermediatePath,
+		*SocketSettings.IpAddress,
+		SocketSettings.Port,
+		SocketSettings.Timeout);
+
+	return TrainerProcess;
+}
+
+FLearningAgentsCommunicator ULearningAgentsCommunicatorLibrary::MakeSocketCommunicator(
+	FLearningAgentsTrainerProcess TrainerProcess,
+	const FLearningAgentsSocketCommunicatorSettings& SocketSettings)
+{
+	UE::Learning::ETrainerResponse Response = UE::Learning::ETrainerResponse::Success;
+
+	FLearningAgentsCommunicator Communicator;
+	Communicator.Trainer = MakeShared<UE::Learning::FSocketTrainer>(
+		Response,
+		TrainerProcess.TrainerProcess,
+		*SocketSettings.IpAddress,
+		SocketSettings.Port,
+		SocketSettings.Timeout);
+
+	if (Response != UE::Learning::ETrainerResponse::Success)
+	{
+		UE_LOG(LogLearning, Error, TEXT("MakeSocketCommunicator: Failed to connect to training process: %s. Check log for additional errors."), UE::Learning::Trainer::GetResponseString(Response));
+		Communicator.Trainer->Terminate();
+	}
+
+	return Communicator;
+}
diff --git a/Engine/Plugins/Experimental/LearningAgents/Source/LearningAgentsTraining/Private/LearningAgentsImitationTrainer.cpp b/Engine/Plugins/Experimental/LearningAgents/Source/LearningAgentsTraining/Private/LearningAgentsImitationTrainer.cpp
--- a/Engine/Plugins/Experimental/LearningAgents/Source/LearningAgentsTraining/Private/LearningAgentsImitationTrainer.cpp
+++ b/Engine/Plugins/Experimental/LearningAgents/Source/LearningAgentsTraining/Private/LearningAgentsImitationTrainer.cpp
@@ -113,7 +113,7 @@ void ULearningAgentsImitationTrainer::BeginTraining(
 	const ULearningAgentsRecording* Recording,
 	const FLearningAgentsImitationTrainerSettings& ImitationTrainerSettings,
 	const FLearningAgentsImitationTrainerTrainingSettings& ImitationTrainerTrainingSettings,
-	const FLearningAgentsTrainerPathSettings& ImitationTrainerPathSettings)
+	const FLearningAgentsTrainerProcessSettings& ImitationTrainerPathSettings)
 {
 	if (!IsSetup())
 	{
@@ -244,7 +244,7 @@ void ULearningAgentsImitationTrainer::BeginTraining(
 	ImitationTrainingSettings.ActionRegularizationWeight = ImitationTrainerTrainingSettings.ActionRegularizationWeight;
 	ImitationTrainingSettings.ActionEntropyWeight = ImitationTrainerTrainingSettings.ActionEntropyWeight;
 	ImitationTrainingSettings.Seed = ImitationTrainerTrainingSettings.RandomSeed;
-	ImitationTrainingSettings.Device = UE::Learning::Agents::GetTrainerDevice(ImitationTrainerTrainingSettings.Device);
+	ImitationTrainingSettings.Device = UE::Learning::Agents::GetTrainingDevice(ImitationTrainerTrainingSettings.Device);
 	ImitationTrainingSettings.bUseTensorboard = ImitationTrainerTrainingSettings.bUseTensorboard;
 	ImitationTrainingSettings.bSaveSnapshots = ImitationTrainerTrainingSettings.bSaveSnapshots;
 
@@ -417,7 +417,7 @@ void ULearningAgentsImitationTrainer::RunTraining(
 	const ULearningAgentsRecording* Recording,
 	const FLearningAgentsImitationTrainerSettings& ImitationTrainerSettings,
 	const FLearningAgentsImitationTrainerTrainingSettings& ImitationTrainerTrainingSettings,
-	const FLearningAgentsTrainerPathSettings& ImitationTrainerPathSettings)
+	const FLearningAgentsTrainerProcessSettings& ImitationTrainerPathSettings)
 {
 	if (!IsSetup())
 	{
diff --git a/Engine/Plugins/Experimental/LearningAgents/Source/LearningAgentsTraining/Private/LearningAgentsPPOTrainer.cpp b/Engine/Plugins/Experimental/LearningAgents/Source/LearningAgentsTraining/Private/LearningAgentsPPOTrainer.cpp
new file mode 100644
--- /dev/null
+++ b/Engine/Plugins/Experimental/LearningAgents/Source/LearningAgentsTraining/Private/LearningAgentsPPOTrainer.cpp
@@ -0,0 +1,772 @@
+// Copyright Epic Games, Inc. All Rights Reserved.
+
+#include "LearningAgentsPPOTrainer.h"
+
+#include "LearningAgentsCritic.h"
+#include "LearningAgentsInteractor.h"
+#include "LearningAgentsManager.h"
+#include "LearningAgentsPolicy.h"
+#include "LearningCompletion.h"
+#include "LearningExperience.h"
+#include "LearningLog.h"
+#include "LearningNeuralNetwork.h"
+#include "LearningTrainer.h"
+#include "LearningAgentsCommunicator.h"
+#include "LearningAgentsTrainingEnvironment.h"
+#include "LearningExternalTrainer.h"
+
+#include "Misc/App.h"
+#include "GameFramework/GameUserSettings.h"
+#include "PhysicsEngine/PhysicsSettings.h"
+#include "Engine/GameViewportClient.h"
+
+#include "Dom/JsonObject.h"
+#include "Serialization/JsonSerializer.h"
+#include "Serialization/JsonWriter.h"
+
+#if WITH_EDITOR
+#include "Editor/EditorPerformanceSettings.h"
+#endif
+
+ULearningAgentsPPOTrainer::ULearningAgentsPPOTrainer() : Super(FObjectInitializer::Get()) {}
+ULearningAgentsPPOTrainer::ULearningAgentsPPOTrainer(FVTableHelper& Helper) : Super(Helper) {}
+ULearningAgentsPPOTrainer::~ULearningAgentsPPOTrainer() = default;
+
+void ULearningAgentsPPOTrainer::BeginDestroy()
+{
+	if (IsTraining())
+	{
+		EndTraining();
+	}
+
+	Super::BeginDestroy();
+}
+
+ULearningAgentsPPOTrainer* ULearningAgentsPPOTrainer::MakePPOTrainer(
+	ULearningAgentsManager* InManager,
+	ULearningAgentsInteractor* InInteractor,
+	ULearningAgentsTrainingEnvironment* InTrainingEnvironment,
+	ULearningAgentsPolicy* InPolicy,
+	ULearningAgentsCritic* InCritic,
+	const FLearningAgentsCommunicator& Communicator,
+	TSubclassOf<ULearningAgentsPPOTrainer> Class,
+	const FName Name,
+	const FLearningAgentsPPOTrainerSettings& TrainerSettings)
+{
+	if (!InManager)
+	{
+		UE_LOG(LogLearning, Error, TEXT("MakePPOTrainer: InManager is nullptr."));
+		return nullptr;
+	}
+
+	if (!Class)
+	{
+		UE_LOG(LogLearning, Error, TEXT("MakePPOTrainer: Class is nullptr."));
+		return nullptr;
+	}
+
+	const FName UniqueName = MakeUniqueObjectName(InManager, Class, Name, EUniqueObjectNameOptions::GloballyUnique);
+
+	ULearningAgentsPPOTrainer* Trainer = NewObject<ULearningAgentsPPOTrainer>(InManager, Class, UniqueName);
+	if (!Trainer) { return nullptr; }
+
+	Trainer->SetupPPOTrainer(InManager, InInteractor, InTrainingEnvironment, InPolicy, InCritic, Communicator, TrainerSettings);
+
+	return Trainer->IsSetup() ? Trainer : nullptr;
+}
+
+void ULearningAgentsPPOTrainer::SetupPPOTrainer(
+	ULearningAgentsManager* InManager,
+	ULearningAgentsInteractor* InInteractor,
+	ULearningAgentsTrainingEnvironment* InTrainingEnvironment,
+	ULearningAgentsPolicy* InPolicy,
+	ULearningAgentsCritic* InCritic,
+	const FLearningAgentsCommunicator& Communicator,
+	const FLearningAgentsPPOTrainerSettings& TrainerSettings)
+{
+	if (IsSetup())
+	{
+		UE_LOG(LogLearning, Error, TEXT("%s: Setup already run!"), *GetName());
+		return;
+	}
+
+	if (!InManager)
+	{
+		UE_LOG(LogLearning, Error, TEXT("%s: InManager is nullptr."), *GetName());
+		return;
+	}
+
+	if (!InInteractor)
+	{
+		UE_LOG(LogLearning, Error, TEXT("%s: InInteractor is nullptr."), *GetName());
+		return;
+	}
+
+	if (!InInteractor->IsSetup())
+	{
+		UE_LOG(LogLearning, Error, TEXT("%s: %s's Setup must be run before it can be used."), *GetName(), *InInteractor->GetName());
+		return;
+	}
+
+	if (!InTrainingEnvironment)
+	{
+		UE_LOG(LogLearning, Error, TEXT("%s: InTrainingEnvironment is nullptr."), *GetName());
+		return;
+	}
+
+	if (!InTrainingEnvironment->IsSetup())
+	{
+		UE_LOG(LogLearning, Error, TEXT("%s: %s's Setup must be run before it can be used."), *GetName(), *InTrainingEnvironment->GetName());
+		return;
+	}
+
+	if (!InPolicy)
+	{
+		UE_LOG(LogLearning, Error, TEXT("%s: InPolicy is nullptr."), *GetName());
+		return;
+	}
+
+	if (!InPolicy->IsSetup())
+	{
+		UE_LOG(LogLearning, Error, TEXT("%s: %s's Setup must be run before it can be used."), *GetName(), *InPolicy->GetName());
+		return;
+	}
+
+	if (!InCritic)
+	{
+		UE_LOG(LogLearning, Error, TEXT("%s: InCritic is nullptr."), *GetName());
+		return;
+	}
+
+	if (!InCritic->IsSetup())
+	{
+		UE_LOG(LogLearning, Error, TEXT("%s: %s's Setup must be run before it can be used."), *GetName(), *InCritic->GetName());
+		return;
+	}
+
+	if (!Communicator.Trainer)
+	{
+		UE_LOG(LogLearning, Error, TEXT("%s: Communicator's Trainer is nullptr."), *GetName());
+		return;
+	}
+
+	Manager = InManager;
+	Interactor = InInteractor;
+	Policy = InPolicy;
+	Critic = InCritic;
+	TrainingEnvironment = InTrainingEnvironment;
+	Trainer = Communicator.Trainer;
+
+	// Create Episode Buffer
+	EpisodeBuffer = MakeUnique<UE::Learning::FEpisodeBuffer>();
+	EpisodeBuffer->Resize(
+		Manager->GetMaxAgentNum(),
+		TrainerSettings.MaxEpisodeStepNum,
+		Interactor->GetObservationVectorSize(),
+		Interactor->GetActionVectorSize(),
+		Policy->GetMemoryStateSize());
+
+	// Create Replay Buffer
+	ReplayBuffer = MakeUnique<UE::Learning::FReplayBuffer>();
+	ReplayBuffer->Resize(
+		Interactor->GetObservationVectorSize(),
+		Interactor->GetActionVectorSize(),
+		Policy->GetMemoryStateSize(),
+		TrainerSettings.MaximumRecordedEpisodesPerIteration,
+		TrainerSettings.MaximumRecordedStepsPerIteration);
+
+	bIsSetup = true;
+
+	Manager->AddListener(this);
+}
+
+void ULearningAgentsPPOTrainer::OnAgentsAdded_Implementation(const TArray<int32>& AgentIds)
+{
+	if (!IsSetup())
+	{
+		UE_LOG(LogLearning, Error, TEXT("%s: Setup not complete."), *GetName());
+		return;
+	}
+
+	EpisodeBuffer->Reset(AgentIds);
+}
+
+void ULearningAgentsPPOTrainer::OnAgentsRemoved_Implementation(const TArray<int32>& AgentIds)
+{
+	if (!IsSetup())
+	{
+		UE_LOG(LogLearning, Error, TEXT("%s: Setup not complete."), *GetName());
+		return;
+	}
+
+	EpisodeBuffer->Reset(AgentIds);
+}
+
+void ULearningAgentsPPOTrainer::OnAgentsReset_Implementation(const TArray<int32>& AgentIds)
+{
+	if (!IsSetup())
+	{
+		UE_LOG(LogLearning, Error, TEXT("%s: Setup not complete."), *GetName());
+		return;
+	}
+
+	EpisodeBuffer->Reset(AgentIds);
+}
+
+const bool ULearningAgentsPPOTrainer::IsTraining() const
+{
+	return bIsTraining;
+}
+
+void ULearningAgentsPPOTrainer::BeginTraining(
+	const FLearningAgentsPPOTrainingSettings& TrainingSettings,
+	const FLearningAgentsTrainingGameSettings& TrainingGameSettings,
+	const bool bResetAgentsOnBegin)
+{
+	if (!PLATFORM_WINDOWS)
+	{
+		UE_LOG(LogLearning, Error, TEXT("%s: Training currently only supported on Windows."), *GetName());
+		return;
+	}
+
+	if (!IsSetup())
+	{
+		UE_LOG(LogLearning, Error, TEXT("%s: Setup not complete."), *GetName());
+		return;
+	}
+
+	if (IsTraining())
+	{
+		UE_LOG(LogLearning, Error, TEXT("%s: Already Training!"), *GetName());
+		return;
+	}
+
+	ApplyGameSettings(TrainingGameSettings);
+
+	// We need to setup the trainer prior to sending the config
+	Trainer->AddNetwork(Policy->GetPolicyNetworkAsset()->GetFName(), *Policy->GetPolicyNetworkAsset()->NeuralNetworkData);
+	Trainer->AddNetwork(Critic->GetCriticNetworkAsset()->GetFName(), *Critic->GetCriticNetworkAsset()->NeuralNetworkData);
+	Trainer->AddNetwork(Policy->GetEncoderNetworkAsset()->GetFName(), *Policy->GetEncoderNetworkAsset()->NeuralNetworkData);
+	Trainer->AddNetwork(Policy->GetDecoderNetworkAsset()->GetFName(), *Policy->GetDecoderNetworkAsset()->NeuralNetworkData);
+	Trainer->AddReplayBuffer(TEXT("ReplayBuffer"), *ReplayBuffer);
+
+	UE_LOG(LogLearning, Display, TEXT("%s: Sending config..."), *GetName());
+	SendConfig(TrainingSettings);
+
+	UE_LOG(LogLearning, Display, TEXT("%s: Sending initial policy..."), *GetName());
+
+	UE::Learning::ETrainerResponse Response = UE::Learning::ETrainerResponse::Success;
+
+	Response = Trainer->SendNetwork(Policy->GetPolicyNetworkAsset()->GetFName(), *Policy->GetPolicyNetworkAsset()->NeuralNetworkData);
+	if (Response != UE::Learning::ETrainerResponse::Success)
+	{
+		UE_LOG(LogLearning, Error, TEXT("%s: Error sending policy to trainer: %s. Check log for additional errors."), *GetName(), UE::Learning::Trainer::GetResponseString(Response));
+		bHasTrainingFailed = true;
+		Trainer->Terminate();
+		return;
+	}
+
+	Response = Trainer->SendNetwork(Critic->GetCriticNetworkAsset()->GetFName(), *Critic->GetCriticNetworkAsset()->NeuralNetworkData);
+	if (Response != UE::Learning::ETrainerResponse::Success)
+	{
+		UE_LOG(LogLearning, Error, TEXT("%s: Error sending critic to trainer: %s. Check log for additional errors."), *GetName(), UE::Learning::Trainer::GetResponseString(Response));
+		bHasTrainingFailed = true;
+		Trainer->Terminate();
+		return;
+	}
+
+	Response = Trainer->SendNetwork(Policy->GetEncoderNetworkAsset()->GetFName(), *Policy->GetEncoderNetworkAsset()->NeuralNetworkData);
+	if (Response != UE::Learning::ETrainerResponse::Success)
+	{
+		UE_LOG(LogLearning, Error, TEXT("%s: Error sending encoder to trainer: %s. Check log for additional errors."), *GetName(), UE::Learning::Trainer::GetResponseString(Response));
+		bHasTrainingFailed = true;
+		Trainer->Terminate();
+		return;
+	}
+
+	Response = Trainer->SendNetwork(Policy->GetDecoderNetworkAsset()->GetFName(), *Policy->GetDecoderNetworkAsset()->NeuralNetworkData);
+	if (Response != UE::Learning::ETrainerResponse::Success)
+	{
+		UE_LOG(LogLearning, Error, TEXT("%s: Error sending decoder to trainer: %s. Check log for additional errors."), *GetName(), UE::Learning::Trainer::GetResponseString(Response));
+		bHasTrainingFailed = true;
+		Trainer->Terminate();
+		return;
+	}
+
+	if (bResetAgentsOnBegin)
+	{
+		Manager->ResetAllAgents();
+	}
+
+	ReplayBuffer->Reset();
+
+	bIsTraining = true;
+}
+
+void ULearningAgentsPPOTrainer::ApplyGameSettings(const FLearningAgentsTrainingGameSettings& Settings)
+{
+	// Record GameState Settings
+
+	bFixedTimestepUsed = FApp::UseFixedTimeStep();
+	FixedTimeStepDeltaTime = FApp::GetFixedDeltaTime();
+
+	UGameUserSettings* GameSettings = UGameUserSettings::GetGameUserSettings();
+	if (GameSettings)
+	{
+		bVSyncEnabled = GameSettings->IsVSyncEnabled();
+	}
+
+	UPhysicsSettings* PhysicsSettings = UPhysicsSettings::Get();
+	if (PhysicsSettings)
+	{
+		MaxPhysicsStep = PhysicsSettings->MaxPhysicsDeltaTime;
+	}
+
+	IConsoleVariable* MaxFPSCVar = IConsoleManager::Get().FindConsoleVariable(TEXT("t.MaxFPS"));
+	if (MaxFPSCVar)
+	{
+		MaxFPS = MaxFPSCVar->GetInt();
+	}
+
+	UGameViewportClient* ViewportClient = GetWorld() ? GetWorld()->GetGameViewport() : nullptr;
+	if (ViewportClient)
+	{
+		ViewModeIndex = ViewportClient->ViewModeIndex;
+	}
+
+#if WITH_EDITOR
+	UEditorPerformanceSettings* EditorPerformanceSettings = GetMutableDefault<UEditorPerformanceSettings>();
+	if (EditorPerformanceSettings)
+	{
+		bUseLessCPUInTheBackground = EditorPerformanceSettings->bThrottleCPUWhenNotForeground;
+		bEditorVSyncEnabled = EditorPerformanceSettings->bEnableVSync;
+	}
+#endif
+
+	// Apply Training GameState Settings
+
+	FApp::SetUseFixedTimeStep(Settings.bUseFixedTimeStep);
+
+	if (Settings.FixedTimeStepFrequency > UE_SMALL_NUMBER)
+	{
+		FApp::SetFixedDeltaTime(1.0f / Settings.FixedTimeStepFrequency);
+		if (Settings.bSetMaxPhysicsStepToFixedTimeStep && PhysicsSettings)
+		{
+			PhysicsSettings->MaxPhysicsDeltaTime = 1.0f / Settings.FixedTimeStepFrequency;
+		}
+	}
+	else
+	{
+		UE_LOG(LogLearning, Warning, TEXT("%s: Provided invalid FixedTimeStepFrequency: %0.5f"), *GetName(), Settings.FixedTimeStepFrequency);
+	}
+
+	if (Settings.bDisableMaxFPS && MaxFPSCVar)
+	{
+		MaxFPSCVar->Set(0);
+	}
+
+	if (Settings.bDisableVSync && GameSettings)
+	{
+		GameSettings->SetVSyncEnabled(false);
+		GameSettings->ApplySettings(false);
+	}
+
+	if (Settings.bUseUnlitViewportRendering && ViewportClient)
+	{
+		ViewportClient->ViewModeIndex = EViewModeIndex::VMI_Unlit;
+	}
+
+#if WITH_EDITOR
+	if (Settings.bDisableUseLessCPUInTheBackground && EditorPerformanceSettings)
+	{
+		EditorPerformanceSettings->bThrottleCPUWhenNotForeground = false;
+		EditorPerformanceSettings->PostEditChange();
+	}
+
+	if (Settings.bDisableEditorVSync && EditorPerformanceSettings)
+	{
+		EditorPerformanceSettings->bEnableVSync = false;
+		EditorPerformanceSettings->PostEditChange();
+	}
+#endif
+}
+
+void ULearningAgentsPPOTrainer::SendConfig(const FLearningAgentsPPOTrainingSettings& Settings)
+{
+	TSharedRef<FJsonObject> ConfigObject = MakeShared<FJsonObject>();
+	ConfigObject->SetStringField(TEXT("TaskName"), TEXT("Training"));
+	ConfigObject->SetStringField(TEXT("TrainerMethod"), TEXT("PPO"));
+	ConfigObject->SetStringField(TEXT("TrainerType"), TEXT("Network"));
+	ConfigObject->SetStringField(TEXT("TimeStamp"), *FDateTime::Now().ToFormattedString(TEXT("%Y-%m-%d_%H-%M-%S")));
+
+	ConfigObject->SetObjectField(TEXT("ObservationSchema"),
+		UE::Learning::Trainer::ConvertObservationSchemaToJSON(Interactor->GetObservationSchema(), Interactor->GetObservationSchemaElement()));
+	ConfigObject->SetObjectField(TEXT("ActionSchema"),
+		UE::Learning::Trainer::ConvertActionSchemaToJSON(Interactor->GetActionSchema(), Interactor->GetActionSchemaElement()));
+	ConfigObject->SetNumberField(TEXT("ObservationVectorDimensionNum"), ReplayBuffer->GetObservations().Num<1>());
+	ConfigObject->SetNumberField(TEXT("ActionVectorDimensionNum"), ReplayBuffer->GetActions().Num<1>());
+	ConfigObject->SetNumberField(TEXT("MemoryStateVectorDimensionNum"), ReplayBuffer->GetMemoryStates().Num<1>());
+	ConfigObject->SetNumberField(TEXT("MaxEpisodeNum"), ReplayBuffer->GetMaxEpisodeNum());
+	ConfigObject->SetNumberField(TEXT("MaxStepNum"), ReplayBuffer->GetMaxStepNum());
+	
+	ConfigObject->SetNumberField(TEXT("PolicyNetworkByteNum"), Policy->GetPolicyNetworkAsset()->NeuralNetworkData->GetSnapshotByteNum());
+	ConfigObject->SetNumberField(TEXT("CriticNetworkByteNum"), Critic->GetCriticNetworkAsset()->NeuralNetworkData->GetSnapshotByteNum());
+	ConfigObject->SetNumberField(TEXT("EncoderNetworkByteNum"), Policy->GetEncoderNetworkAsset()->NeuralNetworkData->GetSnapshotByteNum());
+	ConfigObject->SetNumberField(TEXT("DecoderNetworkByteNum"), Policy->GetDecoderNetworkAsset()->NeuralNetworkData->GetSnapshotByteNum());
+
+	ConfigObject->SetNumberField(TEXT("IterationNum"), Settings.NumberOfIterations);
+	ConfigObject->SetNumberField(TEXT("LearningRatePolicy"), Settings.LearningRatePolicy);
+	ConfigObject->SetNumberField(TEXT("LearningRateCritic"), Settings.LearningRateCritic);
+	ConfigObject->SetNumberField(TEXT("LearningRateDecay"), Settings.LearningRateDecay);
+	ConfigObject->SetNumberField(TEXT("WeightDecay"), Settings.WeightDecay);
+	ConfigObject->SetNumberField(TEXT("PolicyBatchSize"), Settings.PolicyBatchSize);
+	ConfigObject->SetNumberField(TEXT("CriticBatchSize"), Settings.CriticBatchSize);
+	ConfigObject->SetNumberField(TEXT("PolicyWindow"), Settings.PolicyWindowSize);
+	ConfigObject->SetNumberField(TEXT("IterationsPerGather"), Settings.IterationsPerGather);
+	ConfigObject->SetNumberField(TEXT("CriticWarmupIterations"), Settings.CriticWarmupIterations);
+	ConfigObject->SetNumberField(TEXT("EpsilonClip"), Settings.EpsilonClip);
+	ConfigObject->SetNumberField(TEXT("ActionSurrogateWeight"), Settings.ActionSurrogateWeight);
+	ConfigObject->SetNumberField(TEXT("ActionRegularizationWeight"), Settings.ActionRegularizationWeight);
+	ConfigObject->SetNumberField(TEXT("ActionEntropyWeight"), Settings.ActionEntropyWeight);
+	ConfigObject->SetNumberField(TEXT("ReturnRegularizationWeight"), Settings.ReturnRegularizationWeight);
+	ConfigObject->SetNumberField(TEXT("GaeLambda"), Settings.GaeLambda);
+	ConfigObject->SetBoolField(TEXT("AdvantageNormalization"), Settings.bAdvantageNormalization);
+	ConfigObject->SetNumberField(TEXT("AdvantageMin"), Settings.MinimumAdvantage);
+	ConfigObject->SetNumberField(TEXT("AdvantageMax"), Settings.MaximumAdvantage);
+	ConfigObject->SetBoolField(TEXT("UseGradNormMaxClipping"), Settings.bUseGradNormMaxClipping);
+	ConfigObject->SetNumberField(TEXT("GradNormMax"), Settings.GradNormMax);
+	ConfigObject->SetNumberField(TEXT("TrimEpisodeStartStepNum"), Settings.NumberOfStepsToTrimAtStartOfEpisode);
+	ConfigObject->SetNumberField(TEXT("TrimEpisodeEndStepNum"), Settings.NumberOfStepsToTrimAtEndOfEpisode);
+	ConfigObject->SetNumberField(TEXT("Seed"), Settings.RandomSeed);
+	ConfigObject->SetNumberField(TEXT("DiscountFactor"), Settings.DiscountFactor);
+	ConfigObject->SetStringField(TEXT("Device"), UE::Learning::Trainer::GetDeviceString(UE::Learning::Agents::GetTrainingDevice(Settings.Device)));
+	ConfigObject->SetBoolField(TEXT("UseTensorBoard"), Settings.bUseTensorboard);
+	ConfigObject->SetBoolField(TEXT("SaveSnapshots"), Settings.bSaveSnapshots);
+
+	UE::Learning::ETrainerResponse Response = UE::Learning::ETrainerResponse::Success;
+	Response = Trainer->SendConfig(ConfigObject);
+	
+	if (Response != UE::Learning::ETrainerResponse::Success)
+	{
+		UE_LOG(LogLearning, Error, TEXT("%s: Error sending config to trainer: %s. Check log for additional errors."), *GetName(), UE::Learning::Trainer::GetResponseString(Response));
+		bHasTrainingFailed = true;
+		Trainer->Terminate();
+		return;
+	}
+}
+
+void ULearningAgentsPPOTrainer::DoneTraining()
+{
+	if (IsTraining())
+	{
+		// Wait for Trainer to finish
+		Trainer->Wait();
+
+		// If not finished in time, terminate
+		Trainer->Terminate();
+
+		// Apply back previous game settings
+		FApp::SetUseFixedTimeStep(bFixedTimestepUsed);
+		FApp::SetFixedDeltaTime(FixedTimeStepDeltaTime);
+		UGameUserSettings* GameSettings = UGameUserSettings::GetGameUserSettings();
+		if (GameSettings)
+		{
+			GameSettings->SetVSyncEnabled(bVSyncEnabled);
+			GameSettings->ApplySettings(true);
+		}
+
+		UPhysicsSettings* PhysicsSettings = UPhysicsSettings::Get();
+		if (PhysicsSettings)
+		{
+			PhysicsSettings->MaxPhysicsDeltaTime = MaxPhysicsStep;
+		}
+
+		IConsoleVariable* MaxFPSCVar = IConsoleManager::Get().FindConsoleVariable(TEXT("t.MaxFPS"));
+		if (MaxFPSCVar)
+		{
+			MaxFPSCVar->Set(MaxFPS);
+		}
+
+		UGameViewportClient* ViewportClient = GetWorld() ? GetWorld()->GetGameViewport() : nullptr;
+		if (ViewportClient)
+		{
+			ViewportClient->ViewModeIndex = ViewModeIndex;
+		}
+
+#if WITH_EDITOR
+		UEditorPerformanceSettings* EditorPerformanceSettings = GetMutableDefault<UEditorPerformanceSettings>();
+		if (EditorPerformanceSettings)
+		{
+			EditorPerformanceSettings->bThrottleCPUWhenNotForeground = bUseLessCPUInTheBackground;
+			EditorPerformanceSettings->bEnableVSync = bEditorVSyncEnabled;
+			EditorPerformanceSettings->PostEditChange();
+		}
+#endif
+
+		bIsTraining = false;
+	}
+}
+
+void ULearningAgentsPPOTrainer::EndTraining()
+{
+	if (IsTraining())
+	{
+		UE_LOG(LogLearning, Display, TEXT("%s: Stopping training..."), *GetName());
+		Trainer->SendStop();
+		DoneTraining();
+	}
+}
+
+void ULearningAgentsPPOTrainer::ProcessExperience(const bool bResetAgentsOnUpdate)
+{
+	UE_LEARNING_TRACE_CPUPROFILER_EVENT_SCOPE(ULearningAgentsPPOTrainer::ProcessExperience);
+
+	if (!IsSetup())
+	{
+		UE_LOG(LogLearning, Error, TEXT("%s: Setup not complete."), *GetName());
+		return;
+	}
+
+	if (!IsTraining())
+	{
+		UE_LOG(LogLearning, Error, TEXT("%s: Training not running."), *GetName());
+		return;
+	}
+
+	if (Manager->GetAgentNum() == 0)
+	{
+		UE_LOG(LogLearning, Warning, TEXT("%s: No agents added to Manager."), *GetName());
+	}
+
+	// Check Observations, Actions, Rewards, and Completions have been completed and have matching iteration number
+
+	TArray<int32> ValidAgentIds;
+	ValidAgentIds.Empty(Manager->GetMaxAgentNum());
+
+	for (const int32 AgentId : Manager->GetAllAgentSet())
+	{
+		if (Interactor->GetObservationIteration(AgentId) == 0 ||
+			Interactor->GetActionIteration(AgentId) == 0 ||
+			TrainingEnvironment->GetRewardIteration(AgentId) == 0 ||
+			TrainingEnvironment->GetCompletionIteration(AgentId) == 0)
+		{
+			UE_LOG(LogLearning, Display, TEXT("%s: Agent with id %i has not completed a full step of observations, actions, rewards, completions and so experience will not be processed for it."), *GetName(), AgentId);
+			continue;
+		}
+
+		if (Interactor->GetObservationIteration(AgentId) != Interactor->GetActionIteration(AgentId) ||
+			Interactor->GetObservationIteration(AgentId) != TrainingEnvironment->GetRewardIteration(AgentId) ||
+			Interactor->GetObservationIteration(AgentId) != TrainingEnvironment->GetCompletionIteration(AgentId))
+		{
+			UE_LOG(LogLearning, Warning, TEXT("%s: Agent with id %i has non-matching iteration numbers (observation: %i, action: %i, reward: %i, completion: %i). Experience will not be processed for it."), *GetName(), AgentId,
+				Interactor->GetObservationIteration(AgentId),
+				Interactor->GetActionIteration(AgentId),
+				TrainingEnvironment->GetRewardIteration(AgentId),
+				TrainingEnvironment->GetCompletionIteration(AgentId));
+			continue;
+		}
+
+		ValidAgentIds.Add(AgentId);
+	}
+
+	UE::Learning::FIndexSet ValidAgentSet = ValidAgentIds;
+	ValidAgentSet.TryMakeSlice();
+
+	// Check for episodes that have been immediately completed
+
+	for (const int32 AgentId : ValidAgentSet)
+	{
+		if (TrainingEnvironment->GetAgentCompletion(AgentId) != UE::Learning::ECompletionMode::Running && EpisodeBuffer->GetEpisodeStepNums()[AgentId] == 0)
+		{
+			UE_LOG(LogLearning, Warning, TEXT("%s: Agent with id %i has completed episode and will be reset but has not generated any experience."), *GetName(), AgentId);
+		}
+	}
+
+	// Add Experience to Episode Buffer
+	EpisodeBuffer->Push(
+		Interactor->GetObservationVectorArrayView(),
+		Interactor->GetActionVectorArrayView(),
+		Policy->GetPreEvaluationMemoryState(),
+		TrainingEnvironment->GetRewardArrayView(),
+		ValidAgentSet);
+
+	// Find the set of agents which have reached the maximum episode length and mark them as truncated
+	UE::Learning::Completion::EvaluateEndOfEpisodeCompletions(
+		TrainingEnvironment->GetEpisodeCompletions(),
+		EpisodeBuffer->GetEpisodeStepNums(),
+		EpisodeBuffer->GetMaxStepNum(),
+		ValidAgentSet);
+
+	TrainingEnvironment->SetAllCompletions(ValidAgentSet);
+
+	TrainingEnvironment->GetResetBuffer().SetResetInstancesFromCompletions(TrainingEnvironment->GetAllCompletions(), ValidAgentSet);
+
+	// If there are no agents completed we are done
+	if (TrainingEnvironment->GetResetBuffer().GetResetInstanceNum() == 0)
+	{
+		return;
+	}
+
+	// Otherwise Gather Observations for completed Instances without incrementing iteration number
+	Interactor->GatherObservations(TrainingEnvironment->GetResetBuffer().GetResetInstances(), false);
+
+	// And push those episodes to the Replay Buffer
+	const bool bReplayBufferFull = ReplayBuffer->AddEpisodes(
+		TrainingEnvironment->GetAllCompletions(),
+		Interactor->GetObservationVectorArrayView(),
+		Policy->GetMemoryState(),
+		*EpisodeBuffer,
+		TrainingEnvironment->GetResetBuffer().GetResetInstances());
+
+	if (bReplayBufferFull)
+	{
+		UE::Learning::ETrainerResponse Response = Trainer->SendReplayBuffer(TEXT("ReplayBuffer"), *ReplayBuffer);
+
+		if (Response != UE::Learning::ETrainerResponse::Success)
+		{
+			UE_LOG(LogLearning, Error, TEXT("%s: Error waiting to push experience to trainer: %s. Check log for additional errors."), *GetName(), UE::Learning::Trainer::GetResponseString(Response));
+			bHasTrainingFailed = true;
+			EndTraining();
+			return;
+		}
+
+		ReplayBuffer->Reset();
+
+		// Get Updated Policy
+		Response = Trainer->ReceiveNetwork(Policy->GetPolicyNetworkAsset()->GetFName(), *Policy->GetPolicyNetworkAsset()->NeuralNetworkData);
+		Policy->GetPolicyNetworkAsset()->ForceMarkDirty();
+
+		if (Response == UE::Learning::ETrainerResponse::Completed)
+		{
+			UE_LOG(LogLearning, Display, TEXT("%s: Trainer completed training."), *GetName());
+			DoneTraining();
+			return;
+		}
+		else if (Response != UE::Learning::ETrainerResponse::Success)
+		{
+			UE_LOG(LogLearning, Error, TEXT("%s: Error waiting for policy from trainer: %s. Check log for additional errors."), *GetName(), UE::Learning::Trainer::GetResponseString(Response));
+			bHasTrainingFailed = true;
+			EndTraining();
+			return;
+		}
+
+		// Get Updated Critic
+		Response = Trainer->ReceiveNetwork(Critic->GetCriticNetworkAsset()->GetFName(), *Critic->GetCriticNetworkAsset()->NeuralNetworkData);
+		Critic->GetCriticNetworkAsset()->ForceMarkDirty();
+
+		if (Response != UE::Learning::ETrainerResponse::Success)
+		{
+			UE_LOG(LogLearning, Error, TEXT("%s: Error waiting for critic from trainer: %s. Check log for additional errors."), *GetName(), UE::Learning::Trainer::GetResponseString(Response));
+			bHasTrainingFailed = true;
+			EndTraining();
+			return;
+		}
+
+		// Get Updated Encoder
+		Response = Trainer->ReceiveNetwork(Policy->GetEncoderNetworkAsset()->GetFName(), *Policy->GetEncoderNetworkAsset()->NeuralNetworkData);
+		Policy->GetEncoderNetworkAsset()->ForceMarkDirty();
+
+		if (Response != UE::Learning::ETrainerResponse::Success)
+		{
+			UE_LOG(LogLearning, Error, TEXT("%s: Error waiting for encoder from trainer: %s. Check log for additional errors."), *GetName(), UE::Learning::Trainer::GetResponseString(Response));
+			bHasTrainingFailed = true;
+			EndTraining();
+			return;
+		}
+
+		// Get Updated Decoder
+		Response = Trainer->ReceiveNetwork(Policy->GetDecoderNetworkAsset()->GetFName(), *Policy->GetDecoderNetworkAsset()->NeuralNetworkData);
+		Policy->GetDecoderNetworkAsset()->ForceMarkDirty();
+
+		if (Response != UE::Learning::ETrainerResponse::Success)
+		{
+			UE_LOG(LogLearning, Error, TEXT("%s: Error waiting for decoder from trainer: %s. Check log for additional errors."), *GetName(), UE::Learning::Trainer::GetResponseString(Response));
+			bHasTrainingFailed = true;
+			EndTraining();
+			return;
+		}
+
+		if (bResetAgentsOnUpdate)
+		{
+			// Reset all agents since we have a new policy
+			TrainingEnvironment->GetResetBuffer().SetResetInstances(Manager->GetAllAgentSet());
+			Manager->ResetAgents(TrainingEnvironment->GetResetBuffer().GetResetInstancesArray());
+			return;
+		}
+	}
+
+	// Manually reset Episode Buffer for agents who have reached the maximum episode length as 
+	// they wont get it reset via the agent manager's call to ResetAgents
+	TrainingEnvironment->GetResetBuffer().SetResetInstancesFromCompletions(TrainingEnvironment->GetEpisodeCompletions(), ValidAgentSet);
+	EpisodeBuffer->Reset(TrainingEnvironment->GetResetBuffer().GetResetInstances());
+
+	// Call ResetAgents for agents which have manually signaled a completion
+	TrainingEnvironment->GetResetBuffer().SetResetInstancesFromCompletions(TrainingEnvironment->GetAgentCompletions(), ValidAgentSet);
+	if (TrainingEnvironment->GetResetBuffer().GetResetInstanceNum() > 0)
+	{
+		Manager->ResetAgents(TrainingEnvironment->GetResetBuffer().GetResetInstancesArray());
+	}
+}
+
+void ULearningAgentsPPOTrainer::RunTraining(
+	const FLearningAgentsPPOTrainingSettings& TrainingSettings,
+	const FLearningAgentsTrainingGameSettings& TrainingGameSettings,
+	const bool bResetAgentsOnBegin,
+	const bool bResetAgentsOnUpdate)
+{
+	if (!IsSetup())
+	{
+		UE_LOG(LogLearning, Error, TEXT("%s: Setup not complete."), *GetName());
+		return;
+	}
+
+	if (bHasTrainingFailed)
+	{
+		UE_LOG(LogLearning, Error, TEXT("%s: Training has failed. Check log for errors."), *GetName());
+		return;
+	}
+
+	// If we aren't training yet, then start training and do the first inference step.
+	if (!IsTraining())
+	{
+		BeginTraining(TrainingSettings,	TrainingGameSettings, bResetAgentsOnBegin);
+
+		if (!IsTraining())
+		{
+			// If IsTraining is false, then BeginTraining must have failed and we can't continue.
+			return;
+		}
+
+		Policy->RunInference();
+	}
+	// Otherwise, do the regular training process.
+	else
+	{
+		TrainingEnvironment->GatherCompletions();
+		TrainingEnvironment->GatherRewards();
+		ProcessExperience(bResetAgentsOnUpdate);
+		Policy->RunInference();
+	}
+}
+
+int32 ULearningAgentsPPOTrainer::GetEpisodeStepNum(const int32 AgentId) const
+{
+	if (!IsSetup())
+	{
+		UE_LOG(LogLearning, Error, TEXT("%s: Setup not complete."), *GetName());
+		return 0.0f;
+	}
+
+	if (!HasAgent(AgentId))
+	{
+		UE_LOG(LogLearning, Error, TEXT("%s: AgentId %d not found in the agents set."), *GetName(), AgentId);
+		return 0.0f;
+	}
+
+	return EpisodeBuffer->GetEpisodeStepNums()[AgentId];
+}
+
+bool ULearningAgentsPPOTrainer::HasTrainingFailed() const
+{
+	return bHasTrainingFailed;
+}
diff --git a/Engine/Plugins/Experimental/LearningAgents/Source/LearningAgentsTraining/Private/LearningAgentsTrainer.cpp b/Engine/Plugins/Experimental/LearningAgents/Source/LearningAgentsTraining/Private/LearningAgentsTrainer.cpp
--- a/Engine/Plugins/Experimental/LearningAgents/Source/LearningAgentsTraining/Private/LearningAgentsTrainer.cpp
+++ b/Engine/Plugins/Experimental/LearningAgents/Source/LearningAgentsTraining/Private/LearningAgentsTrainer.cpp
@@ -2,63 +2,38 @@
 
 #include "LearningAgentsTrainer.h"
 
-#include "LearningAgentsInteractor.h"
-#include "LearningAgentsManager.h"
-#include "LearningAgentsRewards.h"
-#include "LearningAgentsCompletions.h"
-#include "LearningAgentsObservations.h"
-#include "LearningAgentsPolicy.h"
-#include "LearningAgentsCritic.h"
-#include "LearningArray.h"
-#include "LearningExperience.h"
-#include "LearningLog.h"
-#include "LearningPPOTrainer.h"
-#include "LearningCompletion.h"
-
-#include "GenericPlatform/GenericPlatformMisc.h"
-#include "HAL/FileManager.h"
-#include "UObject/Package.h"
-#include "Misc/App.h"
 #include "Misc/Paths.h"
-#include "GameFramework/GameUserSettings.h"
-#include "PhysicsEngine/PhysicsSettings.h"
-#include "Engine/GameViewportClient.h"
-#include "EngineDefines.h"
-
-#if WITH_EDITOR
-#include "Editor/EditorPerformanceSettings.h"
-#endif
 
 namespace UE::Learning::Agents
 {
-	ELearningAgentsTrainerDevice GetLearningAgentsTrainerDevice(const ETrainerDevice Device)
+	ELearningAgentsTrainingDevice GetLearningAgentsTrainingDevice(const ETrainerDevice Device)
 	{
 		switch (Device)
 		{
-		case ETrainerDevice::CPU: return ELearningAgentsTrainerDevice::CPU;
-		case ETrainerDevice::GPU: return ELearningAgentsTrainerDevice::GPU;
-		default:UE_LOG(LogLearning, Error, TEXT("Unknown Trainer Device.")); return ELearningAgentsTrainerDevice::CPU;
+		case ETrainerDevice::CPU: return ELearningAgentsTrainingDevice::CPU;
+		case ETrainerDevice::GPU: return ELearningAgentsTrainingDevice::GPU;
+		default:UE_LOG(LogLearning, Error, TEXT("Unknown Trainer Device.")); return ELearningAgentsTrainingDevice::CPU;
 		}
 	}
 
-	ETrainerDevice GetTrainerDevice(const ELearningAgentsTrainerDevice Device)
+	ETrainerDevice GetTrainingDevice(const ELearningAgentsTrainingDevice Device)
 	{
 		switch (Device)
 		{
-		case ELearningAgentsTrainerDevice::CPU: return ETrainerDevice::CPU;
-		case ELearningAgentsTrainerDevice::GPU: return ETrainerDevice::GPU;
+		case ELearningAgentsTrainingDevice::CPU: return ETrainerDevice::CPU;
+		case ELearningAgentsTrainingDevice::GPU: return ETrainerDevice::GPU;
 		default:UE_LOG(LogLearning, Error, TEXT("Unknown Trainer Device.")); return ETrainerDevice::CPU;
 		}
 	}
 }
 
-FLearningAgentsTrainerPathSettings::FLearningAgentsTrainerPathSettings()
+FLearningAgentsTrainerProcessSettings::FLearningAgentsTrainerProcessSettings()
 {
 	EditorEngineRelativePath.Path = FPaths::EngineDir();
 	EditorIntermediateRelativePath.Path = FPaths::ProjectIntermediateDir();
 }
 
-FString FLearningAgentsTrainerPathSettings::GetEditorEnginePath() const
+FString FLearningAgentsTrainerProcessSettings::GetEditorEnginePath() const
 {
 #if WITH_EDITOR
 	return EditorEngineRelativePath.Path;
@@ -72,7 +47,7 @@ FString FLearningAgentsTrainerPathSettings::GetEditorEnginePath() const
 #endif
 }
 
-FString FLearningAgentsTrainerPathSettings::GetIntermediatePath() const
+FString FLearningAgentsTrainerProcessSettings::GetIntermediatePath() const
 {
 #if WITH_EDITOR
 	return EditorIntermediateRelativePath.Path;
@@ -85,1010 +60,3 @@ FString FLearningAgentsTrainerPathSettings::GetIntermediatePath() const
 	return NonEditorIntermediateRelativePath;
 #endif
 }
-
-ULearningAgentsTrainer::ULearningAgentsTrainer() : Super(FObjectInitializer::Get()) {}
-ULearningAgentsTrainer::ULearningAgentsTrainer(FVTableHelper& Helper) : Super(Helper) {}
-ULearningAgentsTrainer::~ULearningAgentsTrainer() = default;
-
-void ULearningAgentsTrainer::BeginDestroy()
-{
-	if (IsTraining())
-	{
-		EndTraining();
-	}
-
-	Super::BeginDestroy();
-}
-
-ULearningAgentsTrainer* ULearningAgentsTrainer::MakeTrainer(
-	ULearningAgentsManager* InManager,
-	ULearningAgentsInteractor* InInteractor,
-	ULearningAgentsPolicy* InPolicy,
-	ULearningAgentsCritic* InCritic,
-	TSubclassOf<ULearningAgentsTrainer> Class,
-	const FName Name,
-	const FLearningAgentsTrainerSettings& TrainerSettings)
-{
-	if (!InManager)
-	{
-		UE_LOG(LogLearning, Error, TEXT("MakeTrainer: InManager is nullptr."));
-		return nullptr;
-	}
-
-	if (!Class)
-	{
-		UE_LOG(LogLearning, Error, TEXT("MakeTrainer: Class is nullptr."));
-		return nullptr;
-	}
-
-	const FName UniqueName = MakeUniqueObjectName(InManager, Class, Name, EUniqueObjectNameOptions::GloballyUnique);
-
-	ULearningAgentsTrainer* Trainer = NewObject<ULearningAgentsTrainer>(InManager, Class, UniqueName);
-	if (!Trainer) { return nullptr; }
-
-	Trainer->SetupTrainer(InManager, InInteractor, InPolicy, InCritic, TrainerSettings);
-
-	return Trainer->IsSetup() ? Trainer : nullptr;
-}
-
-void ULearningAgentsTrainer::SetupTrainer(
-	ULearningAgentsManager* InManager,
-	ULearningAgentsInteractor* InInteractor,
-	ULearningAgentsPolicy* InPolicy,
-	ULearningAgentsCritic* InCritic,
-	const FLearningAgentsTrainerSettings& TrainerSettings)
-{
-	if (IsSetup())
-	{
-		UE_LOG(LogLearning, Error, TEXT("%s: Setup already run!"), *GetName());
-		return;
-	}
-
-	if (!InManager)
-	{
-		UE_LOG(LogLearning, Error, TEXT("%s: InManager is nullptr."), *GetName());
-		return;
-	}
-
-	if (!InInteractor)
-	{
-		UE_LOG(LogLearning, Error, TEXT("%s: InInteractor is nullptr."), *GetName());
-		return;
-	}
-
-	if (!InInteractor->IsSetup())
-	{
-		UE_LOG(LogLearning, Error, TEXT("%s: %s's Setup must be run before it can be used."), *GetName(), *InInteractor->GetName());
-		return;
-	}
-
-	if (!InPolicy)
-	{
-		UE_LOG(LogLearning, Error, TEXT("%s: InPolicy is nullptr."), *GetName());
-		return;
-	}
-
-	if (!InPolicy->IsSetup())
-	{
-		UE_LOG(LogLearning, Error, TEXT("%s: %s's Setup must be run before it can be used."), *GetName(), *InPolicy->GetName());
-		return;
-	}
-
-	if (!InCritic)
-	{
-		UE_LOG(LogLearning, Error, TEXT("%s: InCritic is nullptr."), *GetName());
-		return;
-	}
-
-	if (!InCritic->IsSetup())
-	{
-		UE_LOG(LogLearning, Error, TEXT("%s: %s's Setup must be run before it can be used."), *GetName(), *InCritic->GetName());
-		return;
-	}
-
-	Manager = InManager;
-	Interactor = InInteractor;
-	Policy = InPolicy;
-	Critic = InCritic;
-
-	// Create Episode Buffer
-	EpisodeBuffer = MakeUnique<UE::Learning::FEpisodeBuffer>();
-	EpisodeBuffer->Resize(
-		Manager->GetMaxAgentNum(),
-		TrainerSettings.MaxEpisodeStepNum,
-		Interactor->GetObservationVectorSize(),
-		Interactor->GetActionVectorSize(),
-		Policy->GetMemoryStateSize());
-
-	// Create Replay Buffer
-	ReplayBuffer = MakeUnique<UE::Learning::FReplayBuffer>();
-	ReplayBuffer->Resize(
-		Interactor->GetObservationVectorSize(),
-		Interactor->GetActionVectorSize(),
-		Policy->GetMemoryStateSize(),
-		TrainerSettings.MaximumRecordedEpisodesPerIteration,
-		TrainerSettings.MaximumRecordedStepsPerIteration);
-
-	// Create Reset Buffer
-	ResetBuffer = MakeUnique<UE::Learning::FResetInstanceBuffer>();
-	ResetBuffer->Reserve(Manager->GetMaxAgentNum());
-
-	// Record Timeout Setting
-	TrainerTimeout = TrainerSettings.TrainerCommunicationTimeout;
-
-	// Rewards and Completions Buffer
-
-	Rewards.SetNumUninitialized({ Manager->GetMaxAgentNum() });
-	AgentCompletions.SetNumUninitialized({ Manager->GetMaxAgentNum() });
-	EpisodeCompletions.SetNumUninitialized({ Manager->GetMaxAgentNum() });
-	AllCompletions.SetNumUninitialized({ Manager->GetMaxAgentNum() });
-	EpisodeTimes.SetNumUninitialized({ Manager->GetMaxAgentNum() });
-	UE::Learning::Array::Set<1, float>(Rewards, FLT_MAX);
-	UE::Learning::Array::Set<1, UE::Learning::ECompletionMode>(AgentCompletions, UE::Learning::ECompletionMode::Terminated);
-	UE::Learning::Array::Set<1, UE::Learning::ECompletionMode>(EpisodeCompletions, UE::Learning::ECompletionMode::Terminated);
-	UE::Learning::Array::Set<1, UE::Learning::ECompletionMode>(AllCompletions, UE::Learning::ECompletionMode::Terminated);
-	UE::Learning::Array::Set<1, float>(EpisodeTimes, FLT_MAX);
-
-	// Reset Agent iteration
-	RewardIteration.SetNumUninitialized({ Manager->GetMaxAgentNum() });
-	CompletionIteration.SetNumUninitialized({ Manager->GetMaxAgentNum() });
-	UE::Learning::Array::Set<1, uint64>(RewardIteration, INDEX_NONE);
-	UE::Learning::Array::Set<1, uint64>(CompletionIteration, INDEX_NONE);
-
-	bIsSetup = true;
-
-	Manager->AddListener(this);
-}
-
-void ULearningAgentsTrainer::OnAgentsAdded_Implementation(const TArray<int32>& AgentIds)
-{
-	if (!IsSetup())
-	{
-		UE_LOG(LogLearning, Error, TEXT("%s: Setup not complete."), *GetName());
-		return;
-	}
-
-	EpisodeBuffer->Reset(AgentIds);
-
-	UE::Learning::Array::Set<1, float>(Rewards, 0.0f, AgentIds);
-	UE::Learning::Array::Set<1, UE::Learning::ECompletionMode>(AgentCompletions, UE::Learning::ECompletionMode::Running, AgentIds);
-	UE::Learning::Array::Set<1, UE::Learning::ECompletionMode>(EpisodeCompletions, UE::Learning::ECompletionMode::Running, AgentIds);
-	UE::Learning::Array::Set<1, UE::Learning::ECompletionMode>(AllCompletions, UE::Learning::ECompletionMode::Running, AgentIds);
-	UE::Learning::Array::Set<1, uint64>(RewardIteration, 0, AgentIds);
-	UE::Learning::Array::Set<1, uint64>(CompletionIteration, 0, AgentIds);
-	UE::Learning::Array::Set<1, float>(EpisodeTimes, 0.0f, AgentIds);
-}
-
-void ULearningAgentsTrainer::OnAgentsRemoved_Implementation(const TArray<int32>& AgentIds)
-{
-	if (!IsSetup())
-	{
-		UE_LOG(LogLearning, Error, TEXT("%s: Setup not complete."), *GetName());
-		return;
-	}
-
-	EpisodeBuffer->Reset(AgentIds);
-
-	UE::Learning::Array::Set<1, float>(Rewards, FLT_MAX, AgentIds);
-	UE::Learning::Array::Set<1, UE::Learning::ECompletionMode>(AgentCompletions, UE::Learning::ECompletionMode::Terminated, AgentIds);
-	UE::Learning::Array::Set<1, UE::Learning::ECompletionMode>(EpisodeCompletions, UE::Learning::ECompletionMode::Terminated, AgentIds);
-	UE::Learning::Array::Set<1, UE::Learning::ECompletionMode>(AllCompletions, UE::Learning::ECompletionMode::Terminated, AgentIds);
-	UE::Learning::Array::Set<1, uint64>(RewardIteration, INDEX_NONE, AgentIds);
-	UE::Learning::Array::Set<1, uint64>(CompletionIteration, INDEX_NONE, AgentIds);
-	UE::Learning::Array::Set<1, float>(EpisodeTimes, FLT_MAX, AgentIds);
-}
-
-void ULearningAgentsTrainer::OnAgentsReset_Implementation(const TArray<int32>& AgentIds)
-{
-	if (!IsSetup())
-	{
-		UE_LOG(LogLearning, Error, TEXT("%s: Setup not complete."), *GetName());
-		return;
-	}
-
-	ResetAgentEpisodes(AgentIds);
-	EpisodeBuffer->Reset(AgentIds);
-
-	UE::Learning::Array::Set<1, float>(Rewards, 0.0f, AgentIds);
-	UE::Learning::Array::Set<1, UE::Learning::ECompletionMode>(AgentCompletions, UE::Learning::ECompletionMode::Running, AgentIds);
-	UE::Learning::Array::Set<1, UE::Learning::ECompletionMode>(EpisodeCompletions, UE::Learning::ECompletionMode::Running, AgentIds);
-	UE::Learning::Array::Set<1, UE::Learning::ECompletionMode>(AllCompletions, UE::Learning::ECompletionMode::Running, AgentIds);
-	UE::Learning::Array::Set<1, uint64>(RewardIteration, 0, AgentIds);
-	UE::Learning::Array::Set<1, uint64>(CompletionIteration, 0, AgentIds);
-	UE::Learning::Array::Set<1, float>(EpisodeTimes, 0.0f, AgentIds);
-}
-
-void ULearningAgentsTrainer::OnAgentsManagerTick_Implementation(const TArray<int32>& AgentIds, const float DeltaTime)
-{
-	if (!IsSetup())
-	{
-		UE_LOG(LogLearning, Error, TEXT("%s: Setup not complete."), *GetName());
-		return;
-	}
-
-	for (const int32 AgentId : AgentIds)
-	{
-		EpisodeTimes[AgentId] += DeltaTime;
-	}
-}
-
-void ULearningAgentsTrainer::GatherAgentReward_Implementation(float& OutReward, const int32 AgentId)
-{
-	UE_LOG(LogLearning, Error, TEXT("%s: GatherAgentReward function must be overridden!"), *GetName());
-	OutReward = 0.0f;
-}
-
-void ULearningAgentsTrainer::GatherAgentRewards_Implementation(TArray<float>& OutRewards, const TArray<int32>& AgentIds)
-{
-
-	OutRewards.Empty(AgentIds.Num());
-	for (const int32 AgentId : AgentIds)
-	{
-		float OutReward = 0.0f;
-		GatherAgentReward(OutReward, AgentId);
-		OutRewards.Add(OutReward);
-	}
-}
-
-void ULearningAgentsTrainer::GatherAgentCompletion_Implementation(ELearningAgentsCompletion& OutCompletion, const int32 AgentId)
-{
-	UE_LOG(LogLearning, Error, TEXT("%s: GatherAgentCompletion function must be overridden!"), *GetName());
-	OutCompletion = ELearningAgentsCompletion::Running;
-}
-
-void ULearningAgentsTrainer::GatherAgentCompletions_Implementation(TArray<ELearningAgentsCompletion>& OutCompletions, const TArray<int32>& AgentIds)
-{
-
-	OutCompletions.Empty(AgentIds.Num());
-	for (const int32 AgentId : AgentIds)
-	{
-		ELearningAgentsCompletion OutCompletion = ELearningAgentsCompletion::Running;
-		GatherAgentCompletion(OutCompletion, AgentId);
-		OutCompletions.Add(OutCompletion);
-	}
-}
-
-void ULearningAgentsTrainer::ResetAgentEpisode_Implementation(const int32 AgentId)
-{
-	UE_LOG(LogLearning, Error, TEXT("%s: ResetAgentEpisode function must be overridden!"), *GetName());
-}
-
-void ULearningAgentsTrainer::ResetAgentEpisodes_Implementation(const TArray<int32>& AgentIds)
-{
-	for (const int32 AgentId : AgentIds)
-	{
-		ResetAgentEpisode(AgentId);
-	}
-}
-
-const bool ULearningAgentsTrainer::IsTraining() const
-{
-	return bIsTraining;
-}
-
-void ULearningAgentsTrainer::BeginTraining(
-	const FLearningAgentsTrainerTrainingSettings& TrainerTrainingSettings,
-	const FLearningAgentsTrainerGameSettings& TrainerGameSettings,
-	const FLearningAgentsTrainerPathSettings& TrainerPathSettings,
-	const bool bResetAgentsOnBegin)
-{
-	if (!IsSetup())
-	{
-		UE_LOG(LogLearning, Error, TEXT("%s: Setup not complete."), *GetName());
-		return;
-	}
-
-	if (IsTraining())
-	{
-		UE_LOG(LogLearning, Error, TEXT("%s: Already Training!"), *GetName());
-		return;
-	}
-
-	// Check Paths
-
-	const FString PythonExecutablePath = UE::Learning::Trainer::GetPythonExecutablePath(TrainerPathSettings.GetIntermediatePath());
-
-	if (!FPaths::FileExists(PythonExecutablePath))
-	{
-		UE_LOG(LogLearning, Error, TEXT("%s: Can't find Python executable \"%s\"."), *GetName(), *PythonExecutablePath);
-		return;
-	}
-
-	const FString PythonContentPath = UE::Learning::Trainer::GetPythonContentPath(TrainerPathSettings.GetEditorEnginePath());
-
-	if (!FPaths::DirectoryExists(PythonContentPath))
-	{
-		UE_LOG(LogLearning, Error, TEXT("%s: Can't find LearningAgents plugin Content \"%s\"."), *GetName(), *PythonContentPath);
-		return;
-	}
-
-	const FString IntermediatePath = UE::Learning::Trainer::GetIntermediatePath(TrainerPathSettings.GetIntermediatePath());
-
-	// Record GameState Settings
-
-	bFixedTimestepUsed = FApp::UseFixedTimeStep();
-	FixedTimeStepDeltaTime = FApp::GetFixedDeltaTime();
-
-	UGameUserSettings* GameSettings = UGameUserSettings::GetGameUserSettings();
-	if (GameSettings)
-	{
-		bVSyncEnabled = GameSettings->IsVSyncEnabled();
-	}
-
-	UPhysicsSettings* PhysicsSettings = UPhysicsSettings::Get();
-	if (PhysicsSettings)
-	{
-		MaxPhysicsStep = PhysicsSettings->MaxPhysicsDeltaTime;
-	}
-
-	IConsoleVariable* MaxFPSCVar = IConsoleManager::Get().FindConsoleVariable(TEXT("t.MaxFPS"));
-	if (MaxFPSCVar)
-	{
-		MaxFPS = MaxFPSCVar->GetInt();
-	}
-
-	UGameViewportClient* ViewportClient = GetWorld() ? GetWorld()->GetGameViewport() : nullptr;
-	if (ViewportClient)
-	{
-		ViewModeIndex = ViewportClient->ViewModeIndex;
-	}
-
-#if WITH_EDITOR
-	UEditorPerformanceSettings* EditorPerformanceSettings = GetMutableDefault<UEditorPerformanceSettings>();
-	if (EditorPerformanceSettings)
-	{
-		bUseLessCPUInTheBackground = EditorPerformanceSettings->bThrottleCPUWhenNotForeground;
-		bEditorVSyncEnabled = EditorPerformanceSettings->bEnableVSync;
-	}
-#endif
-
-	// Apply Training GameState Settings
-
-	FApp::SetUseFixedTimeStep(TrainerGameSettings.bUseFixedTimeStep);
-
-	if (TrainerGameSettings.FixedTimeStepFrequency > UE_SMALL_NUMBER)
-	{
-		FApp::SetFixedDeltaTime(1.0f / TrainerGameSettings.FixedTimeStepFrequency);
-		if (TrainerGameSettings.bSetMaxPhysicsStepToFixedTimeStep && PhysicsSettings)
-		{
-			PhysicsSettings->MaxPhysicsDeltaTime = 1.0f / TrainerGameSettings.FixedTimeStepFrequency;
-		}
-	}
-	else
-	{
-		UE_LOG(LogLearning, Warning, TEXT("%s: Provided invalid FixedTimeStepFrequency: %0.5f"), *GetName(), TrainerGameSettings.FixedTimeStepFrequency);
-	}
-
-	if (TrainerGameSettings.bDisableMaxFPS && MaxFPSCVar)
-	{
-		MaxFPSCVar->Set(0);
-	}
-
-	if (TrainerGameSettings.bDisableVSync && GameSettings)
-	{
-		GameSettings->SetVSyncEnabled(false);
-		GameSettings->ApplySettings(false);
-	}
-
-	if (TrainerGameSettings.bUseUnlitViewportRendering && ViewportClient)
-	{
-		ViewportClient->ViewModeIndex = EViewModeIndex::VMI_Unlit;
-	}
-
-#if WITH_EDITOR
-	if (TrainerGameSettings.bDisableUseLessCPUInTheBackground && EditorPerformanceSettings)
-	{
-		EditorPerformanceSettings->bThrottleCPUWhenNotForeground = false;
-		EditorPerformanceSettings->PostEditChange();
-	}
-
-	if (TrainerGameSettings.bDisableEditorVSync && EditorPerformanceSettings)
-	{
-		EditorPerformanceSettings->bEnableVSync = false;
-		EditorPerformanceSettings->PostEditChange();
-	}
-#endif
-
-	// Start Trainer
-
-	UE::Learning::FPPOTrainerTrainingSettings PPOTrainingSettings;
-	PPOTrainingSettings.IterationNum = TrainerTrainingSettings.NumberOfIterations;
-	PPOTrainingSettings.LearningRatePolicy = TrainerTrainingSettings.LearningRatePolicy;
-	PPOTrainingSettings.LearningRateCritic = TrainerTrainingSettings.LearningRateCritic;
-	PPOTrainingSettings.LearningRateDecay = TrainerTrainingSettings.LearningRateDecay;
-	PPOTrainingSettings.WeightDecay = TrainerTrainingSettings.WeightDecay;
-	PPOTrainingSettings.PolicyBatchSize = TrainerTrainingSettings.PolicyBatchSize;
-	PPOTrainingSettings.CriticBatchSize = TrainerTrainingSettings.CriticBatchSize;
-	PPOTrainingSettings.PolicyWindow = TrainerTrainingSettings.PolicyWindowSize;
-	PPOTrainingSettings.IterationsPerGather = TrainerTrainingSettings.IterationsPerGather;
-	PPOTrainingSettings.CriticWarmupIterations = TrainerTrainingSettings.CriticWarmupIterations;
-	PPOTrainingSettings.EpsilonClip = TrainerTrainingSettings.EpsilonClip;
-	PPOTrainingSettings.ReturnRegularizationWeight = TrainerTrainingSettings.ReturnRegularizationWeight;
-	PPOTrainingSettings.ActionRegularizationWeight = TrainerTrainingSettings.ActionRegularizationWeight;
-	PPOTrainingSettings.ActionEntropyWeight = TrainerTrainingSettings.ActionEntropyWeight;
-	PPOTrainingSettings.GaeLambda = TrainerTrainingSettings.GaeLambda;
-	PPOTrainingSettings.bAdvantageNormalization = TrainerTrainingSettings.bAdvantageNormalization;
-	PPOTrainingSettings.AdvantageMin = TrainerTrainingSettings.MinimumAdvantage;
-	PPOTrainingSettings.AdvantageMax = TrainerTrainingSettings.MaximumAdvantage;
-	PPOTrainingSettings.bUseGradNormMaxClipping = TrainerTrainingSettings.bUseGradNormMaxClipping;
-	PPOTrainingSettings.GradNormMax = TrainerTrainingSettings.GradNormMax;
-	PPOTrainingSettings.TrimEpisodeStartStepNum = TrainerTrainingSettings.NumberOfStepsToTrimAtStartOfEpisode;
-	PPOTrainingSettings.TrimEpisodeEndStepNum = TrainerTrainingSettings.NumberOfStepsToTrimAtEndOfEpisode;
-	PPOTrainingSettings.Seed = TrainerTrainingSettings.RandomSeed;
-	PPOTrainingSettings.DiscountFactor = TrainerTrainingSettings.DiscountFactor;
-	PPOTrainingSettings.Device = UE::Learning::Agents::GetTrainerDevice(TrainerTrainingSettings.Device);
-	PPOTrainingSettings.bUseTensorboard = TrainerTrainingSettings.bUseTensorboard;
-	PPOTrainingSettings.bSaveSnapshots = TrainerTrainingSettings.bSaveSnapshots;
-
-	// Start Python Training Process (this must be done on game thread)
-	Trainer = MakeUnique<UE::Learning::FSharedMemoryPPOTrainer>(
-		GetName(),
-		PythonExecutablePath,
-		TEXT(""),
-		PythonContentPath,
-		IntermediatePath,
-		*ReplayBuffer,
-		*Policy->GetPolicyNetworkAsset()->NeuralNetworkData,
-		*Critic->GetCriticNetworkAsset()->NeuralNetworkData,
-		*Policy->GetEncoderNetworkAsset()->NeuralNetworkData,
-		*Policy->GetDecoderNetworkAsset()->NeuralNetworkData,
-		Interactor->GetObservationSchema(),
-		Interactor->GetObservationSchemaElement(),
-		Interactor->GetActionSchema(),
-		Interactor->GetActionSchemaElement(),
-		PPOTrainingSettings);
-
-	UE_LOG(LogLearning, Display, TEXT("%s: Sending / Receiving initial policy..."), *GetName());
-
-	UE::Learning::ETrainerResponse Response = UE::Learning::ETrainerResponse::Success;
-
-	Response = Trainer->SendPolicy(*Policy->GetPolicyNetworkAsset()->NeuralNetworkData, TrainerTimeout);
-	if (Response != UE::Learning::ETrainerResponse::Success)
-	{
-		UE_LOG(LogLearning, Error, TEXT("%s: Error sending policy to trainer: %s. Check log for additional errors."), *GetName(), UE::Learning::Trainer::GetResponseString(Response));
-		bHasTrainingFailed = true;
-		Trainer->Terminate();
-		return;
-	}
-
-	Response = Trainer->SendCritic(*Critic->GetCriticNetworkAsset()->NeuralNetworkData, TrainerTimeout);
-	if (Response != UE::Learning::ETrainerResponse::Success)
-	{
-		UE_LOG(LogLearning, Error, TEXT("%s: Error sending critic to trainer: %s. Check log for additional errors."), *GetName(), UE::Learning::Trainer::GetResponseString(Response));
-		bHasTrainingFailed = true;
-		Trainer->Terminate();
-		return;
-	}
-
-	Response = Trainer->SendEncoder(*Policy->GetEncoderNetworkAsset()->NeuralNetworkData, TrainerTimeout);
-	if (Response != UE::Learning::ETrainerResponse::Success)
-	{
-		UE_LOG(LogLearning, Error, TEXT("%s: Error sending encoder to trainer: %s. Check log for additional errors."), *GetName(), UE::Learning::Trainer::GetResponseString(Response));
-		bHasTrainingFailed = true;
-		Trainer->Terminate();
-		return;
-	}
-
-	Response = Trainer->SendDecoder(*Policy->GetDecoderNetworkAsset()->NeuralNetworkData, TrainerTimeout);
-	if (Response != UE::Learning::ETrainerResponse::Success)
-	{
-		UE_LOG(LogLearning, Error, TEXT("%s: Error sending decoder to trainer: %s. Check log for additional errors."), *GetName(), UE::Learning::Trainer::GetResponseString(Response));
-		bHasTrainingFailed = true;
-		Trainer->Terminate();
-		return;
-	}
-
-	// Reset Agents, Replay Buffer
-
-	if (bResetAgentsOnBegin)
-	{
-		Manager->ResetAllAgents();
-	}
-
-	ReplayBuffer->Reset();
-
-	bIsTraining = true;
-}
-
-void ULearningAgentsTrainer::DoneTraining()
-{
-	if (IsTraining())
-	{
-		// Wait for Trainer to finish
-		Trainer->Wait(1.0f);
-
-		// If not finished in time, terminate
-		Trainer->Terminate();
-
-		// Apply back previous game settings
-		FApp::SetUseFixedTimeStep(bFixedTimestepUsed);
-		FApp::SetFixedDeltaTime(FixedTimeStepDeltaTime);
-		UGameUserSettings* GameSettings = UGameUserSettings::GetGameUserSettings();
-		if (GameSettings)
-		{
-			GameSettings->SetVSyncEnabled(bVSyncEnabled);
-			GameSettings->ApplySettings(true);
-		}
-
-		UPhysicsSettings* PhysicsSettings = UPhysicsSettings::Get();
-		if (PhysicsSettings)
-		{
-			PhysicsSettings->MaxPhysicsDeltaTime = MaxPhysicsStep;
-		}
-
-		IConsoleVariable* MaxFPSCVar = IConsoleManager::Get().FindConsoleVariable(TEXT("t.MaxFPS"));
-		if (MaxFPSCVar)
-		{
-			MaxFPSCVar->Set(MaxFPS);
-		}
-
-		UGameViewportClient* ViewportClient = GetWorld() ? GetWorld()->GetGameViewport() : nullptr;
-		if (ViewportClient)
-		{
-			ViewportClient->ViewModeIndex = ViewModeIndex;
-		}
-
-#if WITH_EDITOR
-		UEditorPerformanceSettings* EditorPerformanceSettings = GetMutableDefault<UEditorPerformanceSettings>();
-		if (EditorPerformanceSettings)
-		{
-			EditorPerformanceSettings->bThrottleCPUWhenNotForeground = bUseLessCPUInTheBackground;
-			EditorPerformanceSettings->bEnableVSync = bEditorVSyncEnabled;
-			EditorPerformanceSettings->PostEditChange();
-		}
-#endif
-
-		bIsTraining = false;
-	}
-}
-
-void ULearningAgentsTrainer::EndTraining()
-{
-	if (IsTraining())
-	{
-		UE_LOG(LogLearning, Display, TEXT("%s: Stopping training..."), *GetName());
-		Trainer->SendStop();
-		DoneTraining();
-	}
-}
-
-void ULearningAgentsTrainer::GatherRewards()
-{
-	UE_LEARNING_TRACE_CPUPROFILER_EVENT_SCOPE(ULearningAgentsTrainer::GatherRewards);
-
-	if (!IsSetup())
-	{
-		UE_LOG(LogLearning, Error, TEXT("%s: Setup not complete."), *GetName());
-		return;
-	}
-
-	if (Manager->GetAgentNum() == 0)
-	{
-		UE_LOG(LogLearning, Warning, TEXT("%s: No agents added to Manager."), *GetName());
-	}
-
-	// Get Rewards
-
-	ValidAgentIds = Manager->GetAllAgentIds();
-	ValidAgentSet = Manager->GetAllAgentSet();
-
-	RewardBuffer.Empty(Manager->GetMaxAgentNum());
-	GatherAgentRewards(RewardBuffer, ValidAgentIds);
-
-	if (ValidAgentSet.Num() != RewardBuffer.Num())
-	{
-		UE_LOG(LogLearning, Warning, TEXT("%s: Not enough rewards added by GetAgentRewards. Expected %i, Got %i."), *GetName(), ValidAgentSet.Num(), RewardBuffer.Num());
-		return;
-	}
-
-	for (int32 AgentIdx = 0; AgentIdx < RewardBuffer.Num(); AgentIdx++)
-	{
-		const float RewardValue = RewardBuffer[AgentIdx];
-
-		if (FMath::IsFinite(RewardValue) && RewardValue != MAX_flt && RewardValue != -MAX_flt)
-		{
-			Rewards[ValidAgentSet[AgentIdx]] = RewardValue;
-			RewardIteration[ValidAgentSet[AgentIdx]]++;
-		}
-		else
-		{
-			UE_LOG(LogLearning, Warning, TEXT("%s: Got invalid reward for agent %i: %f."), *GetName(), ValidAgentSet[AgentIdx], RewardValue);
-			continue;
-		}
-	}
-}
-
-void ULearningAgentsTrainer::GatherCompletions()
-{
-	UE_LEARNING_TRACE_CPUPROFILER_EVENT_SCOPE(ULearningAgentsTrainer::GatherCompletions);
-
-	if (!IsSetup())
-	{
-		UE_LOG(LogLearning, Error, TEXT("%s: Setup not complete."), *GetName());
-		return;
-	}
-
-	if (Manager->GetAgentNum() == 0)
-	{
-		UE_LOG(LogLearning, Warning, TEXT("%s: No agents added to Manager."), *GetName());
-	}
-
-	// Get Completions
-
-	ValidAgentIds = Manager->GetAllAgentIds();
-	ValidAgentSet = Manager->GetAllAgentSet();
-
-	CompletionBuffer.Empty(Manager->GetMaxAgentNum());
-	GatherAgentCompletions(CompletionBuffer, ValidAgentIds);
-
-	if (ValidAgentSet.Num() != CompletionBuffer.Num())
-	{
-		UE_LOG(LogLearning, Warning, TEXT("%s: Not enough completions added by GetAgentCompletions. Expected %i, Got %i."), *GetName(), ValidAgentSet.Num(), CompletionBuffer.Num());
-		return;
-	}
-
-	for (int32 AgentIdx = 0; AgentIdx < CompletionBuffer.Num(); AgentIdx++)
-	{
-		AgentCompletions[ValidAgentSet[AgentIdx]] = UE::Learning::Agents::GetCompletionMode(CompletionBuffer[AgentIdx]);
-		CompletionIteration[ValidAgentSet[AgentIdx]]++;
-	}
-}
-
-void ULearningAgentsTrainer::ProcessExperience(const bool bResetAgentsOnUpdate)
-{
-	UE_LEARNING_TRACE_CPUPROFILER_EVENT_SCOPE(ULearningAgentsTrainer::ProcessExperience);
-
-	if (!IsSetup())
-	{
-		UE_LOG(LogLearning, Error, TEXT("%s: Setup not complete."), *GetName());
-		return;
-	}
-
-	if (!IsTraining())
-	{
-		UE_LOG(LogLearning, Error, TEXT("%s: Training not running."), *GetName());
-		return;
-	}
-
-	if (Manager->GetAgentNum() == 0)
-	{
-		UE_LOG(LogLearning, Warning, TEXT("%s: No agents added to Manager."), *GetName());
-	}
-
-	// Check Observations, Actions, Rewards, and Completions have been completed and have matching iteration number
-
-	ValidAgentIds.Empty(Manager->GetMaxAgentNum());
-
-	for (const int32 AgentId : Manager->GetAllAgentSet())
-	{
-		if (Interactor->ObservationVectorIteration[AgentId] == 0 ||
-			Interactor->ActionVectorIteration[AgentId] == 0 ||
-			RewardIteration[AgentId] == 0 ||
-			CompletionIteration[AgentId] == 0)
-		{
-			UE_LOG(LogLearning, Display, TEXT("%s: Agent with id %i has not completed a full step of observations, actions, rewards, completions and so experience will not be processed for it."), *GetName(), AgentId);
-			continue;
-		}
-
-		if (Interactor->ObservationVectorIteration[AgentId] != Interactor->ActionVectorIteration[AgentId] ||
-			Interactor->ObservationVectorIteration[AgentId] != RewardIteration[AgentId] ||
-			Interactor->ObservationVectorIteration[AgentId] != CompletionIteration[AgentId])
-		{
-			UE_LOG(LogLearning, Warning, TEXT("%s: Agent with id %i has non-matching iteration numbers (observation: %i, action: %i, reward: %i, completion: %i). Experience will not be processed for it."), *GetName(), AgentId,
-				Interactor->ObservationVectorIteration[AgentId],
-				Interactor->ActionVectorIteration[AgentId],
-				RewardIteration[AgentId],
-				CompletionIteration[AgentId]);
-			continue;
-		}
-
-		ValidAgentIds.Add(AgentId);
-	}
-
-	ValidAgentSet = ValidAgentIds;
-	ValidAgentSet.TryMakeSlice();
-
-	// Check for episodes that have been immediately completed
-
-	for (const int32 AgentId : ValidAgentSet)
-	{
-		if (AgentCompletions[AgentId] != UE::Learning::ECompletionMode::Running && EpisodeBuffer->GetEpisodeStepNums()[AgentId] == 0)
-		{
-			UE_LOG(LogLearning, Warning, TEXT("%s: Agent with id %i has completed episode and will be reset but has not generated any experience."), *GetName(), AgentId);
-		}
-	}
-
-	// Add Experience to Episode Buffer
-	EpisodeBuffer->Push(
-		Interactor->ObservationVectors,
-		Interactor->ActionVectors,
-		Policy->PreEvaluationMemoryState,
-		Rewards,
-		ValidAgentSet);
-
-	// Find the set of agents which have reached the maximum episode length and mark them as truncated
-	UE::Learning::Completion::EvaluateEndOfEpisodeCompletions(
-		EpisodeCompletions,
-		EpisodeBuffer->GetEpisodeStepNums(),
-		EpisodeBuffer->GetMaxStepNum(),
-		ValidAgentSet);
-
-	// Compute a combined completion buffer for agents that have been completed manually and those which have reached the maximum episode length
-	for (const int32 AgentIdx : ValidAgentSet)
-	{
-		AllCompletions[AgentIdx] = UE::Learning::Completion::Or(AgentCompletions[AgentIdx], EpisodeCompletions[AgentIdx]);
-	}
-
-	ResetBuffer->SetResetInstancesFromCompletions(AllCompletions, ValidAgentSet);
-
-	// If there are no agents completed we are done
-	if (ResetBuffer->GetResetInstanceNum() == 0)
-	{
-		return;
-	}
-
-	// Otherwise Gather Observations for completed Instances without incrementing iteration number
-	Interactor->GatherObservations(ResetBuffer->GetResetInstances(), false);
-
-	// And push those episodes to the Replay Buffer
-	const bool bReplayBufferFull = ReplayBuffer->AddEpisodes(
-		AllCompletions,
-		Interactor->ObservationVectors,
-		Policy->MemoryState,
-		*EpisodeBuffer,
-		ResetBuffer->GetResetInstances());
-
-	if (bReplayBufferFull)
-	{
-		UE::Learning::ETrainerResponse Response = Trainer->SendExperience(*ReplayBuffer, TrainerTimeout);
-
-		if (Response != UE::Learning::ETrainerResponse::Success)
-		{
-			UE_LOG(LogLearning, Error, TEXT("%s: Error waiting to push experience to trainer: %s. Check log for additional errors."), *GetName(), UE::Learning::Trainer::GetResponseString(Response));
-			EndTraining();
-			return;
-		}
-
-		ReplayBuffer->Reset();
-
-		// Get Updated Policy
-		Response = Trainer->RecvPolicy(*Policy->GetPolicyNetworkAsset()->NeuralNetworkData, TrainerTimeout);
-		Policy->GetPolicyNetworkAsset()->ForceMarkDirty();
-
-		if (Response == UE::Learning::ETrainerResponse::Completed)
-		{
-			UE_LOG(LogLearning, Display, TEXT("%s: Trainer completed training."), *GetName());
-			DoneTraining();
-			return;
-		}
-		else if (Response != UE::Learning::ETrainerResponse::Success)
-		{
-			UE_LOG(LogLearning, Error, TEXT("%s: Error waiting for policy from trainer: %s. Check log for additional errors."), *GetName(), UE::Learning::Trainer::GetResponseString(Response));
-			bHasTrainingFailed = true;
-			EndTraining();
-			return;
-		}
-
-		// Get Updated Critic
-		Response = Trainer->RecvCritic(*Critic->GetCriticNetworkAsset()->NeuralNetworkData, TrainerTimeout);
-		Critic->GetCriticNetworkAsset()->ForceMarkDirty();
-
-		if (Response != UE::Learning::ETrainerResponse::Success)
-		{
-			UE_LOG(LogLearning, Error, TEXT("%s: Error waiting for critic from trainer: %s. Check log for additional errors."), *GetName(), UE::Learning::Trainer::GetResponseString(Response));
-			bHasTrainingFailed = true;
-			EndTraining();
-			return;
-		}
-
-		// Get Updated Encoder
-		Response = Trainer->RecvEncoder(*Policy->GetEncoderNetworkAsset()->NeuralNetworkData, TrainerTimeout);
-		Policy->GetEncoderNetworkAsset()->ForceMarkDirty();
-
-		if (Response != UE::Learning::ETrainerResponse::Success)
-		{
-			UE_LOG(LogLearning, Error, TEXT("%s: Error waiting for encoder from trainer: %s. Check log for additional errors."), *GetName(), UE::Learning::Trainer::GetResponseString(Response));
-			bHasTrainingFailed = true;
-			EndTraining();
-			return;
-		}
-
-		// Get Updated Decoder
-		Response = Trainer->RecvDecoder(*Policy->GetDecoderNetworkAsset()->NeuralNetworkData, TrainerTimeout);
-		Policy->GetDecoderNetworkAsset()->ForceMarkDirty();
-
-		if (Response != UE::Learning::ETrainerResponse::Success)
-		{
-			UE_LOG(LogLearning, Error, TEXT("%s: Error waiting for decoder from trainer: %s. Check log for additional errors."), *GetName(), UE::Learning::Trainer::GetResponseString(Response));
-			bHasTrainingFailed = true;
-			EndTraining();
-			return;
-		}
-
-		if (bResetAgentsOnUpdate)
-		{
-			// Reset all agents since we have a new policy
-			ResetBuffer->SetResetInstances(Manager->GetAllAgentSet());
-			Manager->ResetAgents(ResetBuffer->GetResetInstancesArray());
-			return;
-		}
-	}
-
-	// Manually reset Episode Buffer for agents who have reached the maximum episode length as 
-	// they wont get it reset via the agent manager's call to ResetAgents
-	ResetBuffer->SetResetInstancesFromCompletions(EpisodeCompletions, ValidAgentSet);
-	EpisodeBuffer->Reset(ResetBuffer->GetResetInstances());
-
-	// Call ResetAgents for agents which have manually signaled a completion
-	ResetBuffer->SetResetInstancesFromCompletions(AgentCompletions, ValidAgentSet);
-	if (ResetBuffer->GetResetInstanceNum() > 0)
-	{
-		Manager->ResetAgents(ResetBuffer->GetResetInstancesArray());
-	}
-}
-
-void ULearningAgentsTrainer::RunTraining(
-	const FLearningAgentsTrainerTrainingSettings& TrainerTrainingSettings,
-	const FLearningAgentsTrainerGameSettings& TrainerGameSettings,
-	const FLearningAgentsTrainerPathSettings& TrainerPathSettings,
-	const bool bResetAgentsOnBegin,
-	const bool bResetAgentsOnUpdate)
-{
-	if (!IsSetup())
-	{
-		UE_LOG(LogLearning, Error, TEXT("%s: Setup not complete."), *GetName());
-		return;
-	}
-
-	if (bHasTrainingFailed)
-	{
-		UE_LOG(LogLearning, Error, TEXT("%s: Training has failed. Check log for errors."), *GetName());
-		return;
-	}
-
-	// If we aren't training yet, then start training and do the first inference step.
-	if (!IsTraining())
-	{
-		BeginTraining(
-			TrainerTrainingSettings,
-			TrainerGameSettings, 
-			TrainerPathSettings, 
-			bResetAgentsOnBegin);
-
-		if (!IsTraining())
-		{
-			// If IsTraining is false, then BeginTraining must have failed and we can't continue.
-			return;
-		}
-
-		Policy->RunInference();
-	}
-	// Otherwise, do the regular training process.
-	else
-	{
-		GatherCompletions();
-		GatherRewards();
-		ProcessExperience(bResetAgentsOnUpdate);
-		Policy->RunInference();
-	}
-}
-
-bool ULearningAgentsTrainer::HasReward(const int32 AgentId) const
-{
-	if (!IsSetup())
-	{
-		UE_LOG(LogLearning, Error, TEXT("%s: Setup not complete."), *GetName());
-		return false;
-	}
-
-	if (!HasAgent(AgentId))
-	{
-		UE_LOG(LogLearning, Error, TEXT("%s: AgentId %d not found in the agents set."), *GetName(), AgentId);
-		return false;
-	}
-
-	return RewardIteration[AgentId] > 0;
-}
-
-bool ULearningAgentsTrainer::HasCompletion(const int32 AgentId) const
-{
-	if (!IsSetup())
-	{
-		UE_LOG(LogLearning, Error, TEXT("%s: Setup not complete."), *GetName());
-		return false;
-	}
-
-	if (!HasAgent(AgentId))
-	{
-		UE_LOG(LogLearning, Error, TEXT("%s: AgentId %d not found in the agents set."), *GetName(), AgentId);
-		return false;
-	}
-
-	return CompletionIteration[AgentId] > 0;
-}
-
-float ULearningAgentsTrainer::GetReward(const int32 AgentId) const
-{
-	if (!IsSetup())
-	{
-		UE_LOG(LogLearning, Error, TEXT("%s: Setup not complete."), *GetName());
-		return 0.0f;
-	}
-
-	if (!HasAgent(AgentId))
-	{
-		UE_LOG(LogLearning, Error, TEXT("%s: AgentId %d not found in the agents set."), *GetName(), AgentId);
-		return 0.0f;
-	}
-
-	if (RewardIteration[AgentId] == 0)
-	{
-		UE_LOG(LogLearning, Error, TEXT("%s: Agent with id %d has not evaluated rewards. Did you run EvaluateRewards?"), *GetName(), AgentId);
-		return 0.0f;
-	}
-
-	return Rewards[AgentId];
-}
-
-ELearningAgentsCompletion ULearningAgentsTrainer::GetCompletion(const int32 AgentId) const
-{
-	if (!IsSetup())
-	{
-		UE_LOG(LogLearning, Error, TEXT("%s: Setup not complete."), *GetName());
-		return ELearningAgentsCompletion::Running;
-	}
-
-	if (!HasAgent(AgentId))
-	{
-		UE_LOG(LogLearning, Error, TEXT("%s: AgentId %d not found in the agents set."), *GetName(), AgentId);
-		return ELearningAgentsCompletion::Running;
-	}
-
-	if (CompletionIteration[AgentId] == 0)
-	{
-		UE_LOG(LogLearning, Error, TEXT("%s: Agent with id %d has not evaluated completions. Did you run EvaluateCompletions?"), *GetName(), AgentId);
-		return ELearningAgentsCompletion::Running;
-	}
-
-	return UE::Learning::Agents::GetLearningAgentsCompletion(AgentCompletions[AgentId]);
-}
-
-float ULearningAgentsTrainer::GetEpisodeTime(const int32 AgentId) const
-{
-	if (!IsSetup())
-	{
-		UE_LOG(LogLearning, Error, TEXT("%s: Setup not complete."), *GetName());
-		return 0.0f;
-	}
-
-	if (!HasAgent(AgentId))
-	{
-		UE_LOG(LogLearning, Error, TEXT("%s: AgentId %d not found in the agents set."), *GetName(), AgentId);
-		return 0.0f;
-	}
-
-	return EpisodeTimes[AgentId];
-}
-
-int32 ULearningAgentsTrainer::GetEpisodeStepNum(const int32 AgentId) const
-{
-	if (!IsSetup())
-	{
-		UE_LOG(LogLearning, Error, TEXT("%s: Setup not complete."), *GetName());
-		return 0.0f;
-	}
-
-	if (!HasAgent(AgentId))
-	{
-		UE_LOG(LogLearning, Error, TEXT("%s: AgentId %d not found in the agents set."), *GetName(), AgentId);
-		return 0.0f;
-	}
-
-	return EpisodeBuffer->GetEpisodeStepNums()[AgentId];
-}
-
-bool ULearningAgentsTrainer::HasTrainingFailed() const
-{
-	return bHasTrainingFailed;
-}
-
diff --git a/Engine/Plugins/Experimental/LearningAgents/Source/LearningAgentsTraining/Public/LearningAgentsImitationTrainer.h b/Engine/Plugins/Experimental/LearningAgents/Source/LearningAgentsTraining/Public/LearningAgentsImitationTrainer.h
--- a/Engine/Plugins/Experimental/LearningAgents/Source/LearningAgentsTraining/Public/LearningAgentsImitationTrainer.h
+++ b/Engine/Plugins/Experimental/LearningAgents/Source/LearningAgentsTraining/Public/LearningAgentsImitationTrainer.h
@@ -2,7 +2,9 @@
 
 #pragma once
 
-#include "LearningAgentsTrainer.h" // Included for ELearningAgentsTrainerDevice and FLearningAgentsTrainerPathSettings
+#include "LearningAgentsManagerListener.h"
+#include "LearningAgentsTrainer.h"
+
 #include "LearningArray.h"
 
 #include "Templates/SharedPointer.h"
@@ -93,7 +95,7 @@ public:
 
 	/** The device to train on. */
 	UPROPERTY(EditAnywhere, Category = "LearningAgents")
-	ELearningAgentsTrainerDevice Device = ELearningAgentsTrainerDevice::CPU;
+	ELearningAgentsTrainingDevice Device = ELearningAgentsTrainingDevice::CPU;
 
 	/** If true, TensorBoard logs will be emitted to the intermediate directory. */
 	UPROPERTY(EditAnywhere, Category = "LearningAgents")
@@ -182,7 +184,7 @@ public:
 		const ULearningAgentsRecording* Recording,
 		const FLearningAgentsImitationTrainerSettings& ImitationTrainerSettings = FLearningAgentsImitationTrainerSettings(),
 		const FLearningAgentsImitationTrainerTrainingSettings& ImitationTrainerTrainingSettings = FLearningAgentsImitationTrainerTrainingSettings(),
-		const FLearningAgentsTrainerPathSettings& ImitationTrainerPathSettings = FLearningAgentsTrainerPathSettings());
+		const FLearningAgentsTrainerProcessSettings& ImitationTrainerPathSettings = FLearningAgentsTrainerProcessSettings());
 
 	/** Stops the training process. */
 	UFUNCTION(BlueprintCallable, Category = "LearningAgents")
@@ -206,7 +208,7 @@ public:
 		const ULearningAgentsRecording* Recording,
 		const FLearningAgentsImitationTrainerSettings& ImitationTrainerSettings = FLearningAgentsImitationTrainerSettings(),
 		const FLearningAgentsImitationTrainerTrainingSettings& ImitationTrainerTrainingSettings = FLearningAgentsImitationTrainerTrainingSettings(),
-		const FLearningAgentsTrainerPathSettings& ImitationTrainerPathSettings = FLearningAgentsTrainerPathSettings());
+		const FLearningAgentsTrainerProcessSettings& ImitationTrainerPathSettings = FLearningAgentsTrainerProcessSettings());
 
 // ----- Private Data -----
 private:
