Learning Agents Training Rewrite: Part 1

diff --git a/Engine/Plugins/Experimental/LearningAgents/Source/LearningAgentsTraining/Public/LearningAgentsTrainer.h b/Engine/Plugins/Experimental/LearningAgents/Source/LearningAgentsTraining/Public/LearningAgentsTrainer.h
--- a/Engine/Plugins/Experimental/LearningAgents/Source/LearningAgentsTraining/Public/LearningAgentsTrainer.h
+++ b/Engine/Plugins/Experimental/LearningAgents/Source/LearningAgentsTraining/Public/LearningAgentsTrainer.h
@@ -2,124 +2,13 @@
 
 #pragma once
 
-#include "LearningAgentsManagerListener.h"
-#include "LearningAgentsCompletions.h"
-
-#include "Templates/SharedPointer.h"
-#include "UObject/ObjectPtr.h"
-#include "Engine/EngineTypes.h"
+#include "LearningTrainer.h"
 
 #include "LearningAgentsTrainer.generated.h"
 
-namespace UE::Learning
-{
-	struct FEpisodeBuffer;
-	struct FReplayBuffer;
-	struct FResetInstanceBuffer;
-	struct FSharedMemoryPPOTrainer;
-	enum class ETrainerDevice : uint8;
-}
-
-class ALearningAgentsManager;
-class ULearningAgentsInteractor;
-class ULearningAgentsPolicy;
-class ULearningAgentsCritic;
-
-/** The configurable settings for a ULearningAgentsTrainer. */
-USTRUCT(BlueprintType, Category = "LearningAgents")
-struct LEARNINGAGENTSTRAINING_API FLearningAgentsTrainerSettings
-{
-	GENERATED_BODY()
-
-public:
-
-	/**
-	 * Maximum number of steps recorded in an episode before it is added to the replay buffer. This can generally be left at the default value and 
-	 * does not have a large impact on training.
-	 */
-	UPROPERTY(EditAnywhere, Category = "LearningAgents", meta = (ClampMin = "1", UIMin = "1"))
-	int32 MaxEpisodeStepNum = 512;
-
-	/**
-	 * Maximum number of episodes to record before running a training iteration. An iteration of training will be run when either this or
-	 * MaximumRecordedEpisodesPerIteration is reached. Typical values for this should be around 1000. Setting this too small means there is not 
-	 * enough data each iteration for the system to train. Setting it too large means training will be very slow. 
-	 */
-	UPROPERTY(EditAnywhere, Category = "LearningAgents", meta = (ClampMin = "1", UIMin = "1"))
-	int32 MaximumRecordedEpisodesPerIteration = 1000;
-
-	/**
-	 * Maximum number of steps to record before running a training iteration. An iteration of training will be run when either this or
-	 * MaximumRecordedEpisodesPerIteration is reached. Typical values for this should be around 10000. Setting this too small means there is not
-	 * enough data each iteration for the system to train. Setting it too large means training will be very slow.
-	 */
-	UPROPERTY(EditAnywhere, Category = "LearningAgents", meta = (ClampMin = "1", UIMin = "1"))
-	int32 MaximumRecordedStepsPerIteration = 10000;
-
-	/** Time in seconds to wait for the training subprocess before timing out. */
-	UPROPERTY(EditAnywhere, Category = "LearningAgents", meta = (ClampMin = "0.0", UIMin = "0.0"))
-	float TrainerCommunicationTimeout = 10.0f;
-};
-
-/**
- * The configurable game settings for a ULearningAgentsTrainer. These allow the timestep and physics tick to be fixed
- * during training, which can enable ticking faster than real-time.
- */
-USTRUCT(BlueprintType, Category="LearningAgents")
-struct LEARNINGAGENTSTRAINING_API FLearningAgentsTrainerGameSettings
-{
-	GENERATED_BODY()
-
-public:
-
-	/**
-	 * If true, the game will run in fixed time step mode (i.e the frame's delta times will always be the same
-	 * regardless of how much wall time has passed). This can enable faster than real-time training if your game runs
-	 * quickly. If false, the time steps will match real wall time.
-	 */
-	UPROPERTY(EditAnywhere, Category = "LearningAgents")
-	bool bUseFixedTimeStep = true;
-
-	/**
-	 * Determines the amount of time for each frame when bUseFixedTimeStep is true; Ignored if false. You want this
-	 * time step to match as closely as possible to the expected inference time steps, otherwise your training results
-	 * may not generalize to your game.
-	 */
-	UPROPERTY(EditAnywhere, Category = "LearningAgents", meta = (DisplayName = "Fixed Time Step Frequency (Hz)"), meta = (ClampMin = "0.0", UIMin = "0.0"))
-	float FixedTimeStepFrequency = 60.0f;
-
-	/** If true, set the physics delta time to match the fixed time step. */
-	UPROPERTY(EditAnywhere, Category = "LearningAgents")
-	bool bSetMaxPhysicsStepToFixedTimeStep = true;
-
-	/** If true, the MaxFPS console variable will be set to a negative number during training; Otherwise, it will not. */
-	UPROPERTY(EditAnywhere, Category = "LearningAgents")
-	bool bDisableMaxFPS = true;
-
-	/** If true, VSync will be disabled; Otherwise, it will not. Disabling VSync can speed up the game simulation. */
-	UPROPERTY(EditAnywhere, Category = "LearningAgents")
-	bool bDisableVSync = true;
-
-	/** If true, the viewport rendering will be unlit; Otherwise, it will not. Disabling lighting can speed up the game simulation. */
-	UPROPERTY(EditAnywhere, Category = "LearningAgents")
-	bool bUseUnlitViewportRendering = false;
-
-#if WITH_EDITORONLY_DATA
-
-	/** If true, the Use Less CPU In The Background editor setting will be disabled. This prevents the editor from running slowly when minimized. */
-	UPROPERTY(EditAnywhere, Category = "LearningAgents")
-	bool bDisableUseLessCPUInTheBackground = true;
-
-	/** If true, Editor VSync will be disabled; Otherwise, it will not. Disabling Editor VSync can speed up the game simulation. */
-	UPROPERTY(EditAnywhere, Category = "LearningAgents")
-	bool bDisableEditorVSync = true;
-
-#endif
-};
-
 /** Enumeration of the training devices. */
 UENUM(BlueprintType, Category = "LearningAgents")
-enum class ELearningAgentsTrainerDevice : uint8
+enum class ELearningAgentsTrainingDevice : uint8
 {
 	CPU,
 	GPU,
@@ -128,195 +17,21 @@ enum class ELearningAgentsTrainerDevice : uint8
 namespace UE::Learning::Agents
 {
 	/** Get the learning agents trainer device from the UE::Learning trainer device. */
-	LEARNINGAGENTSTRAINING_API ELearningAgentsTrainerDevice GetLearningAgentsTrainerDevice(const ETrainerDevice Device);
+	LEARNINGAGENTSTRAINING_API ELearningAgentsTrainingDevice GetLearningAgentsTrainingDevice(const ETrainerDevice Device);
 
 	/** Get the UE::Learning trainer device from the learning agents trainer device. */
-	LEARNINGAGENTSTRAINING_API ETrainerDevice GetTrainerDevice(const ELearningAgentsTrainerDevice Device);
+	LEARNINGAGENTSTRAINING_API ETrainerDevice GetTrainingDevice(const ELearningAgentsTrainingDevice Device);
 }
 
-/** The configurable settings for the training process. */
-USTRUCT(BlueprintType, Category = "LearningAgents")
-struct LEARNINGAGENTSTRAINING_API FLearningAgentsTrainerTrainingSettings
-{
-	GENERATED_BODY()
-
-public:
-
-	/** The number of iterations to run before ending training. */
-	UPROPERTY(EditAnywhere, Category = "LearningAgents", meta = (ClampMin = "1", UIMin = "1"))
-	int32 NumberOfIterations = 1000000;
-
-	/** Learning rate of the policy network. Typical values are between 0.001 and 0.0001. */
-	UPROPERTY(EditAnywhere, Category = "LearningAgents", meta = (ClampMin = "0.0", UIMin = "0.0", UIMax = "1.0"))
-	float LearningRatePolicy = 0.0001f;
-
-	/**
-	 * Learning rate of the critic network. To avoid instability generally the critic should have a larger learning 
-	 * rate than the policy. Typically this can be set to 10x the rate of the policy.
-	 */
-	UPROPERTY(EditAnywhere, Category = "LearningAgents", meta = (ClampMin = "0.0", UIMin = "0.0", UIMax = "1.0"))
-	float LearningRateCritic = 0.001f;
-
-	/** Amount by which to multiply the learning rate every 1000 iterations. */
-	UPROPERTY(EditAnywhere, Category = "LearningAgents", meta = (ClampMin = "0.0", ClampMax = "1.0", UIMin = "0.0", UIMax = "1.0"))
-	float LearningRateDecay = 1.0f;
-
-	/**
-	 * Amount of weight decay to apply to the network. Larger values encourage network weights to be smaller but too 
-	 * large a value can cause the network weights to collapse to all zeros.
-	 */
-	UPROPERTY(EditAnywhere, Category = "LearningAgents", meta = (ClampMin = "0.0", UIMin = "0.0", UIMax = "1.0"))
-	float WeightDecay = 0.0001f;
-
-	/**
-	 * Batch size to use for training the policy. Large batch sizes are much more computationally efficient when training on the GPU.
-	 */
-	UPROPERTY(EditAnywhere, Category = "LearningAgents", meta = (ClampMin = "1", UIMin = "1", UIMax = "4096"))
-	int32 PolicyBatchSize = 1024;
-
-	/**
-	 * Batch size to use for training the critic. Large batch sizes are much more computationally efficient when training on the GPU.
-	 */
-	UPROPERTY(EditAnywhere, Category = "LearningAgents", meta = (ClampMin = "1", UIMin = "1", UIMax = "4096"))
-	int32 CriticBatchSize = 4096;
-
-	/**
-	 * The number of consecutive steps of observations and actions over which to train the policy. Increasing this value 
-	 * will encourage the policy to use its memory effectively. Too large and training can become slow and unstable.
-	 */
-	UPROPERTY(EditAnywhere, Category = "LearningAgents", meta = (ClampMin = "1", UIMin = "1", UIMax = "128"))
-	int32 PolicyWindowSize = 16;
-
-	/**
-	 * Number of training iterations to perform per buffer of experience gathered. This should be large enough for
-	 * the critic and policy to be effectively updated, but too large and it will simply slow down training.
-	 */
-	UPROPERTY(EditAnywhere, Category = "LearningAgents", meta = (ClampMin = "1", UIMin = "1", UIMax = "1024"))
-	int32 IterationsPerGather = 32;
-
-	/**
-	 * Number of iterations of training to perform to warm - up the Critic. This helps speed up and stabilize training
-	 * at the beginning when the Critic may be producing predictions at the wrong order of magnitude.
-	 */
-	UPROPERTY(EditAnywhere, Category = "LearningAgents", meta = (ClampMin = "1", UIMin = "1", UIMax = "128"))
-	int32 CriticWarmupIterations = 8;
-
-	/**
-	 * Clipping ratio to apply to policy updates. Keeps the training "on-policy". Larger values may speed up training at 
-	 * the cost of stability. Conversely, too small values will keep the policy from being able to learn an 
-	 * optimal policy.
-	 */
-	UPROPERTY(EditAnywhere, Category = "LearningAgents", meta = (ClampMin = "0.0", ClampMax = "1.0", UIMin = "0.0", UIMax = "1.0"))
-	float EpsilonClip = 0.2f;
-
-	/**
-	 * Weight used to regularize returns. Encourages the critic not to over or under estimate returns.
-	 */
-	UPROPERTY(EditAnywhere, Category = "LearningAgents", meta = (ClampMin = "0.0", UIMin = "0.0", UIMax = "1.0"))
-	float ReturnRegularizationWeight = 0.0001f;
-
-	/**
-	 * Weight used to regularize actions. Larger values will encourage exploration and smaller actions, but too large will cause 
-	 * noisy actions centered around zero.
-	 */
-	UPROPERTY(EditAnywhere, Category = "LearningAgents", meta = (ClampMin = "0.0", UIMin = "0.0", UIMax = "1.0"))
-	float ActionRegularizationWeight = 0.001f;
-
-	/**
-	 * Weighting used for the entropy bonus. Larger values encourage larger action noise and therefore greater 
-	 * exploration but can make actions very noisy.
-	 */
-	UPROPERTY(EditAnywhere, Category = "LearningAgents", meta = (ClampMin = "0.0", UIMin = "0.0", UIMax = "1.0"))
-	float ActionEntropyWeight = 0.0f;
-
-	/**
-	 * This is used in the Generalized Advantage Estimation, where larger values will tend to assign more credit to recent actions. Typical
-	 * values should be between 0.9 and 1.0.
-	 */
-	UPROPERTY(EditAnywhere, Category = "LearningAgents", meta = (ClampMin = "0.0", ClampMax = "1.0", UIMin = "0.0", UIMax = "1.0"))
-	float GaeLambda = 0.95f;
-
-	/** When true, advantages are normalized. This tends to makes training more robust to adjustments of the scale of rewards. */
-	UPROPERTY(EditAnywhere, Category = "LearningAgents")
-	bool bAdvantageNormalization = true;
-
-	/**
-	 * The minimum advantage to allow. Setting this below zero will encourage the policy to move away from bad actions, 
-	 * but can introduce instability.
-	 */
-	UPROPERTY(EditAnywhere, Category = "LearningAgents", meta = (UIMin = "-10.0", UIMax = "0.0"))
-	float MinimumAdvantage = 0.0f;
-
-	/**
-	 * The maximum advantage to allow. Making this smaller may increase training stability
-	 * at the cost of some training speed.
-	 */
-	UPROPERTY(EditAnywhere, Category = "LearningAgents", meta = (UIMin = "0.0", UIMax = "10.0"))
-	float MaximumAdvantage = 10.0f;
-
-	/**
-	 * When true, gradient norm max clipping will be used on the policy, critic, encoder, and decoder. Set this as True if
-	 * training is unstable (and adjust GradNormMax) or leave as False if unused.
-	 */
-	UPROPERTY(EditAnywhere, Category = "LearningAgents")
-	bool bUseGradNormMaxClipping = false;
-
-	/**
-	 * The maximum gradient norm to clip updates to. Only used when bUseGradNormMaxClipping is set to true.
-	 * 
-	 * This needs to be carefully chosen based on the size of your gradients during training. Setting too low can make it
-	 * difficult to learn an optimal policy, and too high will have no impact.
-	 */
-	UPROPERTY(EditAnywhere, Category = "LearningAgents", meta = (UIMin = "0.0", UIMax = "10.0"))
-	float GradNormMax = 0.5f;
-
-	/**
-	 * The number of steps to trim from the start of the episode, e.g. can be useful if some things are still getting
-	 * setup at the start of the episode and you don't want them used for training.
-	 */
-	UPROPERTY(EditAnywhere, Category = "LearningAgents", meta = (ClampMin = "0", UIMin = "0"))
-	int32 NumberOfStepsToTrimAtStartOfEpisode = 0;
-
-	/**
-	 * The number of steps to trim from the end of the episode. Can be useful if the end of the episode contains
-	 * irrelevant data.
-	 */
-	UPROPERTY(EditAnywhere, Category = "LearningAgents", meta = (ClampMin = "0", UIMin = "0"))
-	int32 NumberOfStepsToTrimAtEndOfEpisode = 0;
-
-	/** The seed used for any random sampling the trainer will perform, e.g. for weight initialization. */
-	UPROPERTY(EditAnywhere, Category = "LearningAgents", meta = (ClampMin = "0", UIMin = "0"))
-	int32 RandomSeed = 1234;
-
-	/**
-	 * The discount factor to use during training. This affects how much the agent cares about future rewards vs
-	 * near-term rewards. Should typically be a value less than but near 1.0. 
-	 */
-	UPROPERTY(EditAnywhere, Category = "LearningAgents", meta = (ClampMin = "0.0", ClampMax = "1.0", UIMin = "0.0", UIMax = "1.0"))
-	float DiscountFactor = 0.99f;
-
-	/** The device to train on. */
-	UPROPERTY(EditAnywhere, Category = "LearningAgents")
-	ELearningAgentsTrainerDevice Device = ELearningAgentsTrainerDevice::GPU;
-
-	/** If true, TensorBoard logs will be emitted to the intermediate directory. */
-	UPROPERTY(EditAnywhere, Category = "LearningAgents")
-	bool bUseTensorboard = false;
-
-	/** If true, snapshots of the trained networks will be emitted to the intermediate directory. */
-	UPROPERTY(EditAnywhere, Category = "LearningAgents")
-	bool bSaveSnapshots = false;
-};
-
 /** The path settings for the trainer. */
 USTRUCT(BlueprintType, Category = "LearningAgents")
-struct LEARNINGAGENTSTRAINING_API FLearningAgentsTrainerPathSettings
+struct LEARNINGAGENTSTRAINING_API FLearningAgentsTrainerProcessSettings
 {
 	GENERATED_BODY()
 
 public:
 
-	FLearningAgentsTrainerPathSettings();
+	FLearningAgentsTrainerProcessSettings();
 
 	/** The relative path to the engine for editor builds. Defaults to FPaths::EngineDir. */
 	UPROPERTY(EditAnywhere, Category = "LearningAgents", meta = (RelativePath))
@@ -324,7 +39,7 @@ public:
 
 	/**
 	 * The relative path to the editor engine folder for non-editor builds.
-	 * 
+	 *
 	 * If we want to run training in cooked, non-editor builds, then by default we wont have access to python and the
 	 * LearningAgents training scripts - these are editor-only things and are stripped during the cooking process.
 	 *
@@ -332,12 +47,12 @@ public:
 	 * and sound while we are training to make experience gathering as fast as possible - and for any non-trivial game
 	 * is simply may not be realistic to run it for a long time in play-in-editor.
 	 *
-	 * For this reason even in non-editor builds we let you provide the path where all of these editor-only things can 
-	 * be found. This allows you to run training when these things actually exist somewhere accessible to the executable, 
+	 * For this reason even in non-editor builds we let you provide the path where all of these editor-only things can
+	 * be found. This allows you to run training when these things actually exist somewhere accessible to the executable,
 	 * which will usually be the case on a normal development machine or cloud machine if it is set up that way.
 	 *
-	 * Since non-editor builds can be produced in a number of different ways, this is not set by default and cannot 
-	 * use a directory picker since it is relative to the final location of where your cooked, non-editor executable 
+	 * Since non-editor builds can be produced in a number of different ways, this is not set by default and cannot
+	 * use a directory picker since it is relative to the final location of where your cooked, non-editor executable
 	 * will exist rather than the current with-editor executable.
 	 */
 	UPROPERTY(EditAnywhere, Category = "LearningAgents")
@@ -359,351 +74,3 @@ public:
 	/** Gets the Relative Intermediate Path  */
 	FString GetIntermediatePath() const;
 };
-
-/**
- * The ULearningAgentsTrainer is the core class for reinforcement learning training. It defines how agents rewards, completions, and resets are 
- * implemented and provides methods for orchestrating the training process.
- *
- * To use this class, you need to implement the GatherAgentRewards and GatherAgentCompletions functions, which will define the rewards and penalties 
- * the agent receives and what conditions cause an episode to end.
- *
- * @see ULearningAgentsInteractor to understand how observations and actions work.
- */
-UCLASS(Abstract, HideDropdown, BlueprintType, Blueprintable)
-class LEARNINGAGENTSTRAINING_API ULearningAgentsTrainer : public ULearningAgentsManagerListener
-{
-	GENERATED_BODY()
-
-// ----- Setup -----
-public:
-
-	// These constructors/destructors are needed to make forward declarations happy
-	ULearningAgentsTrainer();
-	ULearningAgentsTrainer(FVTableHelper& Helper);
-	virtual ~ULearningAgentsTrainer();
-
-	/** Will automatically call EndTraining if training is still in-progress when the object is destroyed. */
-	virtual void BeginDestroy() override;
-
-	/**
-	 * Constructs the trainer and runs the setup functions for rewards and completions.
-	 * 
-	 * @param InManager			The agent manager we are using.
-	 * @param InInteractor		The agent interactor we are training with.
-	 * @param InPolicy			The policy to be trained.
-	 * @param InCritic			The critic to be trained.
-	 * @param Class				The trainer class
-	 * @param Name				The trainer name
-	 * @param TrainerSettings	The trainer settings to use.
-	 */
-	UFUNCTION(BlueprintCallable, Category = "LearningAgents", meta = (DeterminesOutputType = "Class", AutoCreateRefTerm = "TrainerSettings"))
-	static ULearningAgentsTrainer* MakeTrainer(
-		ULearningAgentsManager* InManager,
-		ULearningAgentsInteractor* InInteractor,
-		ULearningAgentsPolicy* InPolicy,
-		ULearningAgentsCritic* InCritic,
-		TSubclassOf<ULearningAgentsTrainer> Class,
-		const FName Name = TEXT("Trainer"),
-		const FLearningAgentsTrainerSettings& TrainerSettings = FLearningAgentsTrainerSettings());
-
-	/**
-	 * Initializes the trainer and runs the setup functions for rewards and completions.
-	 * 
-	 * @param InManager The agent manager we are using.
-	 * @param InInteractor The agent interactor we are training with.
-	 * @param InPolicy The policy to be trained.
-	 * @param InCritic The critic to be trained.
-	 * @param TrainerSettings The trainer settings to use.
-	 */
-	UFUNCTION(BlueprintCallable, Category = "LearningAgents", meta = (AutoCreateRefTerm = "TrainerSettings"))
-	void SetupTrainer(
-		ULearningAgentsManager* InManager,
-		ULearningAgentsInteractor* InInteractor,
-		ULearningAgentsPolicy* InPolicy,
-		ULearningAgentsCritic* InCritic,
-		const FLearningAgentsTrainerSettings& TrainerSettings = FLearningAgentsTrainerSettings());
-
-public: 
-
-	//~ Begin ULearningAgentsManagerListener Interface
-	virtual void OnAgentsAdded_Implementation(const TArray<int32>& AgentIds) override;
-	virtual void OnAgentsRemoved_Implementation(const TArray<int32>& AgentIds) override;
-	virtual void OnAgentsReset_Implementation(const TArray<int32>& AgentIds) override;
-	virtual void OnAgentsManagerTick_Implementation(const TArray<int32>& AgentIds, const float DeltaTime) override;
-	//~ End ULearningAgentsManagerListener Interface
-
-// ----- Rewards -----
-public:
-
-	/**
-	 * This callback should be overridden by the Trainer and gathers the reward value for the given agent.
-	 *
-	 * @param OutReward			Output reward for the given agent.
-	 * @param AgentId			Agent id to gather reward for.
-	 */
-	UFUNCTION(BlueprintNativeEvent, Category = "LearningAgents", Meta = (ForceAsFunction))
-	void GatherAgentReward(float& OutReward, const int32 AgentId);
-
-
-	/**
-	 * This callback can be overridden by the Trainer and gathers all the reward values for the given set of agents. By default this will call 
-	 * GatherAgentReward on each agent.
-	 *
-	 * @param OutRewards		Output rewards for each agent in AgentIds
-	 * @param AgentIds			Agents to gather rewards for.
-	 */
-	UFUNCTION(BlueprintNativeEvent, Category = "LearningAgents", Meta = (ForceAsFunction))
-	void GatherAgentRewards(TArray<float>& OutRewards, const TArray<int32>& AgentIds);
-
-// ----- Completions ----- 
-public:
-
-	/**
-	 * This callback should be overridden by the Trainer and gathers the completion for a given agent.
-	 *
-	 * @param OutCompletion		Output completion for the given agent.
-	 * @param AgentId			Agent id to gather completion for.
-	 */
-	UFUNCTION(BlueprintNativeEvent, Category = "LearningAgents", Meta = (ForceAsFunction))
-	void GatherAgentCompletion(ELearningAgentsCompletion& OutCompletion, const int32 AgentId);
-
-	/**
-	 * This callback can be overridden by the Trainer and gathers all the completions for the given set of agents. By default this will call 
-	 * GatherAgentCompletion on each agent.
-	 *
-	 * @param OutCompletions	Output completions for each agent in AgentIds
-	 * @param AgentIds			Agents to gather completions for.
-	 */
-	UFUNCTION(BlueprintNativeEvent, Category = "LearningAgents", Meta = (ForceAsFunction))
-	void GatherAgentCompletions(TArray<ELearningAgentsCompletion>& OutCompletions, const TArray<int32>& AgentIds);
-
-// ----- Resets ----- 
-public:
-
-	/**
-	 * This callback should be overridden by the Trainer and resets the episode for the given agent.
-	 *
-	 * @param AgentId			The id of the agent that need resetting.
-	 */
-	UFUNCTION(BlueprintNativeEvent, Category = "LearningAgents", Meta = (ForceAsFunction))
-	void ResetAgentEpisode(const int32 AgentId);
-
-	/**
-	 * This callback can be overridden by the Trainer and resets all episodes for each agent in the given set. By default this will call 
-	 * ResetAgentEpisode on each agent.
-	 * 
-	 * @param AgentIds			The ids of the agents that need resetting.
-	 */
-	UFUNCTION(BlueprintNativeEvent, Category = "LearningAgents", Meta = (ForceAsFunction))
-	void ResetAgentEpisodes(const TArray<int32>& AgentIds);
-
-// ----- Training Process -----
-public:
-
-	/** Returns true if the trainer is currently training; Otherwise, false. */
-	UFUNCTION(BlueprintPure, Category = "LearningAgents")
-	const bool IsTraining() const;
-
-	/**
-	 * Begins the training process with the provided settings.
-	 * 
-	 * @param TrainerTrainingSettings	The settings for this training run.
-	 * @param TrainerGameSettings		The settings that will affect the game's simulation.
-	 * @param TrainerPathSettings		The path settings used by the trainer.
-	 * @param bResetAgentsOnBegin		If true, reset all agents at the beginning of training.
-	 */
-	UFUNCTION(BlueprintCallable, Category = "LearningAgents", meta = (AutoCreateRefTerm = "TrainerTrainingSettings,TrainerGameSettings,TrainerPathSettings"))
-	void BeginTraining(
-		const FLearningAgentsTrainerTrainingSettings& TrainerTrainingSettings = FLearningAgentsTrainerTrainingSettings(),
-		const FLearningAgentsTrainerGameSettings& TrainerGameSettings = FLearningAgentsTrainerGameSettings(),
-		const FLearningAgentsTrainerPathSettings& TrainerPathSettings = FLearningAgentsTrainerPathSettings(),
-		const bool bResetAgentsOnBegin = true);
-
-	/** Stops the training process. */
-	UFUNCTION(BlueprintCallable, Category = "LearningAgents")
-	void EndTraining();
-
-	/**
-	 * Call this function when it is time to evaluate the rewards for your agents. This should be done at the beginning
-	 * of each iteration of your training loop after the initial step, i.e. after taking an action, you want to get into
-	 * the next state before evaluating the rewards.
-	 */
-	UFUNCTION(BlueprintCallable, Category = "LearningAgents")
-	void GatherRewards();
-
-	/**
-	 * Call this function when it is time to evaluate the completions for your agents. This should be done at the beginning
-	 * of each iteration of your training loop after the initial step, i.e. after taking an action, you want to get into
-	 * the next state before evaluating the completions.
-	 */
-	UFUNCTION(BlueprintCallable, Category = "LearningAgents")
-	void GatherCompletions();
-
-	/**
-	 * Call this function at the end of each step of your training loop. This takes the current observations/actions/
-	 * rewards and moves them into the episode experience buffer. All agents with full episode buffers or those which
-	 * have been signaled complete will be reset. If enough experience is gathered, it will be sent to the training 
-	 * process and an iteration of training will be run and the updated policy will be synced back.
-	 *
-	 * @param bResetAgentsOnUpdate				If true, reset all agents whenever an updated policy is received.
-	 */
-	UFUNCTION(BlueprintCallable, Category = "LearningAgents")
-	void ProcessExperience(const bool bResetAgentsOnUpdate = true);
-
-	/**
-	 * Convenience function that runs a basic training loop. If training has not been started, it will start it, and 
-	 * then call RunInference. On each following call to this function, it will call GatherRewards, 
-	 * GatherCompletions, and ProcessExperience, followed by RunInference.
-	 * 
-	 * @param TrainerTrainingSettings			The settings for this training run.
-	 * @param TrainerGameSettings				The settings that will affect the game's simulation.
-	 * @param TrainerPathSettings				The path settings used by the trainer.
-	 * @param bResetAgentsOnBegin				If true, reset all agents at the beginning of training.
-	 * @param bResetAgentsOnUpdate				If true, reset all agents whenever an updated policy is received.
-	 */
-	UFUNCTION(BlueprintCallable, Category = "LearningAgents", meta = (AutoCreateRefTerm = "TrainerTrainingSettings,TrainerGameSettings,TrainerPathSettings"))
-	void RunTraining(
-		const FLearningAgentsTrainerTrainingSettings& TrainerTrainingSettings = FLearningAgentsTrainerTrainingSettings(),
-		const FLearningAgentsTrainerGameSettings& TrainerGameSettings = FLearningAgentsTrainerGameSettings(),
-		const FLearningAgentsTrainerPathSettings& TrainerPathSettings = FLearningAgentsTrainerPathSettings(),
-		const bool bResetAgentsOnBegin = true,
-		const bool bResetAgentsOnUpdate = true);
-
-	/**
-	 * Returns true if GatherRewards has been called and the reward already set for the given agent.
-	 */
-	UFUNCTION(BlueprintPure, Category = "LearningAgents", meta = (AgentId = "-1"))
-	bool HasReward(const int32 AgentId) const;
-
-	/**
-	 * Returns true if GatherCompletions has been called and the completion already set for the given agent.
-	 */
-	UFUNCTION(BlueprintPure, Category = "LearningAgents", meta = (AgentId = "-1"))
-	bool HasCompletion(const int32 AgentId) const;
-
-	/**
-	 * Gets the current reward for an agent. Should be called only after GatherRewards.
-	 *
-	 * @param AgentId	The AgentId to look-up the reward for
-	 * @returns			The reward
-	 */
-	UFUNCTION(BlueprintPure, Category = "LearningAgents", meta = (AgentId = "-1"))
-	float GetReward(const int32 AgentId) const;
-
-	/**
-	 * Gets the current completion for an agent. Should be called only after GatherCompletions.
-	 *
-	 * @param AgentId	The AgentId to look-up the completion for
-	 * @returns			The completion type
-	 */
-	UFUNCTION(BlueprintPure, Category = "LearningAgents", meta = (AgentId = "-1"))
-	ELearningAgentsCompletion GetCompletion(const int32 AgentId) const;
-
-	/**
-	 * Gets the current elapsed episode time for the given agent.
-	 *
-	 * @param AgentId	The AgentId to look-up the episode time for
-	 * @returns			The elapsed episode time
-	 */
-	UFUNCTION(BlueprintPure, Category = "LearningAgents", meta = (AgentId = "-1"))
-	float GetEpisodeTime(const int32 AgentId) const;
-
-	/**
-	 * Gets the number of step recorded in an episode for the given agent.
-	 *
-	 * @param AgentId	The AgentId to look-up the number of recorded episode steps for
-	 * @returns			The number of recorded episode steps
-	 */
-	UFUNCTION(BlueprintPure, Category = "LearningAgents", meta = (AgentId = "-1"))
-	int32 GetEpisodeStepNum(const int32 AgentId) const;
-
-	/**
-	 * Returns true if the trainer has failed to communicate with the external training process. This can be used in
-	 * combination with RunTraining to avoid filling the logs with errors.
-	 *
-	 * @returns				True if the training has failed. Otherwise, false.
-	 */
-	UFUNCTION(BlueprintPure, Category = "LearningAgents")
-	bool HasTrainingFailed() const;
-
-// ----- Private Data ----- 
-private:
-
-	/** The agent interactor associated with this component. */
-	UPROPERTY(VisibleAnywhere, Transient, Category = "LearningAgents")
-	TObjectPtr<ULearningAgentsInteractor> Interactor;
-
-	/** The current policy for experience gathering. */
-	UPROPERTY(VisibleAnywhere, Transient, Category = "LearningAgents")
-	TObjectPtr<ULearningAgentsPolicy> Policy;
-
-	/** The current critic. */
-	UPROPERTY(VisibleAnywhere, Transient, Category = "LearningAgents")
-	TObjectPtr<ULearningAgentsCritic> Critic;
-
-	/** True if training is currently in-progress. Otherwise, false. */
-	UPROPERTY(VisibleAnywhere, Transient, Category = "LearningAgents")
-	bool bIsTraining = false;
-
-	/**
-	 * True if trainer encountered an unrecoverable error during training (e.g. the trainer process timed out). Otherwise, false.
-	 * This exists mainly to keep the editor from locking up if something goes wrong during training.
-	 */
-	UPROPERTY(VisibleAnywhere, Transient, Category = "LearningAgents")
-	bool bHasTrainingFailed = false;
-
-	/** Callback Reward Output */
-	TArray<float> RewardBuffer;
-
-	/** Callback Completion Output */
-	TArray<ELearningAgentsCompletion> CompletionBuffer;
-
-	/** Reward Buffer */
-	TLearningArray<1, float> Rewards;
-	
-	/** Agent Completions Buffer */
-	TLearningArray<1, UE::Learning::ECompletionMode> AgentCompletions;
-
-	/** Episode Completions Buffer */
-	TLearningArray<1, UE::Learning::ECompletionMode> EpisodeCompletions;
-
-	/** All Completions Buffer */
-	TLearningArray<1, UE::Learning::ECompletionMode> AllCompletions;
-
-	/** Agent episode times */
-	TLearningArray<1, float> EpisodeTimes;
-
-	TUniquePtr<UE::Learning::FEpisodeBuffer> EpisodeBuffer;
-	TUniquePtr<UE::Learning::FReplayBuffer> ReplayBuffer;
-	TUniquePtr<UE::Learning::FResetInstanceBuffer> ResetBuffer;
-	TUniquePtr<UE::Learning::FSharedMemoryPPOTrainer> Trainer;
-	
-	float TrainerTimeout = 10.0f;
-
-	void DoneTraining();
-
-	/** Number of times rewards have been evaluated for all agents */
-	TLearningArray<1, uint64, TInlineAllocator<32>> RewardIteration;
-
-	/** Number of times completions have been evaluated for all agents */
-	TLearningArray<1, uint64, TInlineAllocator<32>> CompletionIteration;
-
-	/** Temp buffers used to record the set of agents that are valid for training */
-	TArray<int32> ValidAgentIds;
-	TArray<int32> FinalValidAgentIds;
-	UE::Learning::FIndexSet ValidAgentSet;
-	UE::Learning::FIndexSet FinalValidAgentSet;
-
-// ----- Private Recording of GameSettings ----- 
-private:
-
-	bool bFixedTimestepUsed = false;
-	float FixedTimeStepDeltaTime = -1.0f;
-	float MaxPhysicsStep = -1.0f;
-	int32 MaxFPS = 120;
-	bool bVSyncEnabled = true;
-	int32 ViewModeIndex = -1;
-	
-	bool bUseLessCPUInTheBackground = true;
-	bool bEditorVSyncEnabled = true;
-};
diff --git a/Engine/Plugins/Experimental/LearningAgents/Source/LearningTraining/Private/LearningExternalTrainer.cpp b/Engine/Plugins/Experimental/LearningAgents/Source/LearningTraining/Private/LearningExternalTrainer.cpp
new file mode 100644
--- /dev/null
+++ b/Engine/Plugins/Experimental/LearningAgents/Source/LearningTraining/Private/LearningExternalTrainer.cpp
@@ -0,0 +1,620 @@
+// Copyright Epic Games, Inc. All Rights Reserved.
+
+#include "LearningExternalTrainer.h"
+
+#include "LearningExperience.h"
+#include "LearningNeuralNetwork.h"
+#include "LearningSharedMemoryTraining.h"
+#include "LearningSocketTraining.h"
+
+#include "Misc/CommandLine.h"
+#include "Misc/FileHelper.h"
+#include "Misc/Guid.h"
+#include "Misc/Paths.h"
+
+#include "Dom/JsonObject.h"
+#include "Serialization/JsonSerializer.h"
+#include "Serialization/JsonWriter.h"
+
+#include "Sockets.h"
+#include "Common/TcpSocketBuilder.h"
+#include "SocketSubsystem.h"
+
+namespace UE::Learning
+{
+	FSharedMemoryTrainerServerProcess::FSharedMemoryTrainerServerProcess(
+		const FString& TaskName,
+		const FString& PythonExecutablePath,
+		const FString& PythonContentPath,
+		const FString& InIntermediatePath,
+		const int32 ProcessNum,
+		const float InTimeout,
+		const ESubprocessFlags TrainingProcessFlags)
+	{
+		check(ProcessNum > 0);
+
+		int32 ProcessIdx = 0;
+		FParse::Value(FCommandLine::Get(), TEXT("LearningProcessIdx"), ProcessIdx);
+		
+		Timeout = InTimeout;
+		IntermediatePath = InIntermediatePath;
+
+		if (ProcessIdx == 0)
+		{
+			// Allocate the control memory if we are the parent UE process
+			Controls = SharedMemory::Allocate<2, volatile int32>({ ProcessNum, SharedMemoryTraining::GetControlNum() });
+		}
+		else
+		{
+			FGuid ControlsGuid; ensure(FParse::Value(FCommandLine::Get(), TEXT("LearningControlsGuid"), ControlsGuid));
+			Controls = SharedMemory::Map<2, volatile int32>(ControlsGuid, { ProcessNum, SharedMemoryTraining::GetControlNum() });
+
+			// We do not want to launch another training process if we are a child process
+			return;
+		}
+
+		UE_LEARNING_CHECK(FPaths::FileExists(PythonExecutablePath));
+		UE_LEARNING_CHECK(FPaths::DirectoryExists(PythonContentPath));
+
+		// We need to zero the control memory before we start the training sub-process since it may contain
+		// uninitialized values or those left over from previous runs.
+		Array::Zero(Controls.View);
+
+		const FString TimeStamp = FDateTime::Now().ToFormattedString(TEXT("%Y-%m-%d_%H-%M-%S"));
+		const FString TrainerMethod = TEXT("PPO");
+		const FString TrainerType = TEXT("SharedMemory");
+		ConfigPath = InIntermediatePath / TEXT("Configs") / FString::Printf(TEXT("%s_%s_%s_%s.json"), *TaskName, *TrainerMethod, *TrainerType, *TimeStamp);
+
+		IFileManager& FileManager = IFileManager::Get();
+		const FString CommandLineArguments = FString::Printf(TEXT("\"%s\" SharedMemory \"%s\" %i \"%s\""),
+			*FileManager.ConvertToAbsolutePathForExternalAppForRead(*(PythonContentPath / TEXT("train_ppo.py"))),
+			*Controls.Guid.ToString(EGuidFormats::DigitsWithHyphensInBraces),
+			ProcessNum,
+			*FileManager.ConvertToAbsolutePathForExternalAppForRead(*ConfigPath));
+
+		TrainingProcess.Launch(
+			FileManager.ConvertToAbsolutePathForExternalAppForRead(*PythonExecutablePath),
+			CommandLineArguments, 
+			TrainingProcessFlags);
+	}
+
+	FSharedMemoryTrainerServerProcess::~FSharedMemoryTrainerServerProcess()
+	{
+		Terminate();
+	}
+
+	bool FSharedMemoryTrainerServerProcess::IsRunning() const
+	{
+		return TrainingProcess.IsRunning();
+	}
+
+	bool FSharedMemoryTrainerServerProcess::Wait()
+	{
+		const float SleepTime = 0.001f;
+		float WaitTime = 0.0f;
+
+		while (TrainingProcess.Update())
+		{
+			FPlatformProcess::Sleep(SleepTime);
+			WaitTime += SleepTime;
+
+			if (WaitTime > Timeout)
+			{
+				return false;
+			}
+		}
+
+		return true;
+	}
+
+	void FSharedMemoryTrainerServerProcess::Terminate()
+	{
+		TrainingProcess.Terminate();
+	}
+
+	TSharedMemoryArrayView<2, volatile int32> FSharedMemoryTrainerServerProcess::GetControlsSharedMemoryArrayView() const
+	{
+		return Controls;
+	}
+
+	const FString& FSharedMemoryTrainerServerProcess::GetIntermediatePath() const
+	{
+		return IntermediatePath;
+	}
+
+	const FString& FSharedMemoryTrainerServerProcess::GetConfigPath() const
+	{
+		return ConfigPath;
+	}
+
+	FSubprocess* FSharedMemoryTrainerServerProcess::GetTrainingSubprocess()
+	{
+		return &TrainingProcess;
+	}
+
+	void FSharedMemoryTrainerServerProcess::Deallocate()
+	{
+		if (Controls.Region != nullptr)
+		{
+			SharedMemory::Deallocate(Controls);
+		}
+	}
+
+	void FSharedMemoryTrainer::FSharedMemoryExperienceContainer::Deallocate()
+	{
+		if (EpisodeStarts.Region != nullptr)
+		{
+			SharedMemory::Deallocate(EpisodeStarts);
+			SharedMemory::Deallocate(EpisodeLengths);
+			SharedMemory::Deallocate(EpisodeCompletionModes);
+			SharedMemory::Deallocate(EpisodeFinalObservations);
+			SharedMemory::Deallocate(EpisodeFinalMemoryStates);
+			SharedMemory::Deallocate(Observations);
+			SharedMemory::Deallocate(Actions);
+			SharedMemory::Deallocate(MemoryStates);
+			SharedMemory::Deallocate(Rewards);
+		}
+	}
+
+	FSharedMemoryTrainer::FSharedMemoryTrainer(
+		const FString& InTaskName,
+		const int32 InProcessNum,
+		const TSharedPtr<UE::Learning::ITrainerProcess>& ExternalTrainerProcess,
+		const float InTimeout)
+	{
+		FSharedMemoryTrainerServerProcess* SharedMemoryTrainerProcess = (FSharedMemoryTrainerServerProcess*)ExternalTrainerProcess.Get();
+		if (!SharedMemoryTrainerProcess)
+		{
+			UE_LOG(LogLearning, Error, TEXT("FSharedMemoryTrainer ctor: Trainer process is nullptr. Is it not a shared memory process?"));
+			return;
+		}
+
+		check(InProcessNum > 0);
+
+		TaskName = InTaskName;
+		ConfigPath = SharedMemoryTrainerProcess->GetConfigPath();
+		IntermediatePath = SharedMemoryTrainerProcess->GetIntermediatePath();
+		TrainingProcess = SharedMemoryTrainerProcess->GetTrainingSubprocess();
+		ProcessNum = InProcessNum;
+		Controls = SharedMemoryTrainerProcess->GetControlsSharedMemoryArrayView();
+		Timeout = InTimeout;
+
+		ProcessIdx = 0;
+		FParse::Value(FCommandLine::Get(), TEXT("LearningProcessIdx"), ProcessIdx);
+	}
+
+	FSharedMemoryTrainer::~FSharedMemoryTrainer()
+	{
+		Terminate();
+	}
+
+	ETrainerResponse FSharedMemoryTrainer::Wait()
+	{
+		return ETrainerResponse::Success;
+	}
+
+	void FSharedMemoryTrainer::Terminate()
+	{
+		Deallocate();
+	}
+
+	ETrainerResponse FSharedMemoryTrainer::SendStop()
+	{
+		check(ProcessIdx != INDEX_NONE);
+		checkf(Controls.Region, TEXT("SendStop: Controls Shared Memory Region is nullptr"));
+
+		return SharedMemoryTraining::SendStop(Controls.View[ProcessIdx]);
+	}
+
+	ETrainerResponse FSharedMemoryTrainer::SendConfig(const TSharedRef<FJsonObject>& ConfigObject, const ELogSetting LogSettings)
+	{
+		check(ProcessNum > 0);
+
+		if (ProcessIdx != 0)
+		{
+			// Only the parent process will send the config
+			return ETrainerResponse::Success;
+		}
+
+		IFileManager& FileManager = IFileManager::Get();
+		ConfigObject->SetStringField(TEXT("IntermediatePath"), *FileManager.ConvertToAbsolutePathForExternalAppForRead(*IntermediatePath));
+		ConfigObject->SetBoolField(TEXT("LoggingEnabled"), LogSettings == ELogSetting::Silent ? false : true);
+
+		ConfigObject->SetNumberField(TEXT("ProcessNum"), ProcessNum);
+
+		TArray<TSharedPtr<FJsonValue>> NetworkGuidsArray;
+		for (const TPair<FName, TSharedMemoryArrayView<1, uint8>>& MapEntry : NeuralNetworkSharedMemoryArrayViews)
+		{
+			NetworkGuidsArray.Add(MakeShared<FJsonValueString>(*MapEntry.Value.Guid.ToString(EGuidFormats::DigitsWithHyphensInBraces)));
+		}
+		ConfigObject->SetArrayField(TEXT("NetworkGuids"), NetworkGuidsArray);
+
+		TArray<TSharedPtr<FJsonValue>> ExperienceContainerObjectsArray;
+		for (const TPair<FName, FSharedMemoryExperienceContainer>& MapEntry : SharedMemoryExperienceContainers)
+		{
+			TSharedPtr<FJsonObject> JsonObject = MakeShared<FJsonObject>();
+			JsonObject->SetStringField(TEXT("EpisodeStartsGuid"), *MapEntry.Value.EpisodeStarts.Guid.ToString(EGuidFormats::DigitsWithHyphensInBraces));
+			JsonObject->SetStringField(TEXT("EpisodeLengthsGuid"), *MapEntry.Value.EpisodeLengths.Guid.ToString(EGuidFormats::DigitsWithHyphensInBraces));
+			JsonObject->SetStringField(TEXT("EpisodeCompletionModesGuid"), *MapEntry.Value.EpisodeCompletionModes.Guid.ToString(EGuidFormats::DigitsWithHyphensInBraces));
+			JsonObject->SetStringField(TEXT("EpisodeFinalObservationsGuid"), *MapEntry.Value.EpisodeFinalObservations.Guid.ToString(EGuidFormats::DigitsWithHyphensInBraces));
+			JsonObject->SetStringField(TEXT("EpisodeFinalMemoryStatesGuid"), *MapEntry.Value.EpisodeFinalMemoryStates.Guid.ToString(EGuidFormats::DigitsWithHyphensInBraces));
+			JsonObject->SetStringField(TEXT("ObservationsGuid"), *MapEntry.Value.Observations.Guid.ToString(EGuidFormats::DigitsWithHyphensInBraces));
+			JsonObject->SetStringField(TEXT("ActionsGuid"), *MapEntry.Value.Actions.Guid.ToString(EGuidFormats::DigitsWithHyphensInBraces));
+			JsonObject->SetStringField(TEXT("MemoryStatesGuid"), *MapEntry.Value.MemoryStates.Guid.ToString(EGuidFormats::DigitsWithHyphensInBraces));
+			JsonObject->SetStringField(TEXT("RewardsGuid"), *MapEntry.Value.Rewards.Guid.ToString(EGuidFormats::DigitsWithHyphensInBraces));
+			
+			TSharedRef<FJsonValueObject> JsonValue = MakeShared<FJsonValueObject>(JsonObject);
+			ExperienceContainerObjectsArray.Add(JsonValue);
+		}
+		ConfigObject->SetArrayField(TEXT("ExperienceBuffers"), ExperienceContainerObjectsArray);
+		
+		FString ConfigString;
+		TSharedRef<TJsonWriter<>> JsonWriter = TJsonWriterFactory<>::Create(&ConfigString, 0);
+		FJsonSerializer::Serialize(ConfigObject, JsonWriter, true);
+		FFileHelper::SaveStringToFile(ConfigString, *ConfigPath);
+
+		return SharedMemoryTraining::SendConfigSignal(Controls.View[ProcessIdx], LogSettings);
+	}
+
+	void FSharedMemoryTrainer::AddNetwork(
+		const FName& Name,
+		const ULearningNeuralNetworkData& Network)
+	{
+		NeuralNetworkSharedMemoryArrayViews.Add(Name, SharedMemory::Allocate<1, uint8>({ Network.GetSnapshotByteNum() }));
+	}
+
+	bool FSharedMemoryTrainer::ContainsNetwork(const FName& Name) const
+	{
+		return NeuralNetworkSharedMemoryArrayViews.Contains(Name);
+	}
+
+	ETrainerResponse FSharedMemoryTrainer::ReceiveNetwork(
+		const FName& Name,
+		ULearningNeuralNetworkData& OutNetwork,
+		FRWLock* NetworkLock,
+		const ELogSetting LogSettings)
+	{
+		check(ProcessIdx != INDEX_NONE);
+		checkf(Controls.Region, TEXT("ReceiveNetwork: Controls Shared Memory Region is nullptr"));
+		if (!ensureMsgf(NeuralNetworkSharedMemoryArrayViews.Contains(Name), TEXT("Network %s has not been added. Call AddNetwork prior to ReceiveNetwork."), *Name.ToString()))
+		{
+			return ETrainerResponse::Unexpected;
+		}
+
+		return SharedMemoryTraining::RecvNetwork(
+			Controls.View[ProcessIdx],
+			OutNetwork,
+			*TrainingProcess,
+			NeuralNetworkSharedMemoryArrayViews[Name].View,
+			Timeout,
+			NetworkLock,
+			LogSettings);
+	}
+
+	ETrainerResponse FSharedMemoryTrainer::SendNetwork(
+		const FName& Name,
+		const ULearningNeuralNetworkData& Network,
+		FRWLock* NetworkLock,
+		const ELogSetting LogSettings)
+	{
+		check(ProcessIdx != INDEX_NONE);
+		checkf(Controls.Region, TEXT("SendNetwork: Controls Shared Memory Region is nullptr"));
+		if (!ensureMsgf(NeuralNetworkSharedMemoryArrayViews.Contains(Name), TEXT("Network %s has not been added. Call AddNetwork prior to SendNetwork."), *Name.ToString()))
+		{
+			return ETrainerResponse::Unexpected;
+		}
+
+		return SharedMemoryTraining::SendNetwork(
+			Controls.View[ProcessIdx],
+			NeuralNetworkSharedMemoryArrayViews[Name].View,
+			*TrainingProcess,
+			Network,
+			Timeout,
+			NetworkLock,
+			LogSettings);
+	}
+
+	void FSharedMemoryTrainer::AddReplayBuffer(
+		const FName& Name,
+		const FReplayBuffer& ReplayBuffer)
+	{
+		check(ProcessNum > 0);
+
+		const int32 ObservationVectorDimensionNum = ReplayBuffer.GetObservations().Num<1>();
+		const int32 ActionVectorDimensionNum = ReplayBuffer.GetActions().Num<1>();
+		const int32 MemoryStateVectorDimensionNum = ReplayBuffer.GetMemoryStates().Num<1>();
+
+		FSharedMemoryExperienceContainer ExperienceContainer;
+		if (ProcessIdx == 0)
+		{
+			ExperienceContainer.EpisodeStarts = SharedMemory::Allocate<2, int32>({ ProcessNum, ReplayBuffer.GetMaxEpisodeNum() });
+			ExperienceContainer.EpisodeLengths = SharedMemory::Allocate<2, int32>({ ProcessNum, ReplayBuffer.GetMaxEpisodeNum() });
+			ExperienceContainer.EpisodeCompletionModes = SharedMemory::Allocate<2, ECompletionMode>({ ProcessNum, ReplayBuffer.GetMaxEpisodeNum() });
+			ExperienceContainer.EpisodeFinalObservations = SharedMemory::Allocate<3, float>({ ProcessNum, ReplayBuffer.GetMaxEpisodeNum(), ObservationVectorDimensionNum });
+			ExperienceContainer.EpisodeFinalMemoryStates = SharedMemory::Allocate<3, float>({ ProcessNum, ReplayBuffer.GetMaxEpisodeNum(), MemoryStateVectorDimensionNum });
+			ExperienceContainer.Observations = SharedMemory::Allocate<3, float>({ ProcessNum, ReplayBuffer.GetMaxStepNum(), ObservationVectorDimensionNum });
+			ExperienceContainer.Actions = SharedMemory::Allocate<3, float>({ ProcessNum, ReplayBuffer.GetMaxStepNum(), ActionVectorDimensionNum });
+			ExperienceContainer.MemoryStates = SharedMemory::Allocate<3, float>({ ProcessNum, ReplayBuffer.GetMaxStepNum(), MemoryStateVectorDimensionNum });
+			ExperienceContainer.Rewards = SharedMemory::Allocate<2, float>({ ProcessNum, ReplayBuffer.GetMaxStepNum() });
+		}
+		else
+		{
+			FGuid EpisodeStartsGuid; ensure(FParse::Value(FCommandLine::Get(), TEXT("LearningEpisodeStartsGuid"), EpisodeStartsGuid));
+			FGuid EpisodeLengthsGuid; ensure(FParse::Value(FCommandLine::Get(), TEXT("LearningEpisodeLengthsGuid"), EpisodeLengthsGuid));
+			FGuid EpisodeCompletionModesGuid; ensure(FParse::Value(FCommandLine::Get(), TEXT("LearningEpisodeCompletionModesGuid"), EpisodeCompletionModesGuid));
+			FGuid EpisodeFinalObservationsGuid; ensure(FParse::Value(FCommandLine::Get(), TEXT("LearningEpisodeFinalObservationsGuid"), EpisodeFinalObservationsGuid));
+			FGuid EpisodeFinalMemoryStatesGuid; ensure(FParse::Value(FCommandLine::Get(), TEXT("LearningEpisodeFinalMemoryStatesGuid"), EpisodeFinalMemoryStatesGuid));
+			FGuid ObservationsGuid; ensure(FParse::Value(FCommandLine::Get(), TEXT("LearningObservationsGuid"), ObservationsGuid));
+			FGuid ActionsGuid; ensure(FParse::Value(FCommandLine::Get(), TEXT("LearningActionsGuid"), ActionsGuid));
+			FGuid MemoryStatesGuid; ensure(FParse::Value(FCommandLine::Get(), TEXT("LearningMemoryStatesGuid"), MemoryStatesGuid));
+			FGuid RewardsGuid; ensure(FParse::Value(FCommandLine::Get(), TEXT("LearningRewardsGuid"), RewardsGuid));
+
+			ExperienceContainer.EpisodeStarts = SharedMemory::Map<2, int32>(EpisodeStartsGuid, { ProcessNum, ReplayBuffer.GetMaxEpisodeNum() });
+			ExperienceContainer.EpisodeLengths = SharedMemory::Map<2, int32>(EpisodeLengthsGuid, { ProcessNum, ReplayBuffer.GetMaxEpisodeNum() });
+			ExperienceContainer.EpisodeCompletionModes = SharedMemory::Map<2, ECompletionMode>(EpisodeCompletionModesGuid, { ProcessNum, ReplayBuffer.GetMaxEpisodeNum() });
+			ExperienceContainer.EpisodeFinalObservations = SharedMemory::Map<3, float>(EpisodeFinalObservationsGuid, { ProcessNum, ReplayBuffer.GetMaxEpisodeNum(), ObservationVectorDimensionNum });
+			ExperienceContainer.EpisodeFinalMemoryStates = SharedMemory::Map<3, float>(EpisodeFinalMemoryStatesGuid, { ProcessNum, ReplayBuffer.GetMaxEpisodeNum(), MemoryStateVectorDimensionNum });
+			ExperienceContainer.Observations = SharedMemory::Map<3, float>(ObservationsGuid, { ProcessNum, ReplayBuffer.GetMaxStepNum(), ObservationVectorDimensionNum });
+			ExperienceContainer.Actions = SharedMemory::Map<3, float>(ActionsGuid, { ProcessNum, ReplayBuffer.GetMaxStepNum(), ActionVectorDimensionNum });
+			ExperienceContainer.MemoryStates = SharedMemory::Map<3, float>(MemoryStatesGuid, { ProcessNum, ReplayBuffer.GetMaxStepNum(), MemoryStateVectorDimensionNum });
+			ExperienceContainer.Rewards = SharedMemory::Map<2, float>(RewardsGuid, { ProcessNum, ReplayBuffer.GetMaxStepNum() });
+		}
+
+		SharedMemoryExperienceContainers.Add(Name, ExperienceContainer);
+	}
+
+	bool FSharedMemoryTrainer::ContainsReplayBuffer(const FName& Name) const
+	{
+		return SharedMemoryExperienceContainers.Contains(Name);
+	}
+
+	ETrainerResponse FSharedMemoryTrainer::SendReplayBuffer(const FName& Name, const FReplayBuffer& ReplayBuffer, const ELogSetting LogSettings)
+	{
+		check(ProcessIdx != INDEX_NONE);
+		checkf(Controls.Region, TEXT("SendReplayBuffer: Controls Shared Memory Region is nullptr"));
+		if (!ensureMsgf(SharedMemoryExperienceContainers.Contains(Name), TEXT("ReplayBuffer %s has not been added. Call AddReplayBuffer prior to SendReplayBuffer."), *Name.ToString()))
+		{
+			return ETrainerResponse::Unexpected;
+		}
+
+		return SharedMemoryTraining::SendExperience(
+			SharedMemoryExperienceContainers[Name].EpisodeStarts.View[ProcessIdx],
+			SharedMemoryExperienceContainers[Name].EpisodeLengths.View[ProcessIdx],
+			SharedMemoryExperienceContainers[Name].EpisodeCompletionModes.View[ProcessIdx],
+			SharedMemoryExperienceContainers[Name].EpisodeFinalObservations.View[ProcessIdx],
+			SharedMemoryExperienceContainers[Name].EpisodeFinalMemoryStates.View[ProcessIdx],
+			SharedMemoryExperienceContainers[Name].Observations.View[ProcessIdx],
+			SharedMemoryExperienceContainers[Name].Actions.View[ProcessIdx],
+			SharedMemoryExperienceContainers[Name].MemoryStates.View[ProcessIdx],
+			SharedMemoryExperienceContainers[Name].Rewards.View[ProcessIdx],
+			Controls.View[ProcessIdx],
+			*TrainingProcess,
+			ReplayBuffer,
+			Timeout,
+			LogSettings);
+	}
+
+	void FSharedMemoryTrainer::Deallocate()
+	{
+		for (TPair<FName, TSharedMemoryArrayView<1, uint8>>& MapEntry : NeuralNetworkSharedMemoryArrayViews)
+		{
+			if (MapEntry.Value.Region != nullptr)
+			{
+				SharedMemory::Deallocate(MapEntry.Value);
+			}
+		}
+		NeuralNetworkSharedMemoryArrayViews.Empty();
+
+		for (TPair<FName, FSharedMemoryExperienceContainer>& MapEntry : SharedMemoryExperienceContainers)
+		{
+			MapEntry.Value.Deallocate();
+		}
+		SharedMemoryExperienceContainers.Empty();
+	}
+
+	FSocketTrainerServerProcess::FSocketTrainerServerProcess(
+		const FString& PythonExecutablePath,
+		const FString& PythonContentPath,
+		const FString& IntermediatePath,
+		const TCHAR* IpAddress,
+		const uint32 Port,
+		const float InTimeout,
+		const ESubprocessFlags TrainingProcessFlags,
+		const ELogSetting LogSettings)
+	{
+		Timeout = InTimeout;
+
+		UE_LEARNING_CHECK(FPaths::FileExists(PythonExecutablePath));
+		UE_LEARNING_CHECK(FPaths::DirectoryExists(PythonContentPath));
+
+		IFileManager& FileManager = IFileManager::Get();
+		const FString CommandLineArguments = FString::Printf(TEXT("\"%s\" Socket \"%s:%i\" \"%s\" %i"),
+			*FileManager.ConvertToAbsolutePathForExternalAppForRead(*(PythonContentPath / TEXT("train_ppo.py"))),
+			IpAddress,
+			Port,
+			*FileManager.ConvertToAbsolutePathForExternalAppForRead(*IntermediatePath),
+			LogSettings == ELogSetting::Normal ? 1 : 0);
+
+		TrainingProcess.Launch(
+			FileManager.ConvertToAbsolutePathForExternalAppForRead(*PythonExecutablePath),
+			CommandLineArguments, 
+			TrainingProcessFlags);
+	}
+
+	FSocketTrainerServerProcess::~FSocketTrainerServerProcess()
+	{
+		Terminate();
+	}
+
+	bool FSocketTrainerServerProcess::IsRunning() const
+	{
+		return TrainingProcess.IsRunning();
+	}
+
+	bool FSocketTrainerServerProcess::Wait()
+	{
+		const float SleepTime = 0.001f;
+		float WaitTime = 0.0f;
+
+		while (TrainingProcess.Update())
+		{
+			FPlatformProcess::Sleep(SleepTime);
+			WaitTime += SleepTime;
+
+			if (WaitTime > Timeout)
+			{
+				return false;
+			}
+		}
+
+		return true;
+	}
+
+	void FSocketTrainerServerProcess::Terminate()
+	{
+		TrainingProcess.Terminate();
+	}
+
+	FSubprocess* FSocketTrainerServerProcess::GetTrainingSubprocess()
+	{
+		return &TrainingProcess;
+	}
+
+	FSocketTrainer::FSocketTrainer(
+		ETrainerResponse& OutResponse,
+		const TSharedPtr<UE::Learning::ITrainerProcess>& ExternalTrainerProcess,
+		const TCHAR* IpAddress,
+		const uint32 Port,
+		const float InTimeout)
+	{
+		Timeout = InTimeout;
+
+		FSocketTrainerServerProcess* SocketTrainerProcess = (FSocketTrainerServerProcess*)ExternalTrainerProcess.Get();
+		if (SocketTrainerProcess)
+		{
+			TrainingProcess = SocketTrainerProcess->GetTrainingSubprocess();
+		}
+
+		ISocketSubsystem* SocketSubsystem = ISocketSubsystem::Get(PLATFORM_SOCKETSUBSYSTEM);
+		checkf(SocketSubsystem, TEXT("Could not get socket subsystem"));
+
+		bool bIsValid = false;
+		TSharedRef<FInternetAddr> Address = SocketSubsystem->CreateInternetAddr();
+		Address->SetIp(IpAddress, bIsValid);
+		Address->SetPort(Port);
+
+		if (!bIsValid)
+		{
+			UE_LOG(LogLearning, Error, TEXT("Invalid Ip Address \"%s\"..."), IpAddress);
+			OutResponse = ETrainerResponse::Unexpected;
+			return;
+		}
+
+		Socket = FTcpSocketBuilder(TEXT("LearningTrainerSocket")).AsNonBlocking().Build();
+
+		OutResponse = SocketTraining::WaitForConnection(*Socket, TrainingProcess, *Address, Timeout);
+	}
+
+	FSocketTrainer::~FSocketTrainer()
+	{
+		Terminate();
+	}
+
+	ETrainerResponse FSocketTrainer::Wait()
+	{
+		return ETrainerResponse::Success;
+	}
+
+	void FSocketTrainer::Terminate()
+	{
+		if (Socket)
+		{
+			Socket->Close();
+			Socket = nullptr;
+		}
+	}
+
+	ETrainerResponse FSocketTrainer::SendStop()
+	{
+		checkf(Socket, TEXT("Training socket is nullptr"));
+
+		return SocketTraining::SendStop(*Socket, TrainingProcess, Timeout);
+	}
+
+	ETrainerResponse FSocketTrainer::SendConfig(
+		const TSharedRef<FJsonObject>& ConfigObject,
+		const ELogSetting LogSettings)
+	{
+		checkf(Socket, TEXT("Training socket is nullptr"));
+
+		FString ConfigString;
+		TSharedRef<TJsonWriter<>> JsonWriter = TJsonWriterFactory<>::Create(&ConfigString, 0);
+		FJsonSerializer::Serialize(ConfigObject, JsonWriter, true);
+
+		return SocketTraining::SendConfig(*Socket, ConfigString, TrainingProcess, Timeout, LogSettings);
+	}
+
+	void FSocketTrainer::AddNetwork(
+		const FName& Name,
+		const ULearningNeuralNetworkData& Network)
+	{
+		NetworkBuffers.Add(Name, TLearningArray<1, uint8>());
+		NetworkBuffers[Name].SetNumUninitialized({Network.GetSnapshotByteNum()});
+	}
+
+	bool FSocketTrainer::ContainsNetwork(const FName& Name) const
+	{
+		return NetworkBuffers.Contains(Name);
+	}
+
+	ETrainerResponse FSocketTrainer::ReceiveNetwork(
+		const FName& Name,
+		ULearningNeuralNetworkData& OutNetwork,
+		FRWLock* NetworkLock,
+		const ELogSetting LogSettings)
+	{
+		checkf(Socket, TEXT("Training socket is nullptr"));
+		if (!ensureMsgf(NetworkBuffers.Contains(Name), TEXT("Network %s has not been added. Call AddNetwork prior to ReceiveNetwork."), *Name.ToString()))
+		{
+			return ETrainerResponse::Unexpected;
+		}
+
+		return SocketTraining::RecvNetwork(*Socket, OutNetwork, TrainingProcess, NetworkBuffers[Name], Timeout, NetworkLock, LogSettings);
+	}
+
+	ETrainerResponse FSocketTrainer::SendNetwork(
+		const FName& Name,
+		const ULearningNeuralNetworkData& Network,
+		FRWLock* NetworkLock,
+		const ELogSetting LogSettings)
+	{
+		checkf(Socket, TEXT("Training socket is nullptr"));
+		if (!ensureMsgf(NetworkBuffers.Contains(Name), TEXT("Network %s has not been added. Call AddNetwork prior to SendNetwork."), *Name.ToString()))
+		{
+			return ETrainerResponse::Unexpected;
+		}
+
+		return SocketTraining::SendNetwork(*Socket, NetworkBuffers[Name], TrainingProcess, Network, Timeout, NetworkLock, LogSettings);
+	}
+
+	void FSocketTrainer::AddReplayBuffer(const FName& Name, const FReplayBuffer& ReplayBuffer)
+	{
+		ExperienceBufferNames.Add(Name);
+	}
+
+	bool FSocketTrainer::ContainsReplayBuffer(const FName& Name) const
+	{
+		return ExperienceBufferNames.Contains(Name);
+	}
+
+	ETrainerResponse FSocketTrainer::SendReplayBuffer(const FName& Name, const FReplayBuffer& ReplayBuffer, const ELogSetting LogSettings)
+	{
+		checkf(Socket, TEXT("Training socket is nullptr"));
+		if (!ensureMsgf(ExperienceBufferNames.Contains(Name), TEXT("ReplayBuffer %s has not been added. Call AddReplayBuffer prior to SendReplayBuffer."), *Name.ToString()))
+		{
+			return ETrainerResponse::Unexpected;
+		}
+
+		return SocketTraining::SendExperience(*Socket, ReplayBuffer, TrainingProcess, Timeout, LogSettings);
+	}
+}
diff --git a/Engine/Plugins/Experimental/LearningAgents/Source/LearningTraining/Private/LearningImitationTrainer.cpp b/Engine/Plugins/Experimental/LearningAgents/Source/LearningTraining/Private/LearningImitationTrainer.cpp
--- a/Engine/Plugins/Experimental/LearningAgents/Source/LearningTraining/Private/LearningImitationTrainer.cpp
+++ b/Engine/Plugins/Experimental/LearningAgents/Source/LearningTraining/Private/LearningImitationTrainer.cpp
@@ -242,7 +242,6 @@ namespace UE::Learning
 			Controls.View,
 			OutNetwork,
 			TrainingProcess,
-			SharedMemoryTraining::EControls::PolicySignal,
 			Policy.View,
 			Timeout,
 			NetworkLock,
@@ -259,7 +258,6 @@ namespace UE::Learning
 			Controls.View,
 			OutNetwork,
 			TrainingProcess,
-			SharedMemoryTraining::EControls::EncoderSignal,
 			Encoder.View,
 			Timeout,
 			NetworkLock,
@@ -276,7 +274,6 @@ namespace UE::Learning
 			Controls.View,
 			OutNetwork,
 			TrainingProcess,
-			SharedMemoryTraining::EControls::DecoderSignal,
 			Decoder.View,
 			Timeout,
 			NetworkLock,
@@ -293,7 +290,6 @@ namespace UE::Learning
 			Controls.View,
 			Policy.View,
 			TrainingProcess,
-			SharedMemoryTraining::EControls::PolicySignal,
 			Network,
 			Timeout,
 			NetworkLock,
@@ -310,7 +306,6 @@ namespace UE::Learning
 			Controls.View,
 			Encoder.View,
 			TrainingProcess,
-			SharedMemoryTraining::EControls::EncoderSignal,
 			Network,
 			Timeout,
 			NetworkLock,
@@ -327,7 +322,6 @@ namespace UE::Learning
 			Controls.View,
 			Decoder.View,
 			TrainingProcess,
-			SharedMemoryTraining::EControls::DecoderSignal,
 			Network,
 			Timeout,
 			NetworkLock,
@@ -551,15 +545,14 @@ namespace UE::Learning
 		// Connect  to Server
 
 		Socket = FTcpSocketBuilder(TEXT("LearningNetworkRLTrainerSocket")).AsNonBlocking().Build();
-		Socket->Connect(*Address);
 
-		OutResponse = SocketTraining::WaitForConnection(*Socket, Timeout);
+		OutResponse = SocketTraining::WaitForConnection(*Socket, nullptr, *Address, Timeout);
 		if (OutResponse != ETrainerResponse::Success) { return; }
 
 
 		// Send Config
 
-		OutResponse = SocketTraining::SendConfig(*Socket, JsonString, Timeout);
+		OutResponse = SocketTraining::SendConfig(*Socket, JsonString, nullptr, Timeout);
 		return;
 	}
 
@@ -584,7 +577,7 @@ namespace UE::Learning
 
 	ETrainerResponse FSocketImitationTrainer::SendStop(const float Timeout)
 	{
-		return SocketTraining::SendStop(*Socket, Timeout);
+		return SocketTraining::SendStop(*Socket, nullptr, Timeout);
 	}
 
 	bool FSocketImitationTrainer::HasPolicyOrCompleted()
@@ -598,7 +591,7 @@ namespace UE::Learning
 		FRWLock* NetworkLock,
 		const ELogSetting LogSettings)
 	{
-		return SocketTraining::RecvNetwork(*Socket, OutNetwork, PolicyBuffer, SocketTraining::ESignal::RecvPolicy, Timeout, NetworkLock, LogSettings);
+		return SocketTraining::RecvNetwork(*Socket, OutNetwork, nullptr, PolicyBuffer, Timeout, NetworkLock, LogSettings);
 	}
 
 	ETrainerResponse FSocketImitationTrainer::RecvEncoder(
@@ -607,7 +600,7 @@ namespace UE::Learning
 		FRWLock* NetworkLock,
 		const ELogSetting LogSettings)
 	{
-		return SocketTraining::RecvNetwork(*Socket, OutNetwork, EncoderBuffer, SocketTraining::ESignal::RecvEncoder, Timeout, NetworkLock, LogSettings);
+		return SocketTraining::RecvNetwork(*Socket, OutNetwork, nullptr, EncoderBuffer, Timeout, NetworkLock, LogSettings);
 	}
 
 	ETrainerResponse FSocketImitationTrainer::RecvDecoder(
@@ -616,7 +609,7 @@ namespace UE::Learning
 		FRWLock* NetworkLock,
 		const ELogSetting LogSettings)
 	{
-		return SocketTraining::RecvNetwork(*Socket, OutNetwork, DecoderBuffer, SocketTraining::ESignal::RecvDecoder, Timeout, NetworkLock, LogSettings);
+		return SocketTraining::RecvNetwork(*Socket, OutNetwork, nullptr, DecoderBuffer, Timeout, NetworkLock, LogSettings);
 	}
 
 	ETrainerResponse FSocketImitationTrainer::SendPolicy(
@@ -625,7 +618,7 @@ namespace UE::Learning
 		FRWLock* NetworkLock,
 		const ELogSetting LogSettings)
 	{
-		return SocketTraining::SendNetwork(*Socket, PolicyBuffer, SocketTraining::ESignal::SendPolicy, Network, Timeout, NetworkLock, LogSettings);
+		return SocketTraining::SendNetwork(*Socket, PolicyBuffer, nullptr, Network, Timeout, NetworkLock, LogSettings);
 	}
 
 	ETrainerResponse FSocketImitationTrainer::SendEncoder(
@@ -634,7 +627,7 @@ namespace UE::Learning
 		FRWLock* NetworkLock,
 		const ELogSetting LogSettings)
 	{
-		return SocketTraining::SendNetwork(*Socket, EncoderBuffer, SocketTraining::ESignal::SendEncoder, Network, Timeout, NetworkLock, LogSettings);
+		return SocketTraining::SendNetwork(*Socket, EncoderBuffer, nullptr, Network, Timeout, NetworkLock, LogSettings);
 	}
 
 	ETrainerResponse FSocketImitationTrainer::SendDecoder(
@@ -643,7 +636,7 @@ namespace UE::Learning
 		FRWLock* NetworkLock,
 		const ELogSetting LogSettings)
 	{
-		return SocketTraining::SendNetwork(*Socket, DecoderBuffer, SocketTraining::ESignal::SendDecoder, Network, Timeout, NetworkLock, LogSettings);
+		return SocketTraining::SendNetwork(*Socket, DecoderBuffer, nullptr, Network, Timeout, NetworkLock, LogSettings);
 	}
 
 	ETrainerResponse FSocketImitationTrainer::SendExperience(
@@ -660,6 +653,7 @@ namespace UE::Learning
 			EpisodeLengthsExperience,
 			ObservationsExperience,
 			ActionsExperience,
+			nullptr,
 			Timeout, 
 			LogSettings);
 	}
