Learning Agents Training Rewrite: Part 1

diff --git a/Engine/Plugins/Experimental/LearningAgents/Source/LearningTraining/Private/LearningPPOTrainer.cpp b/Engine/Plugins/Experimental/LearningAgents/Source/LearningTraining/Private/LearningPPOTrainer.cpp
--- a/Engine/Plugins/Experimental/LearningAgents/Source/LearningTraining/Private/LearningPPOTrainer.cpp
+++ b/Engine/Plugins/Experimental/LearningAgents/Source/LearningTraining/Private/LearningPPOTrainer.cpp
@@ -3,6 +3,7 @@
 #include "LearningPPOTrainer.h"
 
 #include "LearningArray.h"
+#include "LearningExternalTrainer.h"
 #include "LearningLog.h"
 #include "LearningNeuralNetwork.h"
 #include "LearningExperience.h"
@@ -20,8 +21,6 @@
 #include "Misc/Paths.h"
 
 #include "Dom/JsonObject.h"
-#include "Serialization/JsonSerializer.h"
-#include "Serialization/JsonWriter.h"
 
 #include "Sockets.h"
 #include "Common/TcpSocketBuilder.h"
@@ -43,7 +42,6 @@ int32 ULearningSocketPPOTrainerServerCommandlet::Main(const FString& Commandline
 	UCommandlet::ParseCommandLine(*Commandline, Tokens, Switches, Params);
 
 	const FString* PythonExecutiblePathParam = Params.Find(TEXT("PythonExecutiblePath"));
-	const FString* ExtraSitePackagesPathParam = Params.Find(TEXT("ExtraSitePackagesPath"));
 	const FString* PythonContentPathParam = Params.Find(TEXT("PythonContentPath"));
 	const FString* IntermediatePathParam = Params.Find(TEXT("IntermediatePath"));
 	const FString* IpAddressParam = Params.Find(TEXT("IpAddress"));
@@ -51,7 +49,6 @@ int32 ULearningSocketPPOTrainerServerCommandlet::Main(const FString& Commandline
 	const FString* LogSettingsParam = Params.Find(TEXT("LogSettings"));
 
 	const FString PythonExecutiblePath = PythonExecutiblePathParam ? *PythonExecutiblePathParam : UE::Learning::Trainer::GetPythonExecutablePath(FPaths::ProjectIntermediateDir());
-	const FString ExtraSitePackagesPath = ExtraSitePackagesPathParam ? *ExtraSitePackagesPathParam : TEXT("");
 	const FString PythonContentPath = PythonContentPathParam ? *PythonContentPathParam : UE::Learning::Trainer::GetPythonContentPath(FPaths::EngineDir());
 	const FString IntermediatePath = IntermediatePathParam ? *IntermediatePathParam : UE::Learning::Trainer::GetIntermediatePath(FPaths::ProjectIntermediateDir());
 
@@ -78,20 +75,19 @@ int32 ULearningSocketPPOTrainerServerCommandlet::Main(const FString& Commandline
 	
 	UE_LOG(LogLearning, Display, TEXT("---  PPO Training Server Arguments ---"));
 	UE_LOG(LogLearning, Display, TEXT("PythonExecutiblePath: %s"), *PythonExecutiblePath);
-	UE_LOG(LogLearning, Display, TEXT("ExtraSitePackagesPath: %s"), *ExtraSitePackagesPath);
 	UE_LOG(LogLearning, Display, TEXT("PythonContentPath: %s"), *PythonContentPath);
 	UE_LOG(LogLearning, Display, TEXT("IntermediatePath: %s"), *IntermediatePath);
 	UE_LOG(LogLearning, Display, TEXT("IpAddress: %s"), IpAddress);
 	UE_LOG(LogLearning, Display, TEXT("Port: %i"), Port);
 	UE_LOG(LogLearning, Display, TEXT("LogSettings: %s"), LogSettings == UE::Learning::ELogSetting::Normal ? TEXT("Normal") : TEXT("Silent"));
 
-	UE::Learning::FSocketPPOTrainerServerProcess ServerProcess(
+	UE::Learning::FSocketTrainerServerProcess ServerProcess(
 		PythonExecutiblePath,
-		ExtraSitePackagesPath,
 		PythonContentPath,
 		IntermediatePath,
 		IpAddress,
 		Port,
+		UE::Learning::Trainer::DefaultTimeout,
 		UE::Learning::ESubprocessFlags::None,
 		LogSettings);
 
@@ -107,540 +103,73 @@ int32 ULearningSocketPPOTrainerServerCommandlet::Main(const FString& Commandline
 	return 0;
 }
 
-namespace UE::Learning
+namespace UE::Learning::PPOTrainer
 {
-	FSharedMemoryPPOTrainer::FSharedMemoryPPOTrainer(
-		const FString& TaskName,
-		const FString& PythonExecutablePath,
-		const FString& ExtraSitePackagesPath,
-		const FString& PythonContentPath,
-		const FString& IntermediatePath,
-		const FReplayBuffer& ReplayBuffer,
-		const ULearningNeuralNetworkData& PolicyNetwork,
-		const ULearningNeuralNetworkData& CriticNetwork,
-		const ULearningNeuralNetworkData& EncoderNetwork,
-		const ULearningNeuralNetworkData& DecoderNetwork,
+	ETrainerResponse Train(
+		IExternalTrainer* ExternalTrainer,
+		FReplayBuffer& ReplayBuffer,
+		FEpisodeBuffer& EpisodeBuffer,
+		FResetInstanceBuffer& ResetBuffer,
+		ULearningNeuralNetworkData& PolicyNetwork,
+		ULearningNeuralNetworkData& CriticNetwork,
+		ULearningNeuralNetworkData& EncoderNetwork,
+		ULearningNeuralNetworkData& DecoderNetwork,
+		TLearningArrayView<2, float> ObservationVectorBuffer,
+		TLearningArrayView<2, float> ActionVectorBuffer,
+		TLearningArrayView<2, float> PreEvaluationMemoryStateVectorBuffer,
+		TLearningArrayView<2, float> MemoryStateVectorBuffer,
+		TLearningArrayView<1, float> RewardBuffer,
+		TLearningArrayView<1, ECompletionMode> CompletionBuffer,
+		TLearningArrayView<1, ECompletionMode> EpisodeCompletionBuffer,
+		TLearningArrayView<1, ECompletionMode> AllCompletionBuffer,
+		const TFunctionRef<void(const FIndexSet Instances)> ResetFunction,
+		const TFunctionRef<void(const FIndexSet Instances)> ObservationFunction,
+		const TFunctionRef<void(const FIndexSet Instances)> PolicyFunction,
+		const TFunctionRef<void(const FIndexSet Instances)> ActionFunction,
+		const TFunctionRef<void(const FIndexSet Instances)> UpdateFunction,
+		const TFunctionRef<void(const FIndexSet Instances)> RewardFunction,
+		const TFunctionRef<void(const FIndexSet Instances)> CompletionFunction,
+		const FIndexSet Instances,
 		const Observation::FSchema& ObservationSchema,
 		const Observation::FSchemaElement& ObservationSchemaElement,
 		const Action::FSchema& ActionSchema,
 		const Action::FSchemaElement& ActionSchemaElement,
-		const FPPOTrainerTrainingSettings& TrainingSettings,
-		const ELogSetting LogSettings,
-		const ESubprocessFlags TrainingProcessFlags,
-		const uint16 ProcessNum,
-		const ESubprocessFlags MultiProcessFlags)
-	{
-		UE_LEARNING_CHECK(FPaths::FileExists(PythonExecutablePath));
-		UE_LEARNING_CHECK(FPaths::DirectoryExists(PythonContentPath));
-
-		const int32 ObservationVectorDimensionNum = ReplayBuffer.GetObservations().Num<1>();
-		const int32 ActionVectorDimensionNum = ReplayBuffer.GetActions().Num<1>();
-		const int32 MemoryStateVectorDimensionNum = ReplayBuffer.GetMemoryStates().Num<1>();
-
-		if (!ensure(Policy.Region == nullptr))
-		{
-			UE_LOG(LogLearning, Warning, TEXT("Training already started!"));
-			return;
-		}
-
-		// Allocate shared memory
-
-		ProcessIdx = 0;
-		FParse::Value(FCommandLine::Get(), TEXT("LearningProcessIdx"), ProcessIdx);
-
-		if (ProcessIdx == 0)
-		{
-			// Allocate Shared Memory
-
-			Policy = SharedMemory::Allocate<1, uint8>({ PolicyNetwork.GetSnapshotByteNum() });
-			Critic = SharedMemory::Allocate<1, uint8>({ CriticNetwork.GetSnapshotByteNum() });
-			Encoder = SharedMemory::Allocate<1, uint8>({ EncoderNetwork.GetSnapshotByteNum() });
-			Decoder = SharedMemory::Allocate<1, uint8>({ DecoderNetwork.GetSnapshotByteNum() });
-			Controls = SharedMemory::Allocate<2, volatile int32>({ ProcessNum, SharedMemoryTraining::GetControlNum() });
-			EpisodeStarts = SharedMemory::Allocate<2, int32>({ ProcessNum, ReplayBuffer.GetMaxEpisodeNum() });
-			EpisodeLengths = SharedMemory::Allocate<2, int32>({ ProcessNum, ReplayBuffer.GetMaxEpisodeNum() });
-			EpisodeCompletionModes = SharedMemory::Allocate<2, ECompletionMode>({ ProcessNum, ReplayBuffer.GetMaxEpisodeNum() });
-			EpisodeFinalObservations = SharedMemory::Allocate<3, float>({ ProcessNum, ReplayBuffer.GetMaxEpisodeNum(), ObservationVectorDimensionNum });
-			EpisodeFinalMemoryStates = SharedMemory::Allocate<3, float>({ ProcessNum, ReplayBuffer.GetMaxEpisodeNum(), MemoryStateVectorDimensionNum });
-			Observations = SharedMemory::Allocate<3, float>({ ProcessNum, ReplayBuffer.GetMaxStepNum(), ObservationVectorDimensionNum });
-			Actions = SharedMemory::Allocate<3, float>({ ProcessNum, ReplayBuffer.GetMaxStepNum(), ActionVectorDimensionNum });
-			MemoryStates = SharedMemory::Allocate<3, float>({ ProcessNum, ReplayBuffer.GetMaxStepNum(), MemoryStateVectorDimensionNum });
-			Rewards = SharedMemory::Allocate<2, float>({ ProcessNum, ReplayBuffer.GetMaxStepNum() });
-
-			// We need to zero the control memory before we start
-			// the training sub-process since it may contain uninitialized 
-			// values or those left over from previous runs.
-
-			Array::Zero(Controls.View);
-
-			// Create Experience Gathering Sub-processes
-
-			for (uint16 SubprocessIdx = 1; SubprocessIdx < ProcessNum; SubprocessIdx++)
-			{
-				ensureMsgf(!WITH_EDITOR || IsRunningCommandlet(), TEXT("Multi-processing generally does not work in-editor as it requires a standalone executable."));
-
-				FString SubprocessCommandLine = FCommandLine::GetOriginal();
-
-				SubprocessCommandLine += FString::Printf(TEXT(" -LearningProcessIdx %i"), SubprocessIdx);
-				SubprocessCommandLine += FString(TEXT(" -LearningPolicyGuid ")) + Policy.Guid.ToString();
-				SubprocessCommandLine += FString(TEXT(" -LearningCriticGuid ")) + Critic.Guid.ToString();
-				SubprocessCommandLine += FString(TEXT(" -LearningEncoderGuid ")) + Encoder.Guid.ToString();
-				SubprocessCommandLine += FString(TEXT(" -LearningDecoderGuid ")) + Decoder.Guid.ToString();
-				SubprocessCommandLine += FString(TEXT(" -LearningControlsGuid ")) + Controls.Guid.ToString();
-				SubprocessCommandLine += FString(TEXT(" -LearningEpisodeStartsGuid ")) + EpisodeStarts.Guid.ToString();
-				SubprocessCommandLine += FString(TEXT(" -LearningEpisodeLengthsGuid ")) + EpisodeLengths.Guid.ToString();
-				SubprocessCommandLine += FString(TEXT(" -LearningEpisodeCompletionModesGuid ")) + EpisodeCompletionModes.Guid.ToString();
-				SubprocessCommandLine += FString(TEXT(" -LearningEpisodeFinalObservationsGuid ")) + EpisodeFinalObservations.Guid.ToString();
-				SubprocessCommandLine += FString(TEXT(" -LearningEpisodeFinalMemoryStatesGuid ")) + EpisodeFinalMemoryStates.Guid.ToString();
-				SubprocessCommandLine += FString(TEXT(" -LearningObservationsGuid ")) + Observations.Guid.ToString();
-				SubprocessCommandLine += FString(TEXT(" -LearningActionsGuid ")) + Actions.Guid.ToString();
-				SubprocessCommandLine += FString(TEXT(" -LearningMemoryStatesGuid ")) + MemoryStates.Guid.ToString();
-				SubprocessCommandLine += FString(TEXT(" -LearningRewardsGuid ")) + Rewards.Guid.ToString();
-
-				FSubprocess& Subprocess = ExperienceGatheringSubprocesses.AddDefaulted_GetRef();
-				Subprocess.Launch(FPlatformProcess::ExecutablePath(), SubprocessCommandLine, MultiProcessFlags);
-
-				UE_LOG(LogLearning, Display, TEXT("Subprocess Command: %s %s"), FPlatformProcess::ExecutablePath(), *SubprocessCommandLine);
-			}
-
-			// Write Config
-
-			IFileManager& FileManager = IFileManager::Get();
-			const FString TimeStamp = FDateTime::Now().ToFormattedString(TEXT("%Y-%m-%d_%H-%M-%S"));
-			const FString TrainerMethod = TEXT("PPO");
-			const FString TrainerType = TEXT("SharedMemory");
-			const FString ConfigPath = IntermediatePath / TEXT("Configs") / FString::Printf(TEXT("%s_%s_%s_%s.json"), *TaskName, *TrainerMethod, *TrainerType, *TimeStamp);
-
-			TSharedRef<FJsonObject> ConfigObject = MakeShared<FJsonObject>();
-			ConfigObject->SetStringField(TEXT("TaskName"), *TaskName);
-			ConfigObject->SetStringField(TEXT("TrainerMethod"), TrainerMethod);
-			ConfigObject->SetStringField(TEXT("TrainerType"), TrainerType);
-			ConfigObject->SetStringField(TEXT("TimeStamp"), *TimeStamp);
-
-			ConfigObject->SetStringField(TEXT("ExtraSitePackagesPath"), *FileManager.ConvertToAbsolutePathForExternalAppForRead(*ExtraSitePackagesPath));
-			ConfigObject->SetStringField(TEXT("IntermediatePath"), *FileManager.ConvertToAbsolutePathForExternalAppForRead(*IntermediatePath));
-
-			ConfigObject->SetStringField(TEXT("PolicyGuid"), *Policy.Guid.ToString(EGuidFormats::DigitsWithHyphensInBraces));
-			ConfigObject->SetStringField(TEXT("CriticGuid"), *Critic.Guid.ToString(EGuidFormats::DigitsWithHyphensInBraces));
-			ConfigObject->SetStringField(TEXT("EncoderGuid"), *Encoder.Guid.ToString(EGuidFormats::DigitsWithHyphensInBraces));
-			ConfigObject->SetStringField(TEXT("DecoderGuid"), *Decoder.Guid.ToString(EGuidFormats::DigitsWithHyphensInBraces));
-			ConfigObject->SetStringField(TEXT("ControlsGuid"), *Controls.Guid.ToString(EGuidFormats::DigitsWithHyphensInBraces));
-			ConfigObject->SetStringField(TEXT("EpisodeStartsGuid"), *EpisodeStarts.Guid.ToString(EGuidFormats::DigitsWithHyphensInBraces));
-			ConfigObject->SetStringField(TEXT("EpisodeLengthsGuid"), *EpisodeLengths.Guid.ToString(EGuidFormats::DigitsWithHyphensInBraces));
-			ConfigObject->SetStringField(TEXT("EpisodeCompletionModesGuid"), *EpisodeCompletionModes.Guid.ToString(EGuidFormats::DigitsWithHyphensInBraces));
-			ConfigObject->SetStringField(TEXT("EpisodeFinalObservationsGuid"), *EpisodeFinalObservations.Guid.ToString(EGuidFormats::DigitsWithHyphensInBraces));
-			ConfigObject->SetStringField(TEXT("EpisodeFinalMemoryStatesGuid"), *EpisodeFinalMemoryStates.Guid.ToString(EGuidFormats::DigitsWithHyphensInBraces));
-			ConfigObject->SetStringField(TEXT("ObservationsGuid"), *Observations.Guid.ToString(EGuidFormats::DigitsWithHyphensInBraces));
-			ConfigObject->SetStringField(TEXT("ActionsGuid"), *Actions.Guid.ToString(EGuidFormats::DigitsWithHyphensInBraces));
-			ConfigObject->SetStringField(TEXT("MemoryStatesGuid"), *MemoryStates.Guid.ToString(EGuidFormats::DigitsWithHyphensInBraces));
-			ConfigObject->SetStringField(TEXT("RewardsGuid"), *Rewards.Guid.ToString(EGuidFormats::DigitsWithHyphensInBraces));
-
-			ConfigObject->SetObjectField(TEXT("ObservationSchema"), Trainer::ConvertObservationSchemaToJSON(ObservationSchema, ObservationSchemaElement));
-			ConfigObject->SetObjectField(TEXT("ActionSchema"), Trainer::ConvertActionSchemaToJSON(ActionSchema, ActionSchemaElement));
-			ConfigObject->SetNumberField(TEXT("ObservationVectorDimensionNum"), ObservationVectorDimensionNum);
-			ConfigObject->SetNumberField(TEXT("ActionVectorDimensionNum"), ActionVectorDimensionNum);
-			ConfigObject->SetNumberField(TEXT("MemoryStateVectorDimensionNum"), MemoryStateVectorDimensionNum);
-			ConfigObject->SetNumberField(TEXT("MaxEpisodeNum"), ReplayBuffer.GetMaxEpisodeNum());
-			ConfigObject->SetNumberField(TEXT("MaxStepNum"), ReplayBuffer.GetMaxStepNum());
-
-			ConfigObject->SetNumberField(TEXT("PolicyNetworkByteNum"), PolicyNetwork.GetSnapshotByteNum());
-			ConfigObject->SetNumberField(TEXT("CriticNetworkByteNum"), CriticNetwork.GetSnapshotByteNum());
-			ConfigObject->SetNumberField(TEXT("EncoderNetworkByteNum"), EncoderNetwork.GetSnapshotByteNum());
-			ConfigObject->SetNumberField(TEXT("DecoderNetworkByteNum"), DecoderNetwork.GetSnapshotByteNum());
-
-			ConfigObject->SetNumberField(TEXT("ProcessNum"), ProcessNum);
-
-			ConfigObject->SetNumberField(TEXT("IterationNum"), TrainingSettings.IterationNum);
-			ConfigObject->SetNumberField(TEXT("LearningRatePolicy"), TrainingSettings.LearningRatePolicy);
-			ConfigObject->SetNumberField(TEXT("LearningRateCritic"), TrainingSettings.LearningRateCritic);
-			ConfigObject->SetNumberField(TEXT("LearningRateDecay"), TrainingSettings.LearningRateDecay);
-			ConfigObject->SetNumberField(TEXT("WeightDecay"), TrainingSettings.WeightDecay);
-			ConfigObject->SetNumberField(TEXT("PolicyBatchSize"), TrainingSettings.PolicyBatchSize);
-			ConfigObject->SetNumberField(TEXT("CriticBatchSize"), TrainingSettings.CriticBatchSize);
-			ConfigObject->SetNumberField(TEXT("PolicyWindow"), TrainingSettings.PolicyWindow);
-			ConfigObject->SetNumberField(TEXT("IterationsPerGather"), TrainingSettings.IterationsPerGather);
-			ConfigObject->SetNumberField(TEXT("CriticWarmupIterations"), TrainingSettings.CriticWarmupIterations);
-			ConfigObject->SetNumberField(TEXT("EpsilonClip"), TrainingSettings.EpsilonClip);
-			ConfigObject->SetNumberField(TEXT("ActionSurrogateWeight"), TrainingSettings.ActionSurrogateWeight);
-			ConfigObject->SetNumberField(TEXT("ActionRegularizationWeight"), TrainingSettings.ActionRegularizationWeight);
-			ConfigObject->SetNumberField(TEXT("ActionEntropyWeight"), TrainingSettings.ActionEntropyWeight);
-			ConfigObject->SetNumberField(TEXT("ReturnRegularizationWeight"), TrainingSettings.ReturnRegularizationWeight);
-			ConfigObject->SetNumberField(TEXT("GaeLambda"), TrainingSettings.GaeLambda);
-			ConfigObject->SetBoolField(TEXT("AdvantageNormalization"), TrainingSettings.bAdvantageNormalization);
-			ConfigObject->SetNumberField(TEXT("AdvantageMin"), TrainingSettings.AdvantageMin);
-			ConfigObject->SetNumberField(TEXT("AdvantageMax"), TrainingSettings.AdvantageMax);
-			ConfigObject->SetBoolField(TEXT("UseGradNormMaxClipping"), TrainingSettings.bUseGradNormMaxClipping);
-			ConfigObject->SetNumberField(TEXT("GradNormMax"), TrainingSettings.GradNormMax);
-			ConfigObject->SetNumberField(TEXT("TrimEpisodeStartStepNum"), TrainingSettings.TrimEpisodeStartStepNum);
-			ConfigObject->SetNumberField(TEXT("TrimEpisodeEndStepNum"), TrainingSettings.TrimEpisodeEndStepNum);
-			ConfigObject->SetNumberField(TEXT("Seed"), TrainingSettings.Seed);
-			ConfigObject->SetNumberField(TEXT("DiscountFactor"), TrainingSettings.DiscountFactor);
-			ConfigObject->SetStringField(TEXT("Device"), Trainer::GetDeviceString(TrainingSettings.Device));
-			ConfigObject->SetBoolField(TEXT("UseTensorBoard"), TrainingSettings.bUseTensorboard);
-			ConfigObject->SetBoolField(TEXT("SaveSnapshots"), TrainingSettings.bSaveSnapshots);
-
-			ConfigObject->SetBoolField(TEXT("LoggingEnabled"), LogSettings == ELogSetting::Silent ? false : true);
-
-			FString JsonString;
-			TSharedRef<TJsonWriter<>> JsonWriter = TJsonWriterFactory<>::Create(&JsonString, 0);
-			FJsonSerializer::Serialize(ConfigObject, JsonWriter, true);
-
-			FFileHelper::SaveStringToFile(JsonString, *ConfigPath);
-
-			// Start Python Training Sub-process
-
-			const FString CommandLineArguments = FString::Printf(TEXT("\"%s\" SharedMemory \"%s\""), 
-				*FileManager.ConvertToAbsolutePathForExternalAppForRead(*(PythonContentPath / TEXT("train_ppo.py"))), 
-				*FileManager.ConvertToAbsolutePathForExternalAppForRead(*ConfigPath));
-
-			TrainingProcess.Launch(
-				FileManager.ConvertToAbsolutePathForExternalAppForRead(*PythonExecutablePath),
-				CommandLineArguments,
-				TrainingProcessFlags);
-		}
-		else
-		{
-			// Parse Guids from command line args
-
-			FGuid PolicyGuid; ensure(FParse::Value(FCommandLine::Get(), TEXT("LearningPolicyGuid"), PolicyGuid));
-			FGuid CriticGuid; ensure(FParse::Value(FCommandLine::Get(), TEXT("LearningCriticGuid"), CriticGuid));
-			FGuid EncoderGuid; ensure(FParse::Value(FCommandLine::Get(), TEXT("LearningEncoderGuid"), EncoderGuid));
-			FGuid DecoderGuid; ensure(FParse::Value(FCommandLine::Get(), TEXT("LearningDecoderGuid"), DecoderGuid));
-			FGuid ControlsGuid; ensure(FParse::Value(FCommandLine::Get(), TEXT("LearningControlsGuid"), ControlsGuid));
-			FGuid EpisodeStartsGuid; ensure(FParse::Value(FCommandLine::Get(), TEXT("LearningEpisodeStartsGuid"), EpisodeStartsGuid));
-			FGuid EpisodeLengthsGuid; ensure(FParse::Value(FCommandLine::Get(), TEXT("LearningEpisodeLengthsGuid"), EpisodeLengthsGuid));
-			FGuid EpisodeCompletionModesGuid; ensure(FParse::Value(FCommandLine::Get(), TEXT("LearningEpisodeCompletionModesGuid"), EpisodeCompletionModesGuid));
-			FGuid EpisodeFinalObservationsGuid; ensure(FParse::Value(FCommandLine::Get(), TEXT("LearningEpisodeFinalObservationsGuid"), EpisodeFinalObservationsGuid));
-			FGuid EpisodeFinalMemoryStatesGuid; ensure(FParse::Value(FCommandLine::Get(), TEXT("LearningEpisodeFinalMemoryStatesGuid"), EpisodeFinalMemoryStatesGuid));
-			FGuid ObservationsGuid; ensure(FParse::Value(FCommandLine::Get(), TEXT("LearningObservationsGuid"), ObservationsGuid));
-			FGuid ActionsGuid; ensure(FParse::Value(FCommandLine::Get(), TEXT("LearningActionsGuid"), ActionsGuid));
-			FGuid MemoryStatesGuid; ensure(FParse::Value(FCommandLine::Get(), TEXT("LearningMemoryStatesGuid"), MemoryStatesGuid));
-			FGuid RewardsGuid; ensure(FParse::Value(FCommandLine::Get(), TEXT("LearningRewardsGuid"), RewardsGuid));
-
-			// Map shared memory
-
-			Policy = SharedMemory::Map<1, uint8>(PolicyGuid, { PolicyNetwork.GetSnapshotByteNum() });
-			Critic = SharedMemory::Map<1, uint8>(CriticGuid, { CriticNetwork.GetSnapshotByteNum() });
-			Encoder = SharedMemory::Map<1, uint8>(EncoderGuid, { EncoderNetwork.GetSnapshotByteNum() });
-			Decoder = SharedMemory::Map<1, uint8>(DecoderGuid, { DecoderNetwork.GetSnapshotByteNum() });
-			Controls = SharedMemory::Map<2, volatile int32>(ControlsGuid, { ProcessNum, SharedMemoryTraining::GetControlNum() });
-			EpisodeStarts = SharedMemory::Map<2, int32>(EpisodeStartsGuid, { ProcessNum, ReplayBuffer.GetMaxEpisodeNum() });
-			EpisodeLengths = SharedMemory::Map<2, int32>(EpisodeLengthsGuid, { ProcessNum, ReplayBuffer.GetMaxEpisodeNum() });
-			EpisodeCompletionModes = SharedMemory::Map<2, ECompletionMode>(EpisodeCompletionModesGuid, { ProcessNum, ReplayBuffer.GetMaxEpisodeNum() });
-			EpisodeFinalObservations = SharedMemory::Map<3, float>(EpisodeFinalObservationsGuid, { ProcessNum, ReplayBuffer.GetMaxEpisodeNum(), ObservationVectorDimensionNum });
-			EpisodeFinalMemoryStates = SharedMemory::Map<3, float>(EpisodeFinalMemoryStatesGuid, { ProcessNum, ReplayBuffer.GetMaxEpisodeNum(), MemoryStateVectorDimensionNum });
-			Observations = SharedMemory::Map<3, float>(ObservationsGuid, { ProcessNum, ReplayBuffer.GetMaxStepNum(), ObservationVectorDimensionNum });
-			Actions = SharedMemory::Map<3, float>(ActionsGuid, { ProcessNum, ReplayBuffer.GetMaxStepNum(), ActionVectorDimensionNum });
-			MemoryStates = SharedMemory::Map<3, float>(MemoryStatesGuid, { ProcessNum, ReplayBuffer.GetMaxStepNum(), MemoryStateVectorDimensionNum });
-			Rewards = SharedMemory::Map<2, float>(RewardsGuid, { ProcessNum, ReplayBuffer.GetMaxStepNum() });
-		}
-	}
-
-	FSharedMemoryPPOTrainer::~FSharedMemoryPPOTrainer()
-	{
-		Terminate();
-	}
-
-	ETrainerResponse FSharedMemoryPPOTrainer::Wait(float Timeout)
-	{
-		const float SleepTime = 0.001f;
-		float WaitTime = 0.0f;
-
-		while (TrainingProcess.Update())
-		{
-			FPlatformProcess::Sleep(SleepTime);
-			WaitTime += SleepTime;
-
-			if (WaitTime > Timeout)
-			{
-				return ETrainerResponse::Timeout;
-			}
-		}
-
-		return ETrainerResponse::Success;
-	}
-
-	void FSharedMemoryPPOTrainer::Terminate()
-	{
-		TrainingProcess.Terminate();
-
-		if (Policy.Region)
-		{
-			Deallocate();
-		}
-	}
-
-	ETrainerResponse FSharedMemoryPPOTrainer::SendStop(const float Timeout)
-	{
-		return SharedMemoryTraining::SendStop(Controls.View[ProcessIdx]);
-	}
-
-	ETrainerResponse FSharedMemoryPPOTrainer::RecvPolicy(
-		ULearningNeuralNetworkData& OutNetwork,
-		const float Timeout,
-		FRWLock* NetworkLock,
-		const ELogSetting LogSettings)
-	{
-		return SharedMemoryTraining::RecvNetwork(
-			Controls.View[ProcessIdx],
-			OutNetwork,
-			TrainingProcess,
-			SharedMemoryTraining::EControls::PolicySignal,
-			Policy.View,
-			Timeout,
-			NetworkLock,
-			LogSettings);
-	}
-
-	ETrainerResponse FSharedMemoryPPOTrainer::RecvCritic(
-		ULearningNeuralNetworkData& OutNetwork,
-		const float Timeout,
-		FRWLock* NetworkLock,
-		const ELogSetting LogSettings)
-	{
-		return SharedMemoryTraining::RecvNetwork(
-			Controls.View[ProcessIdx],
-			OutNetwork,
-			TrainingProcess,
-			SharedMemoryTraining::EControls::CriticSignal,
-			Critic.View,
-			Timeout,
-			NetworkLock,
-			LogSettings);
-	}
-
-	ETrainerResponse FSharedMemoryPPOTrainer::RecvEncoder(
-		ULearningNeuralNetworkData& OutNetwork,
-		const float Timeout,
-		FRWLock* NetworkLock,
-		const ELogSetting LogSettings)
-	{
-		return SharedMemoryTraining::RecvNetwork(
-			Controls.View[ProcessIdx],
-			OutNetwork,
-			TrainingProcess,
-			SharedMemoryTraining::EControls::EncoderSignal,
-			Encoder.View,
-			Timeout,
-			NetworkLock,
-			LogSettings);
-	}
-
-	ETrainerResponse FSharedMemoryPPOTrainer::RecvDecoder(
-		ULearningNeuralNetworkData& OutNetwork,
-		const float Timeout,
-		FRWLock* NetworkLock,
-		const ELogSetting LogSettings)
-	{
-		return SharedMemoryTraining::RecvNetwork(
-			Controls.View[ProcessIdx],
-			OutNetwork,
-			TrainingProcess,
-			SharedMemoryTraining::EControls::DecoderSignal,
-			Decoder.View,
-			Timeout,
-			NetworkLock,
-			LogSettings);
-	}
-
-	ETrainerResponse FSharedMemoryPPOTrainer::SendPolicy(
-		const ULearningNeuralNetworkData& Network,
-		const float Timeout,
-		FRWLock* NetworkLock,
-		const ELogSetting LogSettings)
-	{
-		return SharedMemoryTraining::SendNetwork(
-			Controls.View[ProcessIdx],
-			Policy.View,
-			TrainingProcess,
-			SharedMemoryTraining::EControls::PolicySignal,
-			Network,
-			Timeout,
-			NetworkLock,
-			LogSettings);
-	}
-
-	ETrainerResponse FSharedMemoryPPOTrainer::SendCritic(
-		const ULearningNeuralNetworkData& Network,
-		const float Timeout,
-		FRWLock* NetworkLock,
-		const ELogSetting LogSettings)
-	{
-		return SharedMemoryTraining::SendNetwork(
-			Controls.View[ProcessIdx],
-			Critic.View,
-			TrainingProcess,
-			SharedMemoryTraining::EControls::CriticSignal,
-			Network,
-			Timeout,
-			NetworkLock,
-			LogSettings);
-	}
-
-	ETrainerResponse FSharedMemoryPPOTrainer::SendEncoder(
-		const ULearningNeuralNetworkData& Network,
-		const float Timeout,
-		FRWLock* NetworkLock,
-		const ELogSetting LogSettings)
-	{
-		return SharedMemoryTraining::SendNetwork(
-			Controls.View[ProcessIdx],
-			Encoder.View,
-			TrainingProcess,
-			SharedMemoryTraining::EControls::EncoderSignal,
-			Network,
-			Timeout,
-			NetworkLock,
-			LogSettings);
-	}
-
-	ETrainerResponse FSharedMemoryPPOTrainer::SendDecoder(
-		const ULearningNeuralNetworkData& Network,
-		const float Timeout,
-		FRWLock* NetworkLock,
+		const FPPOTrainerTrainingSettings& TrainerSettings,
+		TAtomic<bool>* bRequestTrainingStopSignal,
+		FRWLock* PolicyNetworkLock,
+		FRWLock* CriticNetworkLock,
+		FRWLock* EncoderNetworkLock,
+		FRWLock* DecoderNetworkLock,
+		TAtomic<bool>* bPolicyNetworkUpdatedSignal,
+		TAtomic<bool>* bCriticNetworkUpdatedSignal,
+		TAtomic<bool>* bEncoderNetworkUpdatedSignal,
+		TAtomic<bool>* bDecoderNetworkUpdatedSignal,
 		const ELogSetting LogSettings)
 	{
-		return SharedMemoryTraining::SendNetwork(
-			Controls.View[ProcessIdx],
-			Decoder.View,
-			TrainingProcess,
-			SharedMemoryTraining::EControls::DecoderSignal,
-			Network,
-			Timeout,
-			NetworkLock,
-			LogSettings);
-	}
-
-	ETrainerResponse FSharedMemoryPPOTrainer::SendExperience(
-		const FReplayBuffer& ReplayBuffer,
-		const float Timeout,
-		const ELogSetting LogSettings)
-	{
-		return SharedMemoryTraining::SendExperience(
-			EpisodeStarts.View[ProcessIdx],
-			EpisodeLengths.View[ProcessIdx],
-			EpisodeCompletionModes.View[ProcessIdx],
-			EpisodeFinalObservations.View[ProcessIdx],
-			EpisodeFinalMemoryStates.View[ProcessIdx],
-			Observations.View[ProcessIdx],
-			Actions.View[ProcessIdx],
-			MemoryStates.View[ProcessIdx],
-			Rewards.View[ProcessIdx],
-			Controls.View[ProcessIdx],
-			TrainingProcess,
-			ReplayBuffer,
-			Timeout,
-			LogSettings);
-	}
-
-	void FSharedMemoryPPOTrainer::Deallocate()
-	{
-		if (Policy.Region != nullptr)
-		{
-			SharedMemory::Deallocate(Policy);
-			SharedMemory::Deallocate(Critic);
-			SharedMemory::Deallocate(Encoder);
-			SharedMemory::Deallocate(Decoder);
-			SharedMemory::Deallocate(Controls);
-			SharedMemory::Deallocate(EpisodeStarts);
-			SharedMemory::Deallocate(EpisodeLengths);
-			SharedMemory::Deallocate(EpisodeCompletionModes);
-			SharedMemory::Deallocate(EpisodeFinalObservations);
-			SharedMemory::Deallocate(EpisodeFinalMemoryStates);
-			SharedMemory::Deallocate(Observations);
-			SharedMemory::Deallocate(Actions);
-			SharedMemory::Deallocate(MemoryStates);
-			SharedMemory::Deallocate(Rewards);
-		}
-	}
+		TRACE_CPUPROFILER_EVENT_SCOPE(Learning::PPOTrainer::Train);
 
-	FSocketPPOTrainerServerProcess::FSocketPPOTrainerServerProcess(
-		const FString& PythonExecutablePath,
-		const FString& ExtraSitePackagesPath,
-		const FString& PythonContentPath,
-		const FString& IntermediatePath,
-		const TCHAR* IpAddress,
-		const uint32 Port,
-		const ESubprocessFlags TrainingProcessFlags,
-		const ELogSetting LogSettings)
-	{
-		UE_LEARNING_CHECK(FPaths::FileExists(PythonExecutablePath));
-		UE_LEARNING_CHECK(FPaths::DirectoryExists(PythonContentPath));
-
-		IFileManager& FileManager = IFileManager::Get();
-		const FString CommandLineArguments = FString::Printf(TEXT("\"%s\" Socket \"%s:%i\" \"%s\" \"%s\" %i"), 
-			*FileManager.ConvertToAbsolutePathForExternalAppForRead(*(PythonContentPath / TEXT("train_ppo.py"))), 
-			IpAddress, 
-			Port, 
-			*FileManager.ConvertToAbsolutePathForExternalAppForRead(*ExtraSitePackagesPath),
-			*FileManager.ConvertToAbsolutePathForExternalAppForRead(*IntermediatePath), 
-			LogSettings == ELogSetting::Normal ? 1 : 0);
-
-		TrainingProcess.Launch(
-			FileManager.ConvertToAbsolutePathForExternalAppForRead(*PythonExecutablePath),
-			CommandLineArguments,
-			TrainingProcessFlags);
-	}
+		ETrainerResponse Response = ETrainerResponse::Success;
 
-	FSocketPPOTrainerServerProcess::~FSocketPPOTrainerServerProcess()
-	{
-		Terminate();
-	}
+		// Send initial Policy
 
-	bool FSocketPPOTrainerServerProcess::IsRunning() const
-	{
-		return TrainingProcess.IsRunning();
-	}
-
-	bool FSocketPPOTrainerServerProcess::Wait(float Timeout)
-	{
-		const float SleepTime = 0.001f;
-		float WaitTime = 0.0f;
-
-		while (TrainingProcess.Update())
+		if (LogSettings != ELogSetting::Silent)
 		{
-			FPlatformProcess::Sleep(SleepTime);
-			WaitTime += SleepTime;
-
-			if (WaitTime > Timeout)
-			{
-				return false;
-			}
+			UE_LOG(LogLearning, Display, TEXT("Sending initial Policy..."));
 		}
 
-		return true;
-	}
+		ExternalTrainer->AddNetwork(TEXT("Policy"), PolicyNetwork);
+		ExternalTrainer->AddNetwork(TEXT("Critic"), CriticNetwork);
+		ExternalTrainer->AddNetwork(TEXT("Encoder"), EncoderNetwork);
+		ExternalTrainer->AddNetwork(TEXT("Decoder"), DecoderNetwork);
+		ExternalTrainer->AddReplayBuffer(TEXT("ReplayBuffer"), ReplayBuffer);
 
-	void FSocketPPOTrainerServerProcess::Terminate()
-	{
-		TrainingProcess.Terminate();
-	}
-
-	FSocketPPOTrainer::FSocketPPOTrainer(
-		ETrainerResponse& OutResponse,
-		const FString& TaskName,
-		const FReplayBuffer& ReplayBuffer,
-		const ULearningNeuralNetworkData& PolicyNetwork,
-		const ULearningNeuralNetworkData& CriticNetwork,
-		const ULearningNeuralNetworkData& EncoderNetwork,
-		const ULearningNeuralNetworkData& DecoderNetwork,
-		const Observation::FSchema& ObservationSchema,
-		const Observation::FSchemaElement& ObservationSchemaElement,
-		const Action::FSchema& ActionSchema,
-		const Action::FSchemaElement& ActionSchemaElement,
-		const TCHAR* IpAddress,
-		const uint32 Port,
-		const float Timeout,
-		const FPPOTrainerTrainingSettings& TrainingSettings)
-	{
 		const int32 ObservationVectorDimensionNum = ReplayBuffer.GetObservations().Num<1>();
 		const int32 ActionVectorDimensionNum = ReplayBuffer.GetActions().Num<1>();
 		const int32 MemoryStateVectorDimensionNum = ReplayBuffer.GetMemoryStates().Num<1>();
 
-		// Write Config
-
+		// Write PPO Config
 		TSharedRef<FJsonObject> ConfigObject = MakeShared<FJsonObject>();
-		ConfigObject->SetStringField(TEXT("TaskName"), *TaskName);
+		ConfigObject->SetStringField(TEXT("TaskName"), TEXT("Training"));
 		ConfigObject->SetStringField(TEXT("TrainerMethod"), TEXT("PPO"));
 		ConfigObject->SetStringField(TEXT("TrainerType"), TEXT("Network"));
 		ConfigObject->SetStringField(TEXT("TimeStamp"), *FDateTime::Now().ToFormattedString(TEXT("%Y-%m-%d_%H-%M-%S")));
@@ -658,474 +187,274 @@ namespace UE::Learning
 		ConfigObject->SetNumberField(TEXT("EncoderNetworkByteNum"), EncoderNetwork.GetSnapshotByteNum());
 		ConfigObject->SetNumberField(TEXT("DecoderNetworkByteNum"), DecoderNetwork.GetSnapshotByteNum());
 
-		ConfigObject->SetNumberField(TEXT("IterationNum"), TrainingSettings.IterationNum);
-		ConfigObject->SetNumberField(TEXT("LearningRatePolicy"), TrainingSettings.LearningRatePolicy);
-		ConfigObject->SetNumberField(TEXT("LearningRateCritic"), TrainingSettings.LearningRateCritic);
-		ConfigObject->SetNumberField(TEXT("LearningRateDecay"), TrainingSettings.LearningRateDecay);
-		ConfigObject->SetNumberField(TEXT("WeightDecay"), TrainingSettings.WeightDecay);
-		ConfigObject->SetNumberField(TEXT("PolicyBatchSize"), TrainingSettings.PolicyBatchSize);
-		ConfigObject->SetNumberField(TEXT("CriticBatchSize"), TrainingSettings.CriticBatchSize);
-		ConfigObject->SetNumberField(TEXT("PolicyWindow"), TrainingSettings.PolicyWindow);
-		ConfigObject->SetNumberField(TEXT("IterationsPerGather"), TrainingSettings.IterationsPerGather);
-		ConfigObject->SetNumberField(TEXT("CriticWarmupIterations"), TrainingSettings.CriticWarmupIterations);
-		ConfigObject->SetNumberField(TEXT("EpsilonClip"), TrainingSettings.EpsilonClip);
-		ConfigObject->SetNumberField(TEXT("ActionSurrogateWeight"), TrainingSettings.ActionSurrogateWeight);
-		ConfigObject->SetNumberField(TEXT("ActionRegularizationWeight"), TrainingSettings.ActionRegularizationWeight);
-		ConfigObject->SetNumberField(TEXT("ActionEntropyWeight"), TrainingSettings.ActionEntropyWeight);
-		ConfigObject->SetNumberField(TEXT("ReturnRegularizationWeight"), TrainingSettings.ReturnRegularizationWeight);
-		ConfigObject->SetNumberField(TEXT("GaeLambda"), TrainingSettings.GaeLambda);
-		ConfigObject->SetBoolField(TEXT("AdvantageNormalization"), TrainingSettings.bAdvantageNormalization);
-		ConfigObject->SetNumberField(TEXT("AdvantageMin"), TrainingSettings.AdvantageMin);
-		ConfigObject->SetNumberField(TEXT("AdvantageMax"), TrainingSettings.AdvantageMax);
-		ConfigObject->SetBoolField(TEXT("UseGradNormMaxClipping"), TrainingSettings.bUseGradNormMaxClipping);
-		ConfigObject->SetNumberField(TEXT("GradNormMax"), TrainingSettings.GradNormMax);
-		ConfigObject->SetNumberField(TEXT("TrimEpisodeStartStepNum"), TrainingSettings.TrimEpisodeStartStepNum);
-		ConfigObject->SetNumberField(TEXT("TrimEpisodeEndStepNum"), TrainingSettings.TrimEpisodeEndStepNum);
-		ConfigObject->SetNumberField(TEXT("Seed"), TrainingSettings.Seed);
-		ConfigObject->SetNumberField(TEXT("DiscountFactor"), TrainingSettings.DiscountFactor);
-		ConfigObject->SetStringField(TEXT("Device"), Trainer::GetDeviceString(TrainingSettings.Device));
-		ConfigObject->SetBoolField(TEXT("UseTensorBoard"), TrainingSettings.bUseTensorboard);
-		ConfigObject->SetBoolField(TEXT("SaveSnapshots"), TrainingSettings.bSaveSnapshots);
-
-		FString JsonString;
-		TSharedRef<TJsonWriter<>> JsonWriter = TJsonWriterFactory<>::Create(&JsonString, 0);
-		FJsonSerializer::Serialize(ConfigObject, JsonWriter, true);
-
-		// Allocate buffer to receive network data in
-
-		PolicyNetworkBuffer.SetNumUninitialized({ PolicyNetwork.GetSnapshotByteNum() });
-		CriticNetworkBuffer.SetNumUninitialized({ CriticNetwork.GetSnapshotByteNum() });
-		EncoderNetworkBuffer.SetNumUninitialized({ EncoderNetwork.GetSnapshotByteNum() });
-		DecoderNetworkBuffer.SetNumUninitialized({ DecoderNetwork.GetSnapshotByteNum() });
-
-		// Create Socket
-
-		ISocketSubsystem* SocketSubsystem = ISocketSubsystem::Get(PLATFORM_SOCKETSUBSYSTEM);
-		if (!SocketSubsystem)
+		ConfigObject->SetNumberField(TEXT("IterationNum"), TrainerSettings.IterationNum);
+		ConfigObject->SetNumberField(TEXT("LearningRatePolicy"), TrainerSettings.LearningRatePolicy);
+		ConfigObject->SetNumberField(TEXT("LearningRateCritic"), TrainerSettings.LearningRateCritic);
+		ConfigObject->SetNumberField(TEXT("LearningRateDecay"), TrainerSettings.LearningRateDecay);
+		ConfigObject->SetNumberField(TEXT("WeightDecay"), TrainerSettings.WeightDecay);
+		ConfigObject->SetNumberField(TEXT("PolicyBatchSize"), TrainerSettings.PolicyBatchSize);
+		ConfigObject->SetNumberField(TEXT("CriticBatchSize"), TrainerSettings.CriticBatchSize);
+		ConfigObject->SetNumberField(TEXT("PolicyWindow"), TrainerSettings.PolicyWindow);
+		ConfigObject->SetNumberField(TEXT("IterationsPerGather"), TrainerSettings.IterationsPerGather);
+		ConfigObject->SetNumberField(TEXT("CriticWarmupIterations"), TrainerSettings.CriticWarmupIterations);
+		ConfigObject->SetNumberField(TEXT("EpsilonClip"), TrainerSettings.EpsilonClip);
+		ConfigObject->SetNumberField(TEXT("ActionSurrogateWeight"), TrainerSettings.ActionSurrogateWeight);
+		ConfigObject->SetNumberField(TEXT("ActionRegularizationWeight"), TrainerSettings.ActionRegularizationWeight);
+		ConfigObject->SetNumberField(TEXT("ActionEntropyWeight"), TrainerSettings.ActionEntropyWeight);
+		ConfigObject->SetNumberField(TEXT("ReturnRegularizationWeight"), TrainerSettings.ReturnRegularizationWeight);
+		ConfigObject->SetNumberField(TEXT("GaeLambda"), TrainerSettings.GaeLambda);
+		ConfigObject->SetBoolField(TEXT("AdvantageNormalization"), TrainerSettings.bAdvantageNormalization);
+		ConfigObject->SetNumberField(TEXT("AdvantageMin"), TrainerSettings.AdvantageMin);
+		ConfigObject->SetNumberField(TEXT("AdvantageMax"), TrainerSettings.AdvantageMax);
+		ConfigObject->SetBoolField(TEXT("UseGradNormMaxClipping"), TrainerSettings.bUseGradNormMaxClipping);
+		ConfigObject->SetNumberField(TEXT("GradNormMax"), TrainerSettings.GradNormMax);
+		ConfigObject->SetNumberField(TEXT("TrimEpisodeStartStepNum"), TrainerSettings.TrimEpisodeStartStepNum);
+		ConfigObject->SetNumberField(TEXT("TrimEpisodeEndStepNum"), TrainerSettings.TrimEpisodeEndStepNum);
+		ConfigObject->SetNumberField(TEXT("Seed"), TrainerSettings.Seed);
+		ConfigObject->SetNumberField(TEXT("DiscountFactor"), TrainerSettings.DiscountFactor);
+		ConfigObject->SetStringField(TEXT("Device"), Trainer::GetDeviceString(TrainerSettings.Device));
+		ConfigObject->SetBoolField(TEXT("UseTensorBoard"), TrainerSettings.bUseTensorboard);
+		ConfigObject->SetBoolField(TEXT("SaveSnapshots"), TrainerSettings.bSaveSnapshots);
+
+		ExternalTrainer->SendConfig(ConfigObject, LogSettings);
+
+		Response = ExternalTrainer->SendNetwork(TEXT("Policy"), PolicyNetwork, PolicyNetworkLock);
+
+		if (Response != ETrainerResponse::Success)
 		{
-			UE_LOG(LogLearning, Error, TEXT("Could not get socket subsystem"));
-			OutResponse = ETrainerResponse::Unexpected;
-			return;
-		}
-
-		bool bIsValid = false;
-		TSharedRef<FInternetAddr> Address = SocketSubsystem->CreateInternetAddr();
-		Address->SetIp(IpAddress, bIsValid);
-		Address->SetPort(Port);
+			if (LogSettings != ELogSetting::Silent)
+			{
+				UE_LOG(LogLearning, Error, TEXT("Error sending initial policy to trainer: %s. Check log for errors."), Trainer::GetResponseString(Response));
+			}
 
-		if (!bIsValid)
-		{
-			UE_LOG(LogLearning, Error, TEXT("Invalid Ip Address \"%s\"..."), IpAddress);
-			OutResponse = ETrainerResponse::Unexpected;
-			return;
+			ExternalTrainer->Terminate();
+			return Response;
 		}
 
-		// Connect  to Server
-
-		Socket = FTcpSocketBuilder(TEXT("LearningNetworkPPOTrainerSocket")).AsNonBlocking().Build();
-		Socket->Connect(*Address);
-
-		OutResponse = SocketTraining::WaitForConnection(*Socket, Timeout);
-		if (OutResponse != ETrainerResponse::Success) { return; }
-
-		// Send Config
-
-		OutResponse = SocketTraining::SendConfig(*Socket, JsonString, Timeout);
-		return;
-	}
-
-	FSocketPPOTrainer::~FSocketPPOTrainer()
-	{
-		Terminate();
-	}
-
-	ETrainerResponse FSocketPPOTrainer::Wait(const float Timeout)
-	{
-		return ETrainerResponse::Success;
-	}
+		// Send initial Critic
 
-	void FSocketPPOTrainer::Terminate()
-	{
-		if (Socket)
+		if (LogSettings != ELogSetting::Silent)
 		{
-			Socket->Close();
-			Socket = nullptr;
+			UE_LOG(LogLearning, Display, TEXT("Sending initial Critic..."));
 		}
-	}
-
-	ETrainerResponse FSocketPPOTrainer::SendStop(const float Timeout)
-	{
-		return SocketTraining::SendStop(*Socket, Timeout);
-	}
-
-	ETrainerResponse FSocketPPOTrainer::RecvPolicy(
-		ULearningNeuralNetworkData& OutNetwork,
-		const float Timeout,
-		FRWLock* NetworkLock,
-		const ELogSetting LogSettings)
-	{
-		return SocketTraining::RecvNetwork(*Socket, OutNetwork, PolicyNetworkBuffer, SocketTraining::ESignal::RecvPolicy, Timeout, NetworkLock, LogSettings);
-	}
-
-	ETrainerResponse FSocketPPOTrainer::RecvCritic(
-		ULearningNeuralNetworkData& OutNetwork,
-		const float Timeout,
-		FRWLock* NetworkLock,
-		const ELogSetting LogSettings)
-	{
-		return SocketTraining::RecvNetwork(*Socket, OutNetwork, CriticNetworkBuffer, SocketTraining::ESignal::RecvCritic, Timeout, NetworkLock, LogSettings);
-	}
-
-	ETrainerResponse FSocketPPOTrainer::RecvEncoder(
-		ULearningNeuralNetworkData& OutNetwork,
-		const float Timeout,
-		FRWLock* NetworkLock,
-		const ELogSetting LogSettings)
-	{
-		return SocketTraining::RecvNetwork(*Socket, OutNetwork, EncoderNetworkBuffer, SocketTraining::ESignal::RecvEncoder, Timeout, NetworkLock, LogSettings);
-	}
 
-	ETrainerResponse FSocketPPOTrainer::RecvDecoder(
-		ULearningNeuralNetworkData& OutNetwork,
-		const float Timeout,
-		FRWLock* NetworkLock,
-		const ELogSetting LogSettings)
-	{
-		return SocketTraining::RecvNetwork(*Socket, OutNetwork, DecoderNetworkBuffer, SocketTraining::ESignal::RecvDecoder, Timeout, NetworkLock, LogSettings);
-	}
-
-	ETrainerResponse FSocketPPOTrainer::SendPolicy(
-		const ULearningNeuralNetworkData& Network,
-		const float Timeout,
-		FRWLock* NetworkLock,
-		const ELogSetting LogSettings)
-	{
-		return SocketTraining::SendNetwork(*Socket, PolicyNetworkBuffer, SocketTraining::ESignal::SendPolicy, Network, Timeout, NetworkLock, LogSettings);
-	}
-
-	ETrainerResponse FSocketPPOTrainer::SendCritic(
-		const ULearningNeuralNetworkData& Network,
-		const float Timeout,
-		FRWLock* NetworkLock,
-		const ELogSetting LogSettings)
-	{
-		return SocketTraining::SendNetwork(*Socket, CriticNetworkBuffer, SocketTraining::ESignal::SendCritic, Network, Timeout, NetworkLock, LogSettings);
-	}
+		Response = ExternalTrainer->SendNetwork(TEXT("Critic"), CriticNetwork, CriticNetworkLock);
 
-	ETrainerResponse FSocketPPOTrainer::SendEncoder(
-		const ULearningNeuralNetworkData& Network,
-		const float Timeout,
-		FRWLock* NetworkLock,
-		const ELogSetting LogSettings)
-	{
-		return SocketTraining::SendNetwork(*Socket, EncoderNetworkBuffer, SocketTraining::ESignal::SendEncoder, Network, Timeout, NetworkLock, LogSettings);
-	}
-
-	ETrainerResponse FSocketPPOTrainer::SendDecoder(
-		const ULearningNeuralNetworkData& Network,
-		const float Timeout,
-		FRWLock* NetworkLock,
-		const ELogSetting LogSettings)
-	{
-		return SocketTraining::SendNetwork(*Socket, DecoderNetworkBuffer, SocketTraining::ESignal::SendDecoder, Network, Timeout, NetworkLock, LogSettings);
-	}
-
-	ETrainerResponse FSocketPPOTrainer::SendExperience(
-		const FReplayBuffer& ReplayBuffer,
-		const float Timeout,
-		const ELogSetting LogSettings)
-	{
-		return SocketTraining::SendExperience(*Socket, ReplayBuffer, Timeout, LogSettings);
-	}
-
-	namespace PPOTrainer
-	{
-		ETrainerResponse Train(
-			IPPOTrainer& Trainer,
-			FReplayBuffer& ReplayBuffer,
-			FEpisodeBuffer& EpisodeBuffer,
-			FResetInstanceBuffer& ResetBuffer,
-			ULearningNeuralNetworkData& PolicyNetwork,
-			ULearningNeuralNetworkData& CriticNetwork,
-			ULearningNeuralNetworkData& EncoderNetwork,
-			ULearningNeuralNetworkData& DecoderNetwork,
-			TLearningArrayView<2, float> ObservationVectorBuffer,
-			TLearningArrayView<2, float> ActionVectorBuffer,
-			TLearningArrayView<2, float> PreEvaluationMemoryStateVectorBuffer,
-			TLearningArrayView<2, float> MemoryStateVectorBuffer,
-			TLearningArrayView<1, float> RewardBuffer,
-			TLearningArrayView<1, ECompletionMode> CompletionBuffer,
-			TLearningArrayView<1, ECompletionMode> EpisodeCompletionBuffer,
-			TLearningArrayView<1, ECompletionMode> AllCompletionBuffer,
-			const TFunctionRef<void(const FIndexSet Instances)> ResetFunction,
-			const TFunctionRef<void(const FIndexSet Instances)> ObservationFunction,
-			const TFunctionRef<void(const FIndexSet Instances)> PolicyFunction,
-			const TFunctionRef<void(const FIndexSet Instances)> ActionFunction,
-			const TFunctionRef<void(const FIndexSet Instances)> UpdateFunction,
-			const TFunctionRef<void(const FIndexSet Instances)> RewardFunction,
-			const TFunctionRef<void(const FIndexSet Instances)> CompletionFunction,
-			const FIndexSet Instances,
-			TAtomic<bool>* bRequestTrainingStopSignal,
-			FRWLock* PolicyNetworkLock,
-			FRWLock* CriticNetworkLock,
-			FRWLock* EncoderNetworkLock,
-			FRWLock* DecoderNetworkLock,
-			TAtomic<bool>* bPolicyNetworkUpdatedSignal,
-			TAtomic<bool>* bCriticNetworkUpdatedSignal,
-			TAtomic<bool>* bEncoderNetworkUpdatedSignal,
-			TAtomic<bool>* bDecoderNetworkUpdatedSignal,
-			const ELogSetting LogSettings)
+		if (Response != ETrainerResponse::Success)
 		{
-			TRACE_CPUPROFILER_EVENT_SCOPE(Learning::PPOTrainer::Train);
-
-			ETrainerResponse Response = ETrainerResponse::Success;
-
-			// Send initial Policy
-
 			if (LogSettings != ELogSetting::Silent)
 			{
-				UE_LOG(LogLearning, Display, TEXT("Sending initial Policy..."));
+				UE_LOG(LogLearning, Error, TEXT("Error sending initial critic to trainer: %s. Check log for errors."), Trainer::GetResponseString(Response));
 			}
 
-			Response = Trainer.SendPolicy(PolicyNetwork, Trainer::DefaultTimeout, PolicyNetworkLock);
+			ExternalTrainer->Terminate();
+			return Response;
+		}
 
-			if (Response != ETrainerResponse::Success)
-			{
-				if (LogSettings != ELogSetting::Silent)
-				{
-					UE_LOG(LogLearning, Error, TEXT("Error sending initial policy to trainer: %s. Check log for errors."), Trainer::GetResponseString(Response));
-				}
+		// Send initial Encoder
 
-				Trainer.Terminate();
-				return Response;
-			}
+		if (LogSettings != ELogSetting::Silent)
+		{
+			UE_LOG(LogLearning, Display, TEXT("Sending initial Encoder..."));
+		}
 
-			// Send initial Critic
+		Response = ExternalTrainer->SendNetwork(TEXT("Encoder"), EncoderNetwork, EncoderNetworkLock);
 
+		if (Response != ETrainerResponse::Success)
+		{
 			if (LogSettings != ELogSetting::Silent)
 			{
-				UE_LOG(LogLearning, Display, TEXT("Sending initial Critic..."));
+				UE_LOG(LogLearning, Error, TEXT("Error sending initial encoder to trainer: %s. Check log for errors."), Trainer::GetResponseString(Response));
 			}
 
-			Response = Trainer.SendCritic(CriticNetwork, Trainer::DefaultTimeout, CriticNetworkLock);
+			ExternalTrainer->Terminate();
+			return Response;
+		}
 
-			if (Response != ETrainerResponse::Success)
-			{
-				if (LogSettings != ELogSetting::Silent)
-				{
-					UE_LOG(LogLearning, Error, TEXT("Error sending initial critic to trainer: %s. Check log for errors."), Trainer::GetResponseString(Response));
-				}
+		// Send initial Decoder
 
-				Trainer.Terminate();
-				return Response;
-			}
+		if (LogSettings != ELogSetting::Silent)
+		{
+			UE_LOG(LogLearning, Display, TEXT("Sending initial Decoder..."));
+		}
 
-			// Send initial Encoder
+		Response = ExternalTrainer->SendNetwork(TEXT("Decoder"), DecoderNetwork, DecoderNetworkLock);
 
+		if (Response != ETrainerResponse::Success)
+		{
 			if (LogSettings != ELogSetting::Silent)
 			{
-				UE_LOG(LogLearning, Display, TEXT("Sending initial Encoder..."));
+				UE_LOG(LogLearning, Error, TEXT("Error sending initial decoder to trainer: %s. Check log for errors."), Trainer::GetResponseString(Response));
 			}
 
-			Response = Trainer.SendEncoder(EncoderNetwork, Trainer::DefaultTimeout, EncoderNetworkLock);
-
-			if (Response != ETrainerResponse::Success)
-			{
-				if (LogSettings != ELogSetting::Silent)
-				{
-					UE_LOG(LogLearning, Error, TEXT("Error sending initial encoder to trainer: %s. Check log for errors."), Trainer::GetResponseString(Response));
-				}
-
-				Trainer.Terminate();
-				return Response;
-			}
-
-			// Send initial Decoder
+			ExternalTrainer->Terminate();
+			return Response;
+		}
 
-			if (LogSettings != ELogSetting::Silent)
+		// Start Training Loop
+		while (true)
+		{
+			if (bRequestTrainingStopSignal && (*bRequestTrainingStopSignal))
 			{
-				UE_LOG(LogLearning, Display, TEXT("Sending initial Decoder..."));
-			}
+				*bRequestTrainingStopSignal = false;
 
-			Response = Trainer.SendDecoder(DecoderNetwork, Trainer::DefaultTimeout, DecoderNetworkLock);
-
-			if (Response != ETrainerResponse::Success)
-			{
 				if (LogSettings != ELogSetting::Silent)
 				{
-					UE_LOG(LogLearning, Error, TEXT("Error sending initial decoder to trainer: %s. Check log for errors."), Trainer::GetResponseString(Response));
+					UE_LOG(LogLearning, Display, TEXT("Stopping Training..."));
 				}
 
-				Trainer.Terminate();
-				return Response;
-			}
-
-			// Start Training Loop
+				Response = ExternalTrainer->SendStop();
 
-			while (true)
-			{
-				if (bRequestTrainingStopSignal && (*bRequestTrainingStopSignal))
+				if (Response != ETrainerResponse::Success)
 				{
-					*bRequestTrainingStopSignal = false;
-
 					if (LogSettings != ELogSetting::Silent)
 					{
-						UE_LOG(LogLearning, Display, TEXT("Stopping Training..."));
-					}
-
-					Response = Trainer.SendStop();
-
-					if (Response != ETrainerResponse::Success)
-					{
-						if (LogSettings != ELogSetting::Silent)
-						{
-							UE_LOG(LogLearning, Error, TEXT("Error sending stop signal to trainer: %s. Check log for errors."), Trainer::GetResponseString(Response));
-						}
-
-						Trainer.Terminate();
-						return Response;
+						UE_LOG(LogLearning, Error, TEXT("Error sending stop signal to trainer: %s. Check log for errors."), Trainer::GetResponseString(Response));
 					}
 
-					break;
+					ExternalTrainer->Terminate();
+					return Response;
 				}
-				else
-				{
-					Experience::GatherExperienceUntilReplayBufferFull(
-						ReplayBuffer,
-						EpisodeBuffer,
-						ResetBuffer,
-						ObservationVectorBuffer,
-						ActionVectorBuffer,
-						PreEvaluationMemoryStateVectorBuffer,
-						MemoryStateVectorBuffer,
-						RewardBuffer,
-						CompletionBuffer,
-						EpisodeCompletionBuffer,
-						AllCompletionBuffer,
-						ResetFunction,
-						ObservationFunction,
-						PolicyFunction,
-						ActionFunction,
-						UpdateFunction,
-						RewardFunction,
-						CompletionFunction,
-						Instances);
-
-					Response = Trainer.SendExperience(ReplayBuffer, Trainer::DefaultTimeout);
-
-					if (Response != ETrainerResponse::Success)
-					{
-						if (LogSettings != ELogSetting::Silent)
-						{
-							UE_LOG(LogLearning, Error, TEXT("Error sending experience to trainer: %s. Check log for errors."), Trainer::GetResponseString(Response));
-						}
-
-						Trainer.Terminate();
-						return Response;
-					}
-				}
-
-				// Update Policy
 
-				Response = Trainer.RecvPolicy(PolicyNetwork, Trainer::DefaultTimeout, PolicyNetworkLock);
+				break;
+			}
+			else
+			{
+				Experience::GatherExperienceUntilReplayBufferFull(
+					ReplayBuffer,
+					EpisodeBuffer,
+					ResetBuffer,
+					ObservationVectorBuffer,
+					ActionVectorBuffer,
+					PreEvaluationMemoryStateVectorBuffer,
+					MemoryStateVectorBuffer,
+					RewardBuffer,
+					CompletionBuffer,
+					EpisodeCompletionBuffer,
+					AllCompletionBuffer,
+					ResetFunction,
+					ObservationFunction,
+					PolicyFunction,
+					ActionFunction,
+					UpdateFunction,
+					RewardFunction,
+					CompletionFunction,
+					Instances);
+
+				Response = ExternalTrainer->SendReplayBuffer(TEXT("ReplayBuffer"), ReplayBuffer);
 
-				if (Response == ETrainerResponse::Completed)
-				{
-					if (LogSettings != ELogSetting::Silent)
-					{
-						UE_LOG(LogLearning, Display, TEXT("Trainer completed training."));
-					}
-					break;
-				}
-				else if (Response != ETrainerResponse::Success)
+				if (Response != ETrainerResponse::Success)
 				{
 					if (LogSettings != ELogSetting::Silent)
 					{
-						UE_LOG(LogLearning, Error, TEXT("Error receiving policy from trainer: %s. Check log for errors."), Trainer::GetResponseString(Response));
+						UE_LOG(LogLearning, Error, TEXT("Error sending replay buffer to trainer: %s. Check log for errors."), Trainer::GetResponseString(Response));
 					}
-					break;
-				}
 
-				if (bPolicyNetworkUpdatedSignal)
-				{
-					*bPolicyNetworkUpdatedSignal = true;
+					ExternalTrainer->Terminate();
+					return Response;
 				}
+			}
 
-				// Update Critic
+			// Update Policy
 
-				Response = Trainer.RecvCritic(CriticNetwork, Trainer::DefaultTimeout, CriticNetworkLock);
+			Response = ExternalTrainer->ReceiveNetwork(TEXT("Policy"), PolicyNetwork, PolicyNetworkLock);
 
-				if (Response != ETrainerResponse::Success)
+			if (Response == ETrainerResponse::Completed)
+			{
+				if (LogSettings != ELogSetting::Silent)
 				{
-					if (LogSettings != ELogSetting::Silent)
-					{
-						UE_LOG(LogLearning, Error, TEXT("Error receiving critic from trainer: %s. Check log for errors."), Trainer::GetResponseString(Response));
-					}
-					break;
+					UE_LOG(LogLearning, Display, TEXT("Trainer completed training."));
 				}
-
-				if (bCriticNetworkUpdatedSignal)
+				break;
+			}
+			else if (Response != ETrainerResponse::Success)
+			{
+				if (LogSettings != ELogSetting::Silent)
 				{
-					*bCriticNetworkUpdatedSignal = true;
+					UE_LOG(LogLearning, Error, TEXT("Error receiving policy from trainer: %s. Check log for errors."), Trainer::GetResponseString(Response));
 				}
+				break;
+			}
 
-				// Update Encoder
+			if (bPolicyNetworkUpdatedSignal)
+			{
+				*bPolicyNetworkUpdatedSignal = true;
+			}
 
-				Response = Trainer.RecvEncoder(EncoderNetwork, Trainer::DefaultTimeout, EncoderNetworkLock);
+			// Update Critic
 
-				if (Response != ETrainerResponse::Success)
-				{
-					if (LogSettings != ELogSetting::Silent)
-					{
-						UE_LOG(LogLearning, Error, TEXT("Error receiving encoder from trainer: %s. Check log for errors."), Trainer::GetResponseString(Response));
-					}
-					break;
-				}
+			Response = ExternalTrainer->ReceiveNetwork(TEXT("Critic"), CriticNetwork, CriticNetworkLock);
 
-				if (bEncoderNetworkUpdatedSignal)
+			if (Response != ETrainerResponse::Success)
+			{
+				if (LogSettings != ELogSetting::Silent)
 				{
-					*bEncoderNetworkUpdatedSignal = true;
+					UE_LOG(LogLearning, Error, TEXT("Error receiving critic from trainer: %s. Check log for errors."), Trainer::GetResponseString(Response));
 				}
+				break;
+			}
 
-				// Update Decoder
+			if (bCriticNetworkUpdatedSignal)
+			{
+				*bCriticNetworkUpdatedSignal = true;
+			}
 
-				Response = Trainer.RecvDecoder(DecoderNetwork, Trainer::DefaultTimeout, DecoderNetworkLock);
+			// Update Encoder
 
-				if (Response != ETrainerResponse::Success)
-				{
-					if (LogSettings != ELogSetting::Silent)
-					{
-						UE_LOG(LogLearning, Error, TEXT("Error receiving decoder from trainer: %s. Check log for errors."), Trainer::GetResponseString(Response));
-					}
-					break;
-				}
+			Response = ExternalTrainer->ReceiveNetwork(TEXT("Encoder"), EncoderNetwork, EncoderNetworkLock);
 
-				if (bDecoderNetworkUpdatedSignal)
+			if (Response != ETrainerResponse::Success)
+			{
+				if (LogSettings != ELogSetting::Silent)
 				{
-					*bDecoderNetworkUpdatedSignal = true;
+					UE_LOG(LogLearning, Error, TEXT("Error receiving encoder from trainer: %s. Check log for errors."), Trainer::GetResponseString(Response));
 				}
+				break;
 			}
 
-			// Allow some time for trainer to shut down gracefully before we kill it...
-			
-			Response = Trainer.Wait(Trainer::DefaultTimeout);
+			if (bEncoderNetworkUpdatedSignal)
+			{
+				*bEncoderNetworkUpdatedSignal = true;
+			}
+
+			// Update Decoder
+
+			Response = ExternalTrainer->ReceiveNetwork(TEXT("Decoder"), DecoderNetwork, DecoderNetworkLock);
 
 			if (Response != ETrainerResponse::Success)
 			{
 				if (LogSettings != ELogSetting::Silent)
 				{
-					UE_LOG(LogLearning, Error, TEXT("Error waiting for trainer to exit: %s. Check log for errors."), Trainer::GetResponseString(Response));
+					UE_LOG(LogLearning, Error, TEXT("Error receiving decoder from trainer: %s. Check log for errors."), Trainer::GetResponseString(Response));
 				}
+				break;
 			}
 
-			Trainer.Terminate();
+			if (bDecoderNetworkUpdatedSignal)
+			{
+				*bDecoderNetworkUpdatedSignal = true;
+			}
+		}
 
+		// Allow some time for trainer to shut down gracefully before we kill it...
+			
+		Response = ExternalTrainer->Wait();
+
+		if (Response != ETrainerResponse::Success)
+		{
 			if (LogSettings != ELogSetting::Silent)
 			{
-				UE_LOG(LogLearning, Display, TEXT("Training Task Done!"));
+				UE_LOG(LogLearning, Error, TEXT("Error waiting for trainer to exit: %s. Check log for errors."), Trainer::GetResponseString(Response));
 			}
+		}
 
-			return ETrainerResponse::Success;
+		ExternalTrainer->Terminate();
+
+		if (LogSettings != ELogSetting::Silent)
+		{
+			UE_LOG(LogLearning, Display, TEXT("Training Task Done!"));
 		}
+
+		return ETrainerResponse::Success;
 	}
 }
diff --git a/Engine/Plugins/Experimental/LearningAgents/Source/LearningTraining/Private/LearningSharedMemoryTraining.cpp b/Engine/Plugins/Experimental/LearningAgents/Source/LearningTraining/Private/LearningSharedMemoryTraining.cpp
--- a/Engine/Plugins/Experimental/LearningAgents/Source/LearningTraining/Private/LearningSharedMemoryTraining.cpp
+++ b/Engine/Plugins/Experimental/LearningAgents/Source/LearningTraining/Private/LearningSharedMemoryTraining.cpp
@@ -24,14 +24,27 @@ namespace UE::Learning::SharedMemoryTraining
 
 	bool HasPolicyOrCompleted(TLearningArrayView<1, volatile int32> Controls)
 	{
-		return Controls[(uint8)EControls::PolicySignal] || Controls[(uint8)EControls::CompleteSignal];
+		return Controls[(uint8)EControls::NetworkSignal] || Controls[(uint8)EControls::CompleteSignal];
+	}
+
+	ETrainerResponse SendConfigSignal(
+		TLearningArrayView<1, volatile int32> Controls,
+		const ELogSetting LogSettings)
+	{
+		if (LogSettings != ELogSetting::Silent)
+		{
+			UE_LOG(LogLearning, Display, TEXT("Sending config signal..."));
+		}
+
+		Controls[(uint8)EControls::ConfigSignal] = true;
+
+		return ETrainerResponse::Success;
 	}
 
 	ETrainerResponse RecvNetwork(
 		TLearningArrayView<1, volatile int32> Controls,
 		ULearningNeuralNetworkData& OutNetwork,
 		FSubprocess& Process,
-		const EControls Signal,
 		const TLearningArrayView<1, const uint8> NetworkData,
 		const float Timeout,
 		FRWLock* NetworkLock,
@@ -40,8 +53,8 @@ namespace UE::Learning::SharedMemoryTraining
 		const float SleepTime = 0.001f;
 		float WaitTime = 0.0f;
 
-		// Wait until the network is done being written by the sub-process
-		while (!Controls[(uint8)Signal])
+		// Wait until the network is done being written by the training process
+		while (!Controls[(uint8)EControls::NetworkSignal])
 		{
 			// Check if Completed Signal has been raised
 			if (Controls[(uint8)EControls::CompleteSignal])
@@ -51,7 +64,7 @@ namespace UE::Learning::SharedMemoryTraining
 				return ETrainerResponse::Completed;
 			}
 
-			// Check if the process has exited
+			// If we're monitoring a process, then has it has exited?
 			if (!Process.Update())
 			{
 				return ETrainerResponse::Unexpected;
@@ -105,7 +118,7 @@ namespace UE::Learning::SharedMemoryTraining
 		}
 
 		// Confirm we have read the network
-		Controls[(uint8)Signal] = false;
+		Controls[(uint8)EControls::NetworkSignal] = false;
 
 		return bSuccess ? ETrainerResponse::Success : ETrainerResponse::Unexpected;
 	}
@@ -114,7 +127,6 @@ namespace UE::Learning::SharedMemoryTraining
 		TLearningArrayView<1, volatile int32> Controls,
 		TLearningArrayView<1, uint8> NetworkData,
 		FSubprocess& Process,
-		const EControls Signal,
 		const ULearningNeuralNetworkData& Network,
 		const float Timeout,
 		FRWLock* NetworkLock,
@@ -123,10 +135,10 @@ namespace UE::Learning::SharedMemoryTraining
 		const float SleepTime = 0.001f;
 		float WaitTime = 0.0f;
 
-		// Wait until the policy is requested by the sub-process
-		while (!Controls[(uint8)Signal])
+		// Wait until the policy is requested by the training process
+		while (!Controls[(uint8)EControls::NetworkSignal])
 		{
-			// Check if the process has exited
+			// If we're monitoring a process, then has it has exited?
 			if (!Process.Update())
 			{
 				return ETrainerResponse::Unexpected;
@@ -172,7 +184,7 @@ namespace UE::Learning::SharedMemoryTraining
 		}
 
 		// Confirm we have written the network
-		Controls[(uint8)Signal] = false;
+		Controls[(uint8)EControls::NetworkSignal] = false;
 
 		return bSuccess ? ETrainerResponse::Success : ETrainerResponse::Unexpected;
 	}
@@ -196,10 +208,10 @@ namespace UE::Learning::SharedMemoryTraining
 		const float SleepTime = 0.001f;
 		float WaitTime = 0.0f;
 
-		// Wait until the sub-process is done reading any experience
+		// Wait until the training process is done reading any experience
 		while (Controls[(uint8)EControls::ExperienceSignal])
 		{
-			// Check if the process has exited
+			// If we're monitoring a process, then has it has exited?
 			if (!Process.Update())
 			{
 				return ETrainerResponse::Unexpected;
@@ -267,10 +279,10 @@ namespace UE::Learning::SharedMemoryTraining
 		const float SleepTime = 0.001f;
 		float WaitTime = 0.0f;
 
-		// Wait until the sub-process is done reading any experience
+		// Wait until the training process is done reading any experience
 		while (Controls[(uint8)EControls::ExperienceSignal])
 		{
-			// Check if the process has exited
+			// If we're monitoring a process, then has it has exited?
 			if (!Process.Update())
 			{
 				return ETrainerResponse::Unexpected;
diff --git a/Engine/Plugins/Experimental/LearningAgents/Source/LearningTraining/Private/LearningSharedMemoryTraining.h b/Engine/Plugins/Experimental/LearningAgents/Source/LearningTraining/Private/LearningSharedMemoryTraining.h
--- a/Engine/Plugins/Experimental/LearningAgents/Source/LearningTraining/Private/LearningSharedMemoryTraining.h
+++ b/Engine/Plugins/Experimental/LearningAgents/Source/LearningTraining/Private/LearningSharedMemoryTraining.h
@@ -18,24 +18,25 @@ namespace UE::Learning
 			ExperienceEpisodeNum	= 0,
 			ExperienceStepNum		= 1,
 			ExperienceSignal		= 2,
-			PolicySignal			= 3,
-			CriticSignal			= 4,
-			EncoderSignal			= 5,
-			DecoderSignal			= 6,
-			CompleteSignal			= 7,
-			StopSignal				= 8,
-			PingSignal				= 9,
+			ConfigSignal			= 3,
+			NetworkSignal			= 4,
+			CompleteSignal			= 5,
+			StopSignal				= 6,
+			PingSignal				= 7,
 
-			ControlNum				= 10,
+			ControlNum				= 8,
 		};
 
 		LEARNINGTRAINING_API uint8 GetControlNum();
 
+		LEARNINGTRAINING_API ETrainerResponse SendConfigSignal(
+			TLearningArrayView<1, volatile int32> Controls,
+			const ELogSetting LogSettings = Trainer::DefaultLogSettings);
+
 		LEARNINGTRAINING_API ETrainerResponse RecvNetwork(
 			TLearningArrayView<1, volatile int32> Controls,
 			ULearningNeuralNetworkData& OutNetwork,
 			FSubprocess& Process,
-			const EControls Signal,
 			const TLearningArrayView<1, const uint8> NetworkData,
 			const float Timeout = Trainer::DefaultTimeout,
 			FRWLock* NetworkLock = nullptr,
@@ -50,7 +51,6 @@ namespace UE::Learning
 			TLearningArrayView<1, volatile int32> Controls,
 			TLearningArrayView<1, uint8> NetworkData,
 			FSubprocess& Process,
-			const EControls Signal,
 			const ULearningNeuralNetworkData& Network,
 			const float Timeout = Trainer::DefaultTimeout,
 			FRWLock* NetworkLock = nullptr,
diff --git a/Engine/Plugins/Experimental/LearningAgents/Source/LearningTraining/Private/LearningSocketTraining.cpp b/Engine/Plugins/Experimental/LearningAgents/Source/LearningTraining/Private/LearningSocketTraining.cpp
--- a/Engine/Plugins/Experimental/LearningAgents/Source/LearningTraining/Private/LearningSocketTraining.cpp
+++ b/Engine/Plugins/Experimental/LearningAgents/Source/LearningTraining/Private/LearningSocketTraining.cpp
@@ -11,19 +11,21 @@
 
 namespace UE::Learning::SocketTraining
 {
-	ETrainerResponse WaitForConnection(FSocket& Socket, const float Timeout)
+	ETrainerResponse WaitForConnection(FSocket& Socket, FSubprocess* Process, const FInternetAddr& Addr, const float Timeout)
 	{
 		float WaitTime = 0.0f;
 		const float SleepTime = 0.001f;
 
 		while (true)
 		{
-			if (Socket.GetConnectionState() == ESocketConnectionState::SCS_ConnectionError)
+			// If we're monitoring a process, then has it has exited?
+			if (Process && !Process->Update())
 			{
-				UE_LOG(LogLearning, Error, TEXT("Socket refused connection..."));
 				return ETrainerResponse::Unexpected;
 			}
 
+			Socket.Connect(Addr);
+
 			if (Socket.GetConnectionState() == ESocketConnectionState::SCS_Connected)
 			{
 				break;
@@ -43,7 +45,7 @@ namespace UE::Learning::SocketTraining
 		return ETrainerResponse::Success;
 	}
 
-	ETrainerResponse RecvWithTimeout(FSocket& Socket, uint8* Bytes, const int32 ByteNum, const float Timeout)
+	ETrainerResponse RecvWithTimeout(FSocket& Socket, FSubprocess* Process, uint8* Bytes, const int32 ByteNum, const float Timeout)
 	{
 		float WaitTime = 0.0f;
 		const float SleepTime = 0.001f;
@@ -51,8 +53,14 @@ namespace UE::Learning::SocketTraining
 		int32 BytesRead = 0;
 		int32 TotalBytesRead = 0;
 
-		while (true)
+		while (Socket.GetConnectionState() == SCS_Connected)
 		{
+			// If we're monitoring a process, then has it has exited?
+			if (Process && !Process->Update())
+			{
+				return ETrainerResponse::Unexpected;
+			}
+
 			if (Socket.Recv(Bytes + TotalBytesRead, ByteNum - TotalBytesRead, BytesRead))
 			{
 				TotalBytesRead += BytesRead;
@@ -71,13 +79,15 @@ namespace UE::Learning::SocketTraining
 				return ETrainerResponse::Timeout;
 			}
 		}
+
+		return ETrainerResponse::Unexpected;
 	}
 
 	ETrainerResponse RecvNetwork(
 		FSocket& Socket,
 		ULearningNeuralNetworkData& OutNetwork,
+		FSubprocess* Process,
 		TLearningArrayView<1, uint8> OutNetworkBuffer,
-		const ESignal NetworkSignal,
 		const float Timeout,
 		FRWLock* NetworkLock,
 		const ELogSetting LogSettings)
@@ -94,7 +104,7 @@ namespace UE::Learning::SocketTraining
 
 		while (true)
 		{
-			Response = RecvWithTimeout(Socket, &Signal, 1, Timeout);
+			Response = RecvWithTimeout(Socket, Process, &Signal, 1, Timeout);
 			if (Response != ETrainerResponse::Success) { return Response; }
 
 			if (Signal == (uint8)ESignal::RecvComplete)
@@ -107,7 +117,7 @@ namespace UE::Learning::SocketTraining
 				continue;
 			}
 
-			if (Signal != (uint8)NetworkSignal)
+			if (Signal != (uint8)ESignal::RecvNetwork)
 			{
 				return ETrainerResponse::Unexpected;
 			}
@@ -115,7 +125,7 @@ namespace UE::Learning::SocketTraining
 			break;
 		}
 
-		Response = RecvWithTimeout(Socket, OutNetworkBuffer.GetData(), OutNetworkBuffer.Num(), Timeout);
+		Response = RecvWithTimeout(Socket, Process, OutNetworkBuffer.GetData(), OutNetworkBuffer.Num(), Timeout);
 		if (Response != ETrainerResponse::Success) { return Response; }
 
 		bool bSuccess = false;
@@ -144,7 +154,7 @@ namespace UE::Learning::SocketTraining
 		return bSuccess ? ETrainerResponse::Success : ETrainerResponse::Unexpected;
 	}
 
-	ETrainerResponse SendWithTimeout(FSocket& Socket, const uint8* Bytes, const int32 ByteNum, const float Timeout)
+	ETrainerResponse SendWithTimeout(FSocket& Socket, FSubprocess* Process, const uint8* Bytes, const int32 ByteNum, const float Timeout)
 	{
 		float WaitTime = 0.0f;
 		const float SleepTime = 0.001f;
@@ -152,7 +162,7 @@ namespace UE::Learning::SocketTraining
 		int32 BytesSent = 0;
 		int32 TotalBytesSent = 0;
 
-		while (true)
+		while ((!Process || Process->Update()) && Socket.GetConnectionState() == SCS_Connected)
 		{
 			if (Socket.Send(Bytes + TotalBytesSent, ByteNum - TotalBytesSent, BytesSent))
 			{
@@ -172,30 +182,42 @@ namespace UE::Learning::SocketTraining
 				return ETrainerResponse::Timeout;
 			}
 		}
+
+		return ETrainerResponse::Unexpected;
 	}
 
-	ETrainerResponse SendConfig(FSocket& Socket, const FString& ConfigString, const float Timeout)
+	ETrainerResponse SendConfig(
+		FSocket& Socket,
+		const FString& ConfigString,
+		FSubprocess* Process,
+		const float Timeout,
+		const ELogSetting LogSettings)
 	{
+		if (LogSettings != ELogSetting::Silent)
+		{
+			UE_LOG(LogLearning, Display, TEXT("Sending config..."));
+		}
+
 		const uint8 Signal = (uint8)ESignal::SendConfig;
-		ETrainerResponse Response = SendWithTimeout(Socket, &Signal, 1, Timeout);
+		ETrainerResponse Response = SendWithTimeout(Socket, Process, &Signal, 1, Timeout);
 		if (Response != ETrainerResponse::Success) { return Response; }
 
 		const FTCHARToUTF8 UTF8String(*ConfigString);
 		const int32 ConfigLength = UTF8String.Length();
 		
-		Response = SendWithTimeout(Socket, (uint8*)&ConfigLength, sizeof(int32), Timeout);
+		Response = SendWithTimeout(Socket, Process, (uint8*)&ConfigLength, sizeof(int32), Timeout);
 		if (Response != ETrainerResponse::Success) { return Response; }
 		
-		Response = SendWithTimeout(Socket, (uint8*)UTF8String.Get(), UTF8String.Length(), Timeout);
+		Response = SendWithTimeout(Socket, Process, (uint8*)UTF8String.Get(), UTF8String.Length(), Timeout);
 		if (Response != ETrainerResponse::Success) { return Response; }
 
 		return ETrainerResponse::Success;
 	}
 
-	ETrainerResponse SendStop(FSocket& Socket, const float Timeout)
+	ETrainerResponse SendStop(FSocket& Socket, FSubprocess* Process, const float Timeout)
 	{
 		const uint8 Signal = (uint8)ESignal::SendStop;
-		return SendWithTimeout(Socket, &Signal, 1, Timeout);
+		return SendWithTimeout(Socket, Process, &Signal, 1, Timeout);
 	}
 
 	bool HasPolicyOrCompleted(FSocket& Socket)
@@ -207,7 +229,7 @@ namespace UE::Learning::SocketTraining
 	ETrainerResponse SendNetwork(
 		FSocket& Socket,
 		TLearningArrayView<1, uint8> NetworkBuffer,
-		const ESignal NetworkSignal,
+		FSubprocess* Process,
 		const ULearningNeuralNetworkData& Network,
 		const float Timeout,
 		FRWLock* NetworkLock,
@@ -233,11 +255,11 @@ namespace UE::Learning::SocketTraining
 			}
 		}
 
-		const uint8 Signal = (uint8)NetworkSignal;
-		ETrainerResponse Response = SendWithTimeout(Socket, &Signal, 1, Timeout);
+		const uint8 Signal = (uint8)ESignal::SendNetwork;
+		ETrainerResponse Response = SendWithTimeout(Socket, Process, &Signal, 1, Timeout);
 		if (Response != ETrainerResponse::Success) { return Response; }
 
-		Response = SendWithTimeout(Socket, NetworkBuffer.GetData(), NetworkBuffer.Num(), Timeout);
+		Response = SendWithTimeout(Socket, Process, NetworkBuffer.GetData(), NetworkBuffer.Num(), Timeout);
 		if (Response != ETrainerResponse::Success) { return Response; }
 
 		return bSuccess ? Response : ETrainerResponse::Unexpected;
@@ -246,6 +268,7 @@ namespace UE::Learning::SocketTraining
 	ETrainerResponse SendExperience(
 		FSocket& Socket,
 		const FReplayBuffer& ReplayBuffer,
+		FSubprocess* Process,
 		const float Timeout,
 		const ELogSetting LogSettings)
 	{
@@ -258,40 +281,40 @@ namespace UE::Learning::SocketTraining
 		const int32 StepNum = ReplayBuffer.GetStepNum();
 
 		const uint8 Signal = (uint8)ESignal::SendExperience;
-		ETrainerResponse Response = SendWithTimeout(Socket, &Signal, 1, Timeout);
+		ETrainerResponse Response = SendWithTimeout(Socket, Process, &Signal, 1, Timeout);
 		if (Response != ETrainerResponse::Success) { return Response; }
 
-		Response = SendWithTimeout(Socket, (const uint8*)&EpisodeNum, sizeof(int32), Timeout);
+		Response = SendWithTimeout(Socket, Process, (const uint8*)&EpisodeNum, sizeof(int32), Timeout);
 		if (Response != ETrainerResponse::Success) { return Response; }
 		
-		Response = SendWithTimeout(Socket, (const uint8*)&StepNum, sizeof(int32), Timeout);
+		Response = SendWithTimeout(Socket, Process, (const uint8*)&StepNum, sizeof(int32), Timeout);
 		if (Response != ETrainerResponse::Success) { return Response; }
 		
-		Response = SendWithTimeout(Socket, (const uint8*)ReplayBuffer.GetEpisodeStarts().GetData(), ReplayBuffer.GetEpisodeStarts().Num() * sizeof(int32), Timeout);
+		Response = SendWithTimeout(Socket, Process, (const uint8*)ReplayBuffer.GetEpisodeStarts().GetData(), ReplayBuffer.GetEpisodeStarts().Num() * sizeof(int32), Timeout);
 		if (Response != ETrainerResponse::Success) { return Response; }
 		
-		Response = SendWithTimeout(Socket, (const uint8*)ReplayBuffer.GetEpisodeLengths().GetData(), ReplayBuffer.GetEpisodeLengths().Num() * sizeof(int32), Timeout);
+		Response = SendWithTimeout(Socket, Process, (const uint8*)ReplayBuffer.GetEpisodeLengths().GetData(), ReplayBuffer.GetEpisodeLengths().Num() * sizeof(int32), Timeout);
 		if (Response != ETrainerResponse::Success) { return Response; }
 		
-		Response = SendWithTimeout(Socket, (const uint8*)ReplayBuffer.GetEpisodeCompletionModes().GetData(), ReplayBuffer.GetEpisodeCompletionModes().Num() * sizeof(ECompletionMode), Timeout);
+		Response = SendWithTimeout(Socket, Process, (const uint8*)ReplayBuffer.GetEpisodeCompletionModes().GetData(), ReplayBuffer.GetEpisodeCompletionModes().Num() * sizeof(ECompletionMode), Timeout);
 		if (Response != ETrainerResponse::Success) { return Response; }
 		
-		Response = SendWithTimeout(Socket, (const uint8*)ReplayBuffer.GetEpisodeFinalObservations().GetData(), ReplayBuffer.GetEpisodeFinalObservations().Num() * sizeof(float), Timeout);
+		Response = SendWithTimeout(Socket, Process, (const uint8*)ReplayBuffer.GetEpisodeFinalObservations().GetData(), ReplayBuffer.GetEpisodeFinalObservations().Num() * sizeof(float), Timeout);
 		if (Response != ETrainerResponse::Success) { return Response; }
 		
-		Response = SendWithTimeout(Socket, (const uint8*)ReplayBuffer.GetEpisodeFinalMemoryStates().GetData(), ReplayBuffer.GetEpisodeFinalMemoryStates().Num() * sizeof(float), Timeout);
+		Response = SendWithTimeout(Socket, Process, (const uint8*)ReplayBuffer.GetEpisodeFinalMemoryStates().GetData(), ReplayBuffer.GetEpisodeFinalMemoryStates().Num() * sizeof(float), Timeout);
 		if (Response != ETrainerResponse::Success) { return Response; }
 
-		Response = SendWithTimeout(Socket, (const uint8*)ReplayBuffer.GetObservations().GetData(), ReplayBuffer.GetObservations().Num() * sizeof(float), Timeout);
+		Response = SendWithTimeout(Socket, Process, (const uint8*)ReplayBuffer.GetObservations().GetData(), ReplayBuffer.GetObservations().Num() * sizeof(float), Timeout);
 		if (Response != ETrainerResponse::Success) { return Response; }
 		
-		Response = SendWithTimeout(Socket, (const uint8*)ReplayBuffer.GetActions().GetData(), ReplayBuffer.GetActions().Num() * sizeof(float), Timeout);
+		Response = SendWithTimeout(Socket, Process, (const uint8*)ReplayBuffer.GetActions().GetData(), ReplayBuffer.GetActions().Num() * sizeof(float), Timeout);
 		if (Response != ETrainerResponse::Success) { return Response; }
 		
-		Response = SendWithTimeout(Socket, (const uint8*)ReplayBuffer.GetMemoryStates().GetData(), ReplayBuffer.GetMemoryStates().Num() * sizeof(float), Timeout);
+		Response = SendWithTimeout(Socket, Process, (const uint8*)ReplayBuffer.GetMemoryStates().GetData(), ReplayBuffer.GetMemoryStates().Num() * sizeof(float), Timeout);
 		if (Response != ETrainerResponse::Success) { return Response; }
 
-		Response = SendWithTimeout(Socket, (const uint8*)ReplayBuffer.GetRewards().GetData(), ReplayBuffer.GetRewards().Num() * sizeof(float), Timeout);
+		Response = SendWithTimeout(Socket, Process, (const uint8*)ReplayBuffer.GetRewards().GetData(), ReplayBuffer.GetRewards().Num() * sizeof(float), Timeout);
 		if (Response != ETrainerResponse::Success) { return Response; }
 
 		return ETrainerResponse::Success;
@@ -303,6 +326,7 @@ namespace UE::Learning::SocketTraining
 		const TLearningArrayView<1, const int32> EpisodeLengthsExperience,
 		const TLearningArrayView<2, const float> ObservationExperience,
 		const TLearningArrayView<2, const float> ActionExperience,
+		FSubprocess* Process,
 		const float Timeout,
 		const ELogSetting LogSettings)
 	{
@@ -315,25 +339,25 @@ namespace UE::Learning::SocketTraining
 		const int32 StepNum = ObservationExperience.Num<0>();
 
 		const uint8 Signal = (uint8)ESignal::SendExperience;
-		ETrainerResponse Response = SendWithTimeout(Socket, &Signal, 1, Timeout);
+		ETrainerResponse Response = SendWithTimeout(Socket, Process, &Signal, 1, Timeout);
 		if (Response != ETrainerResponse::Success) { return Response; }
 
-		Response = SendWithTimeout(Socket, (const uint8*)&EpisodeNum, sizeof(int32), Timeout);
+		Response = SendWithTimeout(Socket, Process, (const uint8*)&EpisodeNum, sizeof(int32), Timeout);
 		if (Response != ETrainerResponse::Success) { return Response; }
 
-		Response = SendWithTimeout(Socket, (const uint8*)&StepNum, sizeof(int32), Timeout);
+		Response = SendWithTimeout(Socket, Process, (const uint8*)&StepNum, sizeof(int32), Timeout);
 		if (Response != ETrainerResponse::Success) { return Response; }
 
-		Response = SendWithTimeout(Socket, (const uint8*)EpisodeStartsExperience.GetData(), EpisodeStartsExperience.Num() * sizeof(float), Timeout);
+		Response = SendWithTimeout(Socket, Process, (const uint8*)EpisodeStartsExperience.GetData(), EpisodeStartsExperience.Num() * sizeof(float), Timeout);
 		if (Response != ETrainerResponse::Success) { return Response; }
 
-		Response = SendWithTimeout(Socket, (const uint8*)EpisodeLengthsExperience.GetData(), EpisodeLengthsExperience.Num() * sizeof(float), Timeout);
+		Response = SendWithTimeout(Socket, Process, (const uint8*)EpisodeLengthsExperience.GetData(), EpisodeLengthsExperience.Num() * sizeof(float), Timeout);
 		if (Response != ETrainerResponse::Success) { return Response; }
 
-		Response = SendWithTimeout(Socket, (const uint8*)ObservationExperience.GetData(), ObservationExperience.Num() * sizeof(float), Timeout);
+		Response = SendWithTimeout(Socket, Process, (const uint8*)ObservationExperience.GetData(), ObservationExperience.Num() * sizeof(float), Timeout);
 		if (Response != ETrainerResponse::Success) { return Response; }
 
-		Response = SendWithTimeout(Socket, (const uint8*)ActionExperience.GetData(), ActionExperience.Num() * sizeof(float), Timeout);
+		Response = SendWithTimeout(Socket, Process, (const uint8*)ActionExperience.GetData(), ActionExperience.Num() * sizeof(float), Timeout);
 		if (Response != ETrainerResponse::Success) { return Response; }
 
 		return ETrainerResponse::Success;
diff --git a/Engine/Plugins/Experimental/LearningAgents/Source/LearningTraining/Private/LearningSocketTraining.h b/Engine/Plugins/Experimental/LearningAgents/Source/LearningTraining/Private/LearningSocketTraining.h
--- a/Engine/Plugins/Experimental/LearningAgents/Source/LearningTraining/Private/LearningSocketTraining.h
+++ b/Engine/Plugins/Experimental/LearningAgents/Source/LearningTraining/Private/LearningSocketTraining.h
@@ -4,6 +4,7 @@
 
 #include "LearningTrainer.h"
 
+class FInternetAddr;
 class FSocket;
 class ULearningNeuralNetworkData;
 
@@ -18,25 +19,22 @@ namespace UE::Learning
 			Invalid				= 0,
 			SendConfig			= 1,
 			SendExperience		= 2,
-			RecvPolicy			= 3,
-			SendPolicy			= 4,
-			RecvCritic			= 5,
-			SendCritic			= 6,
-			RecvEncoder			= 7,
-			SendEncoder			= 8,
-			RecvDecoder			= 9,
-			SendDecoder			= 10,
-			RecvComplete		= 11,
-			SendStop			= 12,
-			RecvPing			= 13,
+			RecvNetwork			= 3,
+			SendNetwork			= 4,
+			RecvComplete		= 5,
+			SendStop			= 6,
+			RecvPing			= 7,
 		};
 
 		LEARNINGTRAINING_API ETrainerResponse WaitForConnection(
-			FSocket& Socket, 
+			FSocket& Socket,
+			FSubprocess* Process,
+			const FInternetAddr& Addr,
 			const float Timeout = Trainer::DefaultTimeout);
 
 		LEARNINGTRAINING_API ETrainerResponse RecvWithTimeout(
-			FSocket& Socket, 
+			FSocket& Socket,
+			FSubprocess* Process,
 			uint8* Bytes, 
 			const int32 ByteNum, 
 			const float Timeout = Trainer::DefaultTimeout);
@@ -44,25 +42,29 @@ namespace UE::Learning
 		LEARNINGTRAINING_API ETrainerResponse RecvNetwork(
 			FSocket& Socket,
 			ULearningNeuralNetworkData& OutNetwork,
+			FSubprocess* Process,
 			TLearningArrayView<1, uint8> OutNetworkBuffer,
-			const ESignal NetworkSignal,
 			const float Timeout = Trainer::DefaultTimeout,
 			FRWLock* NetworkLock = nullptr,
 			const ELogSetting LogSettings = Trainer::DefaultLogSettings);
 
 		LEARNINGTRAINING_API ETrainerResponse SendWithTimeout(
-			FSocket& Socket, 
+			FSocket& Socket,
+			FSubprocess* Process,
 			const uint8* Bytes, 
 			const int32 ByteNum, 
 			const float Timeout = Trainer::DefaultTimeout);
 
 		LEARNINGTRAINING_API ETrainerResponse SendConfig(
-			FSocket& Socket, 
-			const FString& ConfigString, 
-			const float Timeout = Trainer::DefaultTimeout);
+			FSocket& Socket,
+			const FString& ConfigString,
+			FSubprocess* Process,
+			const float Timeout = Trainer::DefaultTimeout,
+			const ELogSetting LogSettings = Trainer::DefaultLogSettings);
 
 		LEARNINGTRAINING_API ETrainerResponse SendStop(
-			FSocket& Socket, 
+			FSocket& Socket,
+			FSubprocess* Process,
 			const float Timeout = Trainer::DefaultTimeout);
 
 		LEARNINGTRAINING_API bool HasPolicyOrCompleted(FSocket& Socket);
@@ -70,7 +72,7 @@ namespace UE::Learning
 		LEARNINGTRAINING_API ETrainerResponse SendNetwork(
 			FSocket& Socket,
 			TLearningArrayView<1, uint8> NetworkBuffer,
-			const ESignal NetworkSignal,
+			FSubprocess* Process,
 			const ULearningNeuralNetworkData& Network,
 			const float Timeout = Trainer::DefaultTimeout,
 			FRWLock* NetworkLock = nullptr,
@@ -79,6 +81,7 @@ namespace UE::Learning
 		LEARNINGTRAINING_API ETrainerResponse SendExperience(
 			FSocket& Socket,
 			const FReplayBuffer& ReplayBuffer,
+			FSubprocess* Process,
 			const float Timeout = Trainer::DefaultTimeout,
 			const ELogSetting LogSettings = Trainer::DefaultLogSettings);
 
@@ -88,6 +91,7 @@ namespace UE::Learning
 			const TLearningArrayView<1, const int32> EpisodeLengthsExperience,
 			const TLearningArrayView<2, const float> ObservationExperience,
 			const TLearningArrayView<2, const float> ActionExperience,
+			FSubprocess* Process,
 			const float Timeout = Trainer::DefaultTimeout,
 			const ELogSetting LogSettings = Trainer::DefaultLogSettings);
 	}
diff --git a/Engine/Plugins/Experimental/LearningAgents/Source/LearningTraining/Public/LearningExternalTrainer.h b/Engine/Plugins/Experimental/LearningAgents/Source/LearningTraining/Public/LearningExternalTrainer.h
new file mode 100644
--- /dev/null
+++ b/Engine/Plugins/Experimental/LearningAgents/Source/LearningTraining/Public/LearningExternalTrainer.h
@@ -0,0 +1,447 @@
+// Copyright Epic Games, Inc. All Rights Reserved.
+
+#pragma once
+
+#include "LearningArray.h"
+#include "LearningLog.h"
+#include "LearningTrainer.h"
+#include "LearningSharedMemory.h"
+
+#include "Misc/Guid.h"
+#include "Templates/SharedPointer.h"
+#include "UObject/NameTypes.h"
+
+class FJsonObject;
+class FSocket;
+class ULearningNeuralNetworkData;
+
+namespace UE::Learning
+{
+	struct FReplayBuffer;
+	enum class ECompletionMode : uint8;
+
+    /** Interface for communicating with an external trainer process. */
+	struct IExternalTrainer
+	{
+		virtual ~IExternalTrainer() {}
+
+		/** Terminate the trainer immediately. */
+		virtual void Terminate() = 0;
+
+		/** Signal for the trainer to stop.	*/
+		virtual ETrainerResponse SendStop() = 0;
+
+		/**
+		* Wait for the trainer to finish.
+		*
+		* @param Timeout		Timeout to wait in seconds
+		* @returns				Trainer response
+		*/
+		virtual ETrainerResponse Wait() = 0;
+
+
+		/**
+		* Sends the given json config to the trainer process.
+		*
+		* @param ConfigObject	The config to send
+		* @param LogSettings	The log verbosity level
+		* @returns				Trainer response
+		*/
+		virtual ETrainerResponse SendConfig(
+			const TSharedRef<FJsonObject>& ConfigObject,
+			const ELogSetting LogSettings = Trainer::DefaultLogSettings) = 0;
+
+		/**
+		 * Adds the network to this external trainer. Allocates buffers, etc.
+		 * Must be called for each network prior to calling Send/Receive.
+		 * 
+		 * @params Name		The unique name of the network, used as a key.
+		 * @params Network	The network to be added.
+		 */
+		virtual void AddNetwork(const FName& Name, const ULearningNeuralNetworkData& Network) = 0;
+
+		/** Returns true if a network with the given Name has already been added. Otherwise, false. */
+		virtual bool ContainsNetwork(const FName& Name) const = 0;
+
+		/**
+		* Wait for the trainer to push an updated network.
+		*
+		* @param OutNetwork		Network to update
+		* @param Timeout		Timeout to wait in seconds
+		* @param NetworkLock	Lock to use when updating network
+		* @param LogSettings	The log verbosity level
+		* @returns				Trainer response
+		*/
+		virtual ETrainerResponse ReceiveNetwork(
+			const FName& Name,
+			ULearningNeuralNetworkData& OutNetwork,
+			FRWLock* NetworkLock = nullptr,
+			const ELogSetting LogSettings = Trainer::DefaultLogSettings) = 0;
+
+		/**
+		* Wait for the trainer to be ready and push the current policy network.
+		*
+		* @param Network		Network to push
+		* @param Timeout		Timeout to wait in seconds
+		* @param NetworkLock	Lock to use when pushing network
+		* @param LogSettings	The log verbosity level
+		* @returns				Trainer response
+		*/
+		virtual ETrainerResponse SendNetwork(
+			const FName& Name,
+			const ULearningNeuralNetworkData& Network,
+			FRWLock* NetworkLock = nullptr,
+			const ELogSetting LogSettings = Trainer::DefaultLogSettings) = 0;
+
+		/**
+		 * Adds a named replay buffer to this external trainer.
+		 * Must be called for each buffer prior to calling SendReplayBuffer.
+		 *
+		 * @params Name						The unique name of the buffer, used as a key
+		 * @params ReplayBuffer				The buffer to be added
+		 */
+		virtual void AddReplayBuffer(const FName& Name,	const FReplayBuffer& ReplayBuffer) = 0;
+
+		/** Returns true if a replay buffer with the given Name has already been added. Otherwise, false. */
+		virtual bool ContainsReplayBuffer(const FName& Name) const = 0;
+
+		/**
+		* Wait for the trainer to be ready and send new experience.
+		*
+		* @params Name			The unique name of the buffer, used as a key
+		* @param ReplayBuffer	Replay buffer to send
+		* @param Timeout		Timeout to wait in seconds
+		* @param LogSettings	The log verbosity level
+		* @returns				Trainer response
+		*/
+		virtual ETrainerResponse SendReplayBuffer(
+			const FName& Name,
+			const FReplayBuffer& ReplayBuffer,
+			const ELogSetting LogSettings = Trainer::DefaultLogSettings) = 0;
+	};
+
+	/** Interface for an external trainer process. */
+	struct ITrainerProcess
+	{
+		virtual ~ITrainerProcess() {}
+
+		/**
+		* Check if the server process is still running.
+		*/
+		virtual bool IsRunning() const = 0;
+
+		/**
+		* Wait for the server process to end
+		*
+		* @returns True if successful, otherwise false if it times out.
+		*/
+		virtual bool Wait() = 0;
+
+		/**
+		* Terminate the server process.
+		*/
+		virtual void Terminate() = 0;
+
+		/** Get the training subprocess. */
+		virtual FSubprocess* GetTrainingSubprocess() = 0;
+	};
+
+	/**
+	* This object allows you to launch the FSharedMemoryTrainer server as a subprocess,
+	* which is convenient when you want to train locally.
+	*/
+	struct LEARNINGTRAINING_API FSharedMemoryTrainerServerProcess : public ITrainerProcess
+	{
+		/**
+		* Creates a training server as a subprocess using shared memory for communication. This will no-op if this UE
+		* process has a non-zero "LearningProcessIdx".
+		*
+		* @param TaskName					The name of this training task (used to disambiguate filenames, etc.)
+		* @param PythonExecutablePath		Path to the python executable used for training. In general should be the
+		*									python shipped with Unreal Editor.
+		* @param PythonContentPath			Path to the Python Content folder provided by the Learning plugin
+		* @param IntermediatePath			Path to the intermediate folder to write temporary files, logs, and
+		*									snapshots to
+		* @param ProcessNum					Number of processes to use for multi-processed experience gathering (used
+		*									to allocate enough shared memory)
+		* @param TrainingProcessFlags		Training server subprocess flags
+		*/
+		FSharedMemoryTrainerServerProcess(
+			const FString& TaskName,
+			const FString& PythonExecutablePath,
+			const FString& PythonContentPath,
+			const FString& IntermediatePath,
+			const int32 ProcessNum,
+			const float InTimeout = Trainer::DefaultTimeout,
+			const ESubprocessFlags TrainingProcessFlags = ESubprocessFlags::None);
+
+		~FSharedMemoryTrainerServerProcess();
+
+		/** Check if the server process is still running. */
+		virtual bool IsRunning() const override final;
+
+		/**
+		* Wait for the server process to end
+		*
+		* @param Timeout		Timeout to wait in seconds
+		* @returns				true if successful, otherwise false if it times out
+		*/
+		virtual bool Wait() override final;
+
+		/** Terminate the server process. */
+		virtual void Terminate() override final;
+
+		/** Get the Controls shared memory array view. */
+		TSharedMemoryArrayView<2, volatile int32> GetControlsSharedMemoryArrayView() const;
+
+		/** Get the intermediate path. */
+		const FString& GetIntermediatePath() const;
+
+		/** Get the config path. */
+		const FString& GetConfigPath() const;
+
+		/** Get the training subprocess. */
+		virtual FSubprocess* GetTrainingSubprocess() override final;
+
+	private:
+
+		/** Free and deallocate all shared memory. */
+		void Deallocate();
+
+		FString IntermediatePath;
+		FString ConfigPath;
+
+		TSharedMemoryArrayView<2, volatile int32> Controls; // Mark as volatile to avoid compiler optimizing away reads without writes etc.
+
+		FSubprocess TrainingProcess;
+		float Timeout = Trainer::DefaultTimeout;
+	};
+
+	/**
+	* Trainer that connects to an external training server to perform training
+	*
+	* This trainer can be used to allow the python training process the run
+	* on a different machine to the experience gathering process.
+	*/
+	struct LEARNINGTRAINING_API FSharedMemoryTrainer : public IExternalTrainer
+	{
+		struct FSharedMemoryExperienceContainer
+		{
+			TSharedMemoryArrayView<2, int32> EpisodeStarts;
+			TSharedMemoryArrayView<2, int32> EpisodeLengths;
+			TSharedMemoryArrayView<2, ECompletionMode> EpisodeCompletionModes;
+			TSharedMemoryArrayView<3, float> EpisodeFinalObservations;
+			TSharedMemoryArrayView<3, float> EpisodeFinalMemoryStates;
+			TSharedMemoryArrayView<3, float> Observations;
+			TSharedMemoryArrayView<3, float> Actions;
+			TSharedMemoryArrayView<3, float> MemoryStates;
+			TSharedMemoryArrayView<2, float> Rewards;
+
+			/** Free and deallocate all shared memory. */
+			void Deallocate();
+		};
+
+		/**
+		* Creates a new SharedMemory trainer
+		*
+		* @param InTaskName				Unique name for this training task - used to avoid config filename conflicts
+		* @param ProcessNum				Number of processes to use for multi-processed experience gathering (used
+		*								to allocate enough shared memory)
+		* @param ExternalTrainerProcess	Shared memory used for communicating status to the trainer server process
+		* @param InTimeout				Timeout to wait in seconds for connection and initial data transfer
+		*/
+		FSharedMemoryTrainer(
+			const FString& InTaskName,
+			const int32 InProcessNum,
+			const TSharedPtr<UE::Learning::ITrainerProcess>& ExternalTrainerProcess,
+			const float InTimeout = Trainer::DefaultTimeout);
+
+		~FSharedMemoryTrainer();
+
+		virtual void Terminate() override final;
+
+		virtual ETrainerResponse SendStop() override final;
+
+		virtual ETrainerResponse Wait() override final;
+
+		virtual ETrainerResponse SendConfig(
+			const TSharedRef<FJsonObject>& ConfigObject,
+			const ELogSetting LogSettings = Trainer::DefaultLogSettings) override final;
+
+		virtual void AddNetwork(const FName& Name, const ULearningNeuralNetworkData& Network) override final;
+
+		virtual bool ContainsNetwork(const FName& Name) const override final;
+
+		virtual ETrainerResponse ReceiveNetwork(
+			const FName& Name,
+			ULearningNeuralNetworkData& OutNetwork,
+			FRWLock* NetworkLock = nullptr,
+			const ELogSetting LogSettings = Trainer::DefaultLogSettings) override final;
+
+		virtual ETrainerResponse SendNetwork(
+			const FName& Name,
+			const ULearningNeuralNetworkData& Network,
+			FRWLock* NetworkLock = nullptr,
+			const ELogSetting LogSettings = Trainer::DefaultLogSettings) override final;
+
+		virtual void AddReplayBuffer(const FName& Name,	const FReplayBuffer& ReplayBuffer) override final;
+
+		virtual bool ContainsReplayBuffer(const FName& Name) const override final;
+
+		virtual ETrainerResponse SendReplayBuffer(
+			const FName& Name,
+			const FReplayBuffer& ReplayBuffer,
+			const ELogSetting LogSettings = Trainer::DefaultLogSettings) override final;
+
+	private:
+
+		/** Free and deallocate all shared memory. */
+		void Deallocate();
+
+		FString TaskName;
+
+		FString IntermediatePath;
+
+		FString ConfigPath;
+
+		int32 ProcessNum;
+
+		FSubprocess* TrainingProcess = nullptr;
+
+		float Timeout = Trainer::DefaultTimeout;
+
+		int32 ProcessIdx = INDEX_NONE;
+		TSharedMemoryArrayView<2, volatile int32> Controls;
+		TMap<FName, TSharedMemoryArrayView<1, uint8>> NeuralNetworkSharedMemoryArrayViews;
+		TMap<FName, FSharedMemoryExperienceContainer> SharedMemoryExperienceContainers;
+	};
+
+	/**
+	* This object allows you to launch the FSocketTrainer server as a subprocess,
+	* which is convenient when you want to train using it locally.
+	*/
+	struct LEARNINGTRAINING_API FSocketTrainerServerProcess : public ITrainerProcess
+	{
+		/**
+		* Creates a training server as a subprocess
+		*
+		* @param PythonExecutablePath		Path to the python executable used for training. In general should be the python shipped with Unreal Editor.
+		* @param PythonContentPath			Path to the Python Content folder provided by the Learning plugin
+		* @param IntermediatePath			Path to the intermediate folder to write temporary files, logs, and snapshots to
+		* @param IpAddress					Ip address to bind the listening socket to. For a local server you will want to use 127.0.0.1
+		* @param Port						Port to use for the listening socket.
+		* @param TrainingProcessFlags		Training server subprocess flags
+		* @param LogSettings				Logging settings to use
+		*/
+		FSocketTrainerServerProcess(
+			const FString& PythonExecutablePath,
+			const FString& PythonContentPath,
+			const FString& IntermediatePath,
+			const TCHAR* IpAddress = Trainer::DefaultIp,
+			const uint32 Port = Trainer::DefaultPort,
+			const float InTimeout = Trainer::DefaultTimeout,
+			const ESubprocessFlags TrainingProcessFlags = ESubprocessFlags::None,
+			const ELogSetting LogSettings = ELogSetting::Normal);
+
+		~FSocketTrainerServerProcess();
+
+		/**
+		* Check if the server process is still running
+		*/
+		virtual bool IsRunning() const override final;
+
+		/**
+		* Wait for the server process to end
+		*
+		* @param Timeout		Timeout to wait in seconds
+		* @returns				true if successful, otherwise false if it times out
+		*/
+		virtual bool Wait() override final;
+
+		/**
+		* Terminate the server process
+		*/
+		virtual void Terminate() override final;
+
+		/** Get the training subprocess. */
+		virtual FSubprocess* GetTrainingSubprocess() override final;
+
+	private:
+
+		FSubprocess TrainingProcess;
+		float Timeout = Trainer::DefaultTimeout;
+	};
+
+	/**
+	* Trainer that connects to an external training server to perform training
+	*
+	* This trainer can be used to allow the python training process the run
+	* on a different machine to the experience gathering process.
+	*/
+	struct LEARNINGTRAINING_API FSocketTrainer : public IExternalTrainer
+	{
+		/**
+		* Creates a new Socket trainer
+		*
+		* @param OutResponse				Response to the initial connection
+		* @param ExternalTrainerProcess		The external trainer process
+		* @param IpAddress					Server Ip address
+		* @param Port						Server Port
+		* @param Timeout					Timeout to wait in seconds for connection and initial data transfer
+		*/
+		FSocketTrainer(
+			ETrainerResponse& OutResponse,
+			const TSharedPtr<UE::Learning::ITrainerProcess>& ExternalTrainerProcess,
+			const TCHAR* IpAddress = Trainer::DefaultIp,
+			const uint32 Port = Trainer::DefaultPort,
+			const float InTimeout = Trainer::DefaultTimeout);
+
+		~FSocketTrainer();
+
+		virtual void Terminate() override final;
+
+		virtual ETrainerResponse SendStop() override final;
+
+		virtual ETrainerResponse Wait() override final;
+
+		virtual ETrainerResponse SendConfig(
+			const TSharedRef<FJsonObject>& ConfigObject,
+			const ELogSetting LogSettings = Trainer::DefaultLogSettings) override final;
+
+		virtual void AddNetwork(const FName& Name, const ULearningNeuralNetworkData& Network) override final;
+
+		virtual bool ContainsNetwork(const FName& Name) const override final;
+
+		virtual ETrainerResponse ReceiveNetwork(
+			const FName& Name,
+			ULearningNeuralNetworkData& OutNetwork,
+			FRWLock* NetworkLock = nullptr,
+			const ELogSetting LogSettings = Trainer::DefaultLogSettings) override final;
+
+		virtual ETrainerResponse SendNetwork(
+			const FName& Name,
+			const ULearningNeuralNetworkData& Network,
+			FRWLock* NetworkLock = nullptr,
+			const ELogSetting LogSettings = Trainer::DefaultLogSettings) override final;
+
+		virtual void AddReplayBuffer(const FName& Name, const FReplayBuffer& ReplayBuffer) override final;
+
+		virtual bool ContainsReplayBuffer(const FName& Name) const override final;
+
+		virtual ETrainerResponse SendReplayBuffer(
+			const FName& Name,
+			const FReplayBuffer& ReplayBuffer,
+			const ELogSetting LogSettings = Trainer::DefaultLogSettings) override final;
+
+	private:
+
+		TMap<FName, TLearningArray<1, uint8>> NetworkBuffers;
+		TSet<FName> ExperienceBufferNames;
+
+		float Timeout = Trainer::DefaultTimeout;
+
+		FSubprocess* TrainingProcess = nullptr;
+
+		FSocket* Socket = nullptr;
+	};
+}
