Learning Agents Training Rewrite: Part 1

diff --git a/Engine/Plugins/Experimental/LearningAgents/Source/LearningTraining/Public/LearningPPOTrainer.h b/Engine/Plugins/Experimental/LearningAgents/Source/LearningTraining/Public/LearningPPOTrainer.h
--- a/Engine/Plugins/Experimental/LearningAgents/Source/LearningTraining/Public/LearningPPOTrainer.h
+++ b/Engine/Plugins/Experimental/LearningAgents/Source/LearningTraining/Public/LearningPPOTrainer.h
@@ -12,7 +12,6 @@
 
 #include "LearningPPOTrainer.generated.h"
 
-class FSocket;
 class ULearningNeuralNetworkData;
 
 UCLASS()
@@ -28,6 +27,8 @@ class LEARNINGTRAINING_API ULearningSocketPPOTrainerServerCommandlet : public UC
 
 namespace UE::Learning
 {
+	struct ITrainerProcess;
+	struct IExternalTrainer;
 	struct FReplayBuffer;
 	struct FResetInstanceBuffer;
 	struct FEpisodeBuffer;
@@ -149,502 +150,12 @@ namespace UE::Learning
 		bool bSaveSnapshots = false;
 	};
 
-	/**
-	* Interface for an object which can train a policy using PPO.
-	*/
-	struct IPPOTrainer
-	{
-		virtual ~IPPOTrainer() {}
-
-		/**
-		* Terminate the trainer immediately.
-		*/
-		virtual void Terminate() = 0;
-
-		/**
-		* Wait for the trainer to finish.
-		* 
-		* @param Timeout		Timeout to wait in seconds
-		* @returns				Trainer response
-		*/
-		virtual ETrainerResponse Wait(const float Timeout = Trainer::DefaultTimeout) = 0;
-
-		/**
-		* Wait for the trainer to push an updated policy network.
-		*
-		* @param OutNetwork		Network to update
-		* @param Timeout		Timeout to wait in seconds
-		* @param NetworkLock	Lock to use when updating network
-		* @param LogSettings	Log settings
-		* @returns				Trainer response
-		*/
-		virtual ETrainerResponse RecvPolicy(
-			ULearningNeuralNetworkData& OutNetwork,
-			const float Timeout = Trainer::DefaultTimeout,
-			FRWLock* NetworkLock = nullptr,
-			const ELogSetting LogSettings = Trainer::DefaultLogSettings) = 0;
-
-		/**
-		* Wait for the trainer to push an updated critic network.
-		*
-		* @param OutNetwork		Network to update
-		* @param Timeout		Timeout to wait in seconds
-		* @param NetworkLock	Lock to use when updating network
-		* @param LogSettings	Log settings
-		* @returns				Trainer response
-		*/
-		virtual ETrainerResponse RecvCritic(
-			ULearningNeuralNetworkData& OutNetwork,
-			const float Timeout = Trainer::DefaultTimeout,
-			FRWLock* NetworkLock = nullptr,
-			const ELogSetting LogSettings = Trainer::DefaultLogSettings) = 0;
-
-		/**
-		* Wait for the trainer to push an updated encoder network.
-		*
-		* @param OutNetwork		Network to update
-		* @param Timeout		Timeout to wait in seconds
-		* @param NetworkLock	Lock to use when updating network
-		* @param LogSettings	Log settings
-		* @returns				Trainer response
-		*/
-		virtual ETrainerResponse RecvEncoder(
-			ULearningNeuralNetworkData& OutNetwork,
-			const float Timeout = Trainer::DefaultTimeout,
-			FRWLock* NetworkLock = nullptr,
-			const ELogSetting LogSettings = Trainer::DefaultLogSettings) = 0;
-
-		/**
-		* Wait for the trainer to push an updated decoder network.
-		*
-		* @param OutNetwork		Network to update
-		* @param Timeout		Timeout to wait in seconds
-		* @param NetworkLock	Lock to use when updating network
-		* @param LogSettings	Log settings
-		* @returns				Trainer response
-		*/
-		virtual ETrainerResponse RecvDecoder(
-			ULearningNeuralNetworkData& OutNetwork,
-			const float Timeout = Trainer::DefaultTimeout,
-			FRWLock* NetworkLock = nullptr,
-			const ELogSetting LogSettings = Trainer::DefaultLogSettings) = 0;
-
-		/*
-		* Signal for the trainer to stop.
-		*/
-		virtual ETrainerResponse SendStop(const float Timeout = Trainer::DefaultTimeout) = 0;
-
-		/**
-		* Wait for the trainer to be ready and push the current policy network.
-		*
-		* @param Network		Network to push
-		* @param Timeout		Timeout to wait in seconds
-		* @param NetworkLock	Lock to use when pushing network
-		* @param LogSettings	Log settings
-		* @returns				Trainer response
-		*/
-		virtual ETrainerResponse SendPolicy(
-			const ULearningNeuralNetworkData& Network,
-			const float Timeout = Trainer::DefaultTimeout,
-			FRWLock* NetworkLock = nullptr,
-			const ELogSetting LogSettings = Trainer::DefaultLogSettings) = 0;
-
-		/**
-		* Wait for the trainer to be ready and push the current critic network.
-		*
-		* @param Network		Network to push
-		* @param Timeout		Timeout to wait in seconds
-		* @param NetworkLock	Lock to use when pushing network
-		* @param LogSettings	Log settings
-		* @returns				Trainer response
-		*/
-		virtual ETrainerResponse SendCritic(
-			const ULearningNeuralNetworkData& Network,
-			const float Timeout = Trainer::DefaultTimeout,
-			FRWLock* NetworkLock = nullptr,
-			const ELogSetting LogSettings = Trainer::DefaultLogSettings) = 0;
-
-		/**
-		* Wait for the trainer to be ready and push the current encoder network.
-		*
-		* @param Network		Network to push
-		* @param Timeout		Timeout to wait in seconds
-		* @param NetworkLock	Lock to use when pushing network
-		* @param LogSettings	Log settings
-		* @returns				Trainer response
-		*/
-		virtual ETrainerResponse SendEncoder(
-			const ULearningNeuralNetworkData& Network,
-			const float Timeout = Trainer::DefaultTimeout,
-			FRWLock* NetworkLock = nullptr,
-			const ELogSetting LogSettings = Trainer::DefaultLogSettings) = 0;
-
-		/**
-		* Wait for the trainer to be ready and push the current decoder network.
-		*
-		* @param Network		Network to push
-		* @param Timeout		Timeout to wait in seconds
-		* @param NetworkLock	Lock to use when pushing network
-		* @param LogSettings	Log settings
-		* @returns				Trainer response
-		*/
-		virtual ETrainerResponse SendDecoder(
-			const ULearningNeuralNetworkData& Network,
-			const float Timeout = Trainer::DefaultTimeout,
-			FRWLock* NetworkLock = nullptr,
-			const ELogSetting LogSettings = Trainer::DefaultLogSettings) = 0;
-
-		/**
-		* Wait for the trainer to be ready and push new experience.
-		*
-		* @param ReplayBuffer	Replay buffer of experience to push
-		* @param Timeout		Timeout to wait in seconds
-		* @param LogSettings	Log settings
-		* @returns				Trainer response
-		*/
-		virtual ETrainerResponse SendExperience(
-			const FReplayBuffer& ReplayBuffer,
-			const float Timeout = Trainer::DefaultTimeout,
-			const ELogSetting LogSettings = Trainer::DefaultLogSettings) = 0;
-	};
-
-	/**
-	* Trainer that uses shared memory and a Python sub-process to perform training
-	* 
-	* This trainer is the most simple and efficient when training the policy on the 
-	* same computer that experience is being gathered on.
-	*/
-	struct LEARNINGTRAINING_API FSharedMemoryPPOTrainer : public IPPOTrainer
-	{
-		/**
-		* Creates a new Shared Memory PPO trainer
-		*
-		* @param TaskName					Name of the training task - used to help identify the logs, snapshots, and other files generated by training
-		* @param PythonExecutablePath		Path to the python executable used for training. In general should be the python shipped with Unreal Editor.
-		* @param ExtraSitePackagesPath		Path to additional site-packages if required otherwise empty string
-		* @param PythonContentPath			Path to the Python Content folder provided by the Learning plugin
-		* @param IntermediatePath			Path to the intermediate folder to write temporary files, logs, and snapshots to
-		* @param ReplayBuffer				Replay buffer used to collect experience
-		* @param PolicyNetwork				Policy Network to use
-		* @param CriticNetwork				Critic Network to use
-		* @param EncoderNetwork				Encoder Network to use
-		* @param DecoderNetwork				Decoder Network to use
-		* @param ObservationSchema			Schema used for Observations
-		* @param ObservationSchemaElement	Schema Observation Element
-		* @param ActionSchema				Schema used for Actions
-		* @param ActionSchemaElement		Schema Action Element
-		* @param TrainingSettings			Trainer Training settings
-		* @param LogSettings				Logging settings to use
-		* @param TrainingProcessFlags		Training subprocess flags
-
-		* @param ProcessNum					Number of processes to use for multi-processed experience gathering
-		*                                   
-		*                                   It is important to know how this multi-process training works so that it can be used
-		*                                   correctly when you set this >1:
-		*                                   
-		*                                   When called with this argument set >1, the process will spawn additional processes running
-		*                                   the same command as is currently being run but with the additional command line argument
-		*                                   `LearningProcessIdx`. Once the training starts on these sub-processes, this command
-		*                                   line argument will be used to indicate that this subprocess should only be used for gathering
-		*                                   experience and should not start a new training process itself.
-		*                                   
-		*                                   This means that this will generally not work in editor (or PIE) since it will end up spawning
-		*									Multiple additional editor windows. Instead, this should generally only be used in cooked, 
-		*									builds - and in most cases only when running headless.
-		*                                   
-		*                                   IMPORTANT: If you are seeding your experience generation process to make it deterministic
-		*                                   you should use this `LearningProcessIdx` command line argument to change the seed of your
-		*                                   experience gathering - otherwise each process will end up gathering identical experience!
-		* 
-		* @param MultiProcessFlags			Flags for the additional experience gathering subprocesses
-		*/
-		FSharedMemoryPPOTrainer(
-			const FString& TaskName,
-			const FString& PythonExecutablePath,
-			const FString& ExtraSitePackagesPath,
-			const FString& PythonContentPath,
-			const FString& IntermediatePath,
-			const FReplayBuffer& ReplayBuffer,
-			const ULearningNeuralNetworkData& PolicyNetwork,
-			const ULearningNeuralNetworkData& CriticNetwork,
-			const ULearningNeuralNetworkData& EncoderNetwork,
-			const ULearningNeuralNetworkData& DecoderNetwork,
-			const Observation::FSchema& ObservationSchema,
-			const Observation::FSchemaElement& ObservationSchemaElement,
-			const Action::FSchema& ActionSchema,
-			const Action::FSchemaElement& ActionSchemaElement,
-			const FPPOTrainerTrainingSettings& TrainingSettings = FPPOTrainerTrainingSettings(),
-			const ELogSetting LogSettings = ELogSetting::Normal,
-			const ESubprocessFlags TrainingProcessFlags = ESubprocessFlags::None,
-			const uint16 ProcessNum = 1,
-			const ESubprocessFlags MultiProcessFlags = ESubprocessFlags::None);
-
-		~FSharedMemoryPPOTrainer();
-
-		virtual void Terminate() override final;
-
-		virtual ETrainerResponse Wait(const float Timeout = Trainer::DefaultTimeout) override final;
-
-		virtual ETrainerResponse RecvPolicy(
-			ULearningNeuralNetworkData& OutNetwork,
-			const float Timeout = Trainer::DefaultTimeout,
-			FRWLock* NetworkLock = nullptr,
-			const ELogSetting LogSettings = Trainer::DefaultLogSettings) override final;
-
-		virtual ETrainerResponse RecvCritic(
-			ULearningNeuralNetworkData& OutNetwork,
-			const float Timeout = Trainer::DefaultTimeout,
-			FRWLock* NetworkLock = nullptr,
-			const ELogSetting LogSettings = Trainer::DefaultLogSettings) override final;
-
-		virtual ETrainerResponse RecvEncoder(
-			ULearningNeuralNetworkData& OutNetwork,
-			const float Timeout = Trainer::DefaultTimeout,
-			FRWLock* NetworkLock = nullptr,
-			const ELogSetting LogSettings = Trainer::DefaultLogSettings) override final;
-
-		virtual ETrainerResponse RecvDecoder(
-			ULearningNeuralNetworkData& OutNetwork,
-			const float Timeout = Trainer::DefaultTimeout,
-			FRWLock* NetworkLock = nullptr,
-			const ELogSetting LogSettings = Trainer::DefaultLogSettings) override final;
-
-		virtual ETrainerResponse SendStop(const float Timeout = Trainer::DefaultTimeout) override final;
-
-		virtual ETrainerResponse SendPolicy(
-			const ULearningNeuralNetworkData& Network,
-			const float Timeout = Trainer::DefaultTimeout,
-			FRWLock* NetworkLock = nullptr,
-			const ELogSetting LogSettings = Trainer::DefaultLogSettings) override final;
-
-		virtual ETrainerResponse SendCritic(
-			const ULearningNeuralNetworkData& Network,
-			const float Timeout = Trainer::DefaultTimeout,
-			FRWLock* NetworkLock = nullptr,
-			const ELogSetting LogSettings = Trainer::DefaultLogSettings) override final;
-
-		virtual ETrainerResponse SendEncoder(
-			const ULearningNeuralNetworkData& Network,
-			const float Timeout = Trainer::DefaultTimeout,
-			FRWLock* NetworkLock = nullptr,
-			const ELogSetting LogSettings = Trainer::DefaultLogSettings) override final;
-
-		virtual ETrainerResponse SendDecoder(
-			const ULearningNeuralNetworkData& Network,
-			const float Timeout = Trainer::DefaultTimeout,
-			FRWLock* NetworkLock = nullptr,
-			const ELogSetting LogSettings = Trainer::DefaultLogSettings) override final;
-
-		virtual ETrainerResponse SendExperience(
-			const FReplayBuffer& ReplayBuffer,
-			const float Timeout = Trainer::DefaultTimeout,
-			const ELogSetting LogSettings = Trainer::DefaultLogSettings) override final;
-
-	private:
-
-		/**
-		* Free and deallocate all shared memory
-		*/
-		void Deallocate();
-
-		// Shared memory
-
-		TSharedMemoryArrayView<1, uint8> Policy;
-		TSharedMemoryArrayView<1, uint8> Critic;
-		TSharedMemoryArrayView<1, uint8> Encoder;
-		TSharedMemoryArrayView<1, uint8> Decoder;
-		TSharedMemoryArrayView<2, volatile int32> Controls; // Mark as volatile to avoid compiler optimizing away reads without writes etc.
-		TSharedMemoryArrayView<2, int32> EpisodeStarts;
-		TSharedMemoryArrayView<2, int32> EpisodeLengths;
-		TSharedMemoryArrayView<2, ECompletionMode> EpisodeCompletionModes;
-		TSharedMemoryArrayView<3, float> EpisodeFinalObservations;
-		TSharedMemoryArrayView<3, float> EpisodeFinalMemoryStates;
-		TSharedMemoryArrayView<3, float> Observations;
-		TSharedMemoryArrayView<3, float> Actions;
-		TSharedMemoryArrayView<3, float> MemoryStates;
-		TSharedMemoryArrayView<2, float> Rewards;
-
-		// Training Process
-
-		uint16 ProcessIdx = INDEX_NONE;
-		FSubprocess TrainingProcess;
-		TArray<FSubprocess, TInlineAllocator<128>> ExperienceGatheringSubprocesses;
-	};
-
-	/**
-	* This object allows you to launch the FSocketPPOTrainer server as a subprocess, 
-	* which is convenient when you want to train using it locally.
-	*/
-	struct LEARNINGTRAINING_API FSocketPPOTrainerServerProcess
-	{
-		/**
-		* Creates a training server as a subprocess
-		*
-		* @param PythonExecutablePath		Path to the python executable used for training. In general should be the python shipped with Unreal Editor.
-		* @param ExtraSitePackagesPath		Path to additional site-packages if required otherwise empty string
-		* @param PythonContentPath			Path to the Python Content folder provided by the Learning plugin
-		* @param IntermediatePath			Path to the intermediate folder to write temporary files, logs, and snapshots to
-		* @param IpAddress					Ip address to bind the listening socket to. For a local server you will want to use 127.0.0.1
-		* @param Port						Port to use for the listening socket.
-		* @param TrainingProcessFlags		Training server subprocess flags
-		* @param LogSettings				Logging settings to use
-		*/
-		FSocketPPOTrainerServerProcess(
-			const FString& PythonExecutablePath,
-			const FString& ExtraSitePackagesPath,
-			const FString& PythonContentPath,
-			const FString& IntermediatePath,
-			const TCHAR* IpAddress = Trainer::DefaultIp,
-			const uint32 Port = Trainer::DefaultPort,
-			const ESubprocessFlags TrainingProcessFlags = ESubprocessFlags::None,
-			const ELogSetting LogSettings = ELogSetting::Normal);
-
-		~FSocketPPOTrainerServerProcess();
-
-		/**
-		* Check if the server process is still running
-		*/
-		bool IsRunning() const;
-
-		/**
-		* Wait for the server process to end
-		*
-		* @param Timeout		Timeout to wait in seconds
-		* @returns				true if successful, otherwise false if it times out
-		*/
-		bool Wait(float Timeout = Trainer::DefaultTimeout);
-
-		/**
-		* Terminate the server process
-		*/
-		void Terminate();
-
-	private:
-
-		FSubprocess TrainingProcess;
-	};
-
-	/**
-	* Trainer that connects to an external training server to perform training
-	*
-	* This trainer can be used to allow the python training process the run
-	* on a different machine to the experience gathering process.
-	*/
-	struct LEARNINGTRAINING_API FSocketPPOTrainer : public IPPOTrainer
-	{
-		/**
-		* Creates a new Socket PPO trainer
-		*
-		* @param OutResponse				Response to the initial connection
-		* @param TaskName					Name of the training task - used to help identify the logs, snapshots, and other files generated by training
-		* @param ReplayBuffer				Replay buffer used to collect experience
-		* @param PolicyNetwork				Policy Network to use
-		* @param CriticNetwork				Critic Network to use
-		* @param EncoderNetwork				Encoder Network to use
-		* @param DecoderNetwork				Decoder Network to use
-		* @param ObservationSchema			Schema used for Observations
-		* @param ObservationSchemaElement	Schema Observation Element
-		* @param ActionSchema				Schema used for Actions
-		* @param ActionSchemaElement		Schema Action Element
-		* @param IpAddress					Server Ip address
-		* @param Port						Server Port
-		* @param Timeout					Timeout to wait in seconds for connection and initial data transfer
-		* @param TrainingSettings			Trainer Training settings
-		*/
-		FSocketPPOTrainer(
-			ETrainerResponse& OutResponse,
-			const FString& TaskName,
-			const FReplayBuffer& ReplayBuffer,
-			const ULearningNeuralNetworkData& PolicyNetwork,
-			const ULearningNeuralNetworkData& CriticNetwork,
-			const ULearningNeuralNetworkData& EncoderNetwork,
-			const ULearningNeuralNetworkData& DecoderNetwork,
-			const Observation::FSchema& ObservationSchema,
-			const Observation::FSchemaElement& ObservationSchemaElement,
-			const Action::FSchema& ActionSchema,
-			const Action::FSchemaElement& ActionSchemaElement,
-			const TCHAR* IpAddress = Trainer::DefaultIp,
-			const uint32 Port = Trainer::DefaultPort,
-			const float Timeout = Trainer::DefaultTimeout,
-			const FPPOTrainerTrainingSettings& TrainingSettings = FPPOTrainerTrainingSettings());
-
-		~FSocketPPOTrainer();
-
-		virtual void Terminate() override final;
-
-		virtual ETrainerResponse Wait(const float Timeout = Trainer::DefaultTimeout) override final;
-
-		virtual ETrainerResponse RecvPolicy(
-			ULearningNeuralNetworkData& OutNetwork,
-			const float Timeout = Trainer::DefaultTimeout,
-			FRWLock* NetworkLock = nullptr,
-			const ELogSetting LogSettings = Trainer::DefaultLogSettings) override final;
-
-		virtual ETrainerResponse RecvCritic(
-			ULearningNeuralNetworkData& OutNetwork,
-			const float Timeout = Trainer::DefaultTimeout,
-			FRWLock* NetworkLock = nullptr,
-			const ELogSetting LogSettings = Trainer::DefaultLogSettings) override final;
-
-		virtual ETrainerResponse RecvEncoder(
-			ULearningNeuralNetworkData& OutNetwork,
-			const float Timeout = Trainer::DefaultTimeout,
-			FRWLock* NetworkLock = nullptr,
-			const ELogSetting LogSettings = Trainer::DefaultLogSettings) override final;
-
-		virtual ETrainerResponse RecvDecoder(
-			ULearningNeuralNetworkData& OutNetwork,
-			const float Timeout = Trainer::DefaultTimeout,
-			FRWLock* NetworkLock = nullptr,
-			const ELogSetting LogSettings = Trainer::DefaultLogSettings) override final;
-
-		virtual ETrainerResponse SendStop(const float Timeout = Trainer::DefaultTimeout) override final;
-
-		virtual ETrainerResponse SendPolicy(
-			const ULearningNeuralNetworkData& Network,
-			const float Timeout = Trainer::DefaultTimeout,
-			FRWLock* NetworkLock = nullptr,
-			const ELogSetting LogSettings = Trainer::DefaultLogSettings) override final;
-
-		virtual ETrainerResponse SendCritic(
-			const ULearningNeuralNetworkData& Network,
-			const float Timeout = Trainer::DefaultTimeout,
-			FRWLock* NetworkLock = nullptr,
-			const ELogSetting LogSettings = Trainer::DefaultLogSettings) override final;
-
-		virtual ETrainerResponse SendEncoder(
-			const ULearningNeuralNetworkData& Network,
-			const float Timeout = Trainer::DefaultTimeout,
-			FRWLock* NetworkLock = nullptr,
-			const ELogSetting LogSettings = Trainer::DefaultLogSettings) override final;
-
-		virtual ETrainerResponse SendDecoder(
-			const ULearningNeuralNetworkData& Network,
-			const float Timeout = Trainer::DefaultTimeout,
-			FRWLock* NetworkLock = nullptr,
-			const ELogSetting LogSettings = Trainer::DefaultLogSettings) override final;
-
-		virtual ETrainerResponse SendExperience(
-			const FReplayBuffer& ReplayBuffer,
-			const float Timeout = Trainer::DefaultTimeout,
-			const ELogSetting LogSettings = Trainer::DefaultLogSettings) override final;
-
-	private:
-
-		TLearningArray<1, uint8> PolicyNetworkBuffer;
-		TLearningArray<1, uint8> CriticNetworkBuffer;
-		TLearningArray<1, uint8> EncoderNetworkBuffer;
-		TLearningArray<1, uint8> DecoderNetworkBuffer;
-		FSocket* Socket = nullptr;
-	};
-
-
 	namespace PPOTrainer
 	{
 		/**
 		* Train a policy while gathering experience
 		*
-		* @param Trainer								Trainer
+		* @param ExternalTrainer						External Trainer
 		* @param ReplayBuffer							Replay Buffer
 		* @param EpisodeBuffer							Episode Buffer
 		* @param ResetBuffer							Reset Buffer
@@ -681,7 +192,7 @@ namespace UE::Learning
 		* @returns										Trainer response in case of errors during communication otherwise Success
 		*/
 		LEARNINGTRAINING_API ETrainerResponse Train(
-			IPPOTrainer& Trainer,
+			IExternalTrainer* ExternalTrainer,
 			FReplayBuffer& ReplayBuffer,
 			FEpisodeBuffer& EpisodeBuffer,
 			FResetInstanceBuffer& ResetBuffer,
@@ -705,6 +216,11 @@ namespace UE::Learning
 			const TFunctionRef<void(const FIndexSet Instances)> RewardFunction,
 			const TFunctionRef<void(const FIndexSet Instances)> CompletionFunction,
 			const FIndexSet Instances,
+			const Observation::FSchema& ObservationSchema,
+			const Observation::FSchemaElement& ObservationSchemaElement,
+			const Action::FSchema& ActionSchema,
+			const Action::FSchemaElement& ActionSchemaElement,
+			const FPPOTrainerTrainingSettings& TrainerSettings = FPPOTrainerTrainingSettings(),
 			TAtomic<bool>* bRequestTrainingStopSignal = nullptr,
 			FRWLock* PolicyNetworkLock = nullptr,
 			FRWLock* CriticNetworkLock = nullptr,
