Made some parameters of some blueprint callables mandatory in bps
- Add UPARAM(ref) to parameters of the type ULearningAgentsManager and classes derived from ULearningAgentsManagerListener

#tests none
#rnx
#virtualized

diff --git a/Engine/Plugins/Experimental/LearningAgents/Source/LearningAgents/Private/LearningAgentsController.cpp b/Engine/Plugins/Experimental/LearningAgents/Source/LearningAgents/Private/LearningAgentsController.cpp
--- a/Engine/Plugins/Experimental/LearningAgents/Source/LearningAgents/Private/LearningAgentsController.cpp
+++ b/Engine/Plugins/Experimental/LearningAgents/Source/LearningAgents/Private/LearningAgentsController.cpp
@@ -47,8 +47,8 @@ void ULearningAgentsController::EvaluateAgentControllers_Implementation(
 }
 
 ULearningAgentsController* ULearningAgentsController::MakeController(
-	ULearningAgentsManager* InManager, 
-	ULearningAgentsInteractor* InInteractor, 
+	ULearningAgentsManager*& InManager, 
+	ULearningAgentsInteractor*& InInteractor, 
 	TSubclassOf<ULearningAgentsController> Class,
 	const FName Name)
 {
@@ -74,7 +74,7 @@ ULearningAgentsController* ULearningAgentsController::MakeController(
 	return Controller->IsSetup() ? Controller : nullptr;
 }
 
-void ULearningAgentsController::SetupController(ULearningAgentsManager* InManager, ULearningAgentsInteractor* InInteractor)
+void ULearningAgentsController::SetupController(ULearningAgentsManager*& InManager, ULearningAgentsInteractor*& InInteractor)
 {
 	if (IsSetup())
 	{
diff --git a/Engine/Plugins/Experimental/LearningAgents/Source/LearningAgents/Private/LearningAgentsCritic.cpp b/Engine/Plugins/Experimental/LearningAgents/Source/LearningAgents/Private/LearningAgentsCritic.cpp
--- a/Engine/Plugins/Experimental/LearningAgents/Source/LearningAgents/Private/LearningAgentsCritic.cpp
+++ b/Engine/Plugins/Experimental/LearningAgents/Source/LearningAgents/Private/LearningAgentsCritic.cpp
@@ -38,9 +38,9 @@ ULearningAgentsCritic::ULearningAgentsCritic(FVTableHelper& Helper) : Super(Help
 ULearningAgentsCritic::~ULearningAgentsCritic() = default;
 
 ULearningAgentsCritic* ULearningAgentsCritic::MakeCritic(
-	ULearningAgentsManager* InManager,
-	ULearningAgentsInteractor* InInteractor,
-	ULearningAgentsPolicy* InPolicy,
+	ULearningAgentsManager*& InManager,
+	ULearningAgentsInteractor*& InInteractor,
+	ULearningAgentsPolicy*& InPolicy,
 	TSubclassOf<ULearningAgentsCritic> Class,
 	const FName Name,
 	ULearningAgentsNeuralNetwork* CriticNeuralNetworkAsset,
@@ -78,9 +78,9 @@ ULearningAgentsCritic* ULearningAgentsCritic::MakeCritic(
 }
 
 void ULearningAgentsCritic::SetupCritic(
-	ULearningAgentsManager* InManager,
-	ULearningAgentsInteractor* InInteractor, 
-	ULearningAgentsPolicy* InPolicy,
+	ULearningAgentsManager*& InManager,
+	ULearningAgentsInteractor*& InInteractor, 
+	ULearningAgentsPolicy*& InPolicy,
 	ULearningAgentsNeuralNetwork* CriticNeuralNetworkAsset,
 	const bool bReinitializeCriticNetwork,
 	const FLearningAgentsCriticSettings& CriticSettings,
diff --git a/Engine/Plugins/Experimental/LearningAgents/Source/LearningAgents/Private/LearningAgentsInteractor.cpp b/Engine/Plugins/Experimental/LearningAgents/Source/LearningAgents/Private/LearningAgentsInteractor.cpp
--- a/Engine/Plugins/Experimental/LearningAgents/Source/LearningAgents/Private/LearningAgentsInteractor.cpp
+++ b/Engine/Plugins/Experimental/LearningAgents/Source/LearningAgents/Private/LearningAgentsInteractor.cpp
@@ -18,7 +18,7 @@ ULearningAgentsInteractor::ULearningAgentsInteractor(FVTableHelper& Helper) : Su
 ULearningAgentsInteractor::~ULearningAgentsInteractor() = default;
 
 ULearningAgentsInteractor* ULearningAgentsInteractor::MakeInteractor(
-	ULearningAgentsManager* InManager,
+	ULearningAgentsManager*& InManager,
 	TSubclassOf<ULearningAgentsInteractor> Class,
 	const FName Name)
 {
@@ -44,7 +44,7 @@ ULearningAgentsInteractor* ULearningAgentsInteractor::MakeInteractor(
 	return Interactor->IsSetup() ? Interactor : nullptr;
 }
 
-void ULearningAgentsInteractor::SetupInteractor(ULearningAgentsManager* InManager)
+void ULearningAgentsInteractor::SetupInteractor(ULearningAgentsManager*& InManager)
 {
 	if (IsSetup())
 	{
diff --git a/Engine/Plugins/Experimental/LearningAgents/Source/LearningAgents/Private/LearningAgentsPolicy.cpp b/Engine/Plugins/Experimental/LearningAgents/Source/LearningAgents/Private/LearningAgentsPolicy.cpp
--- a/Engine/Plugins/Experimental/LearningAgents/Source/LearningAgents/Private/LearningAgentsPolicy.cpp
+++ b/Engine/Plugins/Experimental/LearningAgents/Source/LearningAgents/Private/LearningAgentsPolicy.cpp
@@ -37,8 +37,8 @@ ULearningAgentsPolicy::ULearningAgentsPolicy(FVTableHelper& Helper) : Super(Help
 ULearningAgentsPolicy::~ULearningAgentsPolicy() = default;
 
 ULearningAgentsPolicy* ULearningAgentsPolicy::MakePolicy(
-	ULearningAgentsManager* InManager,
-	ULearningAgentsInteractor* InInteractor,
+	ULearningAgentsManager*& InManager,
+	ULearningAgentsInteractor*& InInteractor,
 	TSubclassOf<ULearningAgentsPolicy> Class,
 	const FName Name,
 	ULearningAgentsNeuralNetwork* EncoderNeuralNetworkAsset,
@@ -83,8 +83,8 @@ ULearningAgentsPolicy* ULearningAgentsPolicy::MakePolicy(
 }
 
 void ULearningAgentsPolicy::SetupPolicy(
-	ULearningAgentsManager* InManager,
-	ULearningAgentsInteractor* InInteractor,
+	ULearningAgentsManager*& InManager,
+	ULearningAgentsInteractor*& InInteractor,
 	ULearningAgentsNeuralNetwork* EncoderNeuralNetworkAsset,
 	ULearningAgentsNeuralNetwork* PolicyNeuralNetworkAsset,
 	ULearningAgentsNeuralNetwork* DecoderNeuralNetworkAsset,
diff --git a/Engine/Plugins/Experimental/LearningAgents/Source/LearningAgents/Public/LearningAgentsController.h b/Engine/Plugins/Experimental/LearningAgents/Source/LearningAgents/Public/LearningAgentsController.h
--- a/Engine/Plugins/Experimental/LearningAgents/Source/LearningAgents/Public/LearningAgentsController.h
+++ b/Engine/Plugins/Experimental/LearningAgents/Source/LearningAgents/Public/LearningAgentsController.h
@@ -37,14 +37,14 @@ public:
 	 */
 	UFUNCTION(BlueprintCallable, Category = "LearningAgents", meta = (DeterminesOutputType = "Class"))
 	static ULearningAgentsController* MakeController(
-		ULearningAgentsManager* InManager, 
-		ULearningAgentsInteractor* InInteractor, 
+		UPARAM(ref) ULearningAgentsManager*& InManager,
+		UPARAM(ref) ULearningAgentsInteractor*& InInteractor, 
 		TSubclassOf<ULearningAgentsController> Class,
 		const FName Name = TEXT("Controller"));
 
 	/** Initializes this object to be used with the given agent interactor. */
 	UFUNCTION(BlueprintCallable, Category = "LearningAgents")
-	void SetupController(ULearningAgentsManager* InManager, ULearningAgentsInteractor* InInteractor);
+	void SetupController(UPARAM(ref) ULearningAgentsManager*& InManager, UPARAM(ref) ULearningAgentsInteractor*& InInteractor);
 
 	/**
 	 * This callback should be overridden by the Controller and produces an Action Object Element from an Observation Object Element.
diff --git a/Engine/Plugins/Experimental/LearningAgents/Source/LearningAgents/Public/LearningAgentsCritic.h b/Engine/Plugins/Experimental/LearningAgents/Source/LearningAgents/Public/LearningAgentsCritic.h
--- a/Engine/Plugins/Experimental/LearningAgents/Source/LearningAgents/Public/LearningAgentsCritic.h
+++ b/Engine/Plugins/Experimental/LearningAgents/Source/LearningAgents/Public/LearningAgentsCritic.h
@@ -70,9 +70,9 @@ public:
 	 */
 	UFUNCTION(BlueprintCallable, Category = "LearningAgents", meta = (Class = "/Script/LearningAgents.LearningAgentsCritic", DeterminesOutputType = "Class", AutoCreateRefTerm = "CriticSettings"))
 	static ULearningAgentsCritic* MakeCritic(
-		ULearningAgentsManager* InManager,
-		ULearningAgentsInteractor* InInteractor,
-		ULearningAgentsPolicy* InPolicy,
+		UPARAM(ref) ULearningAgentsManager*& InManager,
+		UPARAM(ref) ULearningAgentsInteractor*& InInteractor,
+		UPARAM(ref) ULearningAgentsPolicy*& InPolicy,
 		TSubclassOf<ULearningAgentsCritic> Class,
 		const FName Name = TEXT("Critic"),
 		ULearningAgentsNeuralNetwork* CriticNeuralNetworkAsset = nullptr,
@@ -95,9 +95,9 @@ public:
 	 */
 	UFUNCTION(BlueprintCallable, Category = "LearningAgents", meta = (AutoCreateRefTerm = "CriticSettings"))
 	void SetupCritic(
-		ULearningAgentsManager* InManager,
-		ULearningAgentsInteractor* InInteractor,
-		ULearningAgentsPolicy* InPolicy,
+		UPARAM(ref) ULearningAgentsManager*& InManager,
+		UPARAM(ref) ULearningAgentsInteractor*& InInteractor,
+		UPARAM(ref) ULearningAgentsPolicy*& InPolicy,
 		ULearningAgentsNeuralNetwork* CriticNeuralNetworkAsset = nullptr,
 		const bool bReinitializeCriticNetwork = true,
 		const FLearningAgentsCriticSettings& CriticSettings = FLearningAgentsCriticSettings(),
diff --git a/Engine/Plugins/Experimental/LearningAgents/Source/LearningAgents/Public/LearningAgentsInteractor.h b/Engine/Plugins/Experimental/LearningAgents/Source/LearningAgents/Public/LearningAgentsInteractor.h
--- a/Engine/Plugins/Experimental/LearningAgents/Source/LearningAgents/Public/LearningAgentsInteractor.h
+++ b/Engine/Plugins/Experimental/LearningAgents/Source/LearningAgents/Public/LearningAgentsInteractor.h
@@ -52,7 +52,7 @@ public:
 	 */
 	UFUNCTION(BlueprintCallable, Category = "LearningAgents", meta = (DeterminesOutputType = "Class"))
 	static ULearningAgentsInteractor* MakeInteractor(
-		ULearningAgentsManager* InManager, 
+		UPARAM(ref) ULearningAgentsManager*& InManager,
 		TSubclassOf<ULearningAgentsInteractor> Class,
 		const FName Name = TEXT("Interactor"));
 
@@ -62,7 +62,7 @@ public:
 	 * @param InManager						The input Manager
 	 */
 	UFUNCTION(BlueprintCallable, Category = "LearningAgents")
-	void SetupInteractor(ULearningAgentsManager* InManager);
+	void SetupInteractor(UPARAM(ref) ULearningAgentsManager*& InManager);
 
 public:
 
diff --git a/Engine/Plugins/Experimental/LearningAgents/Source/LearningAgents/Public/LearningAgentsPolicy.h b/Engine/Plugins/Experimental/LearningAgents/Source/LearningAgents/Public/LearningAgentsPolicy.h
--- a/Engine/Plugins/Experimental/LearningAgents/Source/LearningAgents/Public/LearningAgentsPolicy.h
+++ b/Engine/Plugins/Experimental/LearningAgents/Source/LearningAgents/Public/LearningAgentsPolicy.h
@@ -91,8 +91,8 @@ public:
 	 */
 	UFUNCTION(BlueprintCallable, Category = "LearningAgents", meta=(Class = "/Script/LearningAgents.LearningAgentsPolicy", DeterminesOutputType = "Class", AutoCreateRefTerm = "PolicySettings"))
 	static ULearningAgentsPolicy* MakePolicy(
-		ULearningAgentsManager* InManager,
-		ULearningAgentsInteractor* InInteractor,
+		UPARAM(ref) ULearningAgentsManager*& InManager,
+		UPARAM(ref) ULearningAgentsInteractor*& InInteractor,
 		TSubclassOf<ULearningAgentsPolicy> Class,
 		const FName Name = TEXT("Policy"),
 		ULearningAgentsNeuralNetwork* EncoderNeuralNetworkAsset = nullptr,
@@ -126,8 +126,8 @@ public:
 	 */
 	UFUNCTION(BlueprintCallable, Category = "LearningAgents", meta = (AutoCreateRefTerm = "PolicySettings"))
 	void SetupPolicy(
-		ULearningAgentsManager* InManager,
-		ULearningAgentsInteractor* InInteractor,
+		UPARAM(ref) ULearningAgentsManager*& InManager,
+		UPARAM(ref) ULearningAgentsInteractor*& InInteractor,
 		ULearningAgentsNeuralNetwork* EncoderNeuralNetworkAsset = nullptr,
 		ULearningAgentsNeuralNetwork* PolicyNeuralNetworkAsset = nullptr,
 		ULearningAgentsNeuralNetwork* DecoderNeuralNetworkAsset = nullptr,
diff --git a/Engine/Plugins/Experimental/LearningAgents/Source/LearningAgentsTraining/Private/LearningAgentsImitationTrainer.cpp b/Engine/Plugins/Experimental/LearningAgents/Source/LearningAgentsTraining/Private/LearningAgentsImitationTrainer.cpp
--- a/Engine/Plugins/Experimental/LearningAgents/Source/LearningAgentsTraining/Private/LearningAgentsImitationTrainer.cpp
+++ b/Engine/Plugins/Experimental/LearningAgents/Source/LearningAgentsTraining/Private/LearningAgentsImitationTrainer.cpp
@@ -29,9 +29,9 @@ void ULearningAgentsImitationTrainer::BeginDestroy()
 }
 
 ULearningAgentsImitationTrainer* ULearningAgentsImitationTrainer::MakeImitationTrainer(
-	ULearningAgentsManager* InManager,
-	ULearningAgentsInteractor* InInteractor,
-	ULearningAgentsPolicy* InPolicy,
+	ULearningAgentsManager*& InManager,
+	ULearningAgentsInteractor*& InInteractor,
+	ULearningAgentsPolicy*& InPolicy,
 	TSubclassOf<ULearningAgentsImitationTrainer> Class)
 {
 	if (!InManager)
diff --git a/Engine/Plugins/Experimental/LearningAgents/Source/LearningAgentsTraining/Private/LearningAgentsPPOTrainer.cpp b/Engine/Plugins/Experimental/LearningAgents/Source/LearningAgentsTraining/Private/LearningAgentsPPOTrainer.cpp
--- a/Engine/Plugins/Experimental/LearningAgents/Source/LearningAgentsTraining/Private/LearningAgentsPPOTrainer.cpp
+++ b/Engine/Plugins/Experimental/LearningAgents/Source/LearningAgentsTraining/Private/LearningAgentsPPOTrainer.cpp
@@ -43,11 +43,11 @@ void ULearningAgentsPPOTrainer::BeginDestroy()
 }
 
 ULearningAgentsPPOTrainer* ULearningAgentsPPOTrainer::MakePPOTrainer(
-	ULearningAgentsManager* InManager,
-	ULearningAgentsInteractor* InInteractor,
-	ULearningAgentsTrainingEnvironment* InTrainingEnvironment,
-	ULearningAgentsPolicy* InPolicy,
-	ULearningAgentsCritic* InCritic,
+	ULearningAgentsManager*& InManager,
+	ULearningAgentsInteractor*& InInteractor,
+	ULearningAgentsTrainingEnvironment*& InTrainingEnvironment,
+	ULearningAgentsPolicy*& InPolicy,
+	ULearningAgentsCritic*& InCritic,
 	const FLearningAgentsCommunicator& Communicator,
 	TSubclassOf<ULearningAgentsPPOTrainer> Class,
 	const FName Name,
@@ -76,11 +76,11 @@ ULearningAgentsPPOTrainer* ULearningAgentsPPOTrainer::MakePPOTrainer(
 }
 
 void ULearningAgentsPPOTrainer::SetupPPOTrainer(
-	ULearningAgentsManager* InManager,
-	ULearningAgentsInteractor* InInteractor,
-	ULearningAgentsTrainingEnvironment* InTrainingEnvironment,
-	ULearningAgentsPolicy* InPolicy,
-	ULearningAgentsCritic* InCritic,
+	ULearningAgentsManager*& InManager,
+	ULearningAgentsInteractor*& InInteractor,
+	ULearningAgentsTrainingEnvironment*& InTrainingEnvironment,
+	ULearningAgentsPolicy*& InPolicy,
+	ULearningAgentsCritic*& InCritic,
 	const FLearningAgentsCommunicator& Communicator,
 	const FLearningAgentsPPOTrainerSettings& TrainerSettings)
 {
diff --git a/Engine/Plugins/Experimental/LearningAgents/Source/LearningAgentsTraining/Private/LearningAgentsRecorder.cpp b/Engine/Plugins/Experimental/LearningAgents/Source/LearningAgentsTraining/Private/LearningAgentsRecorder.cpp
--- a/Engine/Plugins/Experimental/LearningAgents/Source/LearningAgentsTraining/Private/LearningAgentsRecorder.cpp
+++ b/Engine/Plugins/Experimental/LearningAgents/Source/LearningAgentsTraining/Private/LearningAgentsRecorder.cpp
@@ -107,8 +107,8 @@ void ULearningAgentsRecorder::BeginDestroy()
 }
 
 ULearningAgentsRecorder* ULearningAgentsRecorder::MakeRecorder(
-	ULearningAgentsManager* InManager,
-	ULearningAgentsInteractor* InInteractor,
+	ULearningAgentsManager*& InManager,
+	ULearningAgentsInteractor*& InInteractor,
 	TSubclassOf<ULearningAgentsRecorder> Class,
 	const FName Name,
 	const FLearningAgentsRecorderPathSettings& RecorderPathSettings,
@@ -143,8 +143,8 @@ ULearningAgentsRecorder* ULearningAgentsRecorder::MakeRecorder(
 }
 
 void ULearningAgentsRecorder::SetupRecorder(
-	ULearningAgentsManager* InManager,
-	ULearningAgentsInteractor* InInteractor,
+	ULearningAgentsManager*& InManager,
+	ULearningAgentsInteractor*& InInteractor,
 	const FLearningAgentsRecorderPathSettings& RecorderPathSettings,
 	ULearningAgentsRecording* RecordingAsset,
 	bool bReinitializeRecording)
diff --git a/Engine/Plugins/Experimental/LearningAgents/Source/LearningAgentsTraining/Private/LearningAgentsTrainingEnvironment.cpp b/Engine/Plugins/Experimental/LearningAgents/Source/LearningAgentsTraining/Private/LearningAgentsTrainingEnvironment.cpp
--- a/Engine/Plugins/Experimental/LearningAgents/Source/LearningAgentsTraining/Private/LearningAgentsTrainingEnvironment.cpp
+++ b/Engine/Plugins/Experimental/LearningAgents/Source/LearningAgentsTraining/Private/LearningAgentsTrainingEnvironment.cpp
@@ -11,7 +11,7 @@ ULearningAgentsTrainingEnvironment::ULearningAgentsTrainingEnvironment(FVTableHe
 ULearningAgentsTrainingEnvironment::~ULearningAgentsTrainingEnvironment() = default;
 
 ULearningAgentsTrainingEnvironment* ULearningAgentsTrainingEnvironment::MakeTrainingEnvironment(
-	ULearningAgentsManager* InManager,
+	ULearningAgentsManager*& InManager,
 	TSubclassOf<ULearningAgentsTrainingEnvironment> Class,
 	const FName Name)
 {
@@ -37,7 +37,7 @@ ULearningAgentsTrainingEnvironment* ULearningAgentsTrainingEnvironment::MakeTrai
 	return TrainingEnvironment->IsSetup() ? TrainingEnvironment : nullptr;
 }
 
-void ULearningAgentsTrainingEnvironment::SetupTrainingEnvironment(ULearningAgentsManager* InManager)
+void ULearningAgentsTrainingEnvironment::SetupTrainingEnvironment(ULearningAgentsManager*& InManager)
 {
 	if (IsSetup())
 	{
diff --git a/Engine/Plugins/Experimental/LearningAgents/Source/LearningAgentsTraining/Public/LearningAgentsImitationTrainer.h b/Engine/Plugins/Experimental/LearningAgents/Source/LearningAgentsTraining/Public/LearningAgentsImitationTrainer.h
--- a/Engine/Plugins/Experimental/LearningAgents/Source/LearningAgentsTraining/Public/LearningAgentsImitationTrainer.h
+++ b/Engine/Plugins/Experimental/LearningAgents/Source/LearningAgentsTraining/Public/LearningAgentsImitationTrainer.h
@@ -141,9 +141,9 @@ public:
 	 */
 	UFUNCTION(BlueprintCallable, Category = "LearningAgents", meta = (Class = "/Script/LearningAgentsTraining.LearningAgentsImitationTrainer", DeterminesOutputType = "Class"))
 	static ULearningAgentsImitationTrainer* MakeImitationTrainer(
-		ULearningAgentsManager* InManager,
-		ULearningAgentsInteractor* InInteractor,
-		ULearningAgentsPolicy* InPolicy,
+		UPARAM(ref) ULearningAgentsManager*& InManager,
+		UPARAM(ref) ULearningAgentsInteractor*& InInteractor,
+		UPARAM(ref) ULearningAgentsPolicy*& InPolicy,
 		TSubclassOf<ULearningAgentsImitationTrainer> Class);
 
 	/**
diff --git a/Engine/Plugins/Experimental/LearningAgents/Source/LearningAgentsTraining/Public/LearningAgentsPPOTrainer.h b/Engine/Plugins/Experimental/LearningAgents/Source/LearningAgentsTraining/Public/LearningAgentsPPOTrainer.h
--- a/Engine/Plugins/Experimental/LearningAgents/Source/LearningAgentsTraining/Public/LearningAgentsPPOTrainer.h
+++ b/Engine/Plugins/Experimental/LearningAgents/Source/LearningAgentsTraining/Public/LearningAgentsPPOTrainer.h
@@ -322,11 +322,11 @@ public:
 	 */
 	UFUNCTION(BlueprintCallable, Category = "LearningAgents", meta = (Class = "/Script/LearningAgents.LearningAgentsPPOTrainer", DeterminesOutputType = "Class", AutoCreateRefTerm = "TrainerSettings"))
 	static ULearningAgentsPPOTrainer* MakePPOTrainer(
-		ULearningAgentsManager* InManager,
-		ULearningAgentsInteractor* InInteractor,
-		ULearningAgentsTrainingEnvironment* InTrainingEnvironment,
-		ULearningAgentsPolicy* InPolicy,
-		ULearningAgentsCritic* InCritic,
+		UPARAM(ref) ULearningAgentsManager*& InManager,
+		UPARAM(ref) ULearningAgentsInteractor*& InInteractor,
+		UPARAM(ref) ULearningAgentsTrainingEnvironment*& InTrainingEnvironment,
+		UPARAM(ref) ULearningAgentsPolicy*& InPolicy,
+		UPARAM(ref) ULearningAgentsCritic*& InCritic,
 		const FLearningAgentsCommunicator& Communicator,
 		TSubclassOf<ULearningAgentsPPOTrainer> Class,
 		const FName Name = TEXT("PPOTrainer"),
@@ -345,11 +345,11 @@ public:
 	 */
 	UFUNCTION(BlueprintCallable, Category = "LearningAgents", meta = (AutoCreateRefTerm = "TrainerSettings"))
 	void SetupPPOTrainer(
-		ULearningAgentsManager* InManager,
-		ULearningAgentsInteractor* InInteractor,
-		ULearningAgentsTrainingEnvironment* InTrainingEnvironment,
-		ULearningAgentsPolicy* InPolicy,
-		ULearningAgentsCritic* InCritic,
+		UPARAM(ref) ULearningAgentsManager*& InManager,
+		UPARAM(ref) ULearningAgentsInteractor*& InInteractor,
+		UPARAM(ref) ULearningAgentsTrainingEnvironment*& InTrainingEnvironment,
+		UPARAM(ref) ULearningAgentsPolicy*& InPolicy,
+		UPARAM(ref) ULearningAgentsCritic*& InCritic,
 		const FLearningAgentsCommunicator& Communicator,
 		const FLearningAgentsPPOTrainerSettings& TrainerSettings = FLearningAgentsPPOTrainerSettings());
 
diff --git a/Engine/Plugins/Experimental/LearningAgents/Source/LearningAgentsTraining/Public/LearningAgentsRecorder.h b/Engine/Plugins/Experimental/LearningAgents/Source/LearningAgentsTraining/Public/LearningAgentsRecorder.h
--- a/Engine/Plugins/Experimental/LearningAgents/Source/LearningAgentsTraining/Public/LearningAgentsRecorder.h
+++ b/Engine/Plugins/Experimental/LearningAgents/Source/LearningAgentsTraining/Public/LearningAgentsRecorder.h
@@ -63,8 +63,8 @@ public:
 	 */
 	UFUNCTION(BlueprintCallable, Category = "LearningAgents", meta = (Class = "/Script/LearningAgentsTraining.LearningAgentsRecorder", DeterminesOutputType = "Class", AutoCreateRefTerm = "RecorderPathSettings"))
 	static ULearningAgentsRecorder* MakeRecorder(
-		ULearningAgentsManager* InManager,
-		ULearningAgentsInteractor* InInteractor,
+		UPARAM(ref) ULearningAgentsManager*& InManager,
+		UPARAM(ref) ULearningAgentsInteractor*& InInteractor,
 		TSubclassOf<ULearningAgentsRecorder> Class,
 		const FName Name = TEXT("Recorder"),
 		const FLearningAgentsRecorderPathSettings& RecorderPathSettings = FLearningAgentsRecorderPathSettings(),
@@ -83,8 +83,8 @@ public:
 	 */
 	UFUNCTION(BlueprintCallable, Category = "LearningAgents", meta = (AutoCreateRefTerm = "RecorderPathSettings"))
 	void SetupRecorder(
-		ULearningAgentsManager* InManager,
-		ULearningAgentsInteractor* InInteractor,
+		UPARAM(ref) ULearningAgentsManager*& InManager,
+		UPARAM(ref) ULearningAgentsInteractor*& InInteractor,
 		const FLearningAgentsRecorderPathSettings& RecorderPathSettings = FLearningAgentsRecorderPathSettings(),
 		ULearningAgentsRecording* RecordingAsset = nullptr,
 		bool bReinitializeRecording = true);
diff --git a/Engine/Plugins/Experimental/LearningAgents/Source/LearningAgentsTraining/Public/LearningAgentsTrainingEnvironment.h b/Engine/Plugins/Experimental/LearningAgents/Source/LearningAgentsTraining/Public/LearningAgentsTrainingEnvironment.h
--- a/Engine/Plugins/Experimental/LearningAgents/Source/LearningAgentsTraining/Public/LearningAgentsTrainingEnvironment.h
+++ b/Engine/Plugins/Experimental/LearningAgents/Source/LearningAgentsTraining/Public/LearningAgentsTrainingEnvironment.h
@@ -30,7 +30,7 @@ public:
 	 */
 	UFUNCTION(BlueprintCallable, Category = "LearningAgents", meta = (DeterminesOutputType = "Class", AutoCreateRefTerm = "TrainerSettings"))
 	static ULearningAgentsTrainingEnvironment* MakeTrainingEnvironment(
-		ULearningAgentsManager* InManager,
+		UPARAM(ref) ULearningAgentsManager*& InManager,
 		TSubclassOf<ULearningAgentsTrainingEnvironment> Class,
 		const FName Name = TEXT("TrainingEnvironment"));
 
@@ -38,7 +38,7 @@ public:
 	 * Initializes the training environment and runs the setup functions for rewards and completions.
 	 */
 	UFUNCTION(BlueprintCallable, Category = "LearningAgents", meta = (AutoCreateRefTerm = "TrainerSettings"))
-	void SetupTrainingEnvironment(ULearningAgentsManager* InManager);
+	void SetupTrainingEnvironment(UPARAM(ref) ULearningAgentsManager*& InManager);
 
 public:
 
