LearningAgents: Remove all friend classes

#rnx

diff --git a/Engine/Plugins/Experimental/LearningAgents/Source/LearningAgents/Private/LearningAgentsController.cpp b/Engine/Plugins/Experimental/LearningAgents/Source/LearningAgents/Private/LearningAgentsController.cpp
--- a/Engine/Plugins/Experimental/LearningAgents/Source/LearningAgents/Private/LearningAgentsController.cpp
+++ b/Engine/Plugins/Experimental/LearningAgents/Source/LearningAgents/Private/LearningAgentsController.cpp
@@ -122,26 +122,26 @@ void ULearningAgentsController::EvaluateController()
 		UE_LOG(LogLearning, Warning, TEXT("%s: No agents added to Manager."), *GetName());
 	}
 
-	if (Manager->GetAgentNum() != Interactor->ObservationObjectElements.Num())
+	if (Manager->GetAgentNum() != Interactor->GetObservationObjectElements().Num())
 	{
-		UE_LOG(LogLearning, Error, TEXT("%s: Not enough Observation Objects added by GatherAgentObservations. Expected %i, Got %i."), *GetName(), Manager->GetAgentNum(), Interactor->ObservationObjectElements.Num());
+		UE_LOG(LogLearning, Error, TEXT("%s: Not enough Observation Objects added by GatherAgentObservations. Expected %i, Got %i."), *GetName(), Manager->GetAgentNum(), Interactor->GetObservationObjectElements().Num());
 		return;
 	}
 
 	// Run EvaluateAgentControllers Callback
 
-	Interactor->ActionObject->ActionObject.Reset();
-	Interactor->ActionObjectElements.Empty(Manager->GetMaxAgentNum());
+	Interactor->GetActionObject()->ActionObject.Reset();
+	Interactor->GetActionObjectElements().Empty(Manager->GetMaxAgentNum());
 	EvaluateAgentControllers(
-		Interactor->ActionObjectElements, 
-		Interactor->ActionObject, 
-		Interactor->ObservationObject,
-		Interactor->ObservationObjectElements,
+		Interactor->GetActionObjectElements(),
+		Interactor->GetActionObject(),
+		Interactor->GetObservationObject(),
+		Interactor->GetObservationObjectElements(),
 		Manager->GetAllAgentIds());
 
-	if (Manager->GetAgentNum() != Interactor->ActionObjectElements.Num())
+	if (Manager->GetAgentNum() != Interactor->GetActionObjectElements().Num())
 	{
-		UE_LOG(LogLearning, Error, TEXT("%s: Not enough Action Objects added by EvaluateAgentControllers. Expected %i, Got %i."), *GetName(), Manager->GetAgentNum(), Interactor->ActionObjectElements.Num());
+		UE_LOG(LogLearning, Error, TEXT("%s: Not enough Action Objects added by EvaluateAgentControllers. Expected %i, Got %i."), *GetName(), Manager->GetAgentNum(), Interactor->GetActionObjectElements().Num());
 		return;
 	}
 
@@ -150,19 +150,19 @@ void ULearningAgentsController::EvaluateController()
 	for (int32 AgentIdx = 0; AgentIdx < Manager->GetAgentNum(); AgentIdx++)
 	{
 		if (ULearningAgentsActions::ValidateActionObjectMatchesSchema(
-			Interactor->ActionSchema,
-			Interactor->ActionSchemaElement, 
-			Interactor->ActionObject, 
-			Interactor->ActionObjectElements[AgentIdx]))
+			Interactor->GetActionSchema(),
+			Interactor->GetActionSchemaElement(),
+			Interactor->GetActionObject(),
+			Interactor->GetActionObjectElements()[AgentIdx]))
 		{
 			UE::Learning::Action::SetVectorFromObject(
-				Interactor->ActionVectors[Manager->GetAllAgentSet()[AgentIdx]],
-				Interactor->ActionSchema->ActionSchema,
-				Interactor->ActionSchemaElement.SchemaElement,
-				Interactor->ActionObject->ActionObject,
-				Interactor->ActionObjectElements[AgentIdx].ObjectElement);
+				Interactor->GetActionVectorsArrayView()[Manager->GetAllAgentSet()[AgentIdx]],
+				Interactor->GetActionSchema()->ActionSchema,
+				Interactor->GetActionSchemaElement().SchemaElement,
+				Interactor->GetActionObject()->ActionObject,
+				Interactor->GetActionObjectElements()[AgentIdx].ObjectElement);
 
-			Interactor->ActionVectorIteration[Manager->GetAllAgentSet()[AgentIdx]]++;
+			Interactor->GetActionVectorIterationArrayView()[Manager->GetAllAgentSet()[AgentIdx]]++;
 		}
 	}
 }
diff --git a/Engine/Plugins/Experimental/LearningAgents/Source/LearningAgents/Private/LearningAgentsCritic.cpp b/Engine/Plugins/Experimental/LearningAgents/Source/LearningAgents/Private/LearningAgentsCritic.cpp
--- a/Engine/Plugins/Experimental/LearningAgents/Source/LearningAgents/Private/LearningAgentsCritic.cpp
+++ b/Engine/Plugins/Experimental/LearningAgents/Source/LearningAgents/Private/LearningAgentsCritic.cpp
@@ -286,7 +286,7 @@ void ULearningAgentsCritic::EvaluateCritic()
 
 	for (const int32 AgentId : Manager->GetAllAgentSet())
 	{
-		if (Policy->ObservationVectorEncodedIteration[AgentId] == 0)
+		if (!Policy->HasEncodedObservationsForAgent(AgentId))
 		{
 			UE_LOG(LogLearning, Warning, TEXT("%s: Agent with id %i has not made observations so critic will not be evaluated for it. Was EncodeObservations run without error?"), *GetName(), AgentId);
 			continue;
@@ -300,10 +300,10 @@ void ULearningAgentsCritic::EvaluateCritic()
 
 	// Evaluate Critic
 
-	if (CriticObject->GetNeuralNetwork()->GetInputSize() != Policy->ObservationVectorsEncoded.Num<1>() + Policy->MemoryState.Num<1>())
+	if (CriticObject->GetNeuralNetwork()->GetInputSize() != Policy->GetObservationVectorsEncoded().Num<1>() + Policy->GetMemoryState().Num<1>())
 	{
 		UE_LOG(LogLearning, Error, TEXT("%s: Critic Network Input size doesn't match. Network input size is %i but Critic expects %i."), *GetName(),
-			CriticObject->GetNeuralNetwork()->GetInputSize(), Policy->ObservationVectorsEncoded.Num<1>() + Policy->MemoryState.Num<1>());
+			CriticObject->GetNeuralNetwork()->GetInputSize(), Policy->GetObservationVectorsEncoded().Num<1>() + Policy->GetMemoryState().Num<1>());
 		return;
 	}
 
@@ -316,8 +316,8 @@ void ULearningAgentsCritic::EvaluateCritic()
 
 	CriticObject->Evaluate(
 		Returns,
-		Policy->ObservationVectorsEncoded,
-		Policy->MemoryState,
+		Policy->GetObservationVectorsEncoded(),
+		Policy->GetMemoryState(),
 		ValidAgentSet);
 
 	// Increment Discounted Return Iteration
diff --git a/Engine/Plugins/Experimental/LearningAgents/Source/LearningAgents/Private/LearningAgentsInteractor.cpp b/Engine/Plugins/Experimental/LearningAgents/Source/LearningAgents/Private/LearningAgentsInteractor.cpp
--- a/Engine/Plugins/Experimental/LearningAgents/Source/LearningAgents/Private/LearningAgentsInteractor.cpp
+++ b/Engine/Plugins/Experimental/LearningAgents/Source/LearningAgents/Private/LearningAgentsInteractor.cpp
@@ -572,24 +572,24 @@ int32 ULearningAgentsInteractor::GetActionEncodedVectorSize() const
 	return ActionSchema->ActionSchema.GetEncodedVectorSize(ActionSchemaElement.SchemaElement);
 }
 
-const UE::Learning::Observation::FSchema& ULearningAgentsInteractor::GetObservationSchema() const
+const ULearningAgentsObservationSchema* ULearningAgentsInteractor::GetObservationSchema() const
 {
-	return ObservationSchema->ObservationSchema;
+	return ObservationSchema;
 }
 
-UE::Learning::Observation::FSchemaElement ULearningAgentsInteractor::GetObservationSchemaElement() const
+const FLearningAgentsObservationSchemaElement ULearningAgentsInteractor::GetObservationSchemaElement() const
 {
-	return ObservationSchemaElement.SchemaElement;
+	return ObservationSchemaElement;
 }
 
-const UE::Learning::Action::FSchema& ULearningAgentsInteractor::GetActionSchema() const
+const ULearningAgentsActionSchema* ULearningAgentsInteractor::GetActionSchema() const
 {
-	return ActionSchema->ActionSchema;
+	return ActionSchema;
 }
 
-UE::Learning::Action::FSchemaElement ULearningAgentsInteractor::GetActionSchemaElement() const
+const FLearningAgentsActionSchemaElement ULearningAgentsInteractor::GetActionSchemaElement() const
 {
-	return ActionSchemaElement.SchemaElement;
+	return ActionSchemaElement;
 }
 
 TLearningArrayView<2, const float> ULearningAgentsInteractor::GetObservationVectorArrayView() const
@@ -611,3 +611,33 @@ uint64 ULearningAgentsInteractor::GetActionIteration(const int32 AgentId) const
 {
 	return ActionVectorIteration[AgentId];
 }
+
+const ULearningAgentsObservationObject* ULearningAgentsInteractor::GetObservationObject() const
+{
+	return ObservationObject;
+}
+
+const TArray<FLearningAgentsObservationObjectElement>& ULearningAgentsInteractor::GetObservationObjectElements() const
+{
+	return ObservationObjectElements;
+}
+
+ULearningAgentsActionObject* ULearningAgentsInteractor::GetActionObject()
+{
+	return ActionObject;
+}
+
+TArray<FLearningAgentsActionObjectElement>& ULearningAgentsInteractor::GetActionObjectElements()
+{
+	return ActionObjectElements;
+}
+
+TLearningArrayView<2, float> ULearningAgentsInteractor::GetActionVectorsArrayView()
+{
+	return ActionVectors;
+}
+
+TLearningArrayView<1, uint64> ULearningAgentsInteractor::GetActionVectorIterationArrayView()
+{
+	return ActionVectorIteration;
+}
diff --git a/Engine/Plugins/Experimental/LearningAgents/Source/LearningAgents/Private/LearningAgentsPolicy.cpp b/Engine/Plugins/Experimental/LearningAgents/Source/LearningAgents/Private/LearningAgentsPolicy.cpp
--- a/Engine/Plugins/Experimental/LearningAgents/Source/LearningAgents/Private/LearningAgentsPolicy.cpp
+++ b/Engine/Plugins/Experimental/LearningAgents/Source/LearningAgents/Private/LearningAgentsPolicy.cpp
@@ -149,8 +149,8 @@ void ULearningAgentsPolicy::SetupPolicy(
 	const int32 ActionEncodedVectorSize = Interactor->GetActionEncodedVectorSize();
 	const int32 ActionDistributionVectorSize = Interactor->GetActionDistributionVectorSize();
 
-	const int32 ObservationCompatibilityHash = UE::Learning::Observation::GetSchemaObjectsCompatibilityHash(Interactor->GetObservationSchema(), Interactor->GetObservationSchemaElement());
-	const int32 ActionCompatibilityHash = UE::Learning::Action::GetSchemaObjectsCompatibilityHash(Interactor->GetActionSchema(), Interactor->GetActionSchemaElement());
+	const int32 ObservationCompatibilityHash = UE::Learning::Observation::GetSchemaObjectsCompatibilityHash(Interactor->GetObservationSchema()->ObservationSchema, Interactor->GetObservationSchemaElement().SchemaElement);
+	const int32 ActionCompatibilityHash = UE::Learning::Action::GetSchemaObjectsCompatibilityHash(Interactor->GetActionSchema()->ActionSchema, Interactor->GetActionSchemaElement().SchemaElement);
 	
 	const int32 PolicyHashData[3] = { MemoryStateSize, ObservationEncodedVectorSize, ActionEncodedVectorSize };
 	const int32 PolicyCompatibilityHash = CityHash32((const char*)PolicyHashData, 3 * sizeof(int32));
@@ -201,8 +201,8 @@ void ULearningAgentsPolicy::SetupPolicy(
 			FileData,
 			EncoderInputSize,
 			EncoderOutputSize,
-			Interactor->GetObservationSchema(),
-			Interactor->GetObservationSchemaElement(),
+			Interactor->GetObservationSchema()->ObservationSchema,
+			Interactor->GetObservationSchemaElement().SchemaElement,
 			UE::Learning::Observation::FNetworkSettings(),
 			UE::Learning::Random::Int(Seed ^ 0x658868dd));
 
@@ -344,8 +344,8 @@ void ULearningAgentsPolicy::SetupPolicy(
 			FileData,
 			DecoderInputSize,
 			DecoderOutputSize,
-			Interactor->GetActionSchema(),
-			Interactor->GetActionSchemaElement(),
+			Interactor->GetActionSchema()->ActionSchema,
+			Interactor->GetActionSchemaElement().SchemaElement,
 			UE::Learning::Action::FNetworkSettings(),
 			UE::Learning::Random::Int(Seed ^ 0xfa88bb7f));
 
@@ -488,6 +488,16 @@ TLearningArrayView<2, const float> ULearningAgentsPolicy::GetMemoryState() const
 	return MemoryState;
 }
 
+TLearningArrayView<2, const float> ULearningAgentsPolicy::GetObservationVectorsEncoded() const
+{
+	return ObservationVectorsEncoded;
+}
+
+bool ULearningAgentsPolicy::HasEncodedObservationsForAgent(const int32 AgentId) const
+{
+	return ObservationVectorEncodedIteration[AgentId] > 0;
+}
+
 ULearningAgentsNeuralNetwork* ULearningAgentsPolicy::GetEncoderNetworkAsset()
 {
 	if (!IsSetup())
@@ -541,7 +551,7 @@ void ULearningAgentsPolicy::EncodeObservations()
 	ValidAgentIds.Empty(Manager->GetMaxAgentNum());
 	for (const int32 AgentId : Manager->GetAllAgentSet())
 	{
-		if (Interactor->ObservationVectorIteration[AgentId] == 0)
+		if (Interactor->GetObservationIteration(AgentId) == 0)
 		{
 			UE_LOG(LogLearning, Warning, TEXT("%s: Agent with id %i does not have an observation vector ready be encoded. Was GatherObservations run without error?"), *GetName(), AgentId);
 			continue;
@@ -555,10 +565,10 @@ void ULearningAgentsPolicy::EncodeObservations()
 
 	// Encode Observations
 
-	if (EncoderObject->GetNeuralNetwork()->GetInputSize() != Interactor->ObservationVectors.Num<1>())
+	if (EncoderObject->GetNeuralNetwork()->GetInputSize() != Interactor->GetObservationVectorArrayView().Num<1>())
 	{
 		UE_LOG(LogLearning, Error, TEXT("%s: Encoder Network Input size doesn't match. Network input size is %i but Encoder expects %i."), *GetName(),
-			EncoderObject->GetNeuralNetwork()->GetInputSize(), Interactor->ObservationVectors.Num<1>());
+			EncoderObject->GetNeuralNetwork()->GetInputSize(), Interactor->GetObservationVectorArrayView().Num<1>());
 		return;
 	}
 
@@ -569,7 +579,7 @@ void ULearningAgentsPolicy::EncodeObservations()
 		return;
 	}
 
-	EncoderObject->Evaluate(ObservationVectorsEncoded, Interactor->ObservationVectors, ValidAgentSet);
+	EncoderObject->Evaluate(ObservationVectorsEncoded, Interactor->GetObservationVectorArrayView(), ValidAgentSet);
 
 	for (const int32 AgentId : ValidAgentSet)
 	{
@@ -703,16 +713,16 @@ void ULearningAgentsPolicy::DecodeAndSampleActions(const float ActionNoiseScale)
 	{
 		UE::Learning::Action::SampleVectorFromDistributionVector(
 			Seeds[AgentId],
-			Interactor->ActionVectors[AgentId],
+			Interactor->GetActionVectorsArrayView()[AgentId],
 			ActionDistributionVectors[AgentId],
-			Interactor->ActionSchema->ActionSchema,
-			Interactor->ActionSchemaElement.SchemaElement,
+			Interactor->GetActionSchema()->ActionSchema,
+			Interactor->GetActionSchemaElement().SchemaElement,
 			ActionNoiseScale);
 	}
 
 	for (const int32 AgentId : ValidAgentSet)
 	{
-		Interactor->ActionVectorIteration[AgentId]++;
+		Interactor->GetActionVectorIterationArrayView()[AgentId]++;
 	}
 }
 
diff --git a/Engine/Plugins/Experimental/LearningAgents/Source/LearningAgents/Public/LearningAgentsInteractor.h b/Engine/Plugins/Experimental/LearningAgents/Source/LearningAgents/Public/LearningAgentsInteractor.h
--- a/Engine/Plugins/Experimental/LearningAgents/Source/LearningAgents/Public/LearningAgentsInteractor.h
+++ b/Engine/Plugins/Experimental/LearningAgents/Source/LearningAgents/Public/LearningAgentsInteractor.h
@@ -33,11 +33,6 @@ class LEARNINGAGENTS_API ULearningAgentsInteractor : public ULearningAgentsManag
 
 public:
 
-	friend class ULearningAgentsController;
-	friend class ULearningAgentsPolicy;
-	friend class ULearningAgentsRecorder;
-	friend class ULearningAgentsTrainer;
-
 	// These constructors/destructors are needed to make forward declarations happy
 	ULearningAgentsInteractor();
 	ULearningAgentsInteractor(FVTableHelper& Helper);
@@ -236,17 +231,17 @@ public:
 	/** Perform Actions for a specific set of agents */
 	void PerformActions(const UE::Learning::FIndexSet AgentSet);
 
-	/** Gets the internal observation schema object */
-	const UE::Learning::Observation::FSchema& GetObservationSchema() const;
+	/** Gets the observation schema object */
+	const ULearningAgentsObservationSchema* GetObservationSchema() const;
 
-	/** Gets the internal observation schema element */
-	UE::Learning::Observation::FSchemaElement GetObservationSchemaElement() const;
+	/** Gets the observation schema element */
+	const FLearningAgentsObservationSchemaElement GetObservationSchemaElement() const;
 
-	/** Gets the internal action schema object */
-	const UE::Learning::Action::FSchema& GetActionSchema() const;
+	/** Gets the action schema object */
+	const ULearningAgentsActionSchema* GetActionSchema() const;
 
-	/** Gets the internal action schema element */
-	UE::Learning::Action::FSchemaElement GetActionSchemaElement() const;
+	/** Gets the action schema element */
+	const FLearningAgentsActionSchemaElement GetActionSchemaElement() const;
 
 	/** Gets the observation vectors as a const array view. */
 	TLearningArrayView<2, const float> GetObservationVectorArrayView() const;
@@ -260,6 +255,24 @@ public:
 	/** Gets the action iteration value for the given agent id. */
 	uint64 GetActionIteration(const int32 AgentId) const;
 
+	/** Gets the observation object. */
+	const ULearningAgentsObservationObject* GetObservationObject() const;
+
+	/** Gets the observation object elements. */
+	const TArray<FLearningAgentsObservationObjectElement>& GetObservationObjectElements() const;
+
+	/** Gets the action object. */
+	ULearningAgentsActionObject* GetActionObject();
+
+	/** Gets the action object elements. */
+	TArray<FLearningAgentsActionObjectElement>& GetActionObjectElements();
+
+	/** Gets the action vectors as a mutable array view. */
+	TLearningArrayView<2, float> GetActionVectorsArrayView();
+
+	/** Gets the action vector iterations as a mutable array view. */
+	TLearningArrayView<1, uint64> GetActionVectorIterationArrayView();
+
 private:
 
 	/** Observation Schema used by this interactor */
diff --git a/Engine/Plugins/Experimental/LearningAgents/Source/LearningAgents/Public/LearningAgentsPolicy.h b/Engine/Plugins/Experimental/LearningAgents/Source/LearningAgents/Public/LearningAgentsPolicy.h
--- a/Engine/Plugins/Experimental/LearningAgents/Source/LearningAgents/Public/LearningAgentsPolicy.h
+++ b/Engine/Plugins/Experimental/LearningAgents/Source/LearningAgents/Public/LearningAgentsPolicy.h
@@ -59,9 +59,6 @@ class LEARNINGAGENTS_API ULearningAgentsPolicy : public ULearningAgentsManagerLi
 
 public:
 
-	friend class ULearningAgentsCritic;
-	friend class ULearningAgentsTrainer;
-
 	// These constructors/destructors are needed to make forward declarations happy
 	ULearningAgentsPolicy();
 	ULearningAgentsPolicy(FVTableHelper& Helper);
@@ -229,6 +226,12 @@ public:
 	/** Get a const view to this policy's memory state. */
 	TLearningArrayView<2, const float> GetMemoryState() const;
 
+	/** Get a const view to this policy's encoded observation vectors. */
+	TLearningArrayView<2, const float> GetObservationVectorsEncoded() const;
+
+	/** Returns true if this policy has encoded observations for the given agent id. Otherwise, false. */
+	bool HasEncodedObservationsForAgent(const int32 AgentId) const;
+
 // ----- Private Data -----
 private:
 
diff --git a/Engine/Plugins/Experimental/LearningAgents/Source/LearningAgentsTraining/Private/LearningAgentsImitationTrainer.cpp b/Engine/Plugins/Experimental/LearningAgents/Source/LearningAgentsTraining/Private/LearningAgentsImitationTrainer.cpp
--- a/Engine/Plugins/Experimental/LearningAgents/Source/LearningAgentsTraining/Private/LearningAgentsImitationTrainer.cpp
+++ b/Engine/Plugins/Experimental/LearningAgents/Source/LearningAgentsTraining/Private/LearningAgentsImitationTrainer.cpp
@@ -262,10 +262,10 @@ void ULearningAgentsImitationTrainer::BeginTraining(
 		*Policy->GetPolicyNetworkAsset()->NeuralNetworkData,
 		*Policy->GetEncoderNetworkAsset()->NeuralNetworkData,
 		*Policy->GetDecoderNetworkAsset()->NeuralNetworkData,
-		Interactor->GetObservationSchema(),
-		Interactor->GetObservationSchemaElement(),
-		Interactor->GetActionSchema(),
-		Interactor->GetActionSchemaElement(),
+		Interactor->GetObservationSchema()->ObservationSchema,
+		Interactor->GetObservationSchemaElement().SchemaElement,
+		Interactor->GetActionSchema()->ActionSchema,
+		Interactor->GetActionSchemaElement().SchemaElement,
 		ImitationTrainingSettings);
 
 	UE_LOG(LogLearning, Display, TEXT("%s: Sending / Receiving initial policy..."), *GetName());
diff --git a/Engine/Plugins/Experimental/LearningAgents/Source/LearningAgentsTraining/Private/LearningAgentsPPOTrainer.cpp b/Engine/Plugins/Experimental/LearningAgents/Source/LearningAgentsTraining/Private/LearningAgentsPPOTrainer.cpp
--- a/Engine/Plugins/Experimental/LearningAgents/Source/LearningAgentsTraining/Private/LearningAgentsPPOTrainer.cpp
+++ b/Engine/Plugins/Experimental/LearningAgents/Source/LearningAgentsTraining/Private/LearningAgentsPPOTrainer.cpp
@@ -400,9 +400,11 @@ void ULearningAgentsPPOTrainer::SendConfig(const FLearningAgentsPPOTrainingSetti
 	ConfigObject->SetStringField(TEXT("TimeStamp"), *FDateTime::Now().ToFormattedString(TEXT("%Y-%m-%d_%H-%M-%S")));
 
 	ConfigObject->SetObjectField(TEXT("ObservationSchema"),
-		UE::Learning::Trainer::ConvertObservationSchemaToJSON(Interactor->GetObservationSchema(), Interactor->GetObservationSchemaElement()));
+		UE::Learning::Trainer::ConvertObservationSchemaToJSON(Interactor->GetObservationSchema()->ObservationSchema,
+		Interactor->GetObservationSchemaElement().SchemaElement));
 	ConfigObject->SetObjectField(TEXT("ActionSchema"),
-		UE::Learning::Trainer::ConvertActionSchemaToJSON(Interactor->GetActionSchema(), Interactor->GetActionSchemaElement()));
+		UE::Learning::Trainer::ConvertActionSchemaToJSON(Interactor->GetActionSchema()->ActionSchema,
+		Interactor->GetActionSchemaElement().SchemaElement));
 	ConfigObject->SetNumberField(TEXT("ObservationVectorDimensionNum"), ReplayBuffer->GetObservations().Num<1>());
 	ConfigObject->SetNumberField(TEXT("ActionVectorDimensionNum"), ReplayBuffer->GetActions().Num<1>());
 	ConfigObject->SetNumberField(TEXT("MemoryStateVectorDimensionNum"), ReplayBuffer->GetMemoryStates().Num<1>());
diff --git a/Engine/Plugins/Experimental/LearningAgents/Source/LearningAgentsTraining/Private/LearningAgentsRecorder.cpp b/Engine/Plugins/Experimental/LearningAgents/Source/LearningAgentsTraining/Private/LearningAgentsRecorder.cpp
--- a/Engine/Plugins/Experimental/LearningAgents/Source/LearningAgentsTraining/Private/LearningAgentsRecorder.cpp
+++ b/Engine/Plugins/Experimental/LearningAgents/Source/LearningAgentsTraining/Private/LearningAgentsRecorder.cpp
@@ -196,8 +196,8 @@ void ULearningAgentsRecorder::SetupRecorder(
 
 	// Find Compatibility Hash
 
-	const int32 ObservationCompatibilityHash = UE::Learning::Observation::GetSchemaObjectsCompatibilityHash(Interactor->GetObservationSchema(), Interactor->GetObservationSchemaElement());
-	const int32 ActionCompatibilityHash = UE::Learning::Action::GetSchemaObjectsCompatibilityHash(Interactor->GetActionSchema(), Interactor->GetActionSchemaElement());
+	const int32 ObservationCompatibilityHash = UE::Learning::Observation::GetSchemaObjectsCompatibilityHash(Interactor->GetObservationSchema()->ObservationSchema, Interactor->GetObservationSchemaElement().SchemaElement);
+	const int32 ActionCompatibilityHash = UE::Learning::Action::GetSchemaObjectsCompatibilityHash(Interactor->GetActionSchema()->ActionSchema, Interactor->GetActionSchemaElement().SchemaElement);
 
 	// Create Record Buffers
 
@@ -269,19 +269,19 @@ void ULearningAgentsRecorder::AddExperience()
 
 	for (const int32 AgentId : Manager->GetAllAgentSet())
 	{
-		if (Interactor->ObservationVectorIteration[AgentId] == 0 || Interactor->ActionVectorIteration[AgentId] == 0)
+		if (Interactor->GetObservationIteration(AgentId) == 0 || Interactor->GetActionIteration(AgentId) == 0)
 		{
 			UE_LOG(LogLearning, Warning, TEXT("%s: Agent with id %i has not made observations and taken actions so experience will not be recorded for it."), *GetName(), AgentId);
 			continue;
 		}
 
-		if (Interactor->ObservationVectorIteration[AgentId] != Interactor->ActionVectorIteration[AgentId])
+		if (Interactor->GetObservationIteration(AgentId) != Interactor->GetActionIteration(AgentId))
 		{
 			UE_LOG(LogLearning, Warning, TEXT("%s: Agent with id %i does not have matching iteration numbers for observations and actions so experience will not be recorded for it."), *GetName(), AgentId);
 			continue;
 		}
 
-		RecordBuffers[AgentId].Push(Interactor->ObservationVectors[AgentId], Interactor->ActionVectors[AgentId]);
+		RecordBuffers[AgentId].Push(Interactor->GetObservationVectorArrayView()[AgentId], Interactor->GetActionVectorArrayView()[AgentId]);
 	}
 }
 
