LearningAgents: Exposed some more functions to the API

diff --git a/Engine/Plugins/Experimental/LearningAgents/Source/LearningAgents/Public/LearningAgentsActions.h b/Engine/Plugins/Experimental/LearningAgents/Source/LearningAgents/Public/LearningAgentsActions.h
--- a/Engine/Plugins/Experimental/LearningAgents/Source/LearningAgents/Public/LearningAgentsActions.h
+++ b/Engine/Plugins/Experimental/LearningAgents/Source/LearningAgents/Public/LearningAgentsActions.h
@@ -37,10 +37,10 @@ struct LEARNINGAGENTS_API FLearningAgentsActionObjectElement
 };
 
 /** Comparison operator for Action Object Elements */
-bool operator==(const FLearningAgentsActionObjectElement& Lhs, const FLearningAgentsActionObjectElement& Rhs);
+LEARNINGAGENTS_API bool operator==(const FLearningAgentsActionObjectElement& Lhs, const FLearningAgentsActionObjectElement& Rhs);
 
 /** Hashing operator for Action Object Elements */
-uint32 GetTypeHash(const FLearningAgentsActionObjectElement& Element);
+LEARNINGAGENTS_API uint32 GetTypeHash(const FLearningAgentsActionObjectElement& Element);
 
 template<>
 struct TStructOpsTypeTraits<FLearningAgentsActionObjectElement> : public TStructOpsTypeTraitsBase2<FLearningAgentsActionObjectElement>
diff --git a/Engine/Plugins/Experimental/LearningAgents/Source/LearningAgents/Public/LearningAgentsObservations.h b/Engine/Plugins/Experimental/LearningAgents/Source/LearningAgents/Public/LearningAgentsObservations.h
--- a/Engine/Plugins/Experimental/LearningAgents/Source/LearningAgents/Public/LearningAgentsObservations.h
+++ b/Engine/Plugins/Experimental/LearningAgents/Source/LearningAgents/Public/LearningAgentsObservations.h
@@ -40,10 +40,10 @@ struct LEARNINGAGENTS_API FLearningAgentsObservationObjectElement
 };
 
 /** Comparison operator for Observation Object Elements */
-bool operator==(const FLearningAgentsObservationObjectElement& Lhs, const FLearningAgentsObservationObjectElement& Rhs);
+LEARNINGAGENTS_API bool operator==(const FLearningAgentsObservationObjectElement& Lhs, const FLearningAgentsObservationObjectElement& Rhs);
 
 /** Hashing operator for Observation Object Elements */
-uint32 GetTypeHash(const FLearningAgentsObservationObjectElement& Element);
+LEARNINGAGENTS_API uint32 GetTypeHash(const FLearningAgentsObservationObjectElement& Element);
 
 template<>
 struct TStructOpsTypeTraits<FLearningAgentsObservationObjectElement> : public TStructOpsTypeTraitsBase2<FLearningAgentsObservationObjectElement>
