LearningAgents: Added IsEmpty function to observation and action objects.

diff --git a/Engine/Plugins/Experimental/LearningAgents/Source/Learning/Private/LearningAction.cpp b/Engine/Plugins/Experimental/LearningAgents/Source/Learning/Private/LearningAction.cpp
--- a/Engine/Plugins/Experimental/LearningAgents/Source/Learning/Private/LearningAction.cpp
+++ b/Engine/Plugins/Experimental/LearningAgents/Source/Learning/Private/LearningAction.cpp
@@ -469,6 +469,11 @@ namespace UE::Learning::Action
 		Generation++;
 	}
 
+	bool FSchema::IsEmpty() const
+	{
+		return Types.IsEmpty();
+	}
+
 	void FSchema::Reset()
 	{
 		Types.Reset();
@@ -801,6 +806,11 @@ namespace UE::Learning::Action
 		Generation++;
 	}
 
+	bool FObject::IsEmpty() const
+	{
+		return Types.IsEmpty();
+	}
+
 	void FObject::Reset()
 	{
 		Types.Reset();
diff --git a/Engine/Plugins/Experimental/LearningAgents/Source/Learning/Private/LearningObservation.cpp b/Engine/Plugins/Experimental/LearningAgents/Source/Learning/Private/LearningObservation.cpp
--- a/Engine/Plugins/Experimental/LearningAgents/Source/Learning/Private/LearningObservation.cpp
+++ b/Engine/Plugins/Experimental/LearningAgents/Source/Learning/Private/LearningObservation.cpp
@@ -372,6 +372,11 @@ namespace UE::Learning::Observation
 		Generation++;
 	}
 
+	bool FSchema::IsEmpty() const
+	{
+		return Types.IsEmpty();
+	}
+
 	void FSchema::Reset()
 	{
 		Types.Reset();
@@ -654,6 +659,11 @@ namespace UE::Learning::Observation
 		Generation++;
 	}
 
+	bool FObject::IsEmpty() const
+	{
+		return Types.IsEmpty();
+	}
+
 	void FObject::Reset()
 	{
 		Types.Reset();
diff --git a/Engine/Plugins/Experimental/LearningAgents/Source/Learning/Public/LearningAction.h b/Engine/Plugins/Experimental/LearningAgents/Source/Learning/Public/LearningAction.h
--- a/Engine/Plugins/Experimental/LearningAgents/Source/Learning/Public/LearningAction.h
+++ b/Engine/Plugins/Experimental/LearningAgents/Source/Learning/Public/LearningAction.h
@@ -198,6 +198,9 @@ namespace UE::Learning::Action
 		// Get the current generation
 		uint32 GetGeneration() const;
 		
+		// Checks if the given schema is empty of elements.
+		bool IsEmpty() const;
+
 		// Empty all internal buffers of elements. This invalidates all existing elements.
 		void Empty();
 
@@ -391,6 +394,9 @@ namespace UE::Learning::Action
 		// Get the current generation
 		uint32 GetGeneration() const;
 
+		// Checks if the given object is empty of elements.
+		bool IsEmpty() const;
+
 		// Empty all internal buffers of elements. This invalidates all existing elements.
 		void Empty();
 
diff --git a/Engine/Plugins/Experimental/LearningAgents/Source/Learning/Public/LearningObservation.h b/Engine/Plugins/Experimental/LearningAgents/Source/Learning/Public/LearningObservation.h
--- a/Engine/Plugins/Experimental/LearningAgents/Source/Learning/Public/LearningObservation.h
+++ b/Engine/Plugins/Experimental/LearningAgents/Source/Learning/Public/LearningObservation.h
@@ -196,6 +196,9 @@ namespace UE::Learning::Observation
 		// Get the current generation
 		uint32 GetGeneration() const;
 
+		// Checks if the given schema is empty of elements.
+		bool IsEmpty() const;
+
 		// Empty all internal buffers of elements. This invalidates all existing elements.
 		void Empty();
 
@@ -375,6 +378,9 @@ namespace UE::Learning::Observation
 		// Get the current generation
 		uint32 GetGeneration() const;
 
+		// Checks if the given object is empty of elements.
+		bool IsEmpty() const;
+
 		// Empty all internal buffers of elements. This invalidates all existing elements.
 		void Empty();
 
