LearningAgents: Support for SIZE_T when computing total array shape.

diff --git a/Engine/Plugins/Experimental/LearningAgents/Source/Learning/Public/MultiArrayView.h b/Engine/Plugins/Experimental/LearningAgents/Source/Learning/Public/MultiArrayView.h
--- a/Engine/Plugins/Experimental/LearningAgents/Source/Learning/Public/MultiArrayView.h
+++ b/Engine/Plugins/Experimental/LearningAgents/Source/Learning/Public/MultiArrayView.h
@@ -135,15 +135,15 @@ public:
 		return Nums[Dimension];
 	}
 
-	SizeType Total() const
+	SIZE_T Total() const
 	{
 		if constexpr (DimNum == 1)
 		{
-			return Nums[0];
+			return (SIZE_T)Nums[0];
 		}
 		else
 		{
-			InSizeType Total = Nums[0];
+			SIZE_T Total = (SIZE_T)Nums[0];
 			for (uint8 Idx = 1; Idx < DimNum; Idx++)
 			{
 				Total *= Nums[Idx];
diff --git a/Engine/Plugins/Experimental/LearningAgents/Source/LearningTraining/Public/LearningSharedMemory.h b/Engine/Plugins/Experimental/LearningAgents/Source/LearningTraining/Public/LearningSharedMemory.h
--- a/Engine/Plugins/Experimental/LearningAgents/Source/LearningTraining/Public/LearningSharedMemory.h
+++ b/Engine/Plugins/Experimental/LearningAgents/Source/LearningTraining/Public/LearningSharedMemory.h
@@ -40,9 +40,9 @@ namespace UE::Learning
 		template<uint8 DimNum, typename ElementType>
 		TSharedMemoryArrayView<DimNum, ElementType> Map(const FGuid Guid, const TLearningArrayShape<DimNum>& Shape, const bool bCreate = false)
 		{
-			const int32 TotalSize = sizeof(ElementType) * Shape.Total();
+			const SIZE_T TotalSize = sizeof(ElementType) * (SIZE_T)Shape.Total();
 
-			if (TotalSize <= 0)
+			if (TotalSize == 0)
 			{
 				return { FGuid(), TLearningArrayView<DimNum, ElementType>(nullptr, Shape), nullptr };
 			}
