LearningAgents: Added lock for instance creation.

diff --git a/Engine/Plugins/Experimental/LearningAgents/Source/Learning/Private/LearningNeuralNetwork.cpp b/Engine/Plugins/Experimental/LearningAgents/Source/Learning/Private/LearningNeuralNetwork.cpp
--- a/Engine/Plugins/Experimental/LearningAgents/Source/Learning/Private/LearningNeuralNetwork.cpp
+++ b/Engine/Plugins/Experimental/LearningAgents/Source/Learning/Private/LearningNeuralNetwork.cpp
@@ -170,6 +170,7 @@ namespace UE::Learning
 		const int32 MaxBatchSize,
 		const FNeuralNetworkInferenceSettings& Settings)
 	{
+
 		TSharedRef<FNeuralNetworkInference> InferenceObject = MakeShared<FNeuralNetworkInference>(
 			*Model,
 			MaxBatchSize,
@@ -177,7 +178,10 @@ namespace UE::Learning
 			OutputSize,
 			Settings);
 
-		InferenceObjects.Emplace(InferenceObject.ToWeakPtr());
+		{
+			UE::Learning::FScopeNullableWriteLock ScopeLock(&Lock);
+			InferenceObjects.Emplace(InferenceObject.ToWeakPtr());
+		}
 
 		return InferenceObject;
 	}
@@ -199,6 +203,8 @@ namespace UE::Learning
 
 	void FNeuralNetwork::UpdateModel(const TSharedPtr<NNE::IModelCPU>& InModel, const int32 InInputSize, const int32 InOutputSize)
 	{
+		UE::Learning::FScopeNullableWriteLock ScopeLock(&Lock);
+
 		// Update Model
 
 		InputSize = InInputSize;
diff --git a/Engine/Plugins/Experimental/LearningAgents/Source/Learning/Public/LearningNeuralNetwork.h b/Engine/Plugins/Experimental/LearningAgents/Source/Learning/Public/LearningNeuralNetwork.h
--- a/Engine/Plugins/Experimental/LearningAgents/Source/Learning/Public/LearningNeuralNetwork.h
+++ b/Engine/Plugins/Experimental/LearningAgents/Source/Learning/Public/LearningNeuralNetwork.h
@@ -148,6 +148,7 @@ namespace UE::Learning
 		int32 OutputSize = 0;
 		TSharedPtr<NNE::IModelCPU> Model;
 		TArray<TWeakPtr<FNeuralNetworkInference>, TInlineAllocator<64>> InferenceObjects;
+		FRWLock Lock;
 	};
 
 	/**
