LearningAgents: Expand the dimensionality of EpisodeBuffer and ReplayBuffer. Add ability to have custom python trainer files.

#virtualized

diff --git a/Engine/Plugins/Experimental/LearningAgents/Source/Learning/Learning.Build.cs b/Engine/Plugins/Experimental/LearningAgents/Source/Learning/Learning.Build.cs
--- a/Engine/Plugins/Experimental/LearningAgents/Source/Learning/Learning.Build.cs
+++ b/Engine/Plugins/Experimental/LearningAgents/Source/Learning/Learning.Build.cs
@@ -40,6 +40,8 @@ public class Learning : ModuleRules
 			{
 				"CoreUObject",
 				"Engine",
+				"Json",
+				"JsonUtilities",
 				"NNE",
 				"NNERuntimeBasicCpu",
 				// ... add private dependencies that you statically link with here ...	
diff --git a/Engine/Plugins/Experimental/LearningAgents/Source/LearningAgents/Public/LearningAgentsInteractor.h b/Engine/Plugins/Experimental/LearningAgents/Source/LearningAgents/Public/LearningAgentsInteractor.h
--- a/Engine/Plugins/Experimental/LearningAgents/Source/LearningAgents/Public/LearningAgentsInteractor.h
+++ b/Engine/Plugins/Experimental/LearningAgents/Source/LearningAgents/Public/LearningAgentsInteractor.h
@@ -8,7 +8,6 @@
 
 #include "LearningArray.h"
 #include "Containers/Array.h"
-#include "Containers/ArrayView.h"
 #include "Containers/ContainerAllocationPolicies.h"
 #include "Templates/SharedPointer.h"
 #include "UObject/Object.h"
diff --git a/Engine/Plugins/Experimental/LearningAgents/Source/LearningAgentsTraining/Private/LearningAgentsCommunicator.cpp b/Engine/Plugins/Experimental/LearningAgents/Source/LearningAgentsTraining/Private/LearningAgentsCommunicator.cpp
--- a/Engine/Plugins/Experimental/LearningAgents/Source/LearningAgentsTraining/Private/LearningAgentsCommunicator.cpp
+++ b/Engine/Plugins/Experimental/LearningAgents/Source/LearningAgentsTraining/Private/LearningAgentsCommunicator.cpp
@@ -31,6 +31,8 @@ FLearningAgentsTrainerProcess ULearningAgentsCommunicatorLibrary::SpawnSharedMem
 
 	TrainerProcess.TrainerProcess = MakeShared<UE::Learning::FSharedMemoryTrainerServerProcess>(
 		SharedMemorySettings.TaskName,
+		UE::Learning::Trainer::GetProjectPythonContentPath(),
+		TrainerProcessSettings.TrainerFileName,
 		PythonExecutablePath,
 		PythonContentPath,
 		IntermediatePath,
@@ -85,6 +87,8 @@ FLearningAgentsTrainerProcess ULearningAgentsCommunicatorLibrary::SpawnSocketTra
 	const FString IntermediatePath = UE::Learning::Trainer::GetIntermediatePath(TrainerProcessSettings.GetIntermediatePath());
 
 	TrainerProcess.TrainerProcess = MakeShared<UE::Learning::FSocketTrainerServerProcess>(
+		UE::Learning::Trainer::GetProjectPythonContentPath(),
+		TrainerProcessSettings.TrainerFileName,
 		PythonExecutablePath,
 		PythonContentPath,
 		IntermediatePath,
diff --git a/Engine/Plugins/Experimental/LearningAgents/Source/LearningAgentsTraining/Private/LearningAgentsImitationTrainer.cpp b/Engine/Plugins/Experimental/LearningAgents/Source/LearningAgentsTraining/Private/LearningAgentsImitationTrainer.cpp
--- a/Engine/Plugins/Experimental/LearningAgents/Source/LearningAgentsTraining/Private/LearningAgentsImitationTrainer.cpp
+++ b/Engine/Plugins/Experimental/LearningAgents/Source/LearningAgentsTraining/Private/LearningAgentsImitationTrainer.cpp
@@ -6,14 +6,35 @@
 #include "LearningAgentsInteractor.h"
 #include "LearningExperience.h"
 #include "LearningLog.h"
-#include "LearningImitationTrainer.h"
+#include "LearningNeuralNetwork.h"
+#include "LearningExternalTrainer.h"
+#include "LearningAgentsCommunicator.h"
 #include "LearningAgentsRecording.h"
 #include "LearningAgentsPolicy.h"
 
-#include "GenericPlatform/GenericPlatformMisc.h"
-#include "HAL/FileManager.h"
+#include "Dom/JsonObject.h"
 #include "Misc/Paths.h"
 
+TSharedRef<FJsonObject> FLearningAgentsImitationTrainerTrainingSettings::AsJsonConfig() const
+{
+	TSharedRef<FJsonObject> ConfigObject = MakeShared<FJsonObject>();
+
+	ConfigObject->SetNumberField(TEXT("IterationNum"), NumberOfIterations);
+	ConfigObject->SetNumberField(TEXT("LearningRate"), LearningRate);
+	ConfigObject->SetNumberField(TEXT("LearningRateDecay"), LearningRateDecay);
+	ConfigObject->SetNumberField(TEXT("WeightDecay"), WeightDecay);
+	ConfigObject->SetNumberField(TEXT("BatchSize"), BatchSize);
+	ConfigObject->SetNumberField(TEXT("Window"), Window);
+	ConfigObject->SetNumberField(TEXT("ActionRegularizationWeight"), ActionRegularizationWeight);
+	ConfigObject->SetNumberField(TEXT("ActionEntropyWeight"), ActionEntropyWeight);
+	ConfigObject->SetNumberField(TEXT("Seed"), RandomSeed);
+	ConfigObject->SetStringField(TEXT("Device"), UE::Learning::Trainer::GetDeviceString(UE::Learning::Agents::GetTrainingDevice(Device)));
+	ConfigObject->SetBoolField(TEXT("UseTensorBoard"), bUseTensorboard);
+	ConfigObject->SetBoolField(TEXT("SaveSnapshots"), bSaveSnapshots);
+
+	return ConfigObject;
+}
+
 ULearningAgentsImitationTrainer::ULearningAgentsImitationTrainer() : Super(FObjectInitializer::Get()) {}
 ULearningAgentsImitationTrainer::ULearningAgentsImitationTrainer(FVTableHelper& Helper) : Super(Helper) {}
 ULearningAgentsImitationTrainer::~ULearningAgentsImitationTrainer() = default;
@@ -32,7 +53,9 @@ ULearningAgentsImitationTrainer* ULearningAgentsImitationTrainer::MakeImitationT
 	ULearningAgentsManager*& InManager,
 	ULearningAgentsInteractor*& InInteractor,
 	ULearningAgentsPolicy*& InPolicy,
-	TSubclassOf<ULearningAgentsImitationTrainer> Class)
+	const FLearningAgentsCommunicator& Communicator,
+	TSubclassOf<ULearningAgentsImitationTrainer> Class,
+	const FName Name)
 {
 	if (!InManager)
 	{
@@ -46,15 +69,12 @@ ULearningAgentsImitationTrainer* ULearningAgentsImitationTrainer::MakeImitationT
 		return nullptr;
 	}
 
-	const FName UniqueName = MakeUniqueObjectName(InManager, Class, TEXT("ImitationTrainer"), EUniqueObjectNameOptions::GloballyUnique);
+	const FName UniqueName = MakeUniqueObjectName(InManager, Class, Name, EUniqueObjectNameOptions::GloballyUnique);
 
 	ULearningAgentsImitationTrainer* ImitationTrainer = NewObject<ULearningAgentsImitationTrainer>(InManager, Class, UniqueName);
 	if (!ImitationTrainer) { return nullptr; }
 
-	ImitationTrainer->SetupImitationTrainer(
-		InManager,
-		InInteractor,
-		InPolicy);
+	ImitationTrainer->SetupImitationTrainer(InManager, InInteractor, InPolicy, Communicator);
 
 	return ImitationTrainer->IsSetup() ? ImitationTrainer : nullptr;
 }
@@ -62,7 +82,8 @@ ULearningAgentsImitationTrainer* ULearningAgentsImitationTrainer::MakeImitationT
 void ULearningAgentsImitationTrainer::SetupImitationTrainer(
 	ULearningAgentsManager* InManager,
 	ULearningAgentsInteractor* InInteractor,
-	ULearningAgentsPolicy* InPolicy)
+	ULearningAgentsPolicy* InPolicy,
+	const FLearningAgentsCommunicator& Communicator)
 {
 	if (IsSetup())
 	{
@@ -100,9 +121,16 @@ void ULearningAgentsImitationTrainer::SetupImitationTrainer(
 		return;
 	}
 
+	if (!Communicator.Trainer)
+	{
+		UE_LOG(LogLearning, Error, TEXT("%s: Communicator's Trainer is nullptr."), *GetName());
+		return;
+	}
+
 	Manager = InManager;
 	Interactor = InInteractor;
 	Policy = InPolicy;
+	Trainer = Communicator.Trainer;
 
 	bIsSetup = true;
 
@@ -139,10 +167,6 @@ void ULearningAgentsImitationTrainer::BeginTraining(
 		return;
 	}
 
-	// Record Timeout Setting
-
-	TrainerTimeout = ImitationTrainerSettings.TrainerCommunicationTimeout;
-
 	// Check Paths
 
 	const FString PythonExecutablePath = UE::Learning::Trainer::GetPythonExecutablePath(ImitationTrainerPathSettings.GetIntermediatePath());
@@ -166,7 +190,6 @@ void ULearningAgentsImitationTrainer::BeginTraining(
 
 	const int32 ObservationNum = Interactor->GetObservationVectorSize();
 	const int32 ActionNum = Interactor->GetActionVectorSize();
-	const int32 MemoryStateNum = Policy->GetMemoryStateSize();
 
 	// Get Number of Steps
 
@@ -229,98 +252,194 @@ void ULearningAgentsImitationTrainer::BeginTraining(
 	UE_LEARNING_CHECK(EpisodeIdx == TotalEpisodeNum);
 	UE_LEARNING_CHECK(StepIdx == TotalStepNum);
 
-	// Begin Training Properly
+	const int32 ObservationSchemaId = 0;
+	const int32 ActionSchemaId = 0;
 
-	UE_LOG(LogLearning, Display, TEXT("%s: Imitation Training Started"), *GetName());
-
-
-	UE::Learning::FImitationTrainerTrainingSettings ImitationTrainingSettings;
-	ImitationTrainingSettings.IterationNum = ImitationTrainerTrainingSettings.NumberOfIterations;
-	ImitationTrainingSettings.LearningRate = ImitationTrainerTrainingSettings.LearningRate;
-	ImitationTrainingSettings.LearningRateDecay = ImitationTrainerTrainingSettings.LearningRateDecay;
-	ImitationTrainingSettings.WeightDecay = ImitationTrainerTrainingSettings.WeightDecay;
-	ImitationTrainingSettings.BatchSize = ImitationTrainerTrainingSettings.BatchSize;
-	ImitationTrainingSettings.Window = ImitationTrainerTrainingSettings.Window;
-	ImitationTrainingSettings.ActionRegularizationWeight = ImitationTrainerTrainingSettings.ActionRegularizationWeight;
-	ImitationTrainingSettings.ActionEntropyWeight = ImitationTrainerTrainingSettings.ActionEntropyWeight;
-	ImitationTrainingSettings.Seed = ImitationTrainerTrainingSettings.RandomSeed;
-	ImitationTrainingSettings.Device = UE::Learning::Agents::GetTrainingDevice(ImitationTrainerTrainingSettings.Device);
-	ImitationTrainingSettings.bUseTensorboard = ImitationTrainerTrainingSettings.bUseTensorboard;
-	ImitationTrainingSettings.bSaveSnapshots = ImitationTrainerTrainingSettings.bSaveSnapshots;
-
-	ImitationTrainer = MakeUnique<UE::Learning::FSharedMemoryImitationTrainer>(
-		GetName(),
-		PythonExecutablePath,
-		TEXT(""),
-		PythonContentPath,
-		IntermediatePath,
+	// Create Replay Buffer from Records
+	ReplayBuffer = MakeUnique<UE::Learning::FReplayBuffer>();
+	ReplayBuffer->AddRecords(
 		TotalEpisodeNum,
 		TotalStepNum,
+		ObservationSchemaId,
 		ObservationNum,
+		ActionSchemaId,
 		ActionNum,
-		MemoryStateNum,
-		*Policy->GetPolicyNetworkAsset()->NeuralNetworkData,
-		*Policy->GetEncoderNetworkAsset()->NeuralNetworkData,
-		*Policy->GetDecoderNetworkAsset()->NeuralNetworkData,
-		Interactor->GetObservationSchema()->ObservationSchema,
-		Interactor->GetObservationSchemaElement().SchemaElement,
-		Interactor->GetActionSchema()->ActionSchema,
-		Interactor->GetActionSchemaElement().SchemaElement,
-		ImitationTrainingSettings);
+		RecordedEpisodeStarts,
+		RecordedEpisodeLengths,
+		RecordedObservations,
+		RecordedActions);
 
-	UE_LOG(LogLearning, Display, TEXT("%s: Sending / Receiving initial policy..."), *GetName());
+	// We need to setup the trainer prior to sending the config
+	PolicyNetworkId = Trainer->AddNetwork(*Policy->GetPolicyNetworkAsset()->NeuralNetworkData);
+	EncoderNetworkId = Trainer->AddNetwork(*Policy->GetEncoderNetworkAsset()->NeuralNetworkData);
+	DecoderNetworkId = Trainer->AddNetwork(*Policy->GetDecoderNetworkAsset()->NeuralNetworkData);
+	ReplayBufferId = Trainer->AddReplayBuffer(*ReplayBuffer);
+
+	TSharedRef<FJsonObject> ConfigObject = CreateConfig(ImitationTrainerTrainingSettings);
+
+	UE_LOG(LogLearning, Display, TEXT("%s: Sending config..."), *GetName());
+	SendConfig(ConfigObject);
 
+	UE_LOG(LogLearning, Display, TEXT("%s: Imitation Training Started"), *GetName());
 	UE::Learning::ETrainerResponse Response = UE::Learning::ETrainerResponse::Success;
 
-	Response = ImitationTrainer->SendPolicy(*Policy->GetPolicyNetworkAsset()->NeuralNetworkData, TrainerTimeout);
+	UE_LOG(LogLearning, Display, TEXT("%s: Sending / Receiving initial policy..."), *GetName());
+	Response = Trainer->SendNetwork(PolicyNetworkId, *Policy->GetPolicyNetworkAsset()->NeuralNetworkData);
 	if (Response != UE::Learning::ETrainerResponse::Success)
 	{
 		UE_LOG(LogLearning, Error, TEXT("%s: Error sending policy to trainer: %s. Check log for additional errors."), *GetName(), UE::Learning::Trainer::GetResponseString(Response));
-		ImitationTrainer->Terminate();
 		bHasTrainingFailed = true;
+		Trainer->Terminate();
 		return;
 	}
 
-	Response = ImitationTrainer->SendEncoder(*Policy->GetEncoderNetworkAsset()->NeuralNetworkData, TrainerTimeout);
+	Response = Trainer->SendNetwork(EncoderNetworkId, *Policy->GetEncoderNetworkAsset()->NeuralNetworkData);
 	if (Response != UE::Learning::ETrainerResponse::Success)
 	{
 		UE_LOG(LogLearning, Error, TEXT("%s: Error sending encoder to trainer: %s. Check log for additional errors."), *GetName(), UE::Learning::Trainer::GetResponseString(Response));
-		ImitationTrainer->Terminate();
 		bHasTrainingFailed = true;
+		Trainer->Terminate();
 		return;
 	}
 
-	Response = ImitationTrainer->SendDecoder(*Policy->GetDecoderNetworkAsset()->NeuralNetworkData, TrainerTimeout);
+	Response = Trainer->SendNetwork(DecoderNetworkId, *Policy->GetDecoderNetworkAsset()->NeuralNetworkData);
 	if (Response != UE::Learning::ETrainerResponse::Success)
 	{
 		UE_LOG(LogLearning, Error, TEXT("%s: Error sending decoder to trainer: %s. Check log for additional errors."), *GetName(), UE::Learning::Trainer::GetResponseString(Response));
-		ImitationTrainer->Terminate();
 		bHasTrainingFailed = true;
+		Trainer->Terminate();
 		return;
 	}
 
 	UE_LOG(LogLearning, Display, TEXT("%s: Sending Experience..."), *GetName());
-
-	// Send Experience
-
-	Response = ImitationTrainer->SendExperience(
-		RecordedEpisodeStarts,
-		RecordedEpisodeLengths,
-		RecordedObservations, 
-		RecordedActions, 
-		TrainerTimeout);
-
+	Response = Trainer->SendReplayBuffer(ReplayBufferId, *ReplayBuffer);
 	if (Response != UE::Learning::ETrainerResponse::Success)
 	{
 		UE_LOG(LogLearning, Error, TEXT("%s: Error sending experience to trainer: %s. Check log for additional errors."), *GetName(), UE::Learning::Trainer::GetResponseString(Response));
 		bHasTrainingFailed = true;
-		ImitationTrainer->Terminate();
+		Trainer->Terminate();
 		return;
 	}
 
 	bIsTraining = true;
 }
 
+TSharedRef<FJsonObject> ULearningAgentsImitationTrainer::CreateConfig(const FLearningAgentsImitationTrainerTrainingSettings& TrainingSettings) const
+{
+	TSharedRef<FJsonObject> ConfigObject = MakeShared<FJsonObject>();
+	
+	// Add Training Config Entries
+	ConfigObject->SetStringField(TEXT("TaskName"), TEXT("Training"));
+	ConfigObject->SetStringField(TEXT("TrainerMethod"), TEXT("BehaviorCloning"));
+	ConfigObject->SetStringField(TEXT("TimeStamp"), *FDateTime::Now().ToFormattedString(TEXT("%Y-%m-%d_%H-%M-%S")));
+
+	const int32 ObservationSchemaId = 0;
+	const int32 ActionSchemaId = 0;
+
+	// Add Neural Network Config Entries
+	TArray<TSharedPtr<FJsonValue>> NetworkObjects;
+
+	// Policy
+	{
+		TSharedPtr<FJsonObject> NetworkObject = MakeShared<FJsonObject>();
+		NetworkObject->SetNumberField(TEXT("Id"), PolicyNetworkId);
+		NetworkObject->SetStringField(TEXT("Name"), Policy->GetPolicyNetworkAsset()->GetFName().ToString());
+		NetworkObject->SetNumberField(TEXT("MaxByteNum"), Policy->GetPolicyNetworkAsset()->NeuralNetworkData->GetSnapshotByteNum());
+
+		TSharedRef<FJsonValueObject> JsonValue = MakeShared<FJsonValueObject>(NetworkObject);
+		NetworkObjects.Add(JsonValue);
+	}
+
+	// Encoder
+	{
+		TSharedPtr<FJsonObject> NetworkObject = MakeShared<FJsonObject>();
+		NetworkObject->SetNumberField(TEXT("Id"), EncoderNetworkId);
+		NetworkObject->SetStringField(TEXT("Name"), Policy->GetEncoderNetworkAsset()->GetFName().ToString());
+		NetworkObject->SetNumberField(TEXT("MaxByteNum"), Policy->GetEncoderNetworkAsset()->NeuralNetworkData->GetSnapshotByteNum());
+		NetworkObject->SetNumberField(TEXT("InputSchemaId"), ObservationSchemaId);
+
+		TSharedRef<FJsonValueObject> JsonValue = MakeShared<FJsonValueObject>(NetworkObject);
+		NetworkObjects.Add(JsonValue);
+	}
+
+	// Decoder
+	{
+		TSharedPtr<FJsonObject> NetworkObject = MakeShared<FJsonObject>();
+		NetworkObject->SetNumberField(TEXT("Id"), DecoderNetworkId);
+		NetworkObject->SetStringField(TEXT("Name"), Policy->GetDecoderNetworkAsset()->GetFName().ToString());
+		NetworkObject->SetNumberField(TEXT("MaxByteNum"), Policy->GetDecoderNetworkAsset()->NeuralNetworkData->GetSnapshotByteNum());
+		NetworkObject->SetNumberField(TEXT("OutputSchemaId"), ActionSchemaId);
+
+		TSharedRef<FJsonValueObject> JsonValue = MakeShared<FJsonValueObject>(NetworkObject);
+		NetworkObjects.Add(JsonValue);
+	}
+
+	ConfigObject->SetArrayField(TEXT("Networks"), NetworkObjects);
+
+	// Add Replay Buffers Config Entries
+	TArray<TSharedPtr<FJsonValue>> ReplayBufferObjects;
+	TSharedRef<FJsonValueObject> ReplayBufferJsonValue = MakeShared<FJsonValueObject>(ReplayBuffer->AsJsonConfig(ReplayBufferId));
+	ReplayBufferObjects.Add(ReplayBufferJsonValue);
+	ConfigObject->SetArrayField(TEXT("ReplayBuffers"), ReplayBufferObjects);
+
+	// Schemas
+	TSharedPtr<FJsonObject> SchemasObject = MakeShared<FJsonObject>();
+
+	// Add the observation schemas
+	TArray<TSharedPtr<FJsonValue>> ObservationSchemaObjects;
+	{
+		// For this trainer, add the one observation schema we have
+		TSharedPtr<FJsonObject> ObservationSchemaObject = MakeShared<FJsonObject>();
+		ObservationSchemaObject->SetNumberField(TEXT("Id"), ObservationSchemaId);
+		ObservationSchemaObject->SetStringField(TEXT("Name"), "Default");
+		ObservationSchemaObject->SetObjectField(TEXT("Schema"),
+			UE::Learning::Trainer::ConvertObservationSchemaToJSON(Interactor->GetObservationSchema()->ObservationSchema,
+				Interactor->GetObservationSchemaElement().SchemaElement));
+
+		TSharedRef<FJsonValueObject> JsonValue = MakeShared<FJsonValueObject>(ObservationSchemaObject);
+		ObservationSchemaObjects.Add(JsonValue);
+	}
+	SchemasObject->SetArrayField(TEXT("Observations"), ObservationSchemaObjects);
+
+	// Add the action schemas
+	TArray<TSharedPtr<FJsonValue>> ActionSchemaObjects;
+	{
+		// For this trainer, add the one action schema we have
+		TSharedPtr<FJsonObject> ActionSchemaObject = MakeShared<FJsonObject>();
+		ActionSchemaObject->SetNumberField(TEXT("Id"), ActionSchemaId);
+		ActionSchemaObject->SetStringField(TEXT("Name"), "Default");
+		ActionSchemaObject->SetObjectField(TEXT("Schema"),
+			UE::Learning::Trainer::ConvertActionSchemaToJSON(Interactor->GetActionSchema()->ActionSchema,
+				Interactor->GetActionSchemaElement().SchemaElement));
+
+		TSharedRef<FJsonValueObject> JsonValue = MakeShared<FJsonValueObject>(ActionSchemaObject);
+		ActionSchemaObjects.Add(JsonValue);
+	}
+	SchemasObject->SetArrayField(TEXT("Actions"), ActionSchemaObjects);
+
+	ConfigObject->SetObjectField(TEXT("Schemas"), SchemasObject);
+
+	// Add Imitation Specific Config Entries
+	ConfigObject->SetObjectField(TEXT("BehaviorCloningSettings"), TrainingSettings.AsJsonConfig());
+	ConfigObject->SetNumberField(TEXT("MemoryStateNum"), Policy->GetMemoryStateSize());
+
+	return ConfigObject;
+}
+
+void ULearningAgentsImitationTrainer::SendConfig(const TSharedRef<FJsonObject>& ConfigObject)
+{
+	UE::Learning::ETrainerResponse Response = UE::Learning::ETrainerResponse::Success;
+	Response = Trainer->SendConfig(ConfigObject, UE::Learning::ELogSetting::Normal); // todo pipe the log settings up to the blueprint graph
+
+	if (Response != UE::Learning::ETrainerResponse::Success)
+	{
+		UE_LOG(LogLearning, Error, TEXT("%s: Error sending config to trainer: %s. Check log for additional errors."), *GetName(), UE::Learning::Trainer::GetResponseString(Response));
+		bHasTrainingFailed = true;
+		Trainer->Terminate();
+		return;
+	}
+}
+
 void ULearningAgentsImitationTrainer::DoneTraining()
 {
 	if (!IsSetup())
@@ -332,10 +451,10 @@ void ULearningAgentsImitationTrainer::DoneTraining()
 	if (IsTraining())
 	{
 		// Wait for Trainer to finish
-		ImitationTrainer->Wait(1.0f);
+		Trainer->Wait();
 
 		// If not finished in time, terminate
-		ImitationTrainer->Terminate();
+		Trainer->Terminate();
 
 		bIsTraining = false;
 	}
@@ -352,7 +471,7 @@ void ULearningAgentsImitationTrainer::EndTraining()
 	if (IsTraining())
 	{
 		UE_LOG(LogLearning, Display, TEXT("%s: Stopping training..."), *GetName());
-		ImitationTrainer->SendStop();
+		Trainer->SendStop();
 		DoneTraining();
 	}
 }
@@ -373,9 +492,11 @@ void ULearningAgentsImitationTrainer::IterateTraining()
 		return;
 	}
 
-	if (ImitationTrainer->HasPolicyOrCompleted())
+	if (Trainer->HasNetworkOrCompleted())
 	{
-		UE::Learning::ETrainerResponse Response = ImitationTrainer->RecvPolicy(*Policy->GetPolicyNetworkAsset()->NeuralNetworkData, TrainerTimeout);
+		UE_LOG(LogLearning, Display, TEXT("Receiving trained networks..."));
+
+		UE::Learning::ETrainerResponse Response = Trainer->ReceiveNetwork(PolicyNetworkId, *Policy->GetPolicyNetworkAsset()->NeuralNetworkData);
 		if (Response == UE::Learning::ETrainerResponse::Completed)
 		{
 			UE_LOG(LogLearning, Display, TEXT("%s: Trainer completed training."), *GetName());
@@ -391,7 +512,7 @@ void ULearningAgentsImitationTrainer::IterateTraining()
 		}
 		Policy->GetPolicyNetworkAsset()->ForceMarkDirty();
 
-		Response = ImitationTrainer->RecvEncoder(*Policy->GetEncoderNetworkAsset()->NeuralNetworkData, TrainerTimeout);
+		Response = Trainer->ReceiveNetwork(EncoderNetworkId, *Policy->GetEncoderNetworkAsset()->NeuralNetworkData);
 		if (Response != UE::Learning::ETrainerResponse::Success)
 		{
 			UE_LOG(LogLearning, Error, TEXT("%s: Error receiving encoder from trainer: %s. Check log for additional errors."), *GetName(), UE::Learning::Trainer::GetResponseString(Response));
@@ -401,7 +522,7 @@ void ULearningAgentsImitationTrainer::IterateTraining()
 		}
 		Policy->GetEncoderNetworkAsset()->ForceMarkDirty();
 
-		Response = ImitationTrainer->RecvDecoder(*Policy->GetDecoderNetworkAsset()->NeuralNetworkData, TrainerTimeout);
+		Response = Trainer->ReceiveNetwork(DecoderNetworkId, *Policy->GetDecoderNetworkAsset()->NeuralNetworkData);
 		if (Response != UE::Learning::ETrainerResponse::Success)
 		{
 			UE_LOG(LogLearning, Error, TEXT("%s: Error receiving decoder from trainer: %s. Check log for additional errors."), *GetName(), UE::Learning::Trainer::GetResponseString(Response));
diff --git a/Engine/Plugins/Experimental/LearningAgents/Source/LearningAgentsTraining/Private/LearningAgentsPPOTrainer.cpp b/Engine/Plugins/Experimental/LearningAgents/Source/LearningAgentsTraining/Private/LearningAgentsPPOTrainer.cpp
--- a/Engine/Plugins/Experimental/LearningAgents/Source/LearningAgentsTraining/Private/LearningAgentsPPOTrainer.cpp
+++ b/Engine/Plugins/Experimental/LearningAgents/Source/LearningAgentsTraining/Private/LearningAgentsPPOTrainer.cpp
@@ -15,23 +15,48 @@
 #include "LearningAgentsTrainingEnvironment.h"
 #include "LearningExternalTrainer.h"
 
-#include "Misc/App.h"
-#include "GameFramework/GameUserSettings.h"
-#include "PhysicsEngine/PhysicsSettings.h"
-#include "Engine/GameViewportClient.h"
-
 #include "Dom/JsonObject.h"
-#include "Serialization/JsonSerializer.h"
-#include "Serialization/JsonWriter.h"
-
-#if WITH_EDITOR
-#include "Editor/EditorPerformanceSettings.h"
-#endif
 
 ULearningAgentsPPOTrainer::ULearningAgentsPPOTrainer() : Super(FObjectInitializer::Get()) {}
 ULearningAgentsPPOTrainer::ULearningAgentsPPOTrainer(FVTableHelper& Helper) : Super(Helper) {}
 ULearningAgentsPPOTrainer::~ULearningAgentsPPOTrainer() = default;
 
+TSharedRef<FJsonObject> FLearningAgentsPPOTrainingSettings::AsJsonConfig() const
+{
+	TSharedRef<FJsonObject> ConfigObject = MakeShared<FJsonObject>();
+
+	ConfigObject->SetNumberField(TEXT("IterationNum"), NumberOfIterations);
+	ConfigObject->SetNumberField(TEXT("LearningRatePolicy"), LearningRatePolicy);
+	ConfigObject->SetNumberField(TEXT("LearningRateCritic"), LearningRateCritic);
+	ConfigObject->SetNumberField(TEXT("LearningRateDecay"), LearningRateDecay);
+	ConfigObject->SetNumberField(TEXT("WeightDecay"), WeightDecay);
+	ConfigObject->SetNumberField(TEXT("PolicyBatchSize"), PolicyBatchSize);
+	ConfigObject->SetNumberField(TEXT("CriticBatchSize"), CriticBatchSize);
+	ConfigObject->SetNumberField(TEXT("PolicyWindow"), PolicyWindowSize);
+	ConfigObject->SetNumberField(TEXT("IterationsPerGather"), IterationsPerGather);
+	ConfigObject->SetNumberField(TEXT("CriticWarmupIterations"), CriticWarmupIterations);
+	ConfigObject->SetNumberField(TEXT("EpsilonClip"), EpsilonClip);
+	ConfigObject->SetNumberField(TEXT("ActionSurrogateWeight"), ActionSurrogateWeight);
+	ConfigObject->SetNumberField(TEXT("ActionRegularizationWeight"), ActionRegularizationWeight);
+	ConfigObject->SetNumberField(TEXT("ActionEntropyWeight"), ActionEntropyWeight);
+	ConfigObject->SetNumberField(TEXT("ReturnRegularizationWeight"), ReturnRegularizationWeight);
+	ConfigObject->SetNumberField(TEXT("GaeLambda"), GaeLambda);
+	ConfigObject->SetBoolField(TEXT("AdvantageNormalization"), bAdvantageNormalization);
+	ConfigObject->SetNumberField(TEXT("AdvantageMin"), MinimumAdvantage);
+	ConfigObject->SetNumberField(TEXT("AdvantageMax"), MaximumAdvantage);
+	ConfigObject->SetBoolField(TEXT("UseGradNormMaxClipping"), bUseGradNormMaxClipping);
+	ConfigObject->SetNumberField(TEXT("GradNormMax"), GradNormMax);
+	ConfigObject->SetNumberField(TEXT("TrimEpisodeStartStepNum"), NumberOfStepsToTrimAtStartOfEpisode);
+	ConfigObject->SetNumberField(TEXT("TrimEpisodeEndStepNum"), NumberOfStepsToTrimAtEndOfEpisode);
+	ConfigObject->SetNumberField(TEXT("Seed"), RandomSeed);
+	ConfigObject->SetNumberField(TEXT("DiscountFactor"), DiscountFactor);
+	ConfigObject->SetStringField(TEXT("Device"), UE::Learning::Trainer::GetDeviceString(UE::Learning::Agents::GetTrainingDevice(Device)));
+	ConfigObject->SetBoolField(TEXT("UseTensorBoard"), bUseTensorboard);
+	ConfigObject->SetBoolField(TEXT("SaveSnapshots"), bSaveSnapshots);
+
+	return ConfigObject;
+}
+
 void ULearningAgentsPPOTrainer::BeginDestroy()
 {
 	if (IsTraining())
@@ -67,12 +92,12 @@ ULearningAgentsPPOTrainer* ULearningAgentsPPOTrainer::MakePPOTrainer(
 
 	const FName UniqueName = MakeUniqueObjectName(InManager, Class, Name, EUniqueObjectNameOptions::GloballyUnique);
 
-	ULearningAgentsPPOTrainer* Trainer = NewObject<ULearningAgentsPPOTrainer>(InManager, Class, UniqueName);
-	if (!Trainer) { return nullptr; }
+	ULearningAgentsPPOTrainer* PPOTrainer = NewObject<ULearningAgentsPPOTrainer>(InManager, Class, UniqueName);
+	if (!PPOTrainer) { return nullptr; }
 
-	Trainer->SetupPPOTrainer(InManager, InInteractor, InTrainingEnvironment, InPolicy, InCritic, Communicator, TrainerSettings);
+	PPOTrainer->SetupPPOTrainer(InManager, InInteractor, InTrainingEnvironment, InPolicy, InCritic, Communicator, TrainerSettings);
 
-	return Trainer->IsSetup() ? Trainer : nullptr;
+	return PPOTrainer->IsSetup() ? PPOTrainer : nullptr;
 }
 
 void ULearningAgentsPPOTrainer::SetupPPOTrainer(
@@ -157,21 +182,31 @@ void ULearningAgentsPPOTrainer::SetupPPOTrainer(
 	TrainingEnvironment = InTrainingEnvironment;
 	Trainer = Communicator.Trainer;
 
+	const int32 ObservationSchemaId = 0;
+	const int32 ActionSchemaId = 0;
+
 	// Create Episode Buffer
 	EpisodeBuffer = MakeUnique<UE::Learning::FEpisodeBuffer>();
-	EpisodeBuffer->Resize(
-		Manager->GetMaxAgentNum(),
-		TrainerSettings.MaxEpisodeStepNum,
-		Interactor->GetObservationVectorSize(),
-		Interactor->GetActionVectorSize(),
+	EpisodeBuffer->Resize(Manager->GetMaxAgentNum(), TrainerSettings.MaxEpisodeStepNum);
+	ObservationId = EpisodeBuffer->AddObservations(
+		TEXT("Observations"),
+		ObservationSchemaId,
+		Interactor->GetObservationVectorSize());
+	ActionId = EpisodeBuffer->AddActions(
+		TEXT("Actions"),
+		ActionSchemaId,
+		Interactor->GetActionVectorSize());
+	MemoryStateId = EpisodeBuffer->AddMemoryStates(
+		TEXT("MemoryStates"),
 		Policy->GetMemoryStateSize());
+	RewardId = EpisodeBuffer->AddRewards(
+		TEXT("Rewards"),
+		1);
 
 	// Create Replay Buffer
 	ReplayBuffer = MakeUnique<UE::Learning::FReplayBuffer>();
 	ReplayBuffer->Resize(
-		Interactor->GetObservationVectorSize(),
-		Interactor->GetActionVectorSize(),
-		Policy->GetMemoryStateSize(),
+		*EpisodeBuffer,
 		TrainerSettings.MaximumRecordedEpisodesPerIteration,
 		TrainerSettings.MaximumRecordedStepsPerIteration);
 
@@ -241,23 +276,24 @@ void ULearningAgentsPPOTrainer::BeginTraining(
 		return;
 	}
 
-	ApplyGameSettings(TrainingGameSettings);
+	UE::Learning::Agents::ApplyGameSettings(TrainingGameSettings, GetWorld(), PreviousGameSettingsState);
 
 	// We need to setup the trainer prior to sending the config
-	Trainer->AddNetwork(Policy->GetPolicyNetworkAsset()->GetFName(), *Policy->GetPolicyNetworkAsset()->NeuralNetworkData);
-	Trainer->AddNetwork(Critic->GetCriticNetworkAsset()->GetFName(), *Critic->GetCriticNetworkAsset()->NeuralNetworkData);
-	Trainer->AddNetwork(Policy->GetEncoderNetworkAsset()->GetFName(), *Policy->GetEncoderNetworkAsset()->NeuralNetworkData);
-	Trainer->AddNetwork(Policy->GetDecoderNetworkAsset()->GetFName(), *Policy->GetDecoderNetworkAsset()->NeuralNetworkData);
-	Trainer->AddReplayBuffer(TEXT("ReplayBuffer"), *ReplayBuffer);
+	PolicyNetworkId = Trainer->AddNetwork(*Policy->GetPolicyNetworkAsset()->NeuralNetworkData);
+	CriticNetworkId = Trainer->AddNetwork(*Critic->GetCriticNetworkAsset()->NeuralNetworkData);
+	EncoderNetworkId = Trainer->AddNetwork(*Policy->GetEncoderNetworkAsset()->NeuralNetworkData);
+	DecoderNetworkId = Trainer->AddNetwork(*Policy->GetDecoderNetworkAsset()->NeuralNetworkData);
+	ReplayBufferId = Trainer->AddReplayBuffer(*ReplayBuffer);
+	
+	TSharedRef<FJsonObject> ConfigObject = CreateConfig(TrainingSettings);
 
 	UE_LOG(LogLearning, Display, TEXT("%s: Sending config..."), *GetName());
-	SendConfig(TrainingSettings);
+	SendConfig(ConfigObject);
 
 	UE_LOG(LogLearning, Display, TEXT("%s: Sending initial policy..."), *GetName());
-
 	UE::Learning::ETrainerResponse Response = UE::Learning::ETrainerResponse::Success;
 
-	Response = Trainer->SendNetwork(Policy->GetPolicyNetworkAsset()->GetFName(), *Policy->GetPolicyNetworkAsset()->NeuralNetworkData);
+	Response = Trainer->SendNetwork(PolicyNetworkId, *Policy->GetPolicyNetworkAsset()->NeuralNetworkData);
 	if (Response != UE::Learning::ETrainerResponse::Success)
 	{
 		UE_LOG(LogLearning, Error, TEXT("%s: Error sending policy to trainer: %s. Check log for additional errors."), *GetName(), UE::Learning::Trainer::GetResponseString(Response));
@@ -266,7 +302,7 @@ void ULearningAgentsPPOTrainer::BeginTraining(
 		return;
 	}
 
-	Response = Trainer->SendNetwork(Critic->GetCriticNetworkAsset()->GetFName(), *Critic->GetCriticNetworkAsset()->NeuralNetworkData);
+	Response = Trainer->SendNetwork(CriticNetworkId, *Critic->GetCriticNetworkAsset()->NeuralNetworkData);
 	if (Response != UE::Learning::ETrainerResponse::Success)
 	{
 		UE_LOG(LogLearning, Error, TEXT("%s: Error sending critic to trainer: %s. Check log for additional errors."), *GetName(), UE::Learning::Trainer::GetResponseString(Response));
@@ -275,7 +311,7 @@ void ULearningAgentsPPOTrainer::BeginTraining(
 		return;
 	}
 
-	Response = Trainer->SendNetwork(Policy->GetEncoderNetworkAsset()->GetFName(), *Policy->GetEncoderNetworkAsset()->NeuralNetworkData);
+	Response = Trainer->SendNetwork(EncoderNetworkId, *Policy->GetEncoderNetworkAsset()->NeuralNetworkData);
 	if (Response != UE::Learning::ETrainerResponse::Success)
 	{
 		UE_LOG(LogLearning, Error, TEXT("%s: Error sending encoder to trainer: %s. Check log for additional errors."), *GetName(), UE::Learning::Trainer::GetResponseString(Response));
@@ -284,7 +320,7 @@ void ULearningAgentsPPOTrainer::BeginTraining(
 		return;
 	}
 
-	Response = Trainer->SendNetwork(Policy->GetDecoderNetworkAsset()->GetFName(), *Policy->GetDecoderNetworkAsset()->NeuralNetworkData);
+	Response = Trainer->SendNetwork(DecoderNetworkId, *Policy->GetDecoderNetworkAsset()->NeuralNetworkData);
 	if (Response != UE::Learning::ETrainerResponse::Success)
 	{
 		UE_LOG(LogLearning, Error, TEXT("%s: Error sending decoder to trainer: %s. Check log for additional errors."), *GetName(), UE::Learning::Trainer::GetResponseString(Response));
@@ -303,148 +339,121 @@ void ULearningAgentsPPOTrainer::BeginTraining(
 	bIsTraining = true;
 }
 
-void ULearningAgentsPPOTrainer::ApplyGameSettings(const FLearningAgentsTrainingGameSettings& Settings)
+TSharedRef<FJsonObject> ULearningAgentsPPOTrainer::CreateConfig(const FLearningAgentsPPOTrainingSettings& TrainingSettings) const
 {
-	// Record GameState Settings
+	TSharedRef<FJsonObject> ConfigObject = MakeShared<FJsonObject>();
 
-	bFixedTimestepUsed = FApp::UseFixedTimeStep();
-	FixedTimeStepDeltaTime = FApp::GetFixedDeltaTime();
+	// Add Training Config Entries
+	ConfigObject->SetStringField(TEXT("TaskName"), TEXT("Training"));
+	ConfigObject->SetStringField(TEXT("TrainerMethod"), TEXT("PPO"));
+	ConfigObject->SetStringField(TEXT("TimeStamp"), *FDateTime::Now().ToFormattedString(TEXT("%Y-%m-%d_%H-%M-%S")));
 
-	UGameUserSettings* GameSettings = UGameUserSettings::GetGameUserSettings();
-	if (GameSettings)
-	{
-		bVSyncEnabled = GameSettings->IsVSyncEnabled();
-	}
+	const int32 ObservationSchemaId = 0;
+	const int32 ActionSchemaId = 0;
 
-	UPhysicsSettings* PhysicsSettings = UPhysicsSettings::Get();
-	if (PhysicsSettings)
+	// Add Neural Network Config Entries
+	TArray<TSharedPtr<FJsonValue>> NetworkObjects;
+
+	// Policy
 	{
-		MaxPhysicsStep = PhysicsSettings->MaxPhysicsDeltaTime;
+		TSharedPtr<FJsonObject> NetworkObject = MakeShared<FJsonObject>();
+		NetworkObject->SetNumberField(TEXT("Id"), PolicyNetworkId);
+		NetworkObject->SetStringField(TEXT("Name"), Policy->GetPolicyNetworkAsset()->GetFName().ToString());
+		NetworkObject->SetNumberField(TEXT("MaxByteNum"), Policy->GetPolicyNetworkAsset()->NeuralNetworkData->GetSnapshotByteNum());
+		
+		TSharedRef<FJsonValueObject> JsonValue = MakeShared<FJsonValueObject>(NetworkObject);
+		NetworkObjects.Add(JsonValue);
 	}
 
-	IConsoleVariable* MaxFPSCVar = IConsoleManager::Get().FindConsoleVariable(TEXT("t.MaxFPS"));
-	if (MaxFPSCVar)
+	// Critic
 	{
-		MaxFPS = MaxFPSCVar->GetInt();
+		TSharedPtr<FJsonObject> NetworkObject = MakeShared<FJsonObject>();
+		NetworkObject->SetNumberField(TEXT("Id"), CriticNetworkId);
+		NetworkObject->SetStringField(TEXT("Name"), Critic->GetCriticNetworkAsset()->GetFName().ToString());
+		NetworkObject->SetNumberField(TEXT("MaxByteNum"), Critic->GetCriticNetworkAsset()->NeuralNetworkData->GetSnapshotByteNum());
+		NetworkObject->SetNumberField(TEXT("InputSchemaId"), ObservationSchemaId);
+
+		TSharedRef<FJsonValueObject> JsonValue = MakeShared<FJsonValueObject>(NetworkObject);
+		NetworkObjects.Add(JsonValue);
 	}
 
-	UGameViewportClient* ViewportClient = GetWorld() ? GetWorld()->GetGameViewport() : nullptr;
-	if (ViewportClient)
+	// Encoder
 	{
-		ViewModeIndex = ViewportClient->ViewModeIndex;
+		TSharedPtr<FJsonObject> NetworkObject = MakeShared<FJsonObject>();
+		NetworkObject->SetNumberField(TEXT("Id"), EncoderNetworkId);
+		NetworkObject->SetStringField(TEXT("Name"), Policy->GetEncoderNetworkAsset()->GetFName().ToString());
+		NetworkObject->SetNumberField(TEXT("MaxByteNum"), Policy->GetEncoderNetworkAsset()->NeuralNetworkData->GetSnapshotByteNum());
+		NetworkObject->SetNumberField(TEXT("InputSchemaId"), ObservationSchemaId);
+
+		TSharedRef<FJsonValueObject> JsonValue = MakeShared<FJsonValueObject>(NetworkObject);
+		NetworkObjects.Add(JsonValue);
 	}
 
-#if WITH_EDITOR
-	UEditorPerformanceSettings* EditorPerformanceSettings = GetMutableDefault<UEditorPerformanceSettings>();
-	if (EditorPerformanceSettings)
+	// Decoder
 	{
-		bUseLessCPUInTheBackground = EditorPerformanceSettings->bThrottleCPUWhenNotForeground;
-		bEditorVSyncEnabled = EditorPerformanceSettings->bEnableVSync;
+		TSharedPtr<FJsonObject> NetworkObject = MakeShared<FJsonObject>();
+		NetworkObject->SetNumberField(TEXT("Id"), DecoderNetworkId);
+		NetworkObject->SetStringField(TEXT("Name"), Policy->GetDecoderNetworkAsset()->GetFName().ToString());
+		NetworkObject->SetNumberField(TEXT("MaxByteNum"), Policy->GetDecoderNetworkAsset()->NeuralNetworkData->GetSnapshotByteNum());
+		NetworkObject->SetNumberField(TEXT("OutputSchemaId"), ActionSchemaId);
+
+		TSharedRef<FJsonValueObject> JsonValue = MakeShared<FJsonValueObject>(NetworkObject);
+		NetworkObjects.Add(JsonValue);
 	}
-#endif
 
-	// Apply Training GameState Settings
+	ConfigObject->SetArrayField(TEXT("Networks"), NetworkObjects);
 
-	FApp::SetUseFixedTimeStep(Settings.bUseFixedTimeStep);
+	// Add Replay Buffers Config Entries
+	TArray<TSharedPtr<FJsonValue>> ReplayBufferObjects;
+	TSharedRef<FJsonValueObject> ReplayBufferJsonValue = MakeShared<FJsonValueObject>(ReplayBuffer->AsJsonConfig(ReplayBufferId));
+	ReplayBufferObjects.Add(ReplayBufferJsonValue);
+	ConfigObject->SetArrayField(TEXT("ReplayBuffers"), ReplayBufferObjects);
 
-	if (Settings.FixedTimeStepFrequency > UE_SMALL_NUMBER)
-	{
-		FApp::SetFixedDeltaTime(1.0f / Settings.FixedTimeStepFrequency);
-		if (Settings.bSetMaxPhysicsStepToFixedTimeStep && PhysicsSettings)
-		{
-			PhysicsSettings->MaxPhysicsDeltaTime = 1.0f / Settings.FixedTimeStepFrequency;
-		}
-	}
-	else
-	{
-		UE_LOG(LogLearning, Warning, TEXT("%s: Provided invalid FixedTimeStepFrequency: %0.5f"), *GetName(), Settings.FixedTimeStepFrequency);
-	}
+	// Schemas
+	TSharedPtr<FJsonObject> SchemasObject = MakeShared<FJsonObject>();
 
-	if (Settings.bDisableMaxFPS && MaxFPSCVar)
+	// Add the observation schemas
+	TArray<TSharedPtr<FJsonValue>> ObservationSchemaObjects;
 	{
-		MaxFPSCVar->Set(0);
-	}
+		// For this PPO trainer, add the one observation schema we have
+		TSharedPtr<FJsonObject> ObservationSchemaObject = MakeShared<FJsonObject>();
+		ObservationSchemaObject->SetNumberField(TEXT("Id"), ObservationSchemaId);
+		ObservationSchemaObject->SetStringField(TEXT("Name"), "Default");
+		ObservationSchemaObject->SetObjectField(TEXT("Schema"),
+			UE::Learning::Trainer::ConvertObservationSchemaToJSON(Interactor->GetObservationSchema()->ObservationSchema,
+				Interactor->GetObservationSchemaElement().SchemaElement));
 
-	if (Settings.bDisableVSync && GameSettings)
-	{
-		GameSettings->SetVSyncEnabled(false);
-		GameSettings->ApplySettings(false);
+		TSharedRef<FJsonValueObject> JsonValue = MakeShared<FJsonValueObject>(ObservationSchemaObject);
+		ObservationSchemaObjects.Add(JsonValue);
 	}
+	SchemasObject->SetArrayField(TEXT("Observations"), ObservationSchemaObjects);
 
-	if (Settings.bUseUnlitViewportRendering && ViewportClient)
+	// Add the action schemas
+	TArray<TSharedPtr<FJsonValue>> ActionSchemaObjects;
 	{
-		ViewportClient->ViewModeIndex = EViewModeIndex::VMI_Unlit;
-	}
+		// For this PPO trainer, add the one action schema we have
+		TSharedPtr<FJsonObject> ActionSchemaObject = MakeShared<FJsonObject>();
+		ActionSchemaObject->SetNumberField(TEXT("Id"), ActionSchemaId);
+		ActionSchemaObject->SetStringField(TEXT("Name"), "Default");
+		ActionSchemaObject->SetObjectField(TEXT("Schema"),
+			UE::Learning::Trainer::ConvertActionSchemaToJSON(Interactor->GetActionSchema()->ActionSchema,
+				Interactor->GetActionSchemaElement().SchemaElement));
 
-#if WITH_EDITOR
-	if (Settings.bDisableUseLessCPUInTheBackground && EditorPerformanceSettings)
-	{
-		EditorPerformanceSettings->bThrottleCPUWhenNotForeground = false;
-		EditorPerformanceSettings->PostEditChange();
+		TSharedRef<FJsonValueObject> JsonValue = MakeShared<FJsonValueObject>(ActionSchemaObject);
+		ActionSchemaObjects.Add(JsonValue);
 	}
+	SchemasObject->SetArrayField(TEXT("Actions"), ActionSchemaObjects);
 
-	if (Settings.bDisableEditorVSync && EditorPerformanceSettings)
-	{
-		EditorPerformanceSettings->bEnableVSync = false;
-		EditorPerformanceSettings->PostEditChange();
-	}
-#endif
-}
+	ConfigObject->SetObjectField(TEXT("Schemas"), SchemasObject);
 
-void ULearningAgentsPPOTrainer::SendConfig(const FLearningAgentsPPOTrainingSettings& Settings)
-{
-	TSharedRef<FJsonObject> ConfigObject = MakeShared<FJsonObject>();
-	ConfigObject->SetStringField(TEXT("TaskName"), TEXT("Training"));
-	ConfigObject->SetStringField(TEXT("TrainerMethod"), TEXT("PPO"));
-	ConfigObject->SetStringField(TEXT("TrainerType"), TEXT("Network"));
-	ConfigObject->SetStringField(TEXT("TimeStamp"), *FDateTime::Now().ToFormattedString(TEXT("%Y-%m-%d_%H-%M-%S")));
+	// Add PPO Specific Config Entries
+	ConfigObject->SetObjectField(TEXT("PPOSettings"), TrainingSettings.AsJsonConfig());
 
-	ConfigObject->SetObjectField(TEXT("ObservationSchema"),
-		UE::Learning::Trainer::ConvertObservationSchemaToJSON(Interactor->GetObservationSchema()->ObservationSchema,
-		Interactor->GetObservationSchemaElement().SchemaElement));
-	ConfigObject->SetObjectField(TEXT("ActionSchema"),
-		UE::Learning::Trainer::ConvertActionSchemaToJSON(Interactor->GetActionSchema()->ActionSchema,
-		Interactor->GetActionSchemaElement().SchemaElement));
-	ConfigObject->SetNumberField(TEXT("ObservationVectorDimensionNum"), ReplayBuffer->GetObservations().Num<1>());
-	ConfigObject->SetNumberField(TEXT("ActionVectorDimensionNum"), ReplayBuffer->GetActions().Num<1>());
-	ConfigObject->SetNumberField(TEXT("MemoryStateVectorDimensionNum"), ReplayBuffer->GetMemoryStates().Num<1>());
-	ConfigObject->SetNumberField(TEXT("MaxEpisodeNum"), ReplayBuffer->GetMaxEpisodeNum());
-	ConfigObject->SetNumberField(TEXT("MaxStepNum"), ReplayBuffer->GetMaxStepNum());
-	
-	ConfigObject->SetNumberField(TEXT("PolicyNetworkByteNum"), Policy->GetPolicyNetworkAsset()->NeuralNetworkData->GetSnapshotByteNum());
-	ConfigObject->SetNumberField(TEXT("CriticNetworkByteNum"), Critic->GetCriticNetworkAsset()->NeuralNetworkData->GetSnapshotByteNum());
-	ConfigObject->SetNumberField(TEXT("EncoderNetworkByteNum"), Policy->GetEncoderNetworkAsset()->NeuralNetworkData->GetSnapshotByteNum());
-	ConfigObject->SetNumberField(TEXT("DecoderNetworkByteNum"), Policy->GetDecoderNetworkAsset()->NeuralNetworkData->GetSnapshotByteNum());
-
-	ConfigObject->SetNumberField(TEXT("IterationNum"), Settings.NumberOfIterations);
-	ConfigObject->SetNumberField(TEXT("LearningRatePolicy"), Settings.LearningRatePolicy);
-	ConfigObject->SetNumberField(TEXT("LearningRateCritic"), Settings.LearningRateCritic);
-	ConfigObject->SetNumberField(TEXT("LearningRateDecay"), Settings.LearningRateDecay);
-	ConfigObject->SetNumberField(TEXT("WeightDecay"), Settings.WeightDecay);
-	ConfigObject->SetNumberField(TEXT("PolicyBatchSize"), Settings.PolicyBatchSize);
-	ConfigObject->SetNumberField(TEXT("CriticBatchSize"), Settings.CriticBatchSize);
-	ConfigObject->SetNumberField(TEXT("PolicyWindow"), Settings.PolicyWindowSize);
-	ConfigObject->SetNumberField(TEXT("IterationsPerGather"), Settings.IterationsPerGather);
-	ConfigObject->SetNumberField(TEXT("CriticWarmupIterations"), Settings.CriticWarmupIterations);
-	ConfigObject->SetNumberField(TEXT("EpsilonClip"), Settings.EpsilonClip);
-	ConfigObject->SetNumberField(TEXT("ActionSurrogateWeight"), Settings.ActionSurrogateWeight);
-	ConfigObject->SetNumberField(TEXT("ActionRegularizationWeight"), Settings.ActionRegularizationWeight);
-	ConfigObject->SetNumberField(TEXT("ActionEntropyWeight"), Settings.ActionEntropyWeight);
-	ConfigObject->SetNumberField(TEXT("ReturnRegularizationWeight"), Settings.ReturnRegularizationWeight);
-	ConfigObject->SetNumberField(TEXT("GaeLambda"), Settings.GaeLambda);
-	ConfigObject->SetBoolField(TEXT("AdvantageNormalization"), Settings.bAdvantageNormalization);
-	ConfigObject->SetNumberField(TEXT("AdvantageMin"), Settings.MinimumAdvantage);
-	ConfigObject->SetNumberField(TEXT("AdvantageMax"), Settings.MaximumAdvantage);
-	ConfigObject->SetBoolField(TEXT("UseGradNormMaxClipping"), Settings.bUseGradNormMaxClipping);
-	ConfigObject->SetNumberField(TEXT("GradNormMax"), Settings.GradNormMax);
-	ConfigObject->SetNumberField(TEXT("TrimEpisodeStartStepNum"), Settings.NumberOfStepsToTrimAtStartOfEpisode);
-	ConfigObject->SetNumberField(TEXT("TrimEpisodeEndStepNum"), Settings.NumberOfStepsToTrimAtEndOfEpisode);
-	ConfigObject->SetNumberField(TEXT("Seed"), Settings.RandomSeed);
-	ConfigObject->SetNumberField(TEXT("DiscountFactor"), Settings.DiscountFactor);
-	ConfigObject->SetStringField(TEXT("Device"), UE::Learning::Trainer::GetDeviceString(UE::Learning::Agents::GetTrainingDevice(Settings.Device)));
-	ConfigObject->SetBoolField(TEXT("UseTensorBoard"), Settings.bUseTensorboard);
-	ConfigObject->SetBoolField(TEXT("SaveSnapshots"), Settings.bSaveSnapshots);
+	return ConfigObject;
+}
 
+void ULearningAgentsPPOTrainer::SendConfig(const TSharedRef<FJsonObject>& ConfigObject)
+{
 	UE::Learning::ETrainerResponse Response = UE::Learning::ETrainerResponse::Success;
 	Response = Trainer->SendConfig(ConfigObject);
 	
@@ -461,49 +470,12 @@ void ULearningAgentsPPOTrainer::DoneTraining()
 {
 	if (IsTraining())
 	{
-		// Wait for Trainer to finish
 		Trainer->Wait();
 
 		// If not finished in time, terminate
 		Trainer->Terminate();
 
-		// Apply back previous game settings
-		FApp::SetUseFixedTimeStep(bFixedTimestepUsed);
-		FApp::SetFixedDeltaTime(FixedTimeStepDeltaTime);
-		UGameUserSettings* GameSettings = UGameUserSettings::GetGameUserSettings();
-		if (GameSettings)
-		{
-			GameSettings->SetVSyncEnabled(bVSyncEnabled);
-			GameSettings->ApplySettings(true);
-		}
-
-		UPhysicsSettings* PhysicsSettings = UPhysicsSettings::Get();
-		if (PhysicsSettings)
-		{
-			PhysicsSettings->MaxPhysicsDeltaTime = MaxPhysicsStep;
-		}
-
-		IConsoleVariable* MaxFPSCVar = IConsoleManager::Get().FindConsoleVariable(TEXT("t.MaxFPS"));
-		if (MaxFPSCVar)
-		{
-			MaxFPSCVar->Set(MaxFPS);
-		}
-
-		UGameViewportClient* ViewportClient = GetWorld() ? GetWorld()->GetGameViewport() : nullptr;
-		if (ViewportClient)
-		{
-			ViewportClient->ViewModeIndex = ViewModeIndex;
-		}
-
-#if WITH_EDITOR
-		UEditorPerformanceSettings* EditorPerformanceSettings = GetMutableDefault<UEditorPerformanceSettings>();
-		if (EditorPerformanceSettings)
-		{
-			EditorPerformanceSettings->bThrottleCPUWhenNotForeground = bUseLessCPUInTheBackground;
-			EditorPerformanceSettings->bEnableVSync = bEditorVSyncEnabled;
-			EditorPerformanceSettings->PostEditChange();
-		}
-#endif
+		UE::Learning::Agents::RevertGameSettings(PreviousGameSettingsState, GetWorld());
 
 		bIsTraining = false;
 	}
@@ -585,12 +557,11 @@ void ULearningAgentsPPOTrainer::ProcessExperience(const bool bResetAgentsOnUpdat
 	}
 
 	// Add Experience to Episode Buffer
-	EpisodeBuffer->Push(
-		Interactor->GetObservationVectorArrayView(),
-		Interactor->GetActionVectorArrayView(),
-		Policy->GetPreEvaluationMemoryState(),
-		TrainingEnvironment->GetRewardArrayView(),
-		ValidAgentSet);
+	EpisodeBuffer->PushObservations(ObservationId, Interactor->GetObservationVectorArrayView(), ValidAgentSet);
+	EpisodeBuffer->PushActions(ActionId, Interactor->GetActionVectorArrayView(), ValidAgentSet);
+	EpisodeBuffer->PushMemoryStates(MemoryStateId, Policy->GetPreEvaluationMemoryState(), ValidAgentSet);
+	EpisodeBuffer->PushRewards(RewardId, TrainingEnvironment->GetRewardArrayView(), ValidAgentSet);
+	EpisodeBuffer->IncrementEpisodeStepNums(ValidAgentSet);
 
 	// Find the set of agents which have reached the maximum episode length and mark them as truncated
 	UE::Learning::Completion::EvaluateEndOfEpisodeCompletions(
@@ -615,14 +586,14 @@ void ULearningAgentsPPOTrainer::ProcessExperience(const bool bResetAgentsOnUpdat
 	// And push those episodes to the Replay Buffer
 	const bool bReplayBufferFull = ReplayBuffer->AddEpisodes(
 		TrainingEnvironment->GetAllCompletions(),
-		Interactor->GetObservationVectorArrayView(),
-		Policy->GetMemoryState(),
+		{ Interactor->GetObservationVectorArrayView() },
+		{ Policy->GetMemoryState() },
 		*EpisodeBuffer,
 		TrainingEnvironment->GetResetBuffer().GetResetInstances());
 
 	if (bReplayBufferFull)
 	{
-		UE::Learning::ETrainerResponse Response = Trainer->SendReplayBuffer(TEXT("ReplayBuffer"), *ReplayBuffer);
+		UE::Learning::ETrainerResponse Response = Trainer->SendReplayBuffer(ReplayBufferId, *ReplayBuffer);
 
 		if (Response != UE::Learning::ETrainerResponse::Success)
 		{
@@ -635,7 +606,7 @@ void ULearningAgentsPPOTrainer::ProcessExperience(const bool bResetAgentsOnUpdat
 		ReplayBuffer->Reset();
 
 		// Get Updated Policy
-		Response = Trainer->ReceiveNetwork(Policy->GetPolicyNetworkAsset()->GetFName(), *Policy->GetPolicyNetworkAsset()->NeuralNetworkData);
+		Response = Trainer->ReceiveNetwork(PolicyNetworkId, *Policy->GetPolicyNetworkAsset()->NeuralNetworkData);
 		Policy->GetPolicyNetworkAsset()->ForceMarkDirty();
 
 		if (Response == UE::Learning::ETrainerResponse::Completed)
@@ -653,7 +624,7 @@ void ULearningAgentsPPOTrainer::ProcessExperience(const bool bResetAgentsOnUpdat
 		}
 
 		// Get Updated Critic
-		Response = Trainer->ReceiveNetwork(Critic->GetCriticNetworkAsset()->GetFName(), *Critic->GetCriticNetworkAsset()->NeuralNetworkData);
+		Response = Trainer->ReceiveNetwork(CriticNetworkId, *Critic->GetCriticNetworkAsset()->NeuralNetworkData);
 		Critic->GetCriticNetworkAsset()->ForceMarkDirty();
 
 		if (Response != UE::Learning::ETrainerResponse::Success)
@@ -665,7 +636,7 @@ void ULearningAgentsPPOTrainer::ProcessExperience(const bool bResetAgentsOnUpdat
 		}
 
 		// Get Updated Encoder
-		Response = Trainer->ReceiveNetwork(Policy->GetEncoderNetworkAsset()->GetFName(), *Policy->GetEncoderNetworkAsset()->NeuralNetworkData);
+		Response = Trainer->ReceiveNetwork(EncoderNetworkId, *Policy->GetEncoderNetworkAsset()->NeuralNetworkData);
 		Policy->GetEncoderNetworkAsset()->ForceMarkDirty();
 
 		if (Response != UE::Learning::ETrainerResponse::Success)
@@ -677,7 +648,7 @@ void ULearningAgentsPPOTrainer::ProcessExperience(const bool bResetAgentsOnUpdat
 		}
 
 		// Get Updated Decoder
-		Response = Trainer->ReceiveNetwork(Policy->GetDecoderNetworkAsset()->GetFName(), *Policy->GetDecoderNetworkAsset()->NeuralNetworkData);
+		Response = Trainer->ReceiveNetwork(DecoderNetworkId, *Policy->GetDecoderNetworkAsset()->NeuralNetworkData);
 		Policy->GetDecoderNetworkAsset()->ForceMarkDirty();
 
 		if (Response != UE::Learning::ETrainerResponse::Success)
diff --git a/Engine/Plugins/Experimental/LearningAgents/Source/LearningAgentsTraining/Private/LearningAgentsTrainer.cpp b/Engine/Plugins/Experimental/LearningAgents/Source/LearningAgentsTraining/Private/LearningAgentsTrainer.cpp
--- a/Engine/Plugins/Experimental/LearningAgents/Source/LearningAgentsTraining/Private/LearningAgentsTrainer.cpp
+++ b/Engine/Plugins/Experimental/LearningAgents/Source/LearningAgentsTraining/Private/LearningAgentsTrainer.cpp
@@ -4,6 +4,15 @@
 
 #include "Misc/Paths.h"
 
+#include "Misc/App.h"
+#include "GameFramework/GameUserSettings.h"
+#include "PhysicsEngine/PhysicsSettings.h"
+#include "Engine/GameViewportClient.h"
+
+#if WITH_EDITOR
+#include "Editor/EditorPerformanceSettings.h"
+#endif
+
 namespace UE::Learning::Agents
 {
 	ELearningAgentsTrainingDevice GetLearningAgentsTrainingDevice(const ETrainerDevice Device)
@@ -25,6 +34,134 @@ namespace UE::Learning::Agents
 		default:UE_LOG(LogLearning, Error, TEXT("Unknown Trainer Device.")); return ETrainerDevice::CPU;
 		}
 	}
+
+	void ApplyGameSettings(const FLearningAgentsTrainingGameSettings& Settings, const UWorld* World, FGameSettingsState& OutGameSettingsState)
+	{
+		// Record GameState Settings
+
+		OutGameSettingsState.bFixedTimestepUsed = FApp::UseFixedTimeStep();
+		OutGameSettingsState.FixedTimeStepDeltaTime = FApp::GetFixedDeltaTime();
+
+		UGameUserSettings* GameSettings = UGameUserSettings::GetGameUserSettings();
+		if (GameSettings)
+		{
+			OutGameSettingsState.bVSyncEnabled = GameSettings->IsVSyncEnabled();
+		}
+
+		UPhysicsSettings* PhysicsSettings = UPhysicsSettings::Get();
+		if (PhysicsSettings)
+		{
+			OutGameSettingsState.MaxPhysicsStep = PhysicsSettings->MaxPhysicsDeltaTime;
+		}
+
+		IConsoleVariable* MaxFPSCVar = IConsoleManager::Get().FindConsoleVariable(TEXT("t.MaxFPS"));
+		if (MaxFPSCVar)
+		{
+			OutGameSettingsState.MaxFPS = MaxFPSCVar->GetInt();
+		}
+
+		UGameViewportClient* ViewportClient = World ? World->GetGameViewport() : nullptr;
+		if (ViewportClient)
+		{
+			OutGameSettingsState.ViewModeIndex = ViewportClient->ViewModeIndex;
+		}
+
+#if WITH_EDITOR
+		UEditorPerformanceSettings* EditorPerformanceSettings = GetMutableDefault<UEditorPerformanceSettings>();
+		if (EditorPerformanceSettings)
+		{
+			OutGameSettingsState.bUseLessCPUInTheBackground = EditorPerformanceSettings->bThrottleCPUWhenNotForeground;
+			OutGameSettingsState.bEditorVSyncEnabled = EditorPerformanceSettings->bEnableVSync;
+		}
+#endif
+
+		// Apply Training GameState Settings
+
+		FApp::SetUseFixedTimeStep(Settings.bUseFixedTimeStep);
+
+		if (Settings.FixedTimeStepFrequency > UE_SMALL_NUMBER)
+		{
+			FApp::SetFixedDeltaTime(1.0f / Settings.FixedTimeStepFrequency);
+			if (Settings.bSetMaxPhysicsStepToFixedTimeStep && PhysicsSettings)
+			{
+				PhysicsSettings->MaxPhysicsDeltaTime = 1.0f / Settings.FixedTimeStepFrequency;
+			}
+		}
+		else
+		{
+			UE_LOG(LogLearning, Warning, TEXT("Provided invalid FixedTimeStepFrequency: %0.5f"), Settings.FixedTimeStepFrequency);
+		}
+
+		if (Settings.bDisableMaxFPS && MaxFPSCVar)
+		{
+			MaxFPSCVar->Set(0);
+		}
+
+		if (Settings.bDisableVSync && GameSettings)
+		{
+			GameSettings->SetVSyncEnabled(false);
+			GameSettings->ApplySettings(false);
+		}
+
+		if (Settings.bUseUnlitViewportRendering && ViewportClient)
+		{
+			ViewportClient->ViewModeIndex = EViewModeIndex::VMI_Unlit;
+		}
+
+#if WITH_EDITOR
+		if (Settings.bDisableUseLessCPUInTheBackground && EditorPerformanceSettings)
+		{
+			EditorPerformanceSettings->bThrottleCPUWhenNotForeground = false;
+			EditorPerformanceSettings->PostEditChange();
+		}
+
+		if (Settings.bDisableEditorVSync && EditorPerformanceSettings)
+		{
+			EditorPerformanceSettings->bEnableVSync = false;
+			EditorPerformanceSettings->PostEditChange();
+		}
+#endif
+	}
+
+	void RevertGameSettings(const FGameSettingsState& Settings, const UWorld* World)
+	{
+		FApp::SetUseFixedTimeStep(Settings.bFixedTimestepUsed);
+		FApp::SetFixedDeltaTime(Settings.FixedTimeStepDeltaTime);
+		UGameUserSettings* GameSettings = UGameUserSettings::GetGameUserSettings();
+		if (GameSettings)
+		{
+			GameSettings->SetVSyncEnabled(Settings.bVSyncEnabled);
+			GameSettings->ApplySettings(true);
+		}
+
+		UPhysicsSettings* PhysicsSettings = UPhysicsSettings::Get();
+		if (PhysicsSettings)
+		{
+			PhysicsSettings->MaxPhysicsDeltaTime = Settings.MaxPhysicsStep;
+		}
+
+		IConsoleVariable* MaxFPSCVar = IConsoleManager::Get().FindConsoleVariable(TEXT("t.MaxFPS"));
+		if (MaxFPSCVar)
+		{
+			MaxFPSCVar->Set(Settings.MaxFPS);
+		}
+
+		UGameViewportClient* ViewportClient = World ? World->GetGameViewport() : nullptr;
+		if (ViewportClient)
+		{
+			ViewportClient->ViewModeIndex = Settings.ViewModeIndex;
+		}
+
+#if WITH_EDITOR
+		UEditorPerformanceSettings* EditorPerformanceSettings = GetMutableDefault<UEditorPerformanceSettings>();
+		if (EditorPerformanceSettings)
+		{
+			EditorPerformanceSettings->bThrottleCPUWhenNotForeground = Settings.bUseLessCPUInTheBackground;
+			EditorPerformanceSettings->bEnableVSync = Settings.bEditorVSyncEnabled;
+			EditorPerformanceSettings->PostEditChange();
+		}
+#endif
+	}
 }
 
 FLearningAgentsTrainerProcessSettings::FLearningAgentsTrainerProcessSettings()
diff --git a/Engine/Plugins/Experimental/LearningAgents/Source/LearningAgentsTraining/Public/LearningAgentsImitationTrainer.h b/Engine/Plugins/Experimental/LearningAgents/Source/LearningAgentsTraining/Public/LearningAgentsImitationTrainer.h
--- a/Engine/Plugins/Experimental/LearningAgents/Source/LearningAgentsTraining/Public/LearningAgentsImitationTrainer.h
+++ b/Engine/Plugins/Experimental/LearningAgents/Source/LearningAgentsTraining/Public/LearningAgentsImitationTrainer.h
@@ -5,8 +5,6 @@
 #include "LearningAgentsManagerListener.h"
 #include "LearningAgentsTrainer.h"
 
-#include "LearningArray.h"
-
 #include "Templates/SharedPointer.h"
 #include "UObject/ObjectPtr.h"
 
@@ -14,7 +12,8 @@
 
 namespace UE::Learning
 {
-	struct FSharedMemoryImitationTrainer;
+	struct FReplayBuffer;
+	struct IExternalTrainer;
 }
 
 class ULearningAgentsInteractor;
@@ -104,6 +103,8 @@ public:
 	/** If true, snapshots of the trained networks will be emitted to the intermediate directory. */
 	UPROPERTY(EditAnywhere, Category = "LearningAgents")
 	bool bSaveSnapshots = false;
+
+	TSharedRef<FJsonObject> AsJsonConfig() const;
 };
 
 /**
@@ -137,14 +138,18 @@ public:
 	 * @param InManager			The agent manager we are using.
 	 * @param InInteractor		The agent interactor we are recording with.
 	 * @param InPolicy			The policy we are using.
-	 * @param Class				The imitation trainer class
+	 * @param Communicator		The communicator.
+	 * @param Class				The trainer class.
+	 * @param Name				The trainer name.
 	 */
 	UFUNCTION(BlueprintCallable, Category = "LearningAgents", meta = (Class = "/Script/LearningAgentsTraining.LearningAgentsImitationTrainer", DeterminesOutputType = "Class"))
 	static ULearningAgentsImitationTrainer* MakeImitationTrainer(
 		UPARAM(ref) ULearningAgentsManager*& InManager,
 		UPARAM(ref) ULearningAgentsInteractor*& InInteractor,
 		UPARAM(ref) ULearningAgentsPolicy*& InPolicy,
-		TSubclassOf<ULearningAgentsImitationTrainer> Class);
+		const FLearningAgentsCommunicator& Communicator,
+		TSubclassOf<ULearningAgentsImitationTrainer> Class,
+		const FName Name = TEXT("ImitationTrainer"));
 
 	/**
 	 * Initializes the imitation trainer and runs the setup functions.
@@ -152,11 +157,13 @@ public:
 	 * @param InManager			The agent manager we are using.
 	 * @param InInteractor		The agent interactor we are recording with.
 	 * @param InPolicy			The policy we are using.
+	 * @param InCommunicator	The communicator.
 	 */
 	void SetupImitationTrainer(
 		ULearningAgentsManager* InManager,
 		ULearningAgentsInteractor* InInteractor,
-		ULearningAgentsPolicy* InPolicy);
+		ULearningAgentsPolicy* InPolicy,
+		const FLearningAgentsCommunicator& Communicator);
 
 	/** Returns true if the trainer is currently training; Otherwise, false. */
 	UFUNCTION(BlueprintPure, Category = "LearningAgents")
@@ -232,9 +239,21 @@ private:
 	UPROPERTY(VisibleAnywhere, Transient, Category = "LearningAgents")
 	bool bHasTrainingFailed = false;
 
-	float TrainerTimeout = 10.0f;
+	TSharedRef<FJsonObject> CreateConfig(const FLearningAgentsImitationTrainerTrainingSettings& TrainingSettings) const;
+	void SendConfig(const TSharedRef<FJsonObject>& ConfigObject);
 
 	void DoneTraining();
 
-	TUniquePtr<UE::Learning::FSharedMemoryImitationTrainer> ImitationTrainer;
+	TUniquePtr<UE::Learning::FReplayBuffer> ReplayBuffer;
+	TSharedPtr<UE::Learning::IExternalTrainer> Trainer;
+
+	int32 PolicyNetworkId = INDEX_NONE;
+	int32 EncoderNetworkId = INDEX_NONE;
+	int32 DecoderNetworkId = INDEX_NONE;
+
+	int32 ReplayBufferId = INDEX_NONE;
+
+	int32 ObservationId = INDEX_NONE;
+	int32 ActionId = INDEX_NONE;
+	int32 MemoryStateId = INDEX_NONE;
 };
diff --git a/Engine/Plugins/Experimental/LearningAgents/Source/LearningAgentsTraining/Public/LearningAgentsPPOTrainer.h b/Engine/Plugins/Experimental/LearningAgents/Source/LearningAgentsTraining/Public/LearningAgentsPPOTrainer.h
--- a/Engine/Plugins/Experimental/LearningAgents/Source/LearningAgentsTraining/Public/LearningAgentsPPOTrainer.h
+++ b/Engine/Plugins/Experimental/LearningAgents/Source/LearningAgentsTraining/Public/LearningAgentsPPOTrainer.h
@@ -6,9 +6,12 @@
 #include "LearningAgentsTrainer.h"
 
 #include "Templates/SharedPointer.h"
+#include "UObject/ObjectPtr.h"
 
 #include "LearningAgentsPPOTrainer.generated.h"
 
+class FJsonObject;
+
 namespace UE::Learning
 {
 	struct FEpisodeBuffer;
@@ -233,62 +236,8 @@ public:
 	/** If true, snapshots of the trained networks will be emitted to the intermediate directory. */
 	UPROPERTY(EditAnywhere, Category = "LearningAgents")
 	bool bSaveSnapshots = false;
-};
-
-/**
- * The configurable game settings for a ULearningAgentsTrainer. These allow the timestep and physics tick to be fixed
- * during training, which can enable ticking faster than real-time.
- */
-USTRUCT(BlueprintType, Category = "LearningAgents")
-struct LEARNINGAGENTSTRAINING_API FLearningAgentsTrainingGameSettings
-{
-	GENERATED_BODY()
-
-public:
 
-	/**
-	 * If true, the game will run in fixed time step mode (i.e the frame's delta times will always be the same
-	 * regardless of how much wall time has passed). This can enable faster than real-time training if your game runs
-	 * quickly. If false, the time steps will match real wall time.
-	 */
-	UPROPERTY(EditAnywhere, Category = "LearningAgents")
-	bool bUseFixedTimeStep = true;
-
-	/**
-	 * Determines the amount of time for each frame when bUseFixedTimeStep is true; Ignored if false. You want this
-	 * time step to match as closely as possible to the expected inference time steps, otherwise your training results
-	 * may not generalize to your game.
-	 */
-	UPROPERTY(EditAnywhere, Category = "LearningAgents", meta = (DisplayName = "Fixed Time Step Frequency (Hz)"), meta = (ClampMin = "0.0", UIMin = "0.0"))
-	float FixedTimeStepFrequency = 60.0f;
-
-	/** If true, set the physics delta time to match the fixed time step. */
-	UPROPERTY(EditAnywhere, Category = "LearningAgents")
-	bool bSetMaxPhysicsStepToFixedTimeStep = true;
-
-	/** If true, the MaxFPS console variable will be set to a negative number during training; Otherwise, it will not. */
-	UPROPERTY(EditAnywhere, Category = "LearningAgents")
-	bool bDisableMaxFPS = true;
-
-	/** If true, VSync will be disabled; Otherwise, it will not. Disabling VSync can speed up the game simulation. */
-	UPROPERTY(EditAnywhere, Category = "LearningAgents")
-	bool bDisableVSync = true;
-
-	/** If true, the viewport rendering will be unlit; Otherwise, it will not. Disabling lighting can speed up the game simulation. */
-	UPROPERTY(EditAnywhere, Category = "LearningAgents")
-	bool bUseUnlitViewportRendering = false;
-
-#if WITH_EDITORONLY_DATA
-
-	/** If true, the Use Less CPU In The Background editor setting will be disabled. This prevents the editor from running slowly when minimized. */
-	UPROPERTY(EditAnywhere, Category = "LearningAgents")
-	bool bDisableUseLessCPUInTheBackground = true;
-
-	/** If true, Editor VSync will be disabled; Otherwise, it will not. Disabling Editor VSync can speed up the game simulation. */
-	UPROPERTY(EditAnywhere, Category = "LearningAgents")
-	bool bDisableEditorVSync = true;
-
-#endif
+	TSharedRef<FJsonObject> AsJsonConfig() const;
 };
 
 UCLASS(BlueprintType, Blueprintable, meta = (BlueprintSpawnableComponent))
@@ -316,7 +265,7 @@ public:
 	 * @param InPolicy				The policy to be trained.
 	 * @param InCritic				The critic to be trained.
 	 * @param Communicator			The communicator.
-	 * @param Class					The trainer class
+	 * @param Class					The trainer class.
 	 * @param Name					The trainer name.
 	 * @param TrainerSettings		The trainer settings to use.
 	 */
@@ -465,22 +414,22 @@ private:
 	TUniquePtr<UE::Learning::FReplayBuffer> ReplayBuffer;
 	TSharedPtr<UE::Learning::IExternalTrainer> Trainer;
 
-	void ApplyGameSettings(const FLearningAgentsTrainingGameSettings& Settings);
-
-	void SendConfig(const FLearningAgentsPPOTrainingSettings& Settings);
+	TSharedRef<FJsonObject> CreateConfig(const FLearningAgentsPPOTrainingSettings& TrainingSettings) const;
+	void SendConfig(const TSharedRef<FJsonObject>& ConfigObject);
 
 	void DoneTraining();
 
-// ----- Private Recording of GameSettings ----- 
-private:
+	UE::Learning::Agents::FGameSettingsState PreviousGameSettingsState;
+
+	int32 PolicyNetworkId = INDEX_NONE;
+	int32 CriticNetworkId = INDEX_NONE;
+	int32 EncoderNetworkId = INDEX_NONE;
+	int32 DecoderNetworkId = INDEX_NONE;
 
-	bool bFixedTimestepUsed = false;
-	float FixedTimeStepDeltaTime = -1.0f;
-	float MaxPhysicsStep = -1.0f;
-	int32 MaxFPS = 120;
-	bool bVSyncEnabled = true;
-	int32 ViewModeIndex = -1;
+	int32 ReplayBufferId = INDEX_NONE;
 
-	bool bUseLessCPUInTheBackground = true;
-	bool bEditorVSyncEnabled = true;
+	int32 ObservationId = INDEX_NONE;
+	int32 ActionId = INDEX_NONE;
+	int32 MemoryStateId = INDEX_NONE;
+	int32 RewardId = INDEX_NONE;
 };
diff --git a/Engine/Plugins/Experimental/LearningAgents/Source/LearningAgentsTraining/Public/LearningAgentsTrainer.h b/Engine/Plugins/Experimental/LearningAgents/Source/LearningAgentsTraining/Public/LearningAgentsTrainer.h
--- a/Engine/Plugins/Experimental/LearningAgents/Source/LearningAgentsTraining/Public/LearningAgentsTrainer.h
+++ b/Engine/Plugins/Experimental/LearningAgents/Source/LearningAgentsTraining/Public/LearningAgentsTrainer.h
@@ -14,13 +14,87 @@ enum class ELearningAgentsTrainingDevice : uint8
 	GPU,
 };
 
+/**
+ * The configurable game settings for a ULearningAgentsTrainer. These allow the timestep and physics tick to be fixed
+ * during training, which can enable ticking faster than real-time.
+ */
+USTRUCT(BlueprintType, Category = "LearningAgents")
+struct LEARNINGAGENTSTRAINING_API FLearningAgentsTrainingGameSettings
+{
+	GENERATED_BODY()
+
+public:
+
+	/**
+	 * If true, the game will run in fixed time step mode (i.e the frame's delta times will always be the same
+	 * regardless of how much wall time has passed). This can enable faster than real-time training if your game runs
+	 * quickly. If false, the time steps will match real wall time.
+	 */
+	UPROPERTY(EditAnywhere, Category = "LearningAgents")
+	bool bUseFixedTimeStep = true;
+
+	/**
+	 * Determines the amount of time for each frame when bUseFixedTimeStep is true; Ignored if false. You want this
+	 * time step to match as closely as possible to the expected inference time steps, otherwise your training results
+	 * may not generalize to your game.
+	 */
+	UPROPERTY(EditAnywhere, Category = "LearningAgents", meta = (DisplayName = "Fixed Time Step Frequency (Hz)"), meta = (ClampMin = "0.0", UIMin = "0.0"))
+	float FixedTimeStepFrequency = 60.0f;
+
+	/** If true, set the physics delta time to match the fixed time step. */
+	UPROPERTY(EditAnywhere, Category = "LearningAgents")
+	bool bSetMaxPhysicsStepToFixedTimeStep = true;
+
+	/** If true, the MaxFPS console variable will be set to a negative number during training; Otherwise, it will not. */
+	UPROPERTY(EditAnywhere, Category = "LearningAgents")
+	bool bDisableMaxFPS = true;
+
+	/** If true, VSync will be disabled; Otherwise, it will not. Disabling VSync can speed up the game simulation. */
+	UPROPERTY(EditAnywhere, Category = "LearningAgents")
+	bool bDisableVSync = true;
+
+	/** If true, the viewport rendering will be unlit; Otherwise, it will not. Disabling lighting can speed up the game simulation. */
+	UPROPERTY(EditAnywhere, Category = "LearningAgents")
+	bool bUseUnlitViewportRendering = false;
+
+#if WITH_EDITORONLY_DATA
+
+	/** If true, the Use Less CPU In The Background editor setting will be disabled. This prevents the editor from running slowly when minimized. */
+	UPROPERTY(EditAnywhere, Category = "LearningAgents")
+	bool bDisableUseLessCPUInTheBackground = true;
+
+	/** If true, Editor VSync will be disabled; Otherwise, it will not. Disabling Editor VSync can speed up the game simulation. */
+	UPROPERTY(EditAnywhere, Category = "LearningAgents")
+	bool bDisableEditorVSync = true;
+
+#endif
+};
+
 namespace UE::Learning::Agents
 {
+	// ----- Private Recording of GameSettings ----- 
+	struct FGameSettingsState
+	{
+		bool bFixedTimestepUsed = false;
+		float FixedTimeStepDeltaTime = -1.0f;
+		float MaxPhysicsStep = -1.0f;
+		int32 MaxFPS = 120;
+		bool bVSyncEnabled = true;
+		int32 ViewModeIndex = -1;
+
+		bool bUseLessCPUInTheBackground = true;
+		bool bEditorVSyncEnabled = true;
+	};
+
 	/** Get the learning agents trainer device from the UE::Learning trainer device. */
 	LEARNINGAGENTSTRAINING_API ELearningAgentsTrainingDevice GetLearningAgentsTrainingDevice(const ETrainerDevice Device);
 
 	/** Get the UE::Learning trainer device from the learning agents trainer device. */
 	LEARNINGAGENTSTRAINING_API ETrainerDevice GetTrainingDevice(const ELearningAgentsTrainingDevice Device);
+
+	LEARNINGAGENTSTRAINING_API void ApplyGameSettings(const FLearningAgentsTrainingGameSettings& Settings, const UWorld* World, FGameSettingsState& OutGameSettingsState);
+
+	LEARNINGAGENTSTRAINING_API void RevertGameSettings(const FGameSettingsState& Settings, const UWorld* World);
 }
 
 /** The path settings for the trainer. */
@@ -66,6 +140,10 @@ public:
 	UPROPERTY(EditAnywhere, Category = "LearningAgents")
 	FString NonEditorIntermediateRelativePath;
 
+	/** Trainer file name. The name of the python file to use for training. Do NOT include the '.py' file extension. */
+	UPROPERTY(EditAnywhere, Category = "LearningAgents")
+	FString TrainerFileName = TEXT("train_ppo");
+
 public:
 
 	/** Gets the Relative Editor Engine Path accounting for if this is an editor build or not  */
diff --git a/Engine/Plugins/Experimental/LearningAgents/Source/LearningTraining/Private/LearningExperience.cpp b/Engine/Plugins/Experimental/LearningAgents/Source/LearningTraining/Private/LearningExperience.cpp
--- a/Engine/Plugins/Experimental/LearningAgents/Source/LearningTraining/Private/LearningExperience.cpp
+++ b/Engine/Plugins/Experimental/LearningAgents/Source/LearningTraining/Private/LearningExperience.cpp
@@ -2,112 +2,318 @@
 
 #include "LearningExperience.h"
 
+#include "LearningObservation.h"
+#include "LearningAction.h"
 #include "LearningCompletion.h"
 #include "LearningLog.h"
+
+#include "Dom/JsonObject.h"
 #include "Templates/Function.h"
 
 namespace UE::Learning
 {
-	void FEpisodeBuffer::Resize(
-		const int32 InMaxInstanceNum,
-		const int32 InMaxStepNum,
-		const int32 InObservationVectorDimNum,
-		const int32 InActionVectorDimNum,
-		const int32 InMemoryStateVectorDimNum)
+	void FEpisodeBuffer::Resize(const int32 InMaxInstanceNum, const int32 InMaxStepNum)
 	{
 		MaxInstanceNum = InMaxInstanceNum;
 		MaxStepNum = InMaxStepNum;
-		EpisodeStepNums.SetNumUninitialized({ InMaxInstanceNum });
-		Observations.SetNumUninitialized({ InMaxInstanceNum, InMaxStepNum, InObservationVectorDimNum });
-		Actions.SetNumUninitialized({ InMaxInstanceNum, InMaxStepNum, InActionVectorDimNum });
-		MemoryStates.SetNumUninitialized({ InMaxInstanceNum, InMaxStepNum, InMemoryStateVectorDimNum });
-		Rewards.SetNumUninitialized({ InMaxInstanceNum, InMaxStepNum });
 
+		// Observations
+		for (int32 NameIndex = 0; NameIndex < ObservationNames.Num(); NameIndex++)
+		{
+			int32 ObservationSize = ObservationSizes[NameIndex];
+			ObservationArrays[NameIndex].SetNumUninitialized({ InMaxInstanceNum, InMaxStepNum, ObservationSize });
+		}
+
+		// Actions
+		for (int32 NameIndex = 0; NameIndex < ActionNames.Num(); NameIndex++)
+		{
+			int32 ActionSize = ActionSizes[NameIndex];
+			ActionArrays[NameIndex].SetNumUninitialized({ InMaxInstanceNum, InMaxStepNum, ActionSize });
+		}
+
+		// Memory States
+		for (int32 NameIndex = 0; NameIndex < MemoryStateNames.Num(); NameIndex++)
+		{
+			int32 MemoryStateSize = MemoryStateSizes[NameIndex];
+			MemoryStateArrays[NameIndex].SetNumUninitialized({ InMaxInstanceNum, InMaxStepNum, MemoryStateSize });
+		}
+
+		// Rewards
+		for (int32 NameIndex = 0; NameIndex < RewardNames.Num(); NameIndex++)
+		{
+			int32 RewardSize = RewardSizes[NameIndex];
+			RewardArrays[NameIndex].SetNumUninitialized({ InMaxInstanceNum, InMaxStepNum, RewardSize });
+		}
+
+		// Episode Step Nums
+		EpisodeStepNums.SetNumUninitialized({ InMaxInstanceNum });
 		Array::Zero(EpisodeStepNums);
+
+		bHasBeenSized = true;
+	}
+
+	void FEpisodeBuffer::Reset(const FIndexSet Instances)
+	{
+		Array::Zero(EpisodeStepNums, Instances);
+	}
+
+	int32 FEpisodeBuffer::AddObservations(const FName& Name, const int32 SchemaId, const int32 Size)
+	{
+		UE_LEARNING_CHECKF(!ObservationNames.Contains(Name), TEXT("Observation name collision!"));
+
+		int32 ObservationId = ObservationNames.Num();
+		ObservationNames.Add(Name);
+		ObservationSchemaIds.Add(SchemaId);
+		ObservationSizes.Add(Size);
+		TLearningArray<3, float>& NewObservation = ObservationArrays.AddDefaulted_GetRef();
+
+		if (bHasBeenSized)
+		{
+			NewObservation.SetNumUninitialized({ MaxInstanceNum, MaxStepNum, Size });
+		}
+
+		return ObservationId;
 	}
 
-	void FEpisodeBuffer::Push(
-		const TLearningArrayView<2, const float> InObservations,
-		const TLearningArrayView<2, const float> InActions,
-		const TLearningArrayView<2, const float> InMemoryStates,
-		const TLearningArrayView<1, const float> InRewards,
-		const FIndexSet Instances)
+	void FEpisodeBuffer::PushObservations(const int32 ObservationId, const TLearningArrayView<2, const float> InObservations, const FIndexSet Instances)
 	{
-		UE_LEARNING_TRACE_CPUPROFILER_EVENT_SCOPE(Learning::FEpisodeBuffer::Push);
+		UE_LEARNING_TRACE_CPUPROFILER_EVENT_SCOPE(Learning::FEpisodeBuffer::PushObservations);
+
+		UE_LEARNING_CHECKF(ObservationId >= 0 && ObservationId < ObservationArrays.Num(), TEXT("Observation id invalid!"));
 
 		for (const int32 InstanceIdx : Instances)
 		{
-			UE_LEARNING_CHECKF(EpisodeStepNums[InstanceIdx] < MaxStepNum, TEXT("Episode Buffer full!"));
+			Array::Copy(ObservationArrays[ObservationId][InstanceIdx][EpisodeStepNums[InstanceIdx]], InObservations[InstanceIdx]);
+		}
+	}
 
-			Array::Copy(Observations[InstanceIdx][EpisodeStepNums[InstanceIdx]], InObservations[InstanceIdx]);
-			Array::Copy(Actions[InstanceIdx][EpisodeStepNums[InstanceIdx]], InActions[InstanceIdx]);
-			Array::Copy(MemoryStates[InstanceIdx][EpisodeStepNums[InstanceIdx]], InMemoryStates[InstanceIdx]);
-			Rewards[InstanceIdx][EpisodeStepNums[InstanceIdx]] = InRewards[InstanceIdx];
+	const TLearningArrayView<2, const float> FEpisodeBuffer::GetObservations(const int32 ObservationId, const int32 InstanceIdx) const
+	{
+		UE_LEARNING_CHECKF(ObservationId >= 0 && ObservationId < ObservationArrays.Num(), TEXT("Observation id invalid!"));
 
-			EpisodeStepNums[InstanceIdx]++;
+		return ObservationArrays[ObservationId][InstanceIdx].Slice(0, EpisodeStepNums[InstanceIdx]);
+	}
+
+	int32 FEpisodeBuffer::AddActions(const FName& Name, const int32 SchemaId, const int32 Size)
+	{
+		UE_LEARNING_CHECKF(!ActionNames.Contains(Name), TEXT("Action name collision!"));
+
+		int32 ActionId = ActionNames.Num();
+		ActionNames.Add(Name);
+		ActionSchemaIds.Add(SchemaId);
+		ActionSizes.Add(Size);
+		TLearningArray<3, float>& NewAction = ActionArrays.AddDefaulted_GetRef();
+
+		if (bHasBeenSized)
+		{
+			NewAction.SetNumUninitialized({ MaxInstanceNum, MaxStepNum, Size });
 		}
+
+		return ActionId;
 	}
 
-	void FEpisodeBuffer::Reset(const FIndexSet Instances)
+	void FEpisodeBuffer::PushActions(const int32 ActionId, const TLearningArrayView<2, const float> InActions, const FIndexSet Instances)
 	{
-		Array::Zero(EpisodeStepNums, Instances);
+		UE_LEARNING_TRACE_CPUPROFILER_EVENT_SCOPE(Learning::FEpisodeBuffer::PushActions);
+
+		UE_LEARNING_CHECKF(ActionId >= 0 && ActionId < ActionArrays.Num(), TEXT("Action id invalid!"));
+
+		for (const int32 InstanceIdx : Instances)
+		{
+			Array::Copy(ActionArrays[ActionId][InstanceIdx][EpisodeStepNums[InstanceIdx]], InActions[InstanceIdx]);
+		}
 	}
 
-	const int32 FEpisodeBuffer::GetMaxInstanceNum() const
+	const TLearningArrayView<2, const float> FEpisodeBuffer::GetActions(const int32 ActionId, const int32 InstanceIdx) const
 	{
-		return MaxInstanceNum;
+		UE_LEARNING_CHECKF(ActionId >= 0 && ActionId < ActionArrays.Num(), TEXT("Action id invalid!"));
+
+		return ActionArrays[ActionId][InstanceIdx].Slice(0, EpisodeStepNums[InstanceIdx]);
 	}
 
-	const int32 FEpisodeBuffer::GetMaxStepNum() const
+	int32 FEpisodeBuffer::AddMemoryStates(const FName& Name, const int32 Size)
 	{
-		return MaxStepNum;
+		UE_LEARNING_CHECKF(!MemoryStateNames.Contains(Name), TEXT("Memory State name collision!"));
+
+		int32 MemoryStateId = MemoryStateNames.Num();
+		MemoryStateNames.Add(Name);
+		MemoryStateSizes.Add(Size);
+		TLearningArray<3, float>& NewMemoryState = MemoryStateArrays.AddDefaulted_GetRef();
+
+		if (bHasBeenSized)
+		{
+			NewMemoryState.SetNumUninitialized({ MaxInstanceNum, MaxStepNum, Size });
+		}
+
+		return MemoryStateId;
 	}
 
-	const TLearningArrayView<1, const int32> FEpisodeBuffer::GetEpisodeStepNums() const
+	void FEpisodeBuffer::PushMemoryStates(const int32 MemoryStateId, const TLearningArrayView<2, const float> InMemoryStates, const FIndexSet Instances)
 	{
-		return EpisodeStepNums;
+		UE_LEARNING_TRACE_CPUPROFILER_EVENT_SCOPE(Learning::FEpisodeBuffer::PushMemoryStates);
+
+		UE_LEARNING_CHECKF(MemoryStateId >= 0 && MemoryStateId < MemoryStateArrays.Num(), TEXT("Memory state id invalid!"));
+
+		for (const int32 InstanceIdx : Instances)
+		{
+			Array::Copy(MemoryStateArrays[MemoryStateId][InstanceIdx][EpisodeStepNums[InstanceIdx]], InMemoryStates[InstanceIdx]);
+		}
+	}
+
+	const TLearningArrayView<2, const float> FEpisodeBuffer::GetMemoryStates(const int32 MemoryStateId, const int32 InstanceIdx) const
+	{
+		UE_LEARNING_CHECKF(MemoryStateId >= 0 && MemoryStateId < MemoryStateArrays.Num(), TEXT("Memory state id invalid!"));
+
+		return MemoryStateArrays[MemoryStateId][InstanceIdx].Slice(0, EpisodeStepNums[InstanceIdx]);
+	}
+
+	int32 FEpisodeBuffer::AddRewards(const FName& Name, const int32 Size)
+	{
+		UE_LEARNING_CHECKF(!RewardNames.Contains(Name), TEXT("Reward name collision!"));
+
+		int32 RewardId = RewardNames.Num();
+		RewardNames.Add(Name);
+		RewardSizes.Add(Size);
+		TLearningArray<3, float>& NewReward = RewardArrays.AddDefaulted_GetRef();
+
+		if (bHasBeenSized)
+		{
+			NewReward.SetNumUninitialized({ MaxInstanceNum, MaxStepNum, Size });
+		}
+
+		return RewardId;
+	}
+
+	void FEpisodeBuffer::PushRewards(const int32 RewardId, const TLearningArrayView<1, const float> InRewards, const FIndexSet Instances)
+	{
+		UE_LEARNING_TRACE_CPUPROFILER_EVENT_SCOPE(Learning::FEpisodeBuffer::PushRewards);
+
+		UE_LEARNING_CHECKF(RewardId >= 0 && RewardId < RewardArrays.Num(), TEXT("Reward id invalid!"));
+
+		TLearningArrayView<2, const float> RewardsReshaped(InRewards.GetData(), { InRewards.Num(), 1 });
+		for (const int32 InstanceIdx : Instances)
+		{
+			Array::Copy(RewardArrays[RewardId][InstanceIdx][EpisodeStepNums[InstanceIdx]], RewardsReshaped[InstanceIdx]);
+		}
+	}
+
+	void FEpisodeBuffer::PushRewards(const int32 RewardId, const TLearningArrayView<2, const float> InRewards, const FIndexSet Instances)
+	{
+		UE_LEARNING_TRACE_CPUPROFILER_EVENT_SCOPE(Learning::FEpisodeBuffer::PushRewards);
+
+		UE_LEARNING_CHECKF(RewardId >= 0 && RewardId < RewardArrays.Num(), TEXT("Reward id invalid!"));
+
+		for (const int32 InstanceIdx : Instances)
+		{
+			Array::Copy(RewardArrays[RewardId][InstanceIdx][EpisodeStepNums[InstanceIdx]], InRewards[InstanceIdx]);
+		}
 	}
 
-	const TLearningArrayView<2, const float> FEpisodeBuffer::GetObservations(const int32 InstanceIdx) const
+	const TLearningArrayView<2, const float> FEpisodeBuffer::GetRewards(const int32 RewardId, const int32 InstanceIdx) const
 	{
-		return Observations[InstanceIdx].Slice(0, EpisodeStepNums[InstanceIdx]);
+		UE_LEARNING_CHECKF(RewardId >= 0 && RewardId < RewardArrays.Num(), TEXT("Reward id invalid!"));
+
+		return RewardArrays[RewardId][RewardId].Slice(0, EpisodeStepNums[InstanceIdx]);
+	}
+
+	void FEpisodeBuffer::IncrementEpisodeStepNums(const FIndexSet Instances)
+	{
+		UE_LEARNING_TRACE_CPUPROFILER_EVENT_SCOPE(Learning::FEpisodeBuffer::IncrementEpisodeStepNums);
+
+		for (const int32 InstanceIdx : Instances)
+		{
+			UE_LEARNING_CHECKF(EpisodeStepNums[InstanceIdx] < MaxStepNum, TEXT("Episode Buffer full!"));
+			EpisodeStepNums[InstanceIdx]++;
+		}
 	}
 
-	const TLearningArrayView<2, const float> FEpisodeBuffer::GetActions(const int32 InstanceIdx) const
+	const TLearningArrayView<1, const int32> FEpisodeBuffer::GetEpisodeStepNums() const
 	{
-		return Actions[InstanceIdx].Slice(0, EpisodeStepNums[InstanceIdx]);
+		return EpisodeStepNums;
 	}
 
-	const TLearningArrayView<2, const float> FEpisodeBuffer::GetMemoryStates(const int32 InstanceIdx) const
+	int32 FEpisodeBuffer::GetMaxInstanceNum() const
 	{
-		return MemoryStates[InstanceIdx].Slice(0, EpisodeStepNums[InstanceIdx]);
+		return MaxInstanceNum;
 	}
 
-	const TLearningArrayView<1, const float> FEpisodeBuffer::GetRewards(const int32 InstanceIdx) const
+	int32 FEpisodeBuffer::GetMaxStepNum() const
 	{
-		return Rewards[InstanceIdx].Slice(0, EpisodeStepNums[InstanceIdx]);
+		return MaxStepNum;
 	}
 
 	void FReplayBuffer::Resize(
-		const int32 InObservationVectorDimensionNum,
-		const int32 InActionVectorDimensionNum,
-		const int32 InMemoryStateVectorDimensionNum,
+		const FEpisodeBuffer& EpisodeBuffer,
 		const int32 InMaxEpisodeNum,
 		const int32 InMaxStepNum)
 	{
+		bHasCompletions = true;
+		bHasFinalObservations = true;
+		bHasFinalMemoryStates = true;
+
 		MaxEpisodeNum = InMaxEpisodeNum;
-		MaxStepNum = InMaxStepNum;
 		EpisodeNum = 0;
+
+		MaxStepNum = InMaxStepNum;
 		StepNum = 0;
+
 		EpisodeStarts.SetNumUninitialized({ InMaxEpisodeNum });
 		EpisodeLengths.SetNumUninitialized({ InMaxEpisodeNum });
 		EpisodeCompletionModes.SetNumUninitialized({ InMaxEpisodeNum });
-		EpisodeFinalObservations.SetNumUninitialized({ InMaxEpisodeNum, InObservationVectorDimensionNum });
-		EpisodeFinalMemoryStates.SetNumUninitialized({ InMaxEpisodeNum, InMemoryStateVectorDimensionNum });
-		Observations.SetNumUninitialized({ InMaxStepNum, InObservationVectorDimensionNum });
-		Actions.SetNumUninitialized({ InMaxStepNum, InActionVectorDimensionNum });
-		MemoryStates.SetNumUninitialized({ InMaxStepNum, InMemoryStateVectorDimensionNum });
-		Rewards.SetNumUninitialized({ InMaxStepNum });
+
+		// Observations
+		Observations.Empty();
+		for (int32 NameIndex = 0; NameIndex < EpisodeBuffer.ObservationNames.Num(); NameIndex++)
+		{
+			ObservationNames.Add(EpisodeBuffer.ObservationNames[NameIndex]);
+			ObservationSchemaIds.Add(EpisodeBuffer.ObservationSchemaIds[NameIndex]);
+
+			int32 ObservationSize = EpisodeBuffer.ObservationSizes[NameIndex];
+
+			TLearningArray<2, float>& NewObservation = Observations.AddDefaulted_GetRef();
+			NewObservation.SetNumUninitialized({ InMaxStepNum, ObservationSize });
+
+			TLearningArray<2, float>& NewEpisodeFinalObservation = EpisodeFinalObservations.AddDefaulted_GetRef();
+			NewEpisodeFinalObservation.SetNumUninitialized({ InMaxEpisodeNum, ObservationSize });
+		}
+
+		// Actions
+		Actions.Empty();
+		for (int32 NameIndex = 0; NameIndex < EpisodeBuffer.ActionNames.Num(); NameIndex++)
+		{
+			ActionNames.Add(EpisodeBuffer.ActionNames[NameIndex]);
+			ActionSchemaIds.Add(EpisodeBuffer.ActionSchemaIds[NameIndex]);
+
+			int32 ActionSize = EpisodeBuffer.ActionSizes[NameIndex];
+			TLearningArray<2, float>& NewAction = Actions.AddDefaulted_GetRef();
+			NewAction.SetNumUninitialized({ InMaxStepNum, ActionSize });
+		}
+
+		// Memory States
+		MemoryStates.Empty();
+		for (int32 NameIndex = 0; NameIndex < EpisodeBuffer.MemoryStateNames.Num(); NameIndex++)
+		{
+			MemoryStateNames.Add(EpisodeBuffer.MemoryStateNames[NameIndex]);
+
+			int32 MemoryStateSize = EpisodeBuffer.MemoryStateSizes[NameIndex];
+
+			TLearningArray<2, float>& NewMemoryState = MemoryStates.AddDefaulted_GetRef();
+			NewMemoryState.SetNumUninitialized({ InMaxStepNum, MemoryStateSize });
+
+			TLearningArray<2, float>& NewEpisodeFinalMemoryState = EpisodeFinalMemoryStates.AddDefaulted_GetRef();
+			NewEpisodeFinalMemoryState.SetNumUninitialized({ InMaxEpisodeNum, MemoryStateSize });
+		}
+
+		// Rewards
+		Rewards.Empty();
+		for (int32 NameIndex = 0; NameIndex < EpisodeBuffer.RewardNames.Num(); NameIndex++)
+		{
+			RewardNames.Add(EpisodeBuffer.RewardNames[NameIndex]);
+
+			int32 RewardSize = EpisodeBuffer.RewardSizes[NameIndex];
+			TLearningArray<2, float>& NewReward = Rewards.AddDefaulted_GetRef();
+			NewReward.SetNumUninitialized({ InMaxStepNum, RewardSize });
+		}
 	}
 
 	void FReplayBuffer::Reset()
@@ -118,14 +324,21 @@ namespace UE::Learning
 
 	bool FReplayBuffer::AddEpisodes(
 		const TLearningArrayView<1, const ECompletionMode> InEpisodeCompletionModes,
-		const TLearningArrayView<2, const float> InEpisodeFinalObservations,
-		const TLearningArrayView<2, const float> InEpisodeFinalMemoryStates,
+		const TArrayView<const TLearningArrayView<2, const float>> InEpisodeFinalObservations,
+		const TArrayView<const TLearningArrayView<2, const float>> InEpisodeFinalMemoryStates,
 		const FEpisodeBuffer& EpisodeBuffer,
 		const FIndexSet Instances,
 		const bool bAddTruncatedEpisodeWhenFull)
 	{
 		UE_LEARNING_TRACE_CPUPROFILER_EVENT_SCOPE(Learning::FReplayBuffer::AddEpisodes);
 
+		UE_LEARNING_CHECKF(EpisodeBuffer.ObservationArrays.Num() == Observations.Num(), TEXT("Observation number mismatch!"));
+		UE_LEARNING_CHECKF(EpisodeBuffer.ObservationArrays.Num() == InEpisodeFinalObservations.Num(), TEXT("Final Observation number mismatch!"));
+		UE_LEARNING_CHECKF(EpisodeBuffer.ActionArrays.Num() == Actions.Num(), TEXT("Action number mismatch!"));
+		UE_LEARNING_CHECKF(EpisodeBuffer.MemoryStateArrays.Num() == MemoryStates.Num(), TEXT("Memory State number mismatch!"));
+		UE_LEARNING_CHECKF(EpisodeBuffer.MemoryStateArrays.Num() == InEpisodeFinalMemoryStates.Num(), TEXT("Final Memory State number mismatch!"));
+		UE_LEARNING_CHECKF(EpisodeBuffer.RewardArrays.Num() == Rewards.Num(), TEXT("Reward number mismatch!"));
+
 		for (const int32 InstanceIdx : Instances)
 		{
 			UE_LEARNING_CHECKF(InEpisodeCompletionModes[InstanceIdx] != ECompletionMode::Running,
@@ -134,21 +347,48 @@ namespace UE::Learning
 			const int32 EpisodeStepNum = EpisodeBuffer.GetEpisodeStepNums()[InstanceIdx];
 
 			// Is there space for the full episode in the buffer?
-
 			if (EpisodeNum < MaxEpisodeNum && StepNum + EpisodeStepNum <= MaxStepNum)
 			{
 				// Copy the data into the replay buffer
-				Array::Copy(Observations.Slice(StepNum, EpisodeStepNum), EpisodeBuffer.GetObservations(InstanceIdx));
-				Array::Copy(Actions.Slice(StepNum, EpisodeStepNum), EpisodeBuffer.GetActions(InstanceIdx));
-				Array::Copy(MemoryStates.Slice(StepNum, EpisodeStepNum), EpisodeBuffer.GetMemoryStates(InstanceIdx));
-				Array::Copy(Rewards.Slice(StepNum, EpisodeStepNum), EpisodeBuffer.GetRewards(InstanceIdx));
 
-				// Write the Episode start, length, completion, and final observation
+				for(int32 Index = 0; Index < Observations.Num(); Index++)
+				{
+					Array::Copy(
+						Observations[Index].Slice(StepNum, EpisodeStepNum),
+						EpisodeBuffer.ObservationArrays[Index][InstanceIdx].Slice(0, EpisodeBuffer.EpisodeStepNums[InstanceIdx]));
+					Array::Copy(
+						EpisodeFinalObservations[Index][EpisodeNum],
+						InEpisodeFinalObservations[Index][InstanceIdx]);
+				}
+
+				for (int32 Index = 0; Index < Actions.Num(); Index++)
+				{
+					Array::Copy(
+						Actions[Index].Slice(StepNum, EpisodeStepNum),
+						EpisodeBuffer.ActionArrays[Index][InstanceIdx].Slice(0, EpisodeBuffer.EpisodeStepNums[InstanceIdx]));
+				}
+
+				for (int32 Index = 0; Index < MemoryStates.Num(); Index++)
+				{
+					Array::Copy(
+						MemoryStates[Index].Slice(StepNum, EpisodeStepNum),
+						EpisodeBuffer.MemoryStateArrays[Index][InstanceIdx].Slice(0, EpisodeBuffer.EpisodeStepNums[InstanceIdx]));
+					Array::Copy(
+						EpisodeFinalMemoryStates[Index][EpisodeNum],
+						InEpisodeFinalMemoryStates[Index][InstanceIdx]);
+				}
+
+				for (int32 Index = 0; Index < Rewards.Num(); Index++)
+				{
+					Array::Copy(
+						Rewards[Index].Slice(StepNum, EpisodeStepNum),
+						EpisodeBuffer.RewardArrays[Index][InstanceIdx].Slice(0, EpisodeBuffer.EpisodeStepNums[InstanceIdx]));
+				}
+
+				// Write the Episode start, length, completion
 				EpisodeStarts[EpisodeNum] = StepNum;
 				EpisodeLengths[EpisodeNum] = EpisodeStepNum;
 				EpisodeCompletionModes[EpisodeNum] = InEpisodeCompletionModes[InstanceIdx];
-				Array::Copy(EpisodeFinalObservations[EpisodeNum], InEpisodeFinalObservations[InstanceIdx]);
-				Array::Copy(EpisodeFinalMemoryStates[EpisodeNum], InEpisodeFinalMemoryStates[InstanceIdx]);
 
 				// Increment the Counts
 				EpisodeNum++;
@@ -159,24 +399,51 @@ namespace UE::Learning
 			}
 
 			// Is there space for a partial episode in the buffer?
-
 			if (bAddTruncatedEpisodeWhenFull && EpisodeNum < MaxEpisodeNum && StepNum < MaxStepNum)
 			{
 				const int32 PartialStepNum = MaxStepNum - StepNum;
 				UE_LEARNING_CHECK(PartialStepNum > 0 && PartialStepNum < EpisodeStepNum);
 
 				// Copy the data into the replay buffer
-				Array::Copy(Observations.Slice(StepNum, PartialStepNum), EpisodeBuffer.GetObservations(InstanceIdx).Slice(0, PartialStepNum));
-				Array::Copy(Actions.Slice(StepNum, PartialStepNum), EpisodeBuffer.GetActions(InstanceIdx).Slice(0, PartialStepNum));
-				Array::Copy(MemoryStates.Slice(StepNum, PartialStepNum), EpisodeBuffer.GetMemoryStates(InstanceIdx).Slice(0, PartialStepNum));
-				Array::Copy(Rewards.Slice(StepNum, PartialStepNum), EpisodeBuffer.GetRewards(InstanceIdx).Slice(0, PartialStepNum));
 
-				// Write the Episode start, length, completion, and final observation
+				for (int32 Index = 0; Index < Observations.Num(); Index++)
+				{
+					Array::Copy(
+						Observations[Index].Slice(StepNum, PartialStepNum),
+						EpisodeBuffer.ObservationArrays[Index][InstanceIdx].Slice(0, PartialStepNum));
+					Array::Copy(
+						EpisodeFinalObservations[Index][EpisodeNum],
+						InEpisodeFinalObservations[Index][InstanceIdx]);
+				}
+
+				for (int32 Index = 0; Index < Actions.Num(); Index++)
+				{
+					Array::Copy(
+						Actions[Index].Slice(StepNum, PartialStepNum),
+						EpisodeBuffer.ActionArrays[Index][InstanceIdx].Slice(0, PartialStepNum));
+				}
+
+				for (int32 Index = 0; Index < MemoryStates.Num(); Index++)
+				{
+					Array::Copy(
+						MemoryStates[Index].Slice(StepNum, PartialStepNum),
+						EpisodeBuffer.MemoryStateArrays[Index][InstanceIdx].Slice(0, PartialStepNum));
+					Array::Copy(
+						EpisodeFinalMemoryStates[Index][EpisodeNum],
+						InEpisodeFinalMemoryStates[Index][InstanceIdx]);
+				}
+
+				for (int32 Index = 0; Index < Rewards.Num(); Index++)
+				{
+					Array::Copy(
+						Rewards[Index].Slice(StepNum, PartialStepNum),
+						EpisodeBuffer.RewardArrays[Index][InstanceIdx].Slice(0, PartialStepNum));
+				}
+
+				// Write the Episode start, length, completion
 				EpisodeStarts[EpisodeNum] = StepNum;
 				EpisodeLengths[EpisodeNum] = PartialStepNum;
 				EpisodeCompletionModes[EpisodeNum] = ECompletionMode::Truncated;
-				Array::Copy(EpisodeFinalObservations[EpisodeNum], InEpisodeFinalObservations[InstanceIdx]);
-				Array::Copy(EpisodeFinalMemoryStates[EpisodeNum], InEpisodeFinalMemoryStates[InstanceIdx]);
 
 				// Increment the Counts
 				EpisodeNum++;
@@ -190,22 +457,80 @@ namespace UE::Learning
 		return (EpisodeNum == MaxEpisodeNum) || (StepNum == MaxStepNum);
 	}
 
-	const int32 FReplayBuffer::GetMaxEpisodeNum() const
+	void FReplayBuffer::AddRecords(
+		const int32 InEpisodeNum,
+		const int32 InStepNum,
+		const int32 ObservationSchemaId,
+		const int32 ObservationNum,
+		const int32 ActionSchemaId,
+		const int32 ActionNum,
+		const TLearningArrayView<1, const int32> RecordedEpisodeStarts,
+		const TLearningArrayView<1, const int32> RecordedEpisodeLengths,
+		const TLearningArrayView<2, const float> RecordedObservations,
+		const TLearningArrayView<2, const float> RecordedActions)
+	{
+		bHasCompletions = false;
+		bHasFinalObservations = false;
+		bHasFinalMemoryStates = false;
+
+		MaxEpisodeNum = InEpisodeNum;
+		EpisodeNum = InEpisodeNum;
+
+		MaxStepNum = InStepNum;
+		StepNum = InStepNum;
+
+		EpisodeStarts.SetNumUninitialized({ InEpisodeNum });
+		UE::Learning::Array::Copy(EpisodeStarts, RecordedEpisodeStarts);
+
+		EpisodeLengths.SetNumUninitialized({ InEpisodeNum });
+		UE::Learning::Array::Copy(EpisodeLengths, RecordedEpisodeLengths);
+
+		Observations.Empty();
+		ObservationNames.Add("Observations");
+		ObservationSchemaIds.Add(ObservationSchemaId);
+		TLearningArray<2, float>& NewObservation = Observations.AddDefaulted_GetRef();
+		NewObservation.SetNumUninitialized({ InStepNum, ObservationNum });
+		UE::Learning::Array::Copy(NewObservation, RecordedObservations);
+
+		Actions.Empty();
+		ActionNames.Add("Actions");
+		ActionSchemaIds.Add(ActionSchemaId);
+		TLearningArray<2, float>& NewAction = Actions.AddDefaulted_GetRef();
+		NewAction.SetNumUninitialized({ InStepNum, ActionNum });
+		UE::Learning::Array::Copy(NewAction, RecordedActions);
+	}
+
+	bool FReplayBuffer::HasCompletions() const
+	{
+		return bHasCompletions;
+	}
+
+	bool FReplayBuffer::HasFinalObservations() const
+	{
+		return bHasFinalObservations;
+	}
+
+	bool FReplayBuffer::HasFinalMemoryStates() const
+	{
+		return bHasFinalMemoryStates;
+	}
+
+	int32 FReplayBuffer::GetMaxEpisodeNum() const
 	{
 		return MaxEpisodeNum;
 	}
 
-	const int32 FReplayBuffer::GetMaxStepNum() const
+	int32 FReplayBuffer::GetMaxStepNum() const
 	{
 		return MaxStepNum;
 	}
 
-	const int32 FReplayBuffer::GetEpisodeNum() const
+	int32 FReplayBuffer::GetEpisodeNum() const
 	{
 		return EpisodeNum;
 	}
 
-	const int32 FReplayBuffer::GetStepNum() const
+	int32 FReplayBuffer::GetStepNum() const
 	{
 		return StepNum;
 	}
@@ -225,34 +550,129 @@ namespace UE::Learning
 		return EpisodeCompletionModes.Slice(0, EpisodeNum);
 	}
 
-	const TLearningArrayView<2, const float> FReplayBuffer::GetEpisodeFinalObservations() const
+	int32 FReplayBuffer::GetObservationsNum() const
+	{
+		return Observations.Num();
+	}
+
+	const TLearningArrayView<2, const float> FReplayBuffer::GetObservations(const int32 Index) const
+	{
+		return Observations[Index].Slice(0, StepNum);
+	}
+
+	const TLearningArrayView<2, const float> FReplayBuffer::GetEpisodeFinalObservations(const int32 Index) const
+	{
+		return EpisodeFinalObservations[Index].Slice(0, EpisodeNum);
+	}
+
+	int32 FReplayBuffer::GetActionsNum() const
+	{
+		return Actions.Num();
+	}
+
+	const TLearningArrayView<2, const float> FReplayBuffer::GetActions(const int32 Index) const
+	{
+		return Actions[Index].Slice(0, StepNum);
+	}
+
+	int32 FReplayBuffer::GetMemoryStatesNum() const
 	{
-		return EpisodeFinalObservations.Slice(0, EpisodeNum);
+		return MemoryStates.Num();
 	}
 
-	const TLearningArrayView<2, const float> FReplayBuffer::GetEpisodeFinalMemoryStates() const
+	const TLearningArrayView<2, const float> FReplayBuffer::GetMemoryStates(const int32 Index) const
 	{
-		return EpisodeFinalMemoryStates.Slice(0, EpisodeNum);
+		return MemoryStates[Index].Slice(0, StepNum);
 	}
 
-	const TLearningArrayView<2, const float> FReplayBuffer::GetObservations() const
+	const TLearningArrayView<2, const float> FReplayBuffer::GetEpisodeFinalMemoryStates(const int32 Index) const
 	{
-		return Observations.Slice(0, StepNum);
+		return EpisodeFinalMemoryStates[Index].Slice(0, EpisodeNum);
 	}
 
-	const TLearningArrayView<2, const float> FReplayBuffer::GetActions() const
+	int32 FReplayBuffer::GetRewardsNum() const
 	{
-		return Actions.Slice(0, StepNum);
+		return Rewards.Num();
 	}
 
-	const TLearningArrayView<2, const float> FReplayBuffer::GetMemoryStates() const
+	const TLearningArrayView<2, const float> FReplayBuffer::GetRewards(const int32 Index) const
 	{
-		return MemoryStates.Slice(0, StepNum);
+		return Rewards[Index].Slice(0, StepNum);
 	}
 
-	const TLearningArrayView<1, const float> FReplayBuffer::GetRewards() const
+	TSharedRef<FJsonObject> FReplayBuffer::AsJsonConfig(const int32 ReplayBufferId) const
 	{
-		return Rewards.Slice(0, StepNum);
+		TSharedRef<FJsonObject> ConfigObject = MakeShared<FJsonObject>();
+
+		ConfigObject->SetNumberField(TEXT("Id"), ReplayBufferId);
+
+		// Replay Buffer Settings
+		ConfigObject->SetNumberField(TEXT("MaxEpisodeNum"), MaxEpisodeNum);
+		ConfigObject->SetNumberField(TEXT("MaxStepNum"), MaxStepNum);
+
+		ConfigObject->SetBoolField(TEXT("HasCompletions"), HasCompletions());
+		ConfigObject->SetBoolField(TEXT("HasFinalObservations"), HasFinalObservations());
+		ConfigObject->SetBoolField(TEXT("HasFinalMemoryStates"), HasFinalMemoryStates());
+
+		// Observations
+		TArray<TSharedPtr<FJsonValue>> ObservationObjects;
+		for(int32 Index = 0; Index < Observations.Num(); Index++)
+		{
+			TSharedPtr<FJsonObject> BufferObject = MakeShared<FJsonObject>();
+			BufferObject->SetNumberField(TEXT("Id"), Index);
+			BufferObject->SetStringField(TEXT("Name"), ObservationNames[Index].ToString());
+			BufferObject->SetNumberField(TEXT("SchemaId"), ObservationSchemaIds[Index]);
+			BufferObject->SetNumberField(TEXT("VectorDimensionNum"), Observations[Index].Num<1>());
+
+			TSharedRef<FJsonValueObject> JsonValue = MakeShared<FJsonValueObject>(BufferObject);
+			ObservationObjects.Add(JsonValue);
+		}
+		ConfigObject->SetArrayField(TEXT("Observations"), ObservationObjects);
+
+		// Actions
+		TArray<TSharedPtr<FJsonValue>> ActionObjects;
+		for (int32 Index = 0; Index < Actions.Num(); Index++)
+		{
+			TSharedPtr<FJsonObject> BufferObject = MakeShared<FJsonObject>();
+			BufferObject->SetNumberField(TEXT("Id"), Index);
+			BufferObject->SetStringField(TEXT("Name"), ActionNames[Index].ToString());
+			BufferObject->SetNumberField(TEXT("SchemaId"), ActionSchemaIds[Index]);
+			BufferObject->SetNumberField(TEXT("VectorDimensionNum"), Actions[Index].Num<1>());
+
+			TSharedRef<FJsonValueObject> JsonValue = MakeShared<FJsonValueObject>(BufferObject);
+			ActionObjects.Add(JsonValue);
+		}
+		ConfigObject->SetArrayField(TEXT("Actions"), ActionObjects);
+
+		// Memory States
+		TArray<TSharedPtr<FJsonValue>> MemoryStateObjects;
+		for (int32 Index = 0; Index < MemoryStates.Num(); Index++)
+		{
+			TSharedPtr<FJsonObject> BufferObject = MakeShared<FJsonObject>();
+			BufferObject->SetNumberField(TEXT("Id"), Index);
+			BufferObject->SetStringField(TEXT("Name"), MemoryStateNames[Index].ToString());
+			BufferObject->SetNumberField(TEXT("VectorDimensionNum"), MemoryStates[Index].Num<1>());
+
+			TSharedRef<FJsonValueObject> JsonValue = MakeShared<FJsonValueObject>(BufferObject);
+			MemoryStateObjects.Add(JsonValue);
+		}
+		ConfigObject->SetArrayField(TEXT("MemoryStates"), MemoryStateObjects);
+
+		// Rewards
+		TArray<TSharedPtr<FJsonValue>> RewardObjects;
+		for (int32 Index = 0; Index < Rewards.Num(); Index++)
+		{
+			TSharedPtr<FJsonObject> BufferObject = MakeShared<FJsonObject>();
+			BufferObject->SetNumberField(TEXT("Id"), Index);
+			BufferObject->SetStringField(TEXT("Name"), RewardNames[Index].ToString());
+			BufferObject->SetNumberField(TEXT("VectorDimensionNum"), Rewards[Index].Num<1>());
+
+			TSharedRef<FJsonValueObject> JsonValue = MakeShared<FJsonValueObject>(BufferObject);
+			RewardObjects.Add(JsonValue);
+		}
+		ConfigObject->SetArrayField(TEXT("Rewards"), RewardObjects);
+
+		return ConfigObject;
 	}
 
 	namespace Experience
@@ -261,27 +681,26 @@ namespace UE::Learning
 			FReplayBuffer& ReplayBuffer,
 			FEpisodeBuffer& EpisodeBuffer,
 			FResetInstanceBuffer& ResetBuffer,
-			TLearningArrayView<2, float> ObservationVectorBuffer,
-			TLearningArrayView<2, float> ActionVectorBuffer,
-			TLearningArrayView<2, float> PreEvaluationMemoryStateVectorBuffer,
-			TLearningArrayView<2, float> MemoryStateVectorBuffer,
-			TLearningArrayView<1, float> RewardBuffer,
+			const TArrayView<const TLearningArrayView<2, const float>> ObservationVectorBuffers,
+			const TArrayView<const TLearningArrayView<2, const float>> ActionVectorBuffers,
+			const TArrayView<const TLearningArrayView<2, const float>> PreEvaluationMemoryStateVectorBuffers,
+			const TArrayView<const TLearningArrayView<2, const float>> MemoryStateVectorBuffers,
+			const TArrayView<const TLearningArrayView<1, const float>> RewardBuffers,
 			TLearningArrayView<1, ECompletionMode> CompletionBuffer,
 			TLearningArrayView<1, ECompletionMode> EpisodeCompletionBuffer,
 			TLearningArrayView<1, ECompletionMode> AllCompletionBuffer,
 			const TFunctionRef<void(const FIndexSet Instances)> ResetFunction,
-			const TFunctionRef<void(const FIndexSet Instances)> ObservationFunction,
-			const TFunctionRef<void(const FIndexSet Instances)> PolicyFunction,
-			const TFunctionRef<void(const FIndexSet Instances)> ActionFunction,
-			const TFunctionRef<void(const FIndexSet Instances)> UpdateFunction,
-			const TFunctionRef<void(const FIndexSet Instances)> RewardFunction,
+			const TArrayView<const TFunctionRef<void(const FIndexSet Instances)>> ObservationFunctions,
+			const TArrayView<const TFunctionRef<void(const FIndexSet Instances)>> PolicyFunctions,
+			const TArrayView<const TFunctionRef<void(const FIndexSet Instances)>> ActionFunctions,
+			const TArrayView<const TFunctionRef<void(const FIndexSet Instances)>> UpdateFunctions,
+			const TArrayView<const TFunctionRef<void(const FIndexSet Instances)>> RewardFunctions,
 			const TFunctionRef<void(const FIndexSet Instances)> CompletionFunction,
 			const FIndexSet Instances)
 		{
 			UE_LEARNING_TRACE_CPUPROFILER_EVENT_SCOPE(Learning::Experience::GatherExperienceUntilReplayBufferFull);
 
 			// Reset Everything
-
 			ReplayBuffer.Reset();
 			EpisodeBuffer.Reset(Instances);
 			ResetFunction(Instances);
@@ -289,36 +708,42 @@ namespace UE::Learning
 			while (true)
 			{
 				// Encode Observations
-
-				ObservationFunction(Instances);
+				for (int32 Index = 0; Index < ObservationFunctions.Num(); Index++)
+				{
+					ObservationFunctions[Index](Instances);
+					EpisodeBuffer.PushObservations(Index, ObservationVectorBuffers[Index], Instances);
+				}
 
 				// Evaluate Policy
-
-				PolicyFunction(Instances);
+				for (int32 Index = 0; Index < PolicyFunctions.Num(); Index++)
+				{
+					PolicyFunctions[Index](Instances);
+					EpisodeBuffer.PushMemoryStates(Index, PreEvaluationMemoryStateVectorBuffers[Index], Instances);
+				}
 
 				// Decode Actions
-
-				ActionFunction(Instances);
+				for (int32 Index = 0; Index < ActionFunctions.Num(); Index++)
+				{
+					ActionFunctions[Index](Instances);
+					EpisodeBuffer.PushActions(Index, ActionVectorBuffers[Index], Instances);
+				}
 
 				// Update Environment
-
-				UpdateFunction(Instances);
+				for (int32 Index = 0; Index < UpdateFunctions.Num(); Index++)
+				{
+					UpdateFunctions[Index](Instances);
+				}
 
 				// Compute Rewards
+				for (int32 Index = 0; Index < RewardFunctions.Num(); Index++)
+				{
+					RewardFunctions[Index](Instances);
+					EpisodeBuffer.PushRewards(Index, RewardBuffers[Index], Instances);
+				}
 
-				RewardFunction(Instances);
-
-				// Push to Experience Buffer
-
-				EpisodeBuffer.Push(
-					ObservationVectorBuffer,
-					ActionVectorBuffer,
-					PreEvaluationMemoryStateVectorBuffer,
-					RewardBuffer,
-					Instances);
+				EpisodeBuffer.IncrementEpisodeStepNums(Instances);
 
 				// Evaluate Completions
-
 				CompletionFunction(Instances);
 
 				Completion::EvaluateEndOfEpisodeCompletions(
@@ -340,15 +765,17 @@ namespace UE::Learning
 				}
 
 				// Evaluate Observations again for instances that are completed
-
-				ObservationFunction(ResetBuffer.GetResetInstances());
+				for (int32 Index = 0; Index < ObservationFunctions.Num(); Index++)
+				{
+					ObservationFunctions[Index](ResetBuffer.GetResetInstances());
+				}
 
 				// Push completed instances to Replay Buffer and return if full
 
 				if (ReplayBuffer.AddEpisodes(
 					AllCompletionBuffer,
-					ObservationVectorBuffer,
-					MemoryStateVectorBuffer,
+					ObservationVectorBuffers,
+					MemoryStateVectorBuffers,
 					EpisodeBuffer,
 					ResetBuffer.GetResetInstances()))
 				{
diff --git a/Engine/Plugins/Experimental/LearningAgents/Source/LearningTraining/Private/LearningExternalTrainer.cpp b/Engine/Plugins/Experimental/LearningAgents/Source/LearningTraining/Private/LearningExternalTrainer.cpp
--- a/Engine/Plugins/Experimental/LearningAgents/Source/LearningTraining/Private/LearningExternalTrainer.cpp
+++ b/Engine/Plugins/Experimental/LearningAgents/Source/LearningTraining/Private/LearningExternalTrainer.cpp
@@ -24,6 +24,8 @@ namespace UE::Learning
 {
 	FSharedMemoryTrainerServerProcess::FSharedMemoryTrainerServerProcess(
 		const FString& TaskName,
+		const FString& CustomTrainerPath,
+		const FString& TrainerFileName,
 		const FString& PythonExecutablePath,
 		const FString& PythonContentPath,
 		const FString& InIntermediatePath,
@@ -42,12 +44,12 @@ namespace UE::Learning
 		if (ProcessIdx == 0)
 		{
 			// Allocate the control memory if we are the parent UE process
-			Controls = SharedMemory::Allocate<2, volatile int32>({ ProcessNum, SharedMemoryTraining::GetControlNum() });
+			Controls = SharedMemory::Allocate<2, volatile int32>({ ProcessNum, SharedMemoryTraining::GetControlNum()});
 		}
 		else
 		{
 			FGuid ControlsGuid; ensure(FParse::Value(FCommandLine::Get(), TEXT("LearningControlsGuid"), ControlsGuid));
-			Controls = SharedMemory::Map<2, volatile int32>(ControlsGuid, { ProcessNum, SharedMemoryTraining::GetControlNum() });
+			Controls = SharedMemory::Map<2, volatile int32>(ControlsGuid, { ProcessNum, SharedMemoryTraining::GetControlNum()});
 
 			// We do not want to launch another training process if we are a child process
 			return;
@@ -60,14 +62,22 @@ namespace UE::Learning
 		// uninitialized values or those left over from previous runs.
 		Array::Zero(Controls.View);
 
+		// Set the ID columns to -1
+		for (int32 Index = 0; Index < Controls.View.Num<0>(); Index++)
+		{
+			Controls.View[Index][(uint8)SharedMemoryTraining::EControls::NetworkId] = -1;
+			Controls.View[Index][(uint8)SharedMemoryTraining::EControls::ReplayBufferId] = -1;
+		}
+
 		const FString TimeStamp = FDateTime::Now().ToFormattedString(TEXT("%Y-%m-%d_%H-%M-%S"));
-		const FString TrainerMethod = TEXT("PPO");
 		const FString TrainerType = TEXT("SharedMemory");
-		ConfigPath = InIntermediatePath / TEXT("Configs") / FString::Printf(TEXT("%s_%s_%s_%s.json"), *TaskName, *TrainerMethod, *TrainerType, *TimeStamp);
+		ConfigPath = InIntermediatePath / TEXT("Configs") / FString::Printf(TEXT("%s_%s_%s_%s.json"), *TaskName, *TrainerFileName, *TrainerType, *TimeStamp);
 
 		IFileManager& FileManager = IFileManager::Get();
-		const FString CommandLineArguments = FString::Printf(TEXT("\"%s\" SharedMemory \"%s\" %i \"%s\""),
-			*FileManager.ConvertToAbsolutePathForExternalAppForRead(*(PythonContentPath / TEXT("train_ppo.py"))),
+		const FString CommandLineArguments = FString::Printf(TEXT("\"%s\" \"%s\" \"%s\" SharedMemory \"%s\" %i \"%s\""),
+			*FileManager.ConvertToAbsolutePathForExternalAppForRead(*(PythonContentPath / TEXT("train.py"))),
+			*FileManager.ConvertToAbsolutePathForExternalAppForRead(*CustomTrainerPath),
+			*TrainerFileName,
 			*Controls.Guid.ToString(EGuidFormats::DigitsWithHyphensInBraces),
 			ProcessNum,
 			*FileManager.ConvertToAbsolutePathForExternalAppForRead(*ConfigPath));
@@ -147,12 +157,36 @@ namespace UE::Learning
 			SharedMemory::Deallocate(EpisodeStarts);
 			SharedMemory::Deallocate(EpisodeLengths);
 			SharedMemory::Deallocate(EpisodeCompletionModes);
-			SharedMemory::Deallocate(EpisodeFinalObservations);
-			SharedMemory::Deallocate(EpisodeFinalMemoryStates);
-			SharedMemory::Deallocate(Observations);
-			SharedMemory::Deallocate(Actions);
-			SharedMemory::Deallocate(MemoryStates);
-			SharedMemory::Deallocate(Rewards);
+
+			for(TSharedMemoryArrayView<3, float>& SharedMemoryArrayView : EpisodeFinalObservations)
+			{
+				SharedMemory::Deallocate(SharedMemoryArrayView);
+			}
+
+			for (TSharedMemoryArrayView<3, float>& SharedMemoryArrayView : EpisodeFinalMemoryStates)
+			{
+				SharedMemory::Deallocate(SharedMemoryArrayView);
+			}
+
+			for (TSharedMemoryArrayView<3, float>& SharedMemoryArrayView : Observations)
+			{
+				SharedMemory::Deallocate(SharedMemoryArrayView);
+			}
+
+			for (TSharedMemoryArrayView<3, float>& SharedMemoryArrayView : Actions)
+			{
+				SharedMemory::Deallocate(SharedMemoryArrayView);
+			}
+
+			for (TSharedMemoryArrayView<3, float>& SharedMemoryArrayView : MemoryStates)
+			{
+				SharedMemory::Deallocate(SharedMemoryArrayView);
+			}
+
+			for (TSharedMemoryArrayView<3, float>& SharedMemoryArrayView : Rewards)
+			{
+				SharedMemory::Deallocate(SharedMemoryArrayView);
+			}
 		}
 	}
 
@@ -193,6 +227,11 @@ namespace UE::Learning
 		return ETrainerResponse::Success;
 	}
 
+	bool FSharedMemoryTrainer::HasNetworkOrCompleted()
+	{
+		return SharedMemoryTraining::HasNetworkOrCompleted(Controls.View[ProcessIdx]);
+	}
+
 	void FSharedMemoryTrainer::Terminate()
 	{
 		Deallocate();
@@ -220,33 +259,84 @@ namespace UE::Learning
 		ConfigObject->SetStringField(TEXT("IntermediatePath"), *FileManager.ConvertToAbsolutePathForExternalAppForRead(*IntermediatePath));
 		ConfigObject->SetBoolField(TEXT("LoggingEnabled"), LogSettings == ELogSetting::Silent ? false : true);
 
-		ConfigObject->SetNumberField(TEXT("ProcessNum"), ProcessNum);
+		TSharedPtr<FJsonObject> SharedMemoryObject = MakeShared<FJsonObject>();
+
+		SharedMemoryObject->SetNumberField(TEXT("ProcessNum"), ProcessNum);
 
 		TArray<TSharedPtr<FJsonValue>> NetworkGuidsArray;
-		for (const TPair<FName, TSharedMemoryArrayView<1, uint8>>& MapEntry : NeuralNetworkSharedMemoryArrayViews)
+		for(int32 Index = 0; Index < NeuralNetworkSharedMemoryArrayViews.Num(); Index++)
 		{
-			NetworkGuidsArray.Add(MakeShared<FJsonValueString>(*MapEntry.Value.Guid.ToString(EGuidFormats::DigitsWithHyphensInBraces)));
+			TSharedPtr<FJsonObject> JsonObject = MakeShared<FJsonObject>();
+			JsonObject->SetNumberField(TEXT("NetworkId"), Index);
+			JsonObject->SetStringField(TEXT("Guid"), *NeuralNetworkSharedMemoryArrayViews[Index].Guid.ToString(EGuidFormats::DigitsWithHyphensInBraces));
+
+			TSharedRef<FJsonValueObject> JsonValue = MakeShared<FJsonValueObject>(JsonObject);
+			NetworkGuidsArray.Add(JsonValue);
 		}
-		ConfigObject->SetArrayField(TEXT("NetworkGuids"), NetworkGuidsArray);
+		SharedMemoryObject->SetArrayField(TEXT("NetworkGuids"), NetworkGuidsArray);
 
 		TArray<TSharedPtr<FJsonValue>> ExperienceContainerObjectsArray;
-		for (const TPair<FName, FSharedMemoryExperienceContainer>& MapEntry : SharedMemoryExperienceContainers)
+		for (const FSharedMemoryExperienceContainer& SharedMemoryExperienceContainer : SharedMemoryExperienceContainers)
 		{
 			TSharedPtr<FJsonObject> JsonObject = MakeShared<FJsonObject>();
-			JsonObject->SetStringField(TEXT("EpisodeStartsGuid"), *MapEntry.Value.EpisodeStarts.Guid.ToString(EGuidFormats::DigitsWithHyphensInBraces));
-			JsonObject->SetStringField(TEXT("EpisodeLengthsGuid"), *MapEntry.Value.EpisodeLengths.Guid.ToString(EGuidFormats::DigitsWithHyphensInBraces));
-			JsonObject->SetStringField(TEXT("EpisodeCompletionModesGuid"), *MapEntry.Value.EpisodeCompletionModes.Guid.ToString(EGuidFormats::DigitsWithHyphensInBraces));
-			JsonObject->SetStringField(TEXT("EpisodeFinalObservationsGuid"), *MapEntry.Value.EpisodeFinalObservations.Guid.ToString(EGuidFormats::DigitsWithHyphensInBraces));
-			JsonObject->SetStringField(TEXT("EpisodeFinalMemoryStatesGuid"), *MapEntry.Value.EpisodeFinalMemoryStates.Guid.ToString(EGuidFormats::DigitsWithHyphensInBraces));
-			JsonObject->SetStringField(TEXT("ObservationsGuid"), *MapEntry.Value.Observations.Guid.ToString(EGuidFormats::DigitsWithHyphensInBraces));
-			JsonObject->SetStringField(TEXT("ActionsGuid"), *MapEntry.Value.Actions.Guid.ToString(EGuidFormats::DigitsWithHyphensInBraces));
-			JsonObject->SetStringField(TEXT("MemoryStatesGuid"), *MapEntry.Value.MemoryStates.Guid.ToString(EGuidFormats::DigitsWithHyphensInBraces));
-			JsonObject->SetStringField(TEXT("RewardsGuid"), *MapEntry.Value.Rewards.Guid.ToString(EGuidFormats::DigitsWithHyphensInBraces));
+			JsonObject->SetStringField(TEXT("EpisodeStartsGuid"), *SharedMemoryExperienceContainer.EpisodeStarts.Guid.ToString(EGuidFormats::DigitsWithHyphensInBraces));
+			JsonObject->SetStringField(TEXT("EpisodeLengthsGuid"), *SharedMemoryExperienceContainer.EpisodeLengths.Guid.ToString(EGuidFormats::DigitsWithHyphensInBraces));
+			JsonObject->SetStringField(TEXT("EpisodeCompletionModesGuid"), *SharedMemoryExperienceContainer.EpisodeCompletionModes.Guid.ToString(EGuidFormats::DigitsWithHyphensInBraces));
+
+			TArray<TSharedPtr<FJsonValue>> EpisodeFinalObservationsGuidsArray;
+			for (const TSharedMemoryArrayView<3, float>& EpisodeFinalObservations : SharedMemoryExperienceContainer.EpisodeFinalObservations)
+			{
+				EpisodeFinalObservationsGuidsArray.Add(MakeShared<FJsonValueString>(*EpisodeFinalObservations.Guid.ToString(EGuidFormats::DigitsWithHyphensInBraces)));
+				
+			}
+			JsonObject->SetArrayField(TEXT("EpisodeFinalObservationsGuids"), EpisodeFinalObservationsGuidsArray);
+
+			TArray<TSharedPtr<FJsonValue>> EpisodeFinalMemoryStatesGuidsArray;
+			for (const TSharedMemoryArrayView<3, float>& EpisodeFinalMemoryStates : SharedMemoryExperienceContainer.EpisodeFinalMemoryStates)
+			{
+				EpisodeFinalMemoryStatesGuidsArray.Add(MakeShared<FJsonValueString>(*EpisodeFinalMemoryStates.Guid.ToString(EGuidFormats::DigitsWithHyphensInBraces)));
+
+			}
+			JsonObject->SetArrayField(TEXT("EpisodeFinalMemoryStatesGuids"), EpisodeFinalMemoryStatesGuidsArray);
+
+			TArray<TSharedPtr<FJsonValue>> ObservationsGuidsArray;
+			for (const TSharedMemoryArrayView<3, float>& Observations : SharedMemoryExperienceContainer.Observations)
+			{
+				ObservationsGuidsArray.Add(MakeShared<FJsonValueString>(*Observations.Guid.ToString(EGuidFormats::DigitsWithHyphensInBraces)));
+
+			}
+			JsonObject->SetArrayField(TEXT("ObservationsGuids"), ObservationsGuidsArray);
+
+			TArray<TSharedPtr<FJsonValue>> ActionsGuidsArray;
+			for (const TSharedMemoryArrayView<3, float>& Actions : SharedMemoryExperienceContainer.Actions)
+			{
+				ActionsGuidsArray.Add(MakeShared<FJsonValueString>(*Actions.Guid.ToString(EGuidFormats::DigitsWithHyphensInBraces)));
+
+			}
+			JsonObject->SetArrayField(TEXT("ActionsGuids"), ActionsGuidsArray);
+
+			TArray<TSharedPtr<FJsonValue>> MemoryStatesGuidsArray;
+			for (const TSharedMemoryArrayView<3, float>& MemoryStates : SharedMemoryExperienceContainer.MemoryStates)
+			{
+				MemoryStatesGuidsArray.Add(MakeShared<FJsonValueString>(*MemoryStates.Guid.ToString(EGuidFormats::DigitsWithHyphensInBraces)));
+
+			}
+			JsonObject->SetArrayField(TEXT("MemoryStatesGuids"), MemoryStatesGuidsArray);
+
+			TArray<TSharedPtr<FJsonValue>> RewardsGuidsArray;
+			for (const TSharedMemoryArrayView<3, float>& Rewards : SharedMemoryExperienceContainer.Rewards)
+			{
+				RewardsGuidsArray.Add(MakeShared<FJsonValueString>(*Rewards.Guid.ToString(EGuidFormats::DigitsWithHyphensInBraces)));
+
+			}
+			JsonObject->SetArrayField(TEXT("RewardsGuids"), RewardsGuidsArray);
 			
 			TSharedRef<FJsonValueObject> JsonValue = MakeShared<FJsonValueObject>(JsonObject);
 			ExperienceContainerObjectsArray.Add(JsonValue);
 		}
-		ConfigObject->SetArrayField(TEXT("ExperienceBuffers"), ExperienceContainerObjectsArray);
+		SharedMemoryObject->SetArrayField(TEXT("ReplayBuffers"), ExperienceContainerObjectsArray);
+
+		ConfigObject->SetObjectField(TEXT("SharedMemory"), SharedMemoryObject);
 		
 		FString ConfigString;
 		TSharedRef<TJsonWriter<>> JsonWriter = TJsonWriterFactory<>::Create(&ConfigString, 0);
@@ -256,57 +346,54 @@ namespace UE::Learning
 		return SharedMemoryTraining::SendConfigSignal(Controls.View[ProcessIdx], LogSettings);
 	}
 
-	void FSharedMemoryTrainer::AddNetwork(
-		const FName& Name,
-		const ULearningNeuralNetworkData& Network)
-	{
-		NeuralNetworkSharedMemoryArrayViews.Add(Name, SharedMemory::Allocate<1, uint8>({ Network.GetSnapshotByteNum() }));
-	}
-
-	bool FSharedMemoryTrainer::ContainsNetwork(const FName& Name) const
+	int32 FSharedMemoryTrainer::AddNetwork(const ULearningNeuralNetworkData& Network)
 	{
-		return NeuralNetworkSharedMemoryArrayViews.Contains(Name);
+		const int32 NetworkId = NeuralNetworkSharedMemoryArrayViews.Num();
+		NeuralNetworkSharedMemoryArrayViews.Add(SharedMemory::Allocate<1, uint8>({ Network.GetSnapshotByteNum() }));
+		return NetworkId;
 	}
 
 	ETrainerResponse FSharedMemoryTrainer::ReceiveNetwork(
-		const FName& Name,
+		const int32 NetworkId,
 		ULearningNeuralNetworkData& OutNetwork,
 		FRWLock* NetworkLock,
 		const ELogSetting LogSettings)
 	{
 		check(ProcessIdx != INDEX_NONE);
 		checkf(Controls.Region, TEXT("ReceiveNetwork: Controls Shared Memory Region is nullptr"));
-		if (!ensureMsgf(NeuralNetworkSharedMemoryArrayViews.Contains(Name), TEXT("Network %s has not been added. Call AddNetwork prior to ReceiveNetwork."), *Name.ToString()))
+		if (!ensureMsgf(NeuralNetworkSharedMemoryArrayViews.Num() >= NetworkId, TEXT("Network %d has not been added. Call AddNetwork prior to ReceiveNetwork."), NetworkId))
 		{
 			return ETrainerResponse::Unexpected;
 		}
 
 		return SharedMemoryTraining::RecvNetwork(
 			Controls.View[ProcessIdx],
+			NetworkId,
 			OutNetwork,
 			*TrainingProcess,
-			NeuralNetworkSharedMemoryArrayViews[Name].View,
+			NeuralNetworkSharedMemoryArrayViews[NetworkId].View,
 			Timeout,
 			NetworkLock,
 			LogSettings);
 	}
 
 	ETrainerResponse FSharedMemoryTrainer::SendNetwork(
-		const FName& Name,
+		const int32 NetworkId,
 		const ULearningNeuralNetworkData& Network,
 		FRWLock* NetworkLock,
 		const ELogSetting LogSettings)
 	{
 		check(ProcessIdx != INDEX_NONE);
 		checkf(Controls.Region, TEXT("SendNetwork: Controls Shared Memory Region is nullptr"));
-		if (!ensureMsgf(NeuralNetworkSharedMemoryArrayViews.Contains(Name), TEXT("Network %s has not been added. Call AddNetwork prior to SendNetwork."), *Name.ToString()))
+		if (!ensureMsgf(NeuralNetworkSharedMemoryArrayViews.Num() >= NetworkId, TEXT("Network %d has not been added. Call AddNetwork prior to SendNetwork."), NetworkId))
 		{
 			return ETrainerResponse::Unexpected;
 		}
 
 		return SharedMemoryTraining::SendNetwork(
 			Controls.View[ProcessIdx],
-			NeuralNetworkSharedMemoryArrayViews[Name].View,
+			NetworkId,
+			NeuralNetworkSharedMemoryArrayViews[NetworkId].View,
 			*TrainingProcess,
 			Network,
 			Timeout,
@@ -314,81 +401,270 @@ namespace UE::Learning
 			LogSettings);
 	}
 
-	void FSharedMemoryTrainer::AddReplayBuffer(
-		const FName& Name,
-		const FReplayBuffer& ReplayBuffer)
+	int32 FSharedMemoryTrainer::AddReplayBuffer(const FReplayBuffer& ReplayBuffer)
 	{
 		check(ProcessNum > 0);
 
-		const int32 ObservationVectorDimensionNum = ReplayBuffer.GetObservations().Num<1>();
-		const int32 ActionVectorDimensionNum = ReplayBuffer.GetActions().Num<1>();
-		const int32 MemoryStateVectorDimensionNum = ReplayBuffer.GetMemoryStates().Num<1>();
-
 		FSharedMemoryExperienceContainer ExperienceContainer;
 		if (ProcessIdx == 0)
 		{
 			ExperienceContainer.EpisodeStarts = SharedMemory::Allocate<2, int32>({ ProcessNum, ReplayBuffer.GetMaxEpisodeNum() });
 			ExperienceContainer.EpisodeLengths = SharedMemory::Allocate<2, int32>({ ProcessNum, ReplayBuffer.GetMaxEpisodeNum() });
-			ExperienceContainer.EpisodeCompletionModes = SharedMemory::Allocate<2, ECompletionMode>({ ProcessNum, ReplayBuffer.GetMaxEpisodeNum() });
-			ExperienceContainer.EpisodeFinalObservations = SharedMemory::Allocate<3, float>({ ProcessNum, ReplayBuffer.GetMaxEpisodeNum(), ObservationVectorDimensionNum });
-			ExperienceContainer.EpisodeFinalMemoryStates = SharedMemory::Allocate<3, float>({ ProcessNum, ReplayBuffer.GetMaxEpisodeNum(), MemoryStateVectorDimensionNum });
-			ExperienceContainer.Observations = SharedMemory::Allocate<3, float>({ ProcessNum, ReplayBuffer.GetMaxStepNum(), ObservationVectorDimensionNum });
-			ExperienceContainer.Actions = SharedMemory::Allocate<3, float>({ ProcessNum, ReplayBuffer.GetMaxStepNum(), ActionVectorDimensionNum });
-			ExperienceContainer.MemoryStates = SharedMemory::Allocate<3, float>({ ProcessNum, ReplayBuffer.GetMaxStepNum(), MemoryStateVectorDimensionNum });
-			ExperienceContainer.Rewards = SharedMemory::Allocate<2, float>({ ProcessNum, ReplayBuffer.GetMaxStepNum() });
+
+			if (ReplayBuffer.HasCompletions())
+			{
+				ExperienceContainer.EpisodeCompletionModes = SharedMemory::Allocate<2, ECompletionMode>({ ProcessNum, ReplayBuffer.GetMaxEpisodeNum() });
+			}
+
+			if (ReplayBuffer.HasFinalObservations())
+			{
+				for (int32 Index = 0; Index < ReplayBuffer.GetObservationsNum(); Index++)
+				{
+					const int32 DimNum = ReplayBuffer.GetEpisodeFinalObservations(Index).Num<1>();
+					ExperienceContainer.EpisodeFinalObservations.Add(SharedMemory::Allocate<3, float>({ ProcessNum, ReplayBuffer.GetMaxEpisodeNum(), DimNum }));
+				}
+			}
+
+			if (ReplayBuffer.HasFinalMemoryStates())
+			{
+				for (int32 Index = 0; Index < ReplayBuffer.GetMemoryStatesNum(); Index++)
+				{
+					const int32 DimNum = ReplayBuffer.GetEpisodeFinalMemoryStates(Index).Num<1>();
+					ExperienceContainer.EpisodeFinalMemoryStates.Add(SharedMemory::Allocate<3, float>({ ProcessNum, ReplayBuffer.GetMaxEpisodeNum(), DimNum }));
+				}
+			}
+
+			for (int32 Index = 0; Index < ReplayBuffer.GetObservationsNum(); Index++)
+			{
+				const int32 DimNum = ReplayBuffer.GetObservations(Index).Num<1>();
+				ExperienceContainer.Observations.Add(SharedMemory::Allocate<3, float>({ ProcessNum, ReplayBuffer.GetMaxStepNum(), DimNum }));
+			}
+
+			for (int32 Index = 0; Index < ReplayBuffer.GetActionsNum(); Index++)
+			{
+				const int32 DimNum = ReplayBuffer.GetActions(Index).Num<1>();
+				ExperienceContainer.Actions.Add(SharedMemory::Allocate<3, float>({ ProcessNum, ReplayBuffer.GetMaxStepNum(), DimNum }));
+			}
+
+			for (int32 Index = 0; Index < ReplayBuffer.GetMemoryStatesNum(); Index++)
+			{
+				const int32 DimNum = ReplayBuffer.GetMemoryStates(Index).Num<1>();
+				ExperienceContainer.MemoryStates.Add(SharedMemory::Allocate<3, float>({ ProcessNum, ReplayBuffer.GetMaxStepNum(), DimNum }));
+			}
+
+			for (int32 Index = 0; Index < ReplayBuffer.GetRewardsNum(); Index++)
+			{
+				const int32 DimNum = ReplayBuffer.GetRewards(Index).Num<1>();
+				ExperienceContainer.Rewards.Add(SharedMemory::Allocate<3, float>({ ProcessNum, ReplayBuffer.GetMaxStepNum(), DimNum }));
+			}
 		}
 		else
 		{
 			FGuid EpisodeStartsGuid; ensure(FParse::Value(FCommandLine::Get(), TEXT("LearningEpisodeStartsGuid"), EpisodeStartsGuid));
-			FGuid EpisodeLengthsGuid; ensure(FParse::Value(FCommandLine::Get(), TEXT("LearningEpisodeLengthsGuid"), EpisodeLengthsGuid));
-			FGuid EpisodeCompletionModesGuid; ensure(FParse::Value(FCommandLine::Get(), TEXT("LearningEpisodeCompletionModesGuid"), EpisodeCompletionModesGuid));
-			FGuid EpisodeFinalObservationsGuid; ensure(FParse::Value(FCommandLine::Get(), TEXT("LearningEpisodeFinalObservationsGuid"), EpisodeFinalObservationsGuid));
-			FGuid EpisodeFinalMemoryStatesGuid; ensure(FParse::Value(FCommandLine::Get(), TEXT("LearningEpisodeFinalMemoryStatesGuid"), EpisodeFinalMemoryStatesGuid));
-			FGuid ObservationsGuid; ensure(FParse::Value(FCommandLine::Get(), TEXT("LearningObservationsGuid"), ObservationsGuid));
-			FGuid ActionsGuid; ensure(FParse::Value(FCommandLine::Get(), TEXT("LearningActionsGuid"), ActionsGuid));
-			FGuid MemoryStatesGuid; ensure(FParse::Value(FCommandLine::Get(), TEXT("LearningMemoryStatesGuid"), MemoryStatesGuid));
-			FGuid RewardsGuid; ensure(FParse::Value(FCommandLine::Get(), TEXT("LearningRewardsGuid"), RewardsGuid));
-
 			ExperienceContainer.EpisodeStarts = SharedMemory::Map<2, int32>(EpisodeStartsGuid, { ProcessNum, ReplayBuffer.GetMaxEpisodeNum() });
+
+			FGuid EpisodeLengthsGuid; ensure(FParse::Value(FCommandLine::Get(), TEXT("LearningEpisodeLengthsGuid"), EpisodeLengthsGuid));
 			ExperienceContainer.EpisodeLengths = SharedMemory::Map<2, int32>(EpisodeLengthsGuid, { ProcessNum, ReplayBuffer.GetMaxEpisodeNum() });
-			ExperienceContainer.EpisodeCompletionModes = SharedMemory::Map<2, ECompletionMode>(EpisodeCompletionModesGuid, { ProcessNum, ReplayBuffer.GetMaxEpisodeNum() });
-			ExperienceContainer.EpisodeFinalObservations = SharedMemory::Map<3, float>(EpisodeFinalObservationsGuid, { ProcessNum, ReplayBuffer.GetMaxEpisodeNum(), ObservationVectorDimensionNum });
-			ExperienceContainer.EpisodeFinalMemoryStates = SharedMemory::Map<3, float>(EpisodeFinalMemoryStatesGuid, { ProcessNum, ReplayBuffer.GetMaxEpisodeNum(), MemoryStateVectorDimensionNum });
-			ExperienceContainer.Observations = SharedMemory::Map<3, float>(ObservationsGuid, { ProcessNum, ReplayBuffer.GetMaxStepNum(), ObservationVectorDimensionNum });
-			ExperienceContainer.Actions = SharedMemory::Map<3, float>(ActionsGuid, { ProcessNum, ReplayBuffer.GetMaxStepNum(), ActionVectorDimensionNum });
-			ExperienceContainer.MemoryStates = SharedMemory::Map<3, float>(MemoryStatesGuid, { ProcessNum, ReplayBuffer.GetMaxStepNum(), MemoryStateVectorDimensionNum });
-			ExperienceContainer.Rewards = SharedMemory::Map<2, float>(RewardsGuid, { ProcessNum, ReplayBuffer.GetMaxStepNum() });
-		}
 
-		SharedMemoryExperienceContainers.Add(Name, ExperienceContainer);
-	}
+			if (ReplayBuffer.HasCompletions())
+			{
+				FGuid EpisodeCompletionModesGuid; ensure(FParse::Value(FCommandLine::Get(), TEXT("LearningEpisodeCompletionModesGuid"), EpisodeCompletionModesGuid));
+				ExperienceContainer.EpisodeCompletionModes = SharedMemory::Map<2, ECompletionMode>(EpisodeCompletionModesGuid, { ProcessNum, ReplayBuffer.GetMaxEpisodeNum() });
+			}
 
-	bool FSharedMemoryTrainer::ContainsReplayBuffer(const FName& Name) const
-	{
-		return SharedMemoryExperienceContainers.Contains(Name);
+			// Final Observations
+			if (ReplayBuffer.HasFinalObservations())
+			{
+				FString StringOfGuids; ensure(FParse::Value(FCommandLine::Get(), TEXT("LearningEpisodeFinalObservationsGuids"), StringOfGuids));
+				TArray<FString> StringGuids;
+				StringOfGuids.ParseIntoArray(StringGuids, TEXT(","));
+				TArray<FGuid> Guids;
+				for (const FString& StringGuid : StringGuids)
+				{
+					FGuid Guid;
+					FGuid::Parse(StringGuid, Guid);
+					Guids.Add(Guid);
+				}
+				check(Guids.Num() == ReplayBuffer.GetObservationsNum());
+
+				for (int32 Index = 0; Index < ReplayBuffer.GetObservationsNum(); Index++)
+				{
+					const int32 DimNum = ReplayBuffer.GetEpisodeFinalObservations(Index).Num<1>();
+					ExperienceContainer.EpisodeFinalObservations.Add(SharedMemory::Map<3, float>(Guids[Index], { ProcessNum, ReplayBuffer.GetMaxEpisodeNum(), DimNum }));
+				}
+			}
+
+			// Final Memory States
+			if (ReplayBuffer.HasFinalMemoryStates())
+			{
+				FString StringOfGuids; ensure(FParse::Value(FCommandLine::Get(), TEXT("LearningEpisodeFinalMemoryStatesGuid"), StringOfGuids));
+				TArray<FString> StringGuids;
+				StringOfGuids.ParseIntoArray(StringGuids, TEXT(","));
+				TArray<FGuid> Guids;
+				for (const FString& StringGuid : StringGuids)
+				{
+					FGuid Guid;
+					FGuid::Parse(StringGuid, Guid);
+					Guids.Add(Guid);
+				}
+				check(Guids.Num() == ReplayBuffer.GetMemoryStatesNum());
+
+				for (int32 Index = 0; Index < ReplayBuffer.GetMemoryStatesNum(); Index++)
+				{
+					const int32 DimNum = ReplayBuffer.GetEpisodeFinalMemoryStates(Index).Num<1>();
+					ExperienceContainer.EpisodeFinalMemoryStates.Add(SharedMemory::Map<3, float>(Guids[Index], { ProcessNum, ReplayBuffer.GetMaxEpisodeNum(), DimNum }));
+				}
+			}
+
+			// Observations
+			{
+				FString StringOfGuids; ensure(FParse::Value(FCommandLine::Get(), TEXT("LearningObservationsGuid"), StringOfGuids));
+				TArray<FString> StringGuids;
+				StringOfGuids.ParseIntoArray(StringGuids, TEXT(","));
+				TArray<FGuid> Guids;
+				for (const FString& StringGuid : StringGuids)
+				{
+					FGuid Guid;
+					FGuid::Parse(StringGuid, Guid);
+					Guids.Add(Guid);
+				}
+				check(Guids.Num() == ReplayBuffer.GetObservationsNum());
+
+				for (int32 Index = 0; Index < ReplayBuffer.GetObservationsNum(); Index++)
+				{
+					const int32 DimNum = ReplayBuffer.GetObservations(Index).Num<1>();
+					ExperienceContainer.Observations.Add(SharedMemory::Map<3, float>(Guids[Index], { ProcessNum, ReplayBuffer.GetMaxStepNum(), DimNum }));
+				}
+			}
+
+			// Actions
+			{
+				FString StringOfGuids; ensure(FParse::Value(FCommandLine::Get(), TEXT("LearningActionsGuid"), StringOfGuids));
+				TArray<FString> StringGuids;
+				StringOfGuids.ParseIntoArray(StringGuids, TEXT(","));
+				TArray<FGuid> Guids;
+				for (const FString& StringGuid : StringGuids)
+				{
+					FGuid Guid;
+					FGuid::Parse(StringGuid, Guid);
+					Guids.Add(Guid);
+				}
+				check(Guids.Num() == ReplayBuffer.GetActionsNum());
+
+				for (int32 Index = 0; Index < ReplayBuffer.GetActionsNum(); Index++)
+				{
+					const int32 DimNum = ReplayBuffer.GetActions(Index).Num<1>();
+					ExperienceContainer.Actions.Add(SharedMemory::Map<3, float>(Guids[Index], { ProcessNum, ReplayBuffer.GetMaxStepNum(), DimNum }));
+				}
+			}
+
+			// Memory States
+			{
+				FString StringOfGuids; ensure(FParse::Value(FCommandLine::Get(), TEXT("LearningMemoryStatesGuid"), StringOfGuids));
+				TArray<FString> StringGuids;
+				StringOfGuids.ParseIntoArray(StringGuids, TEXT(","));
+				TArray<FGuid> Guids;
+				for (const FString& StringGuid : StringGuids)
+				{
+					FGuid Guid;
+					FGuid::Parse(StringGuid, Guid);
+					Guids.Add(Guid);
+				}
+				check(Guids.Num() == ReplayBuffer.GetMemoryStatesNum());
+
+				for (int32 Index = 0; Index < ReplayBuffer.GetMemoryStatesNum(); Index++)
+				{
+					const int32 DimNum = ReplayBuffer.GetMemoryStates(Index).Num<1>();
+					ExperienceContainer.MemoryStates.Add(SharedMemory::Map<3, float>(Guids[Index], { ProcessNum, ReplayBuffer.GetMaxStepNum(), DimNum }));
+				}
+			}
+
+			// Rewards
+			{
+				FString StringOfGuids; ensure(FParse::Value(FCommandLine::Get(), TEXT("LearningRewardsGuid"), StringOfGuids));
+				TArray<FString> StringGuids;
+				StringOfGuids.ParseIntoArray(StringGuids, TEXT(","));
+				TArray<FGuid> Guids;
+				for (const FString& StringGuid : StringGuids)
+				{
+					FGuid Guid;
+					FGuid::Parse(StringGuid, Guid);
+					Guids.Add(Guid);
+				}
+				check(Guids.Num() == ReplayBuffer.GetRewardsNum());
+
+				for (int32 Index = 0; Index < ReplayBuffer.GetRewardsNum(); Index++)
+				{
+					const int32 DimNum = ReplayBuffer.GetRewards(Index).Num<1>();
+					ExperienceContainer.Rewards.Add(SharedMemory::Map<3, float>(Guids[Index], { ProcessNum, ReplayBuffer.GetMaxStepNum(), DimNum }));
+				}
+			}
+		}
+		
+		const int32 ReplayBufferId = SharedMemoryExperienceContainers.Num();
+		SharedMemoryExperienceContainers.Add(ExperienceContainer);
+		return ReplayBufferId;
 	}
 
-	ETrainerResponse FSharedMemoryTrainer::SendReplayBuffer(const FName& Name, const FReplayBuffer& ReplayBuffer, const ELogSetting LogSettings)
+	ETrainerResponse FSharedMemoryTrainer::SendReplayBuffer(const int32 ReplayBufferId, const FReplayBuffer& ReplayBuffer, const ELogSetting LogSettings)
 	{
 		check(ProcessIdx != INDEX_NONE);
 		checkf(Controls.Region, TEXT("SendReplayBuffer: Controls Shared Memory Region is nullptr"));
-		if (!ensureMsgf(SharedMemoryExperienceContainers.Contains(Name), TEXT("ReplayBuffer %s has not been added. Call AddReplayBuffer prior to SendReplayBuffer."), *Name.ToString()))
+		if (!ensureMsgf(SharedMemoryExperienceContainers.Num() >= ReplayBufferId, TEXT("ReplayBuffer %d has not been added. Call AddReplayBuffer prior to SendReplayBuffer."), ReplayBufferId))
 		{
 			return ETrainerResponse::Unexpected;
 		}
 
+		TArray<TLearningArrayView<2, float>> EpisodeFinalObservations;
+		for (TSharedMemoryArrayView<3, float>& EpisodeFinalObs : SharedMemoryExperienceContainers[ReplayBufferId].EpisodeFinalObservations)
+		{
+			EpisodeFinalObservations.Add(EpisodeFinalObs.View[ProcessIdx]);
+		}
+
+		TArray<TLearningArrayView<2, float>> EpisodeFinalMemoryStates;
+		for (TSharedMemoryArrayView<3, float>& EpisodeFinalMems : SharedMemoryExperienceContainers[ReplayBufferId].EpisodeFinalMemoryStates)
+		{
+			EpisodeFinalMemoryStates.Add(EpisodeFinalMems.View[ProcessIdx]);
+		}
+
+		TArray<TLearningArrayView<2, float>> Observations;
+		for (TSharedMemoryArrayView<3, float>& Obs : SharedMemoryExperienceContainers[ReplayBufferId].Observations)
+		{
+			Observations.Add(Obs.View[ProcessIdx]);
+		}
+
+		TArray<TLearningArrayView<2, float>> Actions;
+		for (TSharedMemoryArrayView<3, float>& Acts : SharedMemoryExperienceContainers[ReplayBufferId].Actions)
+		{
+			Actions.Add(Acts.View[ProcessIdx]);
+		}
+
+		TArray<TLearningArrayView<2, float>> MemoryStates;
+		for (TSharedMemoryArrayView<3, float>& Mems : SharedMemoryExperienceContainers[ReplayBufferId].MemoryStates)
+		{
+			MemoryStates.Add(Mems.View[ProcessIdx]);
+		}
+
+		TArray<TLearningArrayView<2, float>> Rewards;
+		for (TSharedMemoryArrayView<3, float>& Rews : SharedMemoryExperienceContainers[ReplayBufferId].Rewards)
+		{
+			Rewards.Add(Rews.View[ProcessIdx]);
+		}
+
+		TLearningArrayView<1, ECompletionMode> EmptyCompletionsArray;
 		return SharedMemoryTraining::SendExperience(
-			SharedMemoryExperienceContainers[Name].EpisodeStarts.View[ProcessIdx],
-			SharedMemoryExperienceContainers[Name].EpisodeLengths.View[ProcessIdx],
-			SharedMemoryExperienceContainers[Name].EpisodeCompletionModes.View[ProcessIdx],
-			SharedMemoryExperienceContainers[Name].EpisodeFinalObservations.View[ProcessIdx],
-			SharedMemoryExperienceContainers[Name].EpisodeFinalMemoryStates.View[ProcessIdx],
-			SharedMemoryExperienceContainers[Name].Observations.View[ProcessIdx],
-			SharedMemoryExperienceContainers[Name].Actions.View[ProcessIdx],
-			SharedMemoryExperienceContainers[Name].MemoryStates.View[ProcessIdx],
-			SharedMemoryExperienceContainers[Name].Rewards.View[ProcessIdx],
+			SharedMemoryExperienceContainers[ReplayBufferId].EpisodeStarts.View[ProcessIdx],
+			SharedMemoryExperienceContainers[ReplayBufferId].EpisodeLengths.View[ProcessIdx],
+			ReplayBuffer.HasCompletions() ? SharedMemoryExperienceContainers[ReplayBufferId].EpisodeCompletionModes.View[ProcessIdx] : EmptyCompletionsArray,
+			EpisodeFinalObservations,
+			EpisodeFinalMemoryStates,
+			Observations,
+			Actions,
+			MemoryStates,
+			Rewards,
 			Controls.View[ProcessIdx],
 			*TrainingProcess,
+			ReplayBufferId,
 			ReplayBuffer,
 			Timeout,
 			LogSettings);
@@ -396,23 +672,25 @@ namespace UE::Learning
 
 	void FSharedMemoryTrainer::Deallocate()
 	{
-		for (TPair<FName, TSharedMemoryArrayView<1, uint8>>& MapEntry : NeuralNetworkSharedMemoryArrayViews)
+		for (TSharedMemoryArrayView<1, uint8>& SharedMemoryArrayView : NeuralNetworkSharedMemoryArrayViews)
 		{
-			if (MapEntry.Value.Region != nullptr)
+			if (SharedMemoryArrayView.Region != nullptr)
 			{
-				SharedMemory::Deallocate(MapEntry.Value);
+				SharedMemory::Deallocate(SharedMemoryArrayView);
 			}
 		}
 		NeuralNetworkSharedMemoryArrayViews.Empty();
 
-		for (TPair<FName, FSharedMemoryExperienceContainer>& MapEntry : SharedMemoryExperienceContainers)
+		for (FSharedMemoryExperienceContainer& SharedMemoryExperienceContainer : SharedMemoryExperienceContainers)
 		{
-			MapEntry.Value.Deallocate();
+			SharedMemoryExperienceContainer.Deallocate();
 		}
 		SharedMemoryExperienceContainers.Empty();
 	}
 
 	FSocketTrainerServerProcess::FSocketTrainerServerProcess(
+		const FString& CustomTrainerPath,
+		const FString& TrainerFileName,
 		const FString& PythonExecutablePath,
 		const FString& PythonContentPath,
 		const FString& IntermediatePath,
@@ -428,8 +706,10 @@ namespace UE::Learning
 		UE_LEARNING_CHECK(FPaths::DirectoryExists(PythonContentPath));
 
 		IFileManager& FileManager = IFileManager::Get();
-		const FString CommandLineArguments = FString::Printf(TEXT("\"%s\" Socket \"%s:%i\" \"%s\" %i"),
-			*FileManager.ConvertToAbsolutePathForExternalAppForRead(*(PythonContentPath / TEXT("train_ppo.py"))),
+		const FString CommandLineArguments = FString::Printf(TEXT("\"%s\" \"%s\" \"%s\" Socket \"%s:%i\" \"%s\" %i"),
+			*FileManager.ConvertToAbsolutePathForExternalAppForRead(*(PythonContentPath / TEXT("train.py"))),
+			*FileManager.ConvertToAbsolutePathForExternalAppForRead(*CustomTrainerPath),
+			*TrainerFileName,
 			IpAddress,
 			Port,
 			*FileManager.ConvertToAbsolutePathForExternalAppForRead(*IntermediatePath),
@@ -525,6 +805,11 @@ namespace UE::Learning
 		return ETrainerResponse::Success;
 	}
 
+	bool FSocketTrainer::HasNetworkOrCompleted()
+	{
+		return SocketTraining::HasNetworkOrCompleted(*Socket, TrainingProcess);
+	}
+
 	void FSocketTrainer::Terminate()
 	{
 		if (Socket)
@@ -554,67 +839,58 @@ namespace UE::Learning
 		return SocketTraining::SendConfig(*Socket, ConfigString, TrainingProcess, Timeout, LogSettings);
 	}
 
-	void FSocketTrainer::AddNetwork(
-		const FName& Name,
-		const ULearningNeuralNetworkData& Network)
-	{
-		NetworkBuffers.Add(Name, TLearningArray<1, uint8>());
-		NetworkBuffers[Name].SetNumUninitialized({Network.GetSnapshotByteNum()});
-	}
-
-	bool FSocketTrainer::ContainsNetwork(const FName& Name) const
+	int32 FSocketTrainer::AddNetwork(const ULearningNeuralNetworkData& Network)
 	{
-		return NetworkBuffers.Contains(Name);
+		const int32 NetworkId = NetworkBuffers.Num();
+		NetworkBuffers.Add(TLearningArray<1, uint8>());
+		NetworkBuffers[NetworkId].SetNumUninitialized({Network.GetSnapshotByteNum()});
+		return NetworkId;
 	}
 
 	ETrainerResponse FSocketTrainer::ReceiveNetwork(
-		const FName& Name,
+		const int32 NetworkId,
 		ULearningNeuralNetworkData& OutNetwork,
 		FRWLock* NetworkLock,
 		const ELogSetting LogSettings)
 	{
 		checkf(Socket, TEXT("Training socket is nullptr"));
-		if (!ensureMsgf(NetworkBuffers.Contains(Name), TEXT("Network %s has not been added. Call AddNetwork prior to ReceiveNetwork."), *Name.ToString()))
+		if (!ensureMsgf(NetworkBuffers.Num() >= NetworkId, TEXT("Network %d has not been added. Call AddNetwork prior to ReceiveNetwork."), NetworkId))
 		{
 			return ETrainerResponse::Unexpected;
 		}
 
-		return SocketTraining::RecvNetwork(*Socket, OutNetwork, TrainingProcess, NetworkBuffers[Name], Timeout, NetworkLock, LogSettings);
+		return SocketTraining::RecvNetwork(*Socket, NetworkId, OutNetwork, TrainingProcess, NetworkBuffers[NetworkId], Timeout, NetworkLock, LogSettings);
 	}
 
 	ETrainerResponse FSocketTrainer::SendNetwork(
-		const FName& Name,
+		const int32 NetworkId,
 		const ULearningNeuralNetworkData& Network,
 		FRWLock* NetworkLock,
 		const ELogSetting LogSettings)
 	{
 		checkf(Socket, TEXT("Training socket is nullptr"));
-		if (!ensureMsgf(NetworkBuffers.Contains(Name), TEXT("Network %s has not been added. Call AddNetwork prior to SendNetwork."), *Name.ToString()))
+		if (!ensureMsgf(NetworkBuffers.Num() >= NetworkId, TEXT("Network %d has not been added. Call AddNetwork prior to SendNetwork."), NetworkId))
 		{
 			return ETrainerResponse::Unexpected;
 		}
 
-		return SocketTraining::SendNetwork(*Socket, NetworkBuffers[Name], TrainingProcess, Network, Timeout, NetworkLock, LogSettings);
-	}
-
-	void FSocketTrainer::AddReplayBuffer(const FName& Name, const FReplayBuffer& ReplayBuffer)
-	{
-		ExperienceBufferNames.Add(Name);
+		return SocketTraining::SendNetwork(*Socket, NetworkBuffers[NetworkId], TrainingProcess, NetworkId, Network, Timeout, NetworkLock, LogSettings);
 	}
 
-	bool FSocketTrainer::ContainsReplayBuffer(const FName& Name) const
+	int32 FSocketTrainer::AddReplayBuffer(const FReplayBuffer& ReplayBuffer)
 	{
-		return ExperienceBufferNames.Contains(Name);
+		LastReplayBufferId++;
+		return LastReplayBufferId;
 	}
 
-	ETrainerResponse FSocketTrainer::SendReplayBuffer(const FName& Name, const FReplayBuffer& ReplayBuffer, const ELogSetting LogSettings)
+	ETrainerResponse FSocketTrainer::SendReplayBuffer(const int32 ReplayBufferId, const FReplayBuffer& ReplayBuffer, const ELogSetting LogSettings)
 	{
 		checkf(Socket, TEXT("Training socket is nullptr"));
-		if (!ensureMsgf(ExperienceBufferNames.Contains(Name), TEXT("ReplayBuffer %s has not been added. Call AddReplayBuffer prior to SendReplayBuffer."), *Name.ToString()))
+		if (!ensureMsgf(ReplayBufferId <= LastReplayBufferId, TEXT("ReplayBuffer %d has not been added. Call AddReplayBuffer prior to SendReplayBuffer."), ReplayBufferId))
 		{
 			return ETrainerResponse::Unexpected;
 		}
 
-		return SocketTraining::SendExperience(*Socket, ReplayBuffer, TrainingProcess, Timeout, LogSettings);
+		return SocketTraining::SendExperience(*Socket, ReplayBufferId, ReplayBuffer, TrainingProcess, Timeout, LogSettings);
 	}
 }
diff --git a/Engine/Plugins/Experimental/LearningAgents/Source/LearningTraining/Private/LearningImitationTrainer.cpp b/Engine/Plugins/Experimental/LearningAgents/Source/LearningTraining/Private/LearningImitationTrainer.cpp
deleted file mode 100644
diff --git a/Engine/Plugins/Experimental/LearningAgents/Source/LearningTraining/Private/LearningPPOTrainer.cpp b/Engine/Plugins/Experimental/LearningAgents/Source/LearningTraining/Private/LearningPPOTrainer.cpp
--- a/Engine/Plugins/Experimental/LearningAgents/Source/LearningTraining/Private/LearningPPOTrainer.cpp
+++ b/Engine/Plugins/Experimental/LearningAgents/Source/LearningTraining/Private/LearningPPOTrainer.cpp
@@ -82,6 +82,8 @@ int32 ULearningSocketPPOTrainerServerCommandlet::Main(const FString& Commandline
 	UE_LOG(LogLearning, Display, TEXT("LogSettings: %s"), LogSettings == UE::Learning::ELogSetting::Normal ? TEXT("Normal") : TEXT("Silent"));
 
 	UE::Learning::FSocketTrainerServerProcess ServerProcess(
+		UE::Learning::Trainer::GetProjectPythonContentPath(),
+		TEXT("train_ppo"),
 		PythonExecutiblePath,
 		PythonContentPath,
 		IntermediatePath,
@@ -157,68 +159,155 @@ namespace UE::Learning::PPOTrainer
 			UE_LOG(LogLearning, Display, TEXT("Sending initial Policy..."));
 		}
 
-		ExternalTrainer->AddNetwork(TEXT("Policy"), PolicyNetwork);
-		ExternalTrainer->AddNetwork(TEXT("Critic"), CriticNetwork);
-		ExternalTrainer->AddNetwork(TEXT("Encoder"), EncoderNetwork);
-		ExternalTrainer->AddNetwork(TEXT("Decoder"), DecoderNetwork);
-		ExternalTrainer->AddReplayBuffer(TEXT("ReplayBuffer"), ReplayBuffer);
+		const int32 PolicyNetworkId = ExternalTrainer->AddNetwork(PolicyNetwork);
+		const int32 CriticNetworkId = ExternalTrainer->AddNetwork(CriticNetwork);
+		const int32 EncoderNetworkId = ExternalTrainer->AddNetwork(EncoderNetwork);
+		const int32 DecoderNetworkId = ExternalTrainer->AddNetwork(DecoderNetwork);
+		const int32 ReplayBufferId = ExternalTrainer->AddReplayBuffer(ReplayBuffer);
 
-		const int32 ObservationVectorDimensionNum = ReplayBuffer.GetObservations().Num<1>();
-		const int32 ActionVectorDimensionNum = ReplayBuffer.GetActions().Num<1>();
-		const int32 MemoryStateVectorDimensionNum = ReplayBuffer.GetMemoryStates().Num<1>();
+		const int32 ObservationVectorDimensionNum = ReplayBuffer.GetObservations(0).Num<1>();
+		const int32 ActionVectorDimensionNum = ReplayBuffer.GetActions(0).Num<1>();
+		const int32 MemoryStateVectorDimensionNum = ReplayBuffer.GetMemoryStates(0).Num<1>();
 
 		// Write PPO Config
 		TSharedRef<FJsonObject> ConfigObject = MakeShared<FJsonObject>();
 		ConfigObject->SetStringField(TEXT("TaskName"), TEXT("Training"));
 		ConfigObject->SetStringField(TEXT("TrainerMethod"), TEXT("PPO"));
-		ConfigObject->SetStringField(TEXT("TrainerType"), TEXT("Network"));
 		ConfigObject->SetStringField(TEXT("TimeStamp"), *FDateTime::Now().ToFormattedString(TEXT("%Y-%m-%d_%H-%M-%S")));
 
-		ConfigObject->SetObjectField(TEXT("ObservationSchema"), Trainer::ConvertObservationSchemaToJSON(ObservationSchema, ObservationSchemaElement));
-		ConfigObject->SetObjectField(TEXT("ActionSchema"), Trainer::ConvertActionSchemaToJSON(ActionSchema, ActionSchemaElement));
-		ConfigObject->SetNumberField(TEXT("ObservationVectorDimensionNum"), ObservationVectorDimensionNum);
-		ConfigObject->SetNumberField(TEXT("ActionVectorDimensionNum"), ActionVectorDimensionNum);
-		ConfigObject->SetNumberField(TEXT("MemoryStateVectorDimensionNum"), MemoryStateVectorDimensionNum);
-		ConfigObject->SetNumberField(TEXT("MaxEpisodeNum"), ReplayBuffer.GetMaxEpisodeNum());
-		ConfigObject->SetNumberField(TEXT("MaxStepNum"), ReplayBuffer.GetMaxStepNum());
-
-		ConfigObject->SetNumberField(TEXT("PolicyNetworkByteNum"), PolicyNetwork.GetSnapshotByteNum());
-		ConfigObject->SetNumberField(TEXT("CriticNetworkByteNum"), CriticNetwork.GetSnapshotByteNum());
-		ConfigObject->SetNumberField(TEXT("EncoderNetworkByteNum"), EncoderNetwork.GetSnapshotByteNum());
-		ConfigObject->SetNumberField(TEXT("DecoderNetworkByteNum"), DecoderNetwork.GetSnapshotByteNum());
-
-		ConfigObject->SetNumberField(TEXT("IterationNum"), TrainerSettings.IterationNum);
-		ConfigObject->SetNumberField(TEXT("LearningRatePolicy"), TrainerSettings.LearningRatePolicy);
-		ConfigObject->SetNumberField(TEXT("LearningRateCritic"), TrainerSettings.LearningRateCritic);
-		ConfigObject->SetNumberField(TEXT("LearningRateDecay"), TrainerSettings.LearningRateDecay);
-		ConfigObject->SetNumberField(TEXT("WeightDecay"), TrainerSettings.WeightDecay);
-		ConfigObject->SetNumberField(TEXT("PolicyBatchSize"), TrainerSettings.PolicyBatchSize);
-		ConfigObject->SetNumberField(TEXT("CriticBatchSize"), TrainerSettings.CriticBatchSize);
-		ConfigObject->SetNumberField(TEXT("PolicyWindow"), TrainerSettings.PolicyWindow);
-		ConfigObject->SetNumberField(TEXT("IterationsPerGather"), TrainerSettings.IterationsPerGather);
-		ConfigObject->SetNumberField(TEXT("CriticWarmupIterations"), TrainerSettings.CriticWarmupIterations);
-		ConfigObject->SetNumberField(TEXT("EpsilonClip"), TrainerSettings.EpsilonClip);
-		ConfigObject->SetNumberField(TEXT("ActionSurrogateWeight"), TrainerSettings.ActionSurrogateWeight);
-		ConfigObject->SetNumberField(TEXT("ActionRegularizationWeight"), TrainerSettings.ActionRegularizationWeight);
-		ConfigObject->SetNumberField(TEXT("ActionEntropyWeight"), TrainerSettings.ActionEntropyWeight);
-		ConfigObject->SetNumberField(TEXT("ReturnRegularizationWeight"), TrainerSettings.ReturnRegularizationWeight);
-		ConfigObject->SetNumberField(TEXT("GaeLambda"), TrainerSettings.GaeLambda);
-		ConfigObject->SetBoolField(TEXT("AdvantageNormalization"), TrainerSettings.bAdvantageNormalization);
-		ConfigObject->SetNumberField(TEXT("AdvantageMin"), TrainerSettings.AdvantageMin);
-		ConfigObject->SetNumberField(TEXT("AdvantageMax"), TrainerSettings.AdvantageMax);
-		ConfigObject->SetBoolField(TEXT("UseGradNormMaxClipping"), TrainerSettings.bUseGradNormMaxClipping);
-		ConfigObject->SetNumberField(TEXT("GradNormMax"), TrainerSettings.GradNormMax);
-		ConfigObject->SetNumberField(TEXT("TrimEpisodeStartStepNum"), TrainerSettings.TrimEpisodeStartStepNum);
-		ConfigObject->SetNumberField(TEXT("TrimEpisodeEndStepNum"), TrainerSettings.TrimEpisodeEndStepNum);
-		ConfigObject->SetNumberField(TEXT("Seed"), TrainerSettings.Seed);
-		ConfigObject->SetNumberField(TEXT("DiscountFactor"), TrainerSettings.DiscountFactor);
-		ConfigObject->SetStringField(TEXT("Device"), Trainer::GetDeviceString(TrainerSettings.Device));
-		ConfigObject->SetBoolField(TEXT("UseTensorBoard"), TrainerSettings.bUseTensorboard);
-		ConfigObject->SetBoolField(TEXT("SaveSnapshots"), TrainerSettings.bSaveSnapshots);
+		const int32 ObservationSchemaId = 0;
+		const int32 ActionSchemaId = 0;
+
+		// Add Neural Network Config Entries
+		TArray<TSharedPtr<FJsonValue>> NetworkObjects;
+
+		// Policy
+		{
+			TSharedPtr<FJsonObject> NetworkObject = MakeShared<FJsonObject>();
+			NetworkObject->SetNumberField(TEXT("Id"), PolicyNetworkId);
+			NetworkObject->SetStringField(TEXT("Name"), "Policy");
+			NetworkObject->SetNumberField(TEXT("MaxByteNum"), PolicyNetwork.GetSnapshotByteNum());
+
+			TSharedRef<FJsonValueObject> JsonValue = MakeShared<FJsonValueObject>(NetworkObject);
+			NetworkObjects.Add(JsonValue);
+		}
+
+		// Critic
+		{
+			TSharedPtr<FJsonObject> NetworkObject = MakeShared<FJsonObject>();
+			NetworkObject->SetNumberField(TEXT("Id"), CriticNetworkId);
+			NetworkObject->SetStringField(TEXT("Name"), "Critic");
+			NetworkObject->SetNumberField(TEXT("MaxByteNum"), CriticNetwork.GetSnapshotByteNum());
+			NetworkObject->SetNumberField(TEXT("InputSchemaId"), ObservationSchemaId);
+
+			TSharedRef<FJsonValueObject> JsonValue = MakeShared<FJsonValueObject>(NetworkObject);
+			NetworkObjects.Add(JsonValue);
+		}
+
+		// Encoder
+		{
+			TSharedPtr<FJsonObject> NetworkObject = MakeShared<FJsonObject>();
+			NetworkObject->SetNumberField(TEXT("Id"), EncoderNetworkId);
+			NetworkObject->SetStringField(TEXT("Name"), "Encoder");
+			NetworkObject->SetNumberField(TEXT("MaxByteNum"), EncoderNetwork.GetSnapshotByteNum());
+			NetworkObject->SetNumberField(TEXT("InputSchemaId"), ObservationSchemaId);
+
+			TSharedRef<FJsonValueObject> JsonValue = MakeShared<FJsonValueObject>(NetworkObject);
+			NetworkObjects.Add(JsonValue);
+		}
+
+		// Decoder
+		{
+			TSharedPtr<FJsonObject> NetworkObject = MakeShared<FJsonObject>();
+			NetworkObject->SetNumberField(TEXT("Id"), DecoderNetworkId);
+			NetworkObject->SetStringField(TEXT("Name"), "Decoder");
+			NetworkObject->SetNumberField(TEXT("MaxByteNum"), DecoderNetwork.GetSnapshotByteNum());
+			NetworkObject->SetNumberField(TEXT("OutputSchemaId"), ActionSchemaId);
+
+			TSharedRef<FJsonValueObject> JsonValue = MakeShared<FJsonValueObject>(NetworkObject);
+			NetworkObjects.Add(JsonValue);
+		}
+
+		ConfigObject->SetArrayField(TEXT("Networks"), NetworkObjects);
+
+		// Add Replay Buffers Config Entries
+		TArray<TSharedPtr<FJsonValue>> ReplayBufferObjects;
+		TSharedRef<FJsonValueObject> ReplayBufferJsonValue = MakeShared<FJsonValueObject>(ReplayBuffer.AsJsonConfig(ReplayBufferId));
+		ReplayBufferObjects.Add(ReplayBufferJsonValue);
+		ConfigObject->SetArrayField(TEXT("ReplayBuffers"), ReplayBufferObjects);
+
+		// Schemas
+		TSharedPtr<FJsonObject> SchemasObject = MakeShared<FJsonObject>();
+
+		// Add the observation schemas
+		TArray<TSharedPtr<FJsonValue>> ObservationSchemaObjects;
+		{
+			// For this PPO trainer, add the one observation schema we have
+			TSharedPtr<FJsonObject> ObservationSchemaObject = MakeShared<FJsonObject>();
+			ObservationSchemaObject->SetNumberField(TEXT("Id"), ObservationSchemaId);
+			ObservationSchemaObject->SetStringField(TEXT("Name"), "Default");
+			ObservationSchemaObject->SetObjectField(TEXT("Schema"),
+				UE::Learning::Trainer::ConvertObservationSchemaToJSON(ObservationSchema, ObservationSchemaElement));
+
+			TSharedRef<FJsonValueObject> JsonValue = MakeShared<FJsonValueObject>(ObservationSchemaObject);
+			ObservationSchemaObjects.Add(JsonValue);
+		}
+		SchemasObject->SetArrayField(TEXT("Observations"), ObservationSchemaObjects);
+
+		// Add the action schemas
+		TArray<TSharedPtr<FJsonValue>> ActionSchemaObjects;
+		{
+			// For this PPO trainer, add the one action schema we have
+			TSharedPtr<FJsonObject> ActionSchemaObject = MakeShared<FJsonObject>();
+			ActionSchemaObject->SetNumberField(TEXT("Id"), ActionSchemaId);
+			ActionSchemaObject->SetStringField(TEXT("Name"), "Default");
+			ActionSchemaObject->SetObjectField(TEXT("Schema"),
+				UE::Learning::Trainer::ConvertActionSchemaToJSON(ActionSchema, ActionSchemaElement));
+
+			TSharedRef<FJsonValueObject> JsonValue = MakeShared<FJsonValueObject>(ActionSchemaObject);
+			ActionSchemaObjects.Add(JsonValue);
+		}
+		SchemasObject->SetArrayField(TEXT("Actions"), ActionSchemaObjects);
+
+		ConfigObject->SetObjectField(TEXT("Schemas"), SchemasObject);
+
+		// Add PPO Specific Config Entries
+		TSharedRef<FJsonObject> PPOConfigObject = MakeShared<FJsonObject>();
+
+		PPOConfigObject->SetNumberField(TEXT("IterationNum"), TrainerSettings.IterationNum);
+		PPOConfigObject->SetNumberField(TEXT("LearningRatePolicy"), TrainerSettings.LearningRatePolicy);
+		PPOConfigObject->SetNumberField(TEXT("LearningRateCritic"), TrainerSettings.LearningRateCritic);
+		PPOConfigObject->SetNumberField(TEXT("LearningRateDecay"), TrainerSettings.LearningRateDecay);
+		PPOConfigObject->SetNumberField(TEXT("WeightDecay"), TrainerSettings.WeightDecay);
+		PPOConfigObject->SetNumberField(TEXT("PolicyBatchSize"), TrainerSettings.PolicyBatchSize);
+		PPOConfigObject->SetNumberField(TEXT("CriticBatchSize"), TrainerSettings.CriticBatchSize);
+		PPOConfigObject->SetNumberField(TEXT("PolicyWindow"), TrainerSettings.PolicyWindow);
+		PPOConfigObject->SetNumberField(TEXT("IterationsPerGather"), TrainerSettings.IterationsPerGather);
+		PPOConfigObject->SetNumberField(TEXT("CriticWarmupIterations"), TrainerSettings.CriticWarmupIterations);
+		PPOConfigObject->SetNumberField(TEXT("EpsilonClip"), TrainerSettings.EpsilonClip);
+		PPOConfigObject->SetNumberField(TEXT("ActionSurrogateWeight"), TrainerSettings.ActionSurrogateWeight);
+		PPOConfigObject->SetNumberField(TEXT("ActionRegularizationWeight"), TrainerSettings.ActionRegularizationWeight);
+		PPOConfigObject->SetNumberField(TEXT("ActionEntropyWeight"), TrainerSettings.ActionEntropyWeight);
+		PPOConfigObject->SetNumberField(TEXT("ReturnRegularizationWeight"), TrainerSettings.ReturnRegularizationWeight);
+		PPOConfigObject->SetNumberField(TEXT("GaeLambda"), TrainerSettings.GaeLambda);
+		PPOConfigObject->SetBoolField(TEXT("AdvantageNormalization"), TrainerSettings.bAdvantageNormalization);
+		PPOConfigObject->SetNumberField(TEXT("AdvantageMin"), TrainerSettings.AdvantageMin);
+		PPOConfigObject->SetNumberField(TEXT("AdvantageMax"), TrainerSettings.AdvantageMax);
+		PPOConfigObject->SetBoolField(TEXT("UseGradNormMaxClipping"), TrainerSettings.bUseGradNormMaxClipping);
+		PPOConfigObject->SetNumberField(TEXT("GradNormMax"), TrainerSettings.GradNormMax);
+		PPOConfigObject->SetNumberField(TEXT("TrimEpisodeStartStepNum"), TrainerSettings.TrimEpisodeStartStepNum);
+		PPOConfigObject->SetNumberField(TEXT("TrimEpisodeEndStepNum"), TrainerSettings.TrimEpisodeEndStepNum);
+		PPOConfigObject->SetNumberField(TEXT("Seed"), TrainerSettings.Seed);
+		PPOConfigObject->SetNumberField(TEXT("DiscountFactor"), TrainerSettings.DiscountFactor);
+		PPOConfigObject->SetStringField(TEXT("Device"), UE::Learning::Trainer::GetDeviceString(TrainerSettings.Device));
+		PPOConfigObject->SetBoolField(TEXT("UseTensorBoard"), TrainerSettings.bUseTensorboard);
+		PPOConfigObject->SetBoolField(TEXT("SaveSnapshots"), TrainerSettings.bSaveSnapshots);
+
+		ConfigObject->SetObjectField(TEXT("PPOSettings"), PPOConfigObject);
 
 		ExternalTrainer->SendConfig(ConfigObject, LogSettings);
 
-		Response = ExternalTrainer->SendNetwork(TEXT("Policy"), PolicyNetwork, PolicyNetworkLock);
+		Response = ExternalTrainer->SendNetwork(PolicyNetworkId, PolicyNetwork, PolicyNetworkLock);
 
 		if (Response != ETrainerResponse::Success)
 		{
@@ -238,7 +327,7 @@ namespace UE::Learning::PPOTrainer
 			UE_LOG(LogLearning, Display, TEXT("Sending initial Critic..."));
 		}
 
-		Response = ExternalTrainer->SendNetwork(TEXT("Critic"), CriticNetwork, CriticNetworkLock);
+		Response = ExternalTrainer->SendNetwork(CriticNetworkId, CriticNetwork, CriticNetworkLock);
 
 		if (Response != ETrainerResponse::Success)
 		{
@@ -258,7 +347,7 @@ namespace UE::Learning::PPOTrainer
 			UE_LOG(LogLearning, Display, TEXT("Sending initial Encoder..."));
 		}
 
-		Response = ExternalTrainer->SendNetwork(TEXT("Encoder"), EncoderNetwork, EncoderNetworkLock);
+		Response = ExternalTrainer->SendNetwork(EncoderNetworkId, EncoderNetwork, EncoderNetworkLock);
 
 		if (Response != ETrainerResponse::Success)
 		{
@@ -278,7 +367,7 @@ namespace UE::Learning::PPOTrainer
 			UE_LOG(LogLearning, Display, TEXT("Sending initial Decoder..."));
 		}
 
-		Response = ExternalTrainer->SendNetwork(TEXT("Decoder"), DecoderNetwork, DecoderNetworkLock);
+		Response = ExternalTrainer->SendNetwork(DecoderNetworkId, DecoderNetwork, DecoderNetworkLock);
 
 		if (Response != ETrainerResponse::Success)
 		{
@@ -324,24 +413,24 @@ namespace UE::Learning::PPOTrainer
 					ReplayBuffer,
 					EpisodeBuffer,
 					ResetBuffer,
-					ObservationVectorBuffer,
-					ActionVectorBuffer,
-					PreEvaluationMemoryStateVectorBuffer,
-					MemoryStateVectorBuffer,
-					RewardBuffer,
+					{ ObservationVectorBuffer },
+					{ ActionVectorBuffer},
+					{ PreEvaluationMemoryStateVectorBuffer },
+					{ MemoryStateVectorBuffer },
+					{ RewardBuffer },
 					CompletionBuffer,
 					EpisodeCompletionBuffer,
 					AllCompletionBuffer,
 					ResetFunction,
-					ObservationFunction,
-					PolicyFunction,
-					ActionFunction,
-					UpdateFunction,
-					RewardFunction,
+					{ ObservationFunction },
+					{ PolicyFunction },
+					{ ActionFunction },
+					{ UpdateFunction },
+					{ RewardFunction },
 					CompletionFunction,
 					Instances);
 
-				Response = ExternalTrainer->SendReplayBuffer(TEXT("ReplayBuffer"), ReplayBuffer);
+				Response = ExternalTrainer->SendReplayBuffer(ReplayBufferId, ReplayBuffer);
 
 				if (Response != ETrainerResponse::Success)
 				{
@@ -357,7 +446,7 @@ namespace UE::Learning::PPOTrainer
 
 			// Update Policy
 
-			Response = ExternalTrainer->ReceiveNetwork(TEXT("Policy"), PolicyNetwork, PolicyNetworkLock);
+			Response = ExternalTrainer->ReceiveNetwork(PolicyNetworkId, PolicyNetwork, PolicyNetworkLock);
 
 			if (Response == ETrainerResponse::Completed)
 			{
@@ -383,7 +472,7 @@ namespace UE::Learning::PPOTrainer
 
 			// Update Critic
 
-			Response = ExternalTrainer->ReceiveNetwork(TEXT("Critic"), CriticNetwork, CriticNetworkLock);
+			Response = ExternalTrainer->ReceiveNetwork(CriticNetworkId, CriticNetwork, CriticNetworkLock);
 
 			if (Response != ETrainerResponse::Success)
 			{
@@ -401,7 +490,7 @@ namespace UE::Learning::PPOTrainer
 
 			// Update Encoder
 
-			Response = ExternalTrainer->ReceiveNetwork(TEXT("Encoder"), EncoderNetwork, EncoderNetworkLock);
+			Response = ExternalTrainer->ReceiveNetwork(EncoderNetworkId, EncoderNetwork, EncoderNetworkLock);
 
 			if (Response != ETrainerResponse::Success)
 			{
@@ -419,7 +508,7 @@ namespace UE::Learning::PPOTrainer
 
 			// Update Decoder
 
-			Response = ExternalTrainer->ReceiveNetwork(TEXT("Decoder"), DecoderNetwork, DecoderNetworkLock);
+			Response = ExternalTrainer->ReceiveNetwork(DecoderNetworkId, DecoderNetwork, DecoderNetworkLock);
 
 			if (Response != ETrainerResponse::Success)
 			{
diff --git a/Engine/Plugins/Experimental/LearningAgents/Source/LearningTraining/Private/LearningSharedMemoryTraining.cpp b/Engine/Plugins/Experimental/LearningAgents/Source/LearningTraining/Private/LearningSharedMemoryTraining.cpp
--- a/Engine/Plugins/Experimental/LearningAgents/Source/LearningTraining/Private/LearningSharedMemoryTraining.cpp
+++ b/Engine/Plugins/Experimental/LearningAgents/Source/LearningTraining/Private/LearningSharedMemoryTraining.cpp
@@ -22,7 +22,7 @@ namespace UE::Learning::SharedMemoryTraining
 		return ETrainerResponse::Success;
 	}
 
-	bool HasPolicyOrCompleted(TLearningArrayView<1, volatile int32> Controls)
+	bool HasNetworkOrCompleted(TLearningArrayView<1, volatile int32> Controls)
 	{
 		return Controls[(uint8)EControls::NetworkSignal] || Controls[(uint8)EControls::CompleteSignal];
 	}
@@ -43,6 +43,7 @@ namespace UE::Learning::SharedMemoryTraining
 
 	ETrainerResponse RecvNetwork(
 		TLearningArrayView<1, volatile int32> Controls,
+		const int32 NetworkId,
 		ULearningNeuralNetworkData& OutNetwork,
 		FSubprocess& Process,
 		const TLearningArrayView<1, const uint8> NetworkData,
@@ -98,6 +99,8 @@ namespace UE::Learning::SharedMemoryTraining
 		{
 			FScopeNullableWriteLock ScopeLock(NetworkLock);
 
+			UE_LEARNING_CHECKF(Controls[(uint8)EControls::NetworkId] == NetworkId, TEXT("Received unexpected NetworkId!"));
+
 			if (NetworkData.Num() != OutNetwork.GetSnapshotByteNum())
 			{
 				UE_LOG(LogLearning, Error, TEXT("Error receiving network. Incorrect buffer size. Buffer is %i bytes, expected %i."), NetworkData.Num(), OutNetwork.GetSnapshotByteNum());
@@ -118,6 +121,7 @@ namespace UE::Learning::SharedMemoryTraining
 		}
 
 		// Confirm we have read the network
+		Controls[(uint8)EControls::NetworkId] = -1;
 		Controls[(uint8)EControls::NetworkSignal] = false;
 
 		return bSuccess ? ETrainerResponse::Success : ETrainerResponse::Unexpected;
@@ -125,6 +129,7 @@ namespace UE::Learning::SharedMemoryTraining
 
 	ETrainerResponse SendNetwork(
 		TLearningArrayView<1, volatile int32> Controls,
+		const int32 NetworkId,
 		TLearningArrayView<1, uint8> NetworkData,
 		FSubprocess& Process,
 		const ULearningNeuralNetworkData& Network,
@@ -184,6 +189,7 @@ namespace UE::Learning::SharedMemoryTraining
 		}
 
 		// Confirm we have written the network
+		Controls[(uint8)EControls::NetworkId] = NetworkId;
 		Controls[(uint8)EControls::NetworkSignal] = false;
 
 		return bSuccess ? ETrainerResponse::Success : ETrainerResponse::Unexpected;
@@ -193,14 +199,15 @@ namespace UE::Learning::SharedMemoryTraining
 		TLearningArrayView<1, int32> EpisodeStarts,
 		TLearningArrayView<1, int32> EpisodeLengths,
 		TLearningArrayView<1, ECompletionMode> EpisodeCompletionModes,
-		TLearningArrayView<2, float> EpisodeFinalObservations,
-		TLearningArrayView<2, float> EpisodeFinalMemoryStates,
-		TLearningArrayView<2, float> Observations,
-		TLearningArrayView<2, float> Actions,
-		TLearningArrayView<2, float> MemoryStates,
-		TLearningArrayView<1, float> Rewards,
+		TArrayView<TLearningArrayView<2, float>> EpisodeFinalObservations,
+		TArrayView<TLearningArrayView<2, float>> EpisodeFinalMemoryStates,
+		TArrayView<TLearningArrayView<2, float>> Observations,
+		TArrayView<TLearningArrayView<2, float>> Actions,
+		TArrayView<TLearningArrayView<2, float>> MemoryStates,
+		TArrayView<TLearningArrayView<2, float>> Rewards,
 		TLearningArrayView<1, volatile int32> Controls,
 		FSubprocess& Process,
+		const int32 ReplayBufferId,
 		const FReplayBuffer& ReplayBuffer,
 		const float Timeout,
 		const ELogSetting LogSettings)
@@ -246,86 +253,54 @@ namespace UE::Learning::SharedMemoryTraining
 		// Write experience to the shared memory
 		Array::Copy(EpisodeStarts.Slice(0, EpisodeNum), ReplayBuffer.GetEpisodeStarts());
 		Array::Copy(EpisodeLengths.Slice(0, EpisodeNum), ReplayBuffer.GetEpisodeLengths());
-		Array::Copy(EpisodeCompletionModes.Slice(0, EpisodeNum), ReplayBuffer.GetEpisodeCompletionModes());
-		Array::Copy(EpisodeFinalObservations.Slice(0, EpisodeNum), ReplayBuffer.GetEpisodeFinalObservations());
-		Array::Copy(EpisodeFinalMemoryStates.Slice(0, EpisodeNum), ReplayBuffer.GetEpisodeFinalMemoryStates());
-		Array::Copy(Observations.Slice(0, StepNum), ReplayBuffer.GetObservations());
-		Array::Copy(Actions.Slice(0, StepNum), ReplayBuffer.GetActions());
-		Array::Copy(MemoryStates.Slice(0, StepNum), ReplayBuffer.GetMemoryStates());
-		Array::Copy(Rewards.Slice(0, StepNum), ReplayBuffer.GetRewards());
-
-		// Indicate that experience is written
-		Controls[(uint8)EControls::ExperienceEpisodeNum] = EpisodeNum;
-		Controls[(uint8)EControls::ExperienceStepNum] = StepNum;
-		Controls[(uint8)EControls::ExperienceSignal] = true;
-
-		return ETrainerResponse::Success;
-	}
-
-	ETrainerResponse SendExperience(
-		TLearningArrayView<1, int32> EpisodeStarts,
-		TLearningArrayView<1, int32> EpisodeLengths,
-		TLearningArrayView<2, float> Observations,
-		TLearningArrayView<2, float> Actions,
-		TLearningArrayView<1, volatile int32> Controls,
-		FSubprocess& Process,
-		const TLearningArrayView<1, const int32> EpisodeStartsExperience,
-		const TLearningArrayView<1, const int32> EpisodeLengthsExperience,
-		const TLearningArrayView<2, const float> ObservationExperience,
-		const TLearningArrayView<2, const float> ActionExperience,
-		const float Timeout,
-		const ELogSetting LogSettings)
-	{
-		const float SleepTime = 0.001f;
-		float WaitTime = 0.0f;
 
-		// Wait until the training process is done reading any experience
-		while (Controls[(uint8)EControls::ExperienceSignal])
+		if (ReplayBuffer.HasCompletions())
 		{
-			// If we're monitoring a process, then has it has exited?
-			if (!Process.Update())
-			{
-				return ETrainerResponse::Unexpected;
-			}
+			Array::Copy(EpisodeCompletionModes.Slice(0, EpisodeNum), ReplayBuffer.GetEpisodeCompletionModes());
+		}
 
-			// Check if we've timed out
-			if (WaitTime > Timeout)
+		if (ReplayBuffer.HasFinalObservations())
+		{
+			for (int32 Index = 0; Index < ReplayBuffer.GetObservationsNum(); Index++)
 			{
-				return ETrainerResponse::Timeout;
+				Array::Copy(EpisodeFinalObservations[Index].Slice(0, EpisodeNum), ReplayBuffer.GetEpisodeFinalObservations(Index));
 			}
+		}
 
-			// Check if ping has been sent
-			if (Controls[(uint8)EControls::PingSignal])
+		if (ReplayBuffer.HasFinalMemoryStates())
+		{
+			for (int32 Index = 0; Index < ReplayBuffer.GetMemoryStatesNum(); Index++)
 			{
-				Controls[(uint8)EControls::PingSignal] = false;
-				WaitTime = 0.0f;
+				Array::Copy(EpisodeFinalMemoryStates[Index].Slice(0, EpisodeNum), ReplayBuffer.GetEpisodeFinalMemoryStates(Index));
 			}
+		}
 
-			// Sleep for some time
-			FPlatformProcess::Sleep(SleepTime);
-			WaitTime += SleepTime;
+		for (int32 Index = 0; Index < ReplayBuffer.GetObservationsNum(); Index++)
+		{
+			Array::Copy(Observations[Index].Slice(0, StepNum), ReplayBuffer.GetObservations(Index));
 		}
 
-		if (LogSettings != ELogSetting::Silent)
+		for (int32 Index = 0; Index < ReplayBuffer.GetActionsNum(); Index++)
 		{
-			UE_LOG(LogLearning, Display, TEXT("Pushing Experience..."));
+			Array::Copy(Actions[Index].Slice(0, StepNum), ReplayBuffer.GetActions(Index));
 		}
 
-		const int32 EpisodeNum = EpisodeStartsExperience.Num<0>();
-		const int32 StepNum = ObservationExperience.Num<0>();
+		for (int32 Index = 0; Index < ReplayBuffer.GetMemoryStatesNum(); Index++)
+		{
+			Array::Copy(MemoryStates[Index].Slice(0, StepNum), ReplayBuffer.GetMemoryStates(Index));
+		}
 
-		// Write experience to the shared memory
-		Array::Copy(EpisodeStarts.Slice(0, EpisodeNum), EpisodeStartsExperience);
-		Array::Copy(EpisodeLengths.Slice(0, EpisodeNum), EpisodeLengthsExperience);
-		Array::Copy(Observations.Slice(0, StepNum), ObservationExperience);
-		Array::Copy(Actions.Slice(0, StepNum), ActionExperience);
+		for (int32 Index = 0; Index < ReplayBuffer.GetRewardsNum(); Index++)
+		{
+			Array::Copy(Rewards[Index].Slice(0, StepNum), ReplayBuffer.GetRewards(Index));
+		}
 
-		// Confirm that experience is written
+		// Indicate that experience is written
 		Controls[(uint8)EControls::ExperienceEpisodeNum] = EpisodeNum;
 		Controls[(uint8)EControls::ExperienceStepNum] = StepNum;
+		Controls[(uint8)EControls::ReplayBufferId] = ReplayBufferId;
 		Controls[(uint8)EControls::ExperienceSignal] = true;
 
 		return ETrainerResponse::Success;
 	}
-
 }
\ No newline at end of file
diff --git a/Engine/Plugins/Experimental/LearningAgents/Source/LearningTraining/Private/LearningSharedMemoryTraining.h b/Engine/Plugins/Experimental/LearningAgents/Source/LearningTraining/Private/LearningSharedMemoryTraining.h
--- a/Engine/Plugins/Experimental/LearningAgents/Source/LearningTraining/Private/LearningSharedMemoryTraining.h
+++ b/Engine/Plugins/Experimental/LearningAgents/Source/LearningTraining/Private/LearningSharedMemoryTraining.h
@@ -23,8 +23,10 @@ namespace UE::Learning
 			CompleteSignal			= 5,
 			StopSignal				= 6,
 			PingSignal				= 7,
+			NetworkId				= 8,
+			ReplayBufferId			= 9,
 
-			ControlNum				= 8,
+			ControlNum				= 10,
 		};
 
 		LEARNINGTRAINING_API uint8 GetControlNum();
@@ -35,6 +37,7 @@ namespace UE::Learning
 
 		LEARNINGTRAINING_API ETrainerResponse RecvNetwork(
 			TLearningArrayView<1, volatile int32> Controls,
+			const int32 NetworkId,
 			ULearningNeuralNetworkData& OutNetwork,
 			FSubprocess& Process,
 			const TLearningArrayView<1, const uint8> NetworkData,
@@ -45,10 +48,11 @@ namespace UE::Learning
 		LEARNINGTRAINING_API ETrainerResponse SendStop(
 			TLearningArrayView<1, volatile int32> Controls);
 
-		LEARNINGTRAINING_API bool HasPolicyOrCompleted(TLearningArrayView<1, volatile int32> Controls);
+		LEARNINGTRAINING_API bool HasNetworkOrCompleted(TLearningArrayView<1, volatile int32> Controls);
 
 		LEARNINGTRAINING_API ETrainerResponse SendNetwork(
 			TLearningArrayView<1, volatile int32> Controls,
+			const int32 NetworkId,
 			TLearningArrayView<1, uint8> NetworkData,
 			FSubprocess& Process,
 			const ULearningNeuralNetworkData& Network,
@@ -60,31 +64,17 @@ namespace UE::Learning
 			TLearningArrayView<1, int32> EpisodeStarts,
 			TLearningArrayView<1, int32> EpisodeLengths,
 			TLearningArrayView<1, ECompletionMode> EpisodeCompletionModes,
-			TLearningArrayView<2, float> EpisodeFinalObservations,
-			TLearningArrayView<2, float> EpisodeFinalMemoryStates,
-			TLearningArrayView<2, float> Observations,
-			TLearningArrayView<2, float> Actions,
-			TLearningArrayView<2, float> MemoryStates,
-			TLearningArrayView<1, float> Rewards,
+			TArrayView<TLearningArrayView<2, float>> EpisodeFinalObservations,
+			TArrayView<TLearningArrayView<2, float>> EpisodeFinalMemoryStates,
+			TArrayView<TLearningArrayView<2, float>> Observations,
+			TArrayView<TLearningArrayView<2, float>> Actions,
+			TArrayView<TLearningArrayView<2, float>> MemoryStates,
+			TArrayView<TLearningArrayView<2, float>> Rewards,
 			TLearningArrayView<1, volatile int32> Controls,
 			FSubprocess& Process,
+			const int32 ReplayBufferId,
 			const FReplayBuffer& ReplayBuffer,
 			const float Timeout = Trainer::DefaultTimeout,
 			const ELogSetting LogSettings = Trainer::DefaultLogSettings);
-
-		LEARNINGTRAINING_API ETrainerResponse SendExperience(
-			TLearningArrayView<1, int32> EpisodeStarts,
-			TLearningArrayView<1, int32> EpisodeLengths,
-			TLearningArrayView<2, float> Observations,
-			TLearningArrayView<2, float> Actions,
-			TLearningArrayView<1, volatile int32> Controls,
-			FSubprocess& Process,
-			const TLearningArrayView<1, const int32> EpisodeStartsExperience,
-			const TLearningArrayView<1, const int32> EpisodeLengthsExperience,
-			const TLearningArrayView<2, const float> ObservationExperience,
-			const TLearningArrayView<2, const float> ActionExperience,
-			const float Timeout = Trainer::DefaultTimeout,
-			const ELogSetting LogSettings = Trainer::DefaultLogSettings);
-
 	}
 }
diff --git a/Engine/Plugins/Experimental/LearningAgents/Source/LearningTraining/Private/LearningSocketTraining.cpp b/Engine/Plugins/Experimental/LearningAgents/Source/LearningTraining/Private/LearningSocketTraining.cpp
--- a/Engine/Plugins/Experimental/LearningAgents/Source/LearningTraining/Private/LearningSocketTraining.cpp
+++ b/Engine/Plugins/Experimental/LearningAgents/Source/LearningTraining/Private/LearningSocketTraining.cpp
@@ -85,6 +85,7 @@ namespace UE::Learning::SocketTraining
 
 	ETrainerResponse RecvNetwork(
 		FSocket& Socket,
+		const int32 NetworkId,
 		ULearningNeuralNetworkData& OutNetwork,
 		FSubprocess* Process,
 		TLearningArrayView<1, uint8> OutNetworkBuffer,
@@ -125,6 +126,11 @@ namespace UE::Learning::SocketTraining
 			break;
 		}
 
+		int32 Id = -1;
+		Response = RecvWithTimeout(Socket, Process, (uint8*)&Id, sizeof(int32), Timeout);
+		if (Response != ETrainerResponse::Success) { return Response; }
+		checkf(Id == NetworkId, TEXT("Received unexpected network id!"))
+
 		Response = RecvWithTimeout(Socket, Process, OutNetworkBuffer.GetData(), OutNetworkBuffer.Num(), Timeout);
 		if (Response != ETrainerResponse::Success) { return Response; }
 
@@ -220,16 +226,39 @@ namespace UE::Learning::SocketTraining
 		return SendWithTimeout(Socket, Process, &Signal, 1, Timeout);
 	}
 
-	bool HasPolicyOrCompleted(FSocket& Socket)
+	bool HasNetworkOrCompleted(FSocket& Socket, FSubprocess* Process)
 	{
+		// If we're monitoring a process, then has it has exited?
+		if (Process && !Process->Update())
+		{
+			return false;
+		}
+
 		uint32 PendingDataSize;
-		return Socket.HasPendingData(PendingDataSize);
+		bool bHasPendingData = Socket.HasPendingData(PendingDataSize);
+
+		uint8 Signal = (uint8)ESignal::Invalid;
+		int32 BytesRead = 0;
+		if (bHasPendingData)
+		{ 
+			Socket.Recv(&Signal, 1, BytesRead, ESocketReceiveFlags::Type::Peek);
+
+			if (Signal == (uint8)ESignal::RecvPing)
+			{
+				// Consume the ping
+				Socket.Recv(&Signal, 1, BytesRead);
+				return false;
+			}
+		}
+
+		return Signal == (uint8)ESignal::RecvComplete || Signal == (uint8)ESignal::RecvNetwork;
 	}
 
 	ETrainerResponse SendNetwork(
 		FSocket& Socket,
 		TLearningArrayView<1, uint8> NetworkBuffer,
 		FSubprocess* Process,
+		const int32 NetworkId,
 		const ULearningNeuralNetworkData& Network,
 		const float Timeout,
 		FRWLock* NetworkLock,
@@ -259,6 +288,9 @@ namespace UE::Learning::SocketTraining
 		ETrainerResponse Response = SendWithTimeout(Socket, Process, &Signal, 1, Timeout);
 		if (Response != ETrainerResponse::Success) { return Response; }
 
+		Response = SendWithTimeout(Socket, Process, (const uint8*)&NetworkId, sizeof(int32), Timeout);
+		if (Response != ETrainerResponse::Success) { return Response; }
+
 		Response = SendWithTimeout(Socket, Process, NetworkBuffer.GetData(), NetworkBuffer.Num(), Timeout);
 		if (Response != ETrainerResponse::Success) { return Response; }
 
@@ -267,6 +299,7 @@ namespace UE::Learning::SocketTraining
 
 	ETrainerResponse SendExperience(
 		FSocket& Socket,
+		const int32 ReplayBufferId,
 		const FReplayBuffer& ReplayBuffer,
 		FSubprocess* Process,
 		const float Timeout,
@@ -277,16 +310,18 @@ namespace UE::Learning::SocketTraining
 			UE_LOG(LogLearning, Display, TEXT("Pushing Experience..."));
 		}
 
-		const int32 EpisodeNum = ReplayBuffer.GetEpisodeNum();
-		const int32 StepNum = ReplayBuffer.GetStepNum();
-
 		const uint8 Signal = (uint8)ESignal::SendExperience;
 		ETrainerResponse Response = SendWithTimeout(Socket, Process, &Signal, 1, Timeout);
 		if (Response != ETrainerResponse::Success) { return Response; }
 
+		Response = SendWithTimeout(Socket, Process, (const uint8*)&ReplayBufferId, sizeof(int32), Timeout);
+		if (Response != ETrainerResponse::Success) { return Response; }
+
+		const int32 EpisodeNum = ReplayBuffer.GetEpisodeNum();
 		Response = SendWithTimeout(Socket, Process, (const uint8*)&EpisodeNum, sizeof(int32), Timeout);
 		if (Response != ETrainerResponse::Success) { return Response; }
 		
+		const int32 StepNum = ReplayBuffer.GetStepNum();
 		Response = SendWithTimeout(Socket, Process, (const uint8*)&StepNum, sizeof(int32), Timeout);
 		if (Response != ETrainerResponse::Success) { return Response; }
 		
@@ -296,26 +331,53 @@ namespace UE::Learning::SocketTraining
 		Response = SendWithTimeout(Socket, Process, (const uint8*)ReplayBuffer.GetEpisodeLengths().GetData(), ReplayBuffer.GetEpisodeLengths().Num() * sizeof(int32), Timeout);
 		if (Response != ETrainerResponse::Success) { return Response; }
 		
-		Response = SendWithTimeout(Socket, Process, (const uint8*)ReplayBuffer.GetEpisodeCompletionModes().GetData(), ReplayBuffer.GetEpisodeCompletionModes().Num() * sizeof(ECompletionMode), Timeout);
-		if (Response != ETrainerResponse::Success) { return Response; }
+		if (ReplayBuffer.HasCompletions())
+		{
+			Response = SendWithTimeout(Socket, Process, (const uint8*)ReplayBuffer.GetEpisodeCompletionModes().GetData(), ReplayBuffer.GetEpisodeCompletionModes().Num() * sizeof(ECompletionMode), Timeout);
+			if (Response != ETrainerResponse::Success) { return Response; }
+		}
 		
-		Response = SendWithTimeout(Socket, Process, (const uint8*)ReplayBuffer.GetEpisodeFinalObservations().GetData(), ReplayBuffer.GetEpisodeFinalObservations().Num() * sizeof(float), Timeout);
-		if (Response != ETrainerResponse::Success) { return Response; }
+		if (ReplayBuffer.HasFinalObservations())
+		{
+			for (int32 Index = 0; Index < ReplayBuffer.GetObservationsNum(); Index++)
+			{
+				Response = SendWithTimeout(Socket, Process, (const uint8*)ReplayBuffer.GetEpisodeFinalObservations(Index).GetData(), ReplayBuffer.GetEpisodeFinalObservations(Index).Num() * sizeof(float), Timeout);
+				if (Response != ETrainerResponse::Success) { return Response; }
+			}
+		}
 		
-		Response = SendWithTimeout(Socket, Process, (const uint8*)ReplayBuffer.GetEpisodeFinalMemoryStates().GetData(), ReplayBuffer.GetEpisodeFinalMemoryStates().Num() * sizeof(float), Timeout);
-		if (Response != ETrainerResponse::Success) { return Response; }
+		if (ReplayBuffer.HasFinalMemoryStates())
+		{
+			for (int32 Index = 0; Index < ReplayBuffer.GetMemoryStatesNum(); Index++)
+			{
+				Response = SendWithTimeout(Socket, Process, (const uint8*)ReplayBuffer.GetEpisodeFinalMemoryStates(Index).GetData(), ReplayBuffer.GetEpisodeFinalMemoryStates(Index).Num() * sizeof(float), Timeout);
+				if (Response != ETrainerResponse::Success) { return Response; }
+			}
+		}
 
-		Response = SendWithTimeout(Socket, Process, (const uint8*)ReplayBuffer.GetObservations().GetData(), ReplayBuffer.GetObservations().Num() * sizeof(float), Timeout);
-		if (Response != ETrainerResponse::Success) { return Response; }
+		for (int32 Index = 0; Index < ReplayBuffer.GetObservationsNum(); Index++)
+		{
+			Response = SendWithTimeout(Socket, Process, (const uint8*)ReplayBuffer.GetObservations(Index).GetData(), ReplayBuffer.GetObservations(Index).Num() * sizeof(float), Timeout);
+			if (Response != ETrainerResponse::Success) { return Response; }
+		}
 		
-		Response = SendWithTimeout(Socket, Process, (const uint8*)ReplayBuffer.GetActions().GetData(), ReplayBuffer.GetActions().Num() * sizeof(float), Timeout);
-		if (Response != ETrainerResponse::Success) { return Response; }
+		for (int32 Index = 0; Index < ReplayBuffer.GetActionsNum(); Index++)
+		{
+			Response = SendWithTimeout(Socket, Process, (const uint8*)ReplayBuffer.GetActions(Index).GetData(), ReplayBuffer.GetActions(Index).Num() * sizeof(float), Timeout);
+			if (Response != ETrainerResponse::Success) { return Response; }
+		}
 		
-		Response = SendWithTimeout(Socket, Process, (const uint8*)ReplayBuffer.GetMemoryStates().GetData(), ReplayBuffer.GetMemoryStates().Num() * sizeof(float), Timeout);
-		if (Response != ETrainerResponse::Success) { return Response; }
+		for (int32 Index = 0; Index < ReplayBuffer.GetMemoryStatesNum(); Index++)
+		{
+			Response = SendWithTimeout(Socket, Process, (const uint8*)ReplayBuffer.GetMemoryStates(Index).GetData(), ReplayBuffer.GetMemoryStates(Index).Num() * sizeof(float), Timeout);
+			if (Response != ETrainerResponse::Success) { return Response; }
+		}
 
-		Response = SendWithTimeout(Socket, Process, (const uint8*)ReplayBuffer.GetRewards().GetData(), ReplayBuffer.GetRewards().Num() * sizeof(float), Timeout);
-		if (Response != ETrainerResponse::Success) { return Response; }
+		for (int32 Index = 0; Index < ReplayBuffer.GetRewardsNum(); Index++)
+		{
+			Response = SendWithTimeout(Socket, Process, (const uint8*)ReplayBuffer.GetRewards(Index).GetData(), ReplayBuffer.GetRewards(Index).Num() * sizeof(float), Timeout);
+			if (Response != ETrainerResponse::Success) { return Response; }
+		}
 
 		return ETrainerResponse::Success;
 	}
diff --git a/Engine/Plugins/Experimental/LearningAgents/Source/LearningTraining/Private/LearningSocketTraining.h b/Engine/Plugins/Experimental/LearningAgents/Source/LearningTraining/Private/LearningSocketTraining.h
--- a/Engine/Plugins/Experimental/LearningAgents/Source/LearningTraining/Private/LearningSocketTraining.h
+++ b/Engine/Plugins/Experimental/LearningAgents/Source/LearningTraining/Private/LearningSocketTraining.h
@@ -41,6 +41,7 @@ namespace UE::Learning
 
 		LEARNINGTRAINING_API ETrainerResponse RecvNetwork(
 			FSocket& Socket,
+			const int32 NetworkId,
 			ULearningNeuralNetworkData& OutNetwork,
 			FSubprocess* Process,
 			TLearningArrayView<1, uint8> OutNetworkBuffer,
@@ -67,12 +68,15 @@ namespace UE::Learning
 			FSubprocess* Process,
 			const float Timeout = Trainer::DefaultTimeout);
 
-		LEARNINGTRAINING_API bool HasPolicyOrCompleted(FSocket& Socket);
+		LEARNINGTRAINING_API bool HasNetworkOrCompleted(
+			FSocket& Socket,
+			FSubprocess* Process);
 
 		LEARNINGTRAINING_API ETrainerResponse SendNetwork(
 			FSocket& Socket,
 			TLearningArrayView<1, uint8> NetworkBuffer,
 			FSubprocess* Process,
+			const int32 NetworkId,
 			const ULearningNeuralNetworkData& Network,
 			const float Timeout = Trainer::DefaultTimeout,
 			FRWLock* NetworkLock = nullptr,
@@ -80,6 +84,7 @@ namespace UE::Learning
 
 		LEARNINGTRAINING_API ETrainerResponse SendExperience(
 			FSocket& Socket,
+			const int32 ReplayBufferId,
 			const FReplayBuffer& ReplayBuffer,
 			FSubprocess* Process,
 			const float Timeout = Trainer::DefaultTimeout,
diff --git a/Engine/Plugins/Experimental/LearningAgents/Source/LearningTraining/Private/LearningTrainer.cpp b/Engine/Plugins/Experimental/LearningAgents/Source/LearningTraining/Private/LearningTrainer.cpp
--- a/Engine/Plugins/Experimental/LearningAgents/Source/LearningTraining/Private/LearningTrainer.cpp
+++ b/Engine/Plugins/Experimental/LearningAgents/Source/LearningTraining/Private/LearningTrainer.cpp
@@ -402,6 +402,11 @@ namespace UE::Learning::Trainer
 		return EngineDir / TEXT("Plugins/Experimental/LearningAgents/Content/Python/");
 	}
 
+	FString GetProjectPythonContentPath()
+	{
+		return FPaths::ProjectContentDir() / TEXT("Python/");
+	}
+
 	FString GetIntermediatePath(const FString& IntermediateDir)
 	{
 		return IntermediateDir / TEXT("LearningAgents");
diff --git a/Engine/Plugins/Experimental/LearningAgents/Source/LearningTraining/Public/LearningExperience.h b/Engine/Plugins/Experimental/LearningAgents/Source/LearningTraining/Public/LearningExperience.h
--- a/Engine/Plugins/Experimental/LearningAgents/Source/LearningTraining/Public/LearningExperience.h
+++ b/Engine/Plugins/Experimental/LearningAgents/Source/LearningTraining/Public/LearningExperience.h
@@ -3,6 +3,14 @@
 #pragma once
 
 #include "LearningArray.h"
+#include "LearningObservation.h"
+#include "LearningAction.h"
+
+#include "Containers/Set.h"
+#include "Templates/SharedPointer.h"
+#include "UObject/NameTypes.h"
+
+class FJsonObject;
 
 namespace UE::Learning
 {
@@ -14,63 +22,102 @@ namespace UE::Learning
 	*/
 	struct LEARNINGTRAINING_API FEpisodeBuffer
 	{
+		friend struct FReplayBuffer;
+
 		/**
 		* Resize the experience buffer
 		*
 		* @param InMaxInstanceNum				Maximum number of instances
 		* @param InMaxStepNum					Maximum number of steps in an episode
-		* @param InObservationVectorDimNum		Number of dimensions of the observation vector
-		* @param InActionVectorDimNum			Number of dimensions of the action vector
-		* @param InMemoryStateVectorDimNum		Number of dimensions of the memory state vector
 		*/
-		void Resize(
-			const int32 InMaxInstanceNum,
-			const int32 InMaxStepNum,
-			const int32 InObservationVectorDimNum,
-			const int32 InActionVectorDimNum,
-			const int32 InMemoryStateVectorDimNum);
+		void Resize(const int32 InMaxInstanceNum, const int32 InMaxStepNum);
 
 		/**
 		* Reset the buffer for the given set of instances
 		*/
 		void Reset(const FIndexSet Instances);
 
-		/**
-		* Push new experience to the buffer
-		*
-		* @param InObservations					Observation vectors of shape (MaxInstanceNum, ObservationVectorDimNum)
-		* @param InActions						Action vectors of shape (MaxInstanceNum, ActionVectorDimNum)
-		* @param InMemoryStates					Memory state vectors (pre-evaluation) of shape (MaxInstanceNum, MemoryStateVectorDimNum)
-		* @param InRewards						Rewards of shape (MaxInstanceNum)
-		* @param Instances						Instances to push to buffer
-		*/
-		void Push(
-			const TLearningArrayView<2, const float> InObservations,
-			const TLearningArrayView<2, const float> InActions,
-			const TLearningArrayView<2, const float> InMemoryStates,
-			const TLearningArrayView<1, const float> InRewards,
-			const FIndexSet Instances);
 
-		const int32 GetMaxInstanceNum() const;
-		const int32 GetMaxStepNum() const;
+		// Observations
+
+		int32 AddObservations(const FName& Name, const int32 SchemaId, const int32 Size);
+
+		void PushObservations(const int32 ObservationId, const TLearningArrayView<2, const float> InObservations, const FIndexSet Instances);
+
+		const TLearningArrayView<2, const float> GetObservations(const int32 ObservationId, const int32 InstanceIdx) const;
+
+
+		// Actions
+
+		int32 AddActions(const FName& Name, const int32 SchemaId, const int32 Size);
+
+		void PushActions(const int32 ActionId, const TLearningArrayView<2, const float> InActions, const FIndexSet Instances);
+
+		const TLearningArrayView<2, const float> GetActions(const int32 ActionId, const int32 InstanceIdx) const;
+
+
+		// Memory States
+
+		int32 AddMemoryStates(const FName& Name, const int32 Size);
+
+		void PushMemoryStates(const int32 MemoryStateId, const TLearningArrayView<2, const float> InMemoryStates, const FIndexSet Instances);
+
+		const TLearningArrayView<2, const float> GetMemoryStates(const int32 MemoryStateId, const int32 InstanceIdx) const;
+
+
+		// Rewards
+
+		int32 AddRewards(const FName& Name, const int32 Size);
+
+		// Convenience overload for pushing rewards from a Training Environment
+		void PushRewards(const int32 RewardId, const TLearningArrayView<1, const float> InRewards, const FIndexSet Instances);
+
+		void PushRewards(const int32 RewardId, const TLearningArrayView<2, const float> InRewards, const FIndexSet Instances);
+
+		const TLearningArrayView<2, const float> GetRewards(const int32 RewardId, const int32 InstanceIdx) const;
+
+
+		// Episode Step Nums
+
+		void IncrementEpisodeStepNums(const FIndexSet Instances);
+
 		const TLearningArrayView<1, const int32> GetEpisodeStepNums() const;
-		const TLearningArrayView<2, const float> GetObservations(const int32 InstanceIdx) const;
-		const TLearningArrayView<2, const float> GetActions(const int32 InstanceIdx) const;
-		const TLearningArrayView<2, const float> GetMemoryStates(const int32 InstanceIdx) const;
-		const TLearningArrayView<1, const float> GetRewards(const int32 InstanceIdx) const;
+
+		int32 GetMaxInstanceNum() const;
+		int32 GetMaxStepNum() const;
 
 	private:
 
+		bool bHasBeenSized = false;
 		int32 MaxInstanceNum = 0;
 		int32 MaxStepNum = 0;
+
+		// Observations
+		TArray<FName, TInlineAllocator<1>> ObservationNames;
+		TArray<int32, TInlineAllocator<1>> ObservationSchemaIds;
+		TArray<int32, TInlineAllocator<1>> ObservationSizes;
+		TArray<TLearningArray<3, float>, TInlineAllocator<1>> ObservationArrays;
+
+		// Actions
+		TArray<FName, TInlineAllocator<1>> ActionNames;
+		TArray<int32, TInlineAllocator<1>> ActionSchemaIds;
+		TArray<int32, TInlineAllocator<1>> ActionSizes;
+		TArray<TLearningArray<3, float>, TInlineAllocator<1>> ActionArrays;
+
+		// Memory States
+		TArray<FName, TInlineAllocator<1>> MemoryStateNames;
+		TArray<int32, TInlineAllocator<1>> MemoryStateSizes;
+		TArray<TLearningArray<3, float>, TInlineAllocator<1>> MemoryStateArrays;
+
+		// Rewards
+		TArray<FName, TInlineAllocator<1>> RewardNames;
+		TArray<int32, TInlineAllocator<1>> RewardSizes;
+		TArray<TLearningArray<3, float>, TInlineAllocator<1>> RewardArrays;
+
+		// Episode Step Nums
 		TLearningArray<1, int32> EpisodeStepNums;
-		TLearningArray<3, float> Observations;
-		TLearningArray<3, float> Actions;
-		TLearningArray<3, float> MemoryStates;
-		TLearningArray<2, float> Rewards;
 	};
 
-
 	/**
 	* Large buffer that sequentially concatenates a series of episodes in a large
 	* flat array. Used to collate episodic data together from multiple instances.
@@ -80,16 +127,12 @@ namespace UE::Learning
 		/**
 		* Resizes the replay buffer.
 		*
-		* @param ObservationVectorDimNum		Dimensionality of observation vector
-		* @param ActionVectorDimNum				Dimensionality of action vector
-		* @param MemoryStateVectorDimNum		Dimensionality of memory state vector
-		* @param MaxEpisodeNum					Maximum number of episodes to be stored in the buffer
-		* @param MaxStepNum						Maximum number of steps to be stored in the buffer
+		* @param EpisodeBuffer		Determines the dimensionality of the various arrays
+		* @param MaxEpisodeNum		Maximum number of episodes to be stored in the buffer
+		* @param MaxStepNum			Maximum number of steps to be stored in the buffer
 		*/
 		void Resize(
-			const int32 InObservationVectorDimNum,
-			const int32 InActionVectorDimNum,
-			const int32 InMemoryStateVectorDimNum,
+			const FEpisodeBuffer& EpisodeBuffer,
 			const int32 InMaxEpisodeNum = 2048,
 			const int32 InMaxStepNum = 16384);
 
@@ -112,28 +155,68 @@ namespace UE::Learning
 		*/
 		bool AddEpisodes(
 			const TLearningArrayView<1, const ECompletionMode> InEpisodeCompletionModes,
-			const TLearningArrayView<2, const float> InEpisodeFinalObservations,
-			const TLearningArrayView<2, const float> InEpisodeFinalMemoryStates,
+			const TArrayView<const TLearningArrayView<2, const float>> InEpisodeFinalObservations,
+			const TArrayView<const TLearningArrayView<2, const float>> InEpisodeFinalMemoryStates,
 			const FEpisodeBuffer& EpisodeBuffer,
 			const FIndexSet Instances,
 			const bool bAddTruncatedEpisodeWhenFull = true);
 
-		const int32 GetMaxEpisodeNum() const;
-		const int32 GetMaxStepNum() const;
-		const int32 GetEpisodeNum() const;
-		const int32 GetStepNum() const;
+		/**
+		 * Alternate way to add data from records generated via imitation learning. Does the resizing needed.
+		 * 
+		 * Putting this in place until we have more time later to rewrite how imitation learning stores data. At this
+		 * time, it seems logical to have the data records be implemented in terms of EpisodeBuffer(s), in which case
+		 * this method may no longer be needed.
+		 */
+		void AddRecords(
+			const int32 InEpisodeNum,
+			const int32 InMaxStepNum,
+			const int32 ObservationSchemaId,
+			const int32 ObservationNum,
+			const int32 ActionSchemaId,
+			const int32 ActionNum,
+			const TLearningArrayView<1, const int32> RecordedEpisodeStarts,
+			const TLearningArrayView<1, const int32> RecordedEpisodeLengths,
+			const TLearningArrayView<2, const float> RecordedObservations,
+			const TLearningArrayView<2, const float> RecordedActions);
+
+		bool HasCompletions() const;
+
+		bool HasFinalObservations() const;
+
+		bool HasFinalMemoryStates() const;
+
+		int32 GetMaxEpisodeNum() const;
+		int32 GetMaxStepNum() const;
+		int32 GetEpisodeNum() const;
+		int32 GetStepNum() const;
+
 		const TLearningArrayView<1, const int32> GetEpisodeStarts() const;
 		const TLearningArrayView<1, const int32> GetEpisodeLengths() const;
 		const TLearningArrayView<1, const ECompletionMode> GetEpisodeCompletionModes() const;
-		const TLearningArrayView<2, const float> GetEpisodeFinalObservations() const;
-		const TLearningArrayView<2, const float> GetEpisodeFinalMemoryStates() const;
-		const TLearningArrayView<2, const float> GetObservations() const;
-		const TLearningArrayView<2, const float> GetActions() const;
-		const TLearningArrayView<2, const float> GetMemoryStates() const;
-		const TLearningArrayView<1, const float> GetRewards() const;
+		
+		int32 GetObservationsNum() const;
+		const TLearningArrayView<2, const float> GetObservations(const int32 Index) const;
+		const TLearningArrayView<2, const float> GetEpisodeFinalObservations(const int32 Index) const;
+
+		int32 GetActionsNum() const;
+		const TLearningArrayView<2, const float> GetActions(const int32 Index) const;
+
+		int32 GetMemoryStatesNum() const;
+		const TLearningArrayView<2, const float> GetMemoryStates(const int32 Index) const;
+		const TLearningArrayView<2, const float> GetEpisodeFinalMemoryStates(const int32 Index) const;
+		
+		int32 GetRewardsNum() const;
+		const TLearningArrayView<2, const float> GetRewards(const int32 Index) const;
+
+		TSharedRef<FJsonObject> AsJsonConfig(const int32 ReplayBufferId) const;
 
 	private:
 
+		bool bHasCompletions = false;
+		bool bHasFinalObservations = false;
+		bool bHasFinalMemoryStates = false;
+
 		int32 MaxEpisodeNum = 0;
 		int32 MaxStepNum = 0;
 
@@ -143,12 +226,22 @@ namespace UE::Learning
 		TLearningArray<1, int32> EpisodeStarts;
 		TLearningArray<1, int32> EpisodeLengths;
 		TLearningArray<1, ECompletionMode> EpisodeCompletionModes;
-		TLearningArray<2, float> EpisodeFinalObservations;
-		TLearningArray<2, float> EpisodeFinalMemoryStates;
-		TLearningArray<2, float> Observations;
-		TLearningArray<2, float> Actions;
-		TLearningArray<2, float> MemoryStates;
-		TLearningArray<1, float> Rewards;
+		
+		TArray<FName, TInlineAllocator<1>> ObservationNames;
+		TArray<int32, TInlineAllocator<1>> ObservationSchemaIds;
+		TArray<TLearningArray<2, float>, TInlineAllocator<1>> Observations;
+		TArray<TLearningArray<2, float>, TInlineAllocator<1>> EpisodeFinalObservations;
+
+		TArray<FName, TInlineAllocator<1>> ActionNames;
+		TArray<int32, TInlineAllocator<1>> ActionSchemaIds;
+		TArray<TLearningArray<2, float>, TInlineAllocator<1>> Actions;
+
+		TArray<FName, TInlineAllocator<1>> MemoryStateNames;
+		TArray<TLearningArray<2, float>, TInlineAllocator<1>> MemoryStates;
+		TArray<TLearningArray<2, float>, TInlineAllocator<1>> EpisodeFinalMemoryStates;
+
+		TArray<FName, TInlineAllocator<1>> RewardNames;
+		TArray<TLearningArray<2, float>, TInlineAllocator<1>> Rewards;
 	};
 
 	namespace Experience
@@ -159,20 +252,20 @@ namespace UE::Learning
 		* @param ReplayBuffer							Replay Buffer
 		* @param EpisodeBuffer							Episode Buffer
 		* @param ResetBuffer							Reset Buffer
-		* @param ObservationVectorBuffer				Buffer to read/write observation vectors into
-		* @param ActionVectorBuffer						Buffer to read/write action vectors into
-		* @param PreEvaluationMemoryStateVectorBuffer	Buffer to read/write pre-evaluation memory state vectors into
-		* @param MemoryStateVectorBuffer				Buffer to read/write (post-evaluation) memory state vectors into
-		* @param RewardBuffer							Buffer to read/write rewards into
+		* @param ObservationVectorBuffers				Buffers to read/write observation vectors into
+		* @param ActionVectorBuffers					Buffers to read/write action vectors into
+		* @param PreEvaluationMemoryStateVectorBuffers	Buffers to read/write pre-evaluation memory state vectors into
+		* @param MemoryStateVectorBuffers				Buffers to read/write (post-evaluation) memory state vectors into
+		* @param RewardBuffers							Buffers to read/write rewards into
 		* @param CompletionBuffer						Buffer to read/write completions into
 		* @param EpisodeCompletionBuffer				Additional buffer to record completions from full episode buffers
 		* @param AllCompletionBuffer					Additional buffer to record all completions from full episodes and normal completions
 		* @param ResetFunction							Function to run for resetting the environment
-		* @param ObservationFunction					Function to run for evaluating observations
-		* @param PolicyFunction							Function to run generating actions from observations
-		* @param ActionFunction							Function to run for evaluating actions
-		* @param UpdateFunction							Function to run for updating the environment
-		* @param RewardFunction							Function to run for evaluating rewards
+		* @param ObservationFunctions					Functions to run for evaluating observations
+		* @param PolicyFunctions						Functions to run generating actions from observations
+		* @param ActionFunctions						Functions to run for evaluating actions
+		* @param UpdateFunctions						Functions to run for updating the environment
+		* @param RewardFunctions						Functions to run for evaluating rewards
 		* @param CompletionFunction						Function to run for evaluating completions
 		* @param Instances								Set of instances to gather experience for
 		*/
@@ -180,20 +273,20 @@ namespace UE::Learning
 			FReplayBuffer& ReplayBuffer,
 			FEpisodeBuffer& EpisodeBuffer,
 			FResetInstanceBuffer& ResetBuffer,
-			TLearningArrayView<2, float> ObservationVectorBuffer,
-			TLearningArrayView<2, float> ActionVectorBuffer,
-			TLearningArrayView<2, float> PreEvaluationMemoryStateVectorBuffer,
-			TLearningArrayView<2, float> MemoryStateVectorBuffer,
-			TLearningArrayView<1, float> RewardBuffer,
+			const TArrayView<const TLearningArrayView<2, const float>> ObservationVectorBuffers,
+			const TArrayView<const TLearningArrayView<2, const float>> ActionVectorBuffers,
+			const TArrayView<const TLearningArrayView<2, const float>> PreEvaluationMemoryStateVectorBuffers,
+			const TArrayView<const TLearningArrayView<2, const float>> MemoryStateVectorBuffers,
+			const TArrayView<const TLearningArrayView<1, const float>> RewardBuffers,
 			TLearningArrayView<1, ECompletionMode> CompletionBuffer,
 			TLearningArrayView<1, ECompletionMode> EpisodeCompletionBuffer,
 			TLearningArrayView<1, ECompletionMode> AllCompletionBuffer,
 			const TFunctionRef<void(const FIndexSet Instances)> ResetFunction,
-			const TFunctionRef<void(const FIndexSet Instances)> ObservationFunction,
-			const TFunctionRef<void(const FIndexSet Instances)> PolicyFunction,
-			const TFunctionRef<void(const FIndexSet Instances)> ActionFunction,
-			const TFunctionRef<void(const FIndexSet Instances)> UpdateFunction,
-			const TFunctionRef<void(const FIndexSet Instances)> RewardFunction,
+			const TArrayView<const TFunctionRef<void(const FIndexSet Instances)>> ObservationFunctions,
+			const TArrayView<const TFunctionRef<void(const FIndexSet Instances)>> PolicyFunctions,
+			const TArrayView<const TFunctionRef<void(const FIndexSet Instances)>> ActionFunctions,
+			const TArrayView<const TFunctionRef<void(const FIndexSet Instances)>> UpdateFunctions,
+			const TArrayView<const TFunctionRef<void(const FIndexSet Instances)>> RewardFunctions,
 			const TFunctionRef<void(const FIndexSet Instances)> CompletionFunction,
 			const FIndexSet Instances);
 	}
diff --git a/Engine/Plugins/Experimental/LearningAgents/Source/LearningTraining/Public/LearningExternalTrainer.h b/Engine/Plugins/Experimental/LearningAgents/Source/LearningTraining/Public/LearningExternalTrainer.h
--- a/Engine/Plugins/Experimental/LearningAgents/Source/LearningTraining/Public/LearningExternalTrainer.h
+++ b/Engine/Plugins/Experimental/LearningAgents/Source/LearningTraining/Public/LearningExternalTrainer.h
@@ -39,6 +39,8 @@ namespace UE::Learning
 		*/
 		virtual ETrainerResponse Wait() = 0;
 
+		/** Returns true if we can receive a network or training completed. Otherwise, false. */
+		virtual bool HasNetworkOrCompleted() = 0;
 
 		/**
 		* Sends the given json config to the trainer process.
@@ -55,17 +57,15 @@ namespace UE::Learning
 		 * Adds the network to this external trainer. Allocates buffers, etc.
 		 * Must be called for each network prior to calling Send/Receive.
 		 * 
-		 * @params Name		The unique name of the network, used as a key.
-		 * @params Network	The network to be added.
+		 * @params Network	The network to be added
+		 * @returns			The network's unique id
 		 */
-		virtual void AddNetwork(const FName& Name, const ULearningNeuralNetworkData& Network) = 0;
-
-		/** Returns true if a network with the given Name has already been added. Otherwise, false. */
-		virtual bool ContainsNetwork(const FName& Name) const = 0;
+		virtual int32 AddNetwork(const ULearningNeuralNetworkData& Network) = 0;
 
 		/**
 		* Wait for the trainer to push an updated network.
 		*
+		* @param NetworkId		Unique network id
 		* @param OutNetwork		Network to update
 		* @param Timeout		Timeout to wait in seconds
 		* @param NetworkLock	Lock to use when updating network
@@ -73,7 +73,7 @@ namespace UE::Learning
 		* @returns				Trainer response
 		*/
 		virtual ETrainerResponse ReceiveNetwork(
-			const FName& Name,
+			const int32 NetworkId,
 			ULearningNeuralNetworkData& OutNetwork,
 			FRWLock* NetworkLock = nullptr,
 			const ELogSetting LogSettings = Trainer::DefaultLogSettings) = 0;
@@ -81,6 +81,7 @@ namespace UE::Learning
 		/**
 		* Wait for the trainer to be ready and push the current policy network.
 		*
+		* @param NetworkId		Unique network id
 		* @param Network		Network to push
 		* @param Timeout		Timeout to wait in seconds
 		* @param NetworkLock	Lock to use when pushing network
@@ -88,7 +89,7 @@ namespace UE::Learning
 		* @returns				Trainer response
 		*/
 		virtual ETrainerResponse SendNetwork(
-			const FName& Name,
+			const int32 NetworkId,
 			const ULearningNeuralNetworkData& Network,
 			FRWLock* NetworkLock = nullptr,
 			const ELogSetting LogSettings = Trainer::DefaultLogSettings) = 0;
@@ -97,25 +98,23 @@ namespace UE::Learning
 		 * Adds a named replay buffer to this external trainer.
 		 * Must be called for each buffer prior to calling SendReplayBuffer.
 		 *
-		 * @params Name						The unique name of the buffer, used as a key
-		 * @params ReplayBuffer				The buffer to be added
+		 * @params ReplayBuffer		The buffer to be added
+		 * @returns					The replay buffer's unique id		
 		 */
-		virtual void AddReplayBuffer(const FName& Name,	const FReplayBuffer& ReplayBuffer) = 0;
-
-		/** Returns true if a replay buffer with the given Name has already been added. Otherwise, false. */
-		virtual bool ContainsReplayBuffer(const FName& Name) const = 0;
+		virtual int32 AddReplayBuffer(const FReplayBuffer& ReplayBuffer) = 0;
 
 		/**
 		* Wait for the trainer to be ready and send new experience.
 		*
-		* @params Name			The unique name of the buffer, used as a key
-		* @param ReplayBuffer	Replay buffer to send
-		* @param Timeout		Timeout to wait in seconds
-		* @param LogSettings	The log verbosity level
-		* @returns				Trainer response
+		* @params ReplayBufferId	Unique replay buffer id
+		* @params Name				The unique name of the buffer, used as a key
+		* @param ReplayBuffer		Replay buffer to send
+		* @param Timeout			Timeout to wait in seconds
+		* @param LogSettings		The log verbosity level
+		* @returns					Trainer response
 		*/
 		virtual ETrainerResponse SendReplayBuffer(
-			const FName& Name,
+			const int32 ReplayBufferId,
 			const FReplayBuffer& ReplayBuffer,
 			const ELogSetting LogSettings = Trainer::DefaultLogSettings) = 0;
 	};
@@ -157,6 +156,8 @@ namespace UE::Learning
 		* process has a non-zero "LearningProcessIdx".
 		*
 		* @param TaskName					The name of this training task (used to disambiguate filenames, etc.)
+		* @param CustomTrainerPath			Path to check for custom trainer files
+		* @param TrainerFileName			The name of the training file to use
 		* @param PythonExecutablePath		Path to the python executable used for training. In general should be the
 		*									python shipped with Unreal Editor.
 		* @param PythonContentPath			Path to the Python Content folder provided by the Learning plugin
@@ -168,6 +169,8 @@ namespace UE::Learning
 		*/
 		FSharedMemoryTrainerServerProcess(
 			const FString& TaskName,
+			const FString& CustomTrainerPath,
+			const FString& TrainerFileName,
 			const FString& PythonExecutablePath,
 			const FString& PythonContentPath,
 			const FString& IntermediatePath,
@@ -230,12 +233,13 @@ namespace UE::Learning
 			TSharedMemoryArrayView<2, int32> EpisodeStarts;
 			TSharedMemoryArrayView<2, int32> EpisodeLengths;
 			TSharedMemoryArrayView<2, ECompletionMode> EpisodeCompletionModes;
-			TSharedMemoryArrayView<3, float> EpisodeFinalObservations;
-			TSharedMemoryArrayView<3, float> EpisodeFinalMemoryStates;
-			TSharedMemoryArrayView<3, float> Observations;
-			TSharedMemoryArrayView<3, float> Actions;
-			TSharedMemoryArrayView<3, float> MemoryStates;
-			TSharedMemoryArrayView<2, float> Rewards;
+
+			TArray<TSharedMemoryArrayView<3, float>, TInlineAllocator<1>> EpisodeFinalObservations;
+			TArray<TSharedMemoryArrayView<3, float>, TInlineAllocator<1>> EpisodeFinalMemoryStates;
+			TArray<TSharedMemoryArrayView<3, float>, TInlineAllocator<1>> Observations;
+			TArray<TSharedMemoryArrayView<3, float>, TInlineAllocator<1>> Actions;
+			TArray<TSharedMemoryArrayView<3, float>, TInlineAllocator<1>> MemoryStates;
+			TArray<TSharedMemoryArrayView<3, float>, TInlineAllocator<1>> Rewards;
 
 			/** Free and deallocate all shared memory. */
 			void Deallocate();
@@ -264,32 +268,30 @@ namespace UE::Learning
 
 		virtual ETrainerResponse Wait() override final;
 
+		virtual bool HasNetworkOrCompleted() override final;
+
 		virtual ETrainerResponse SendConfig(
 			const TSharedRef<FJsonObject>& ConfigObject,
 			const ELogSetting LogSettings = Trainer::DefaultLogSettings) override final;
 
-		virtual void AddNetwork(const FName& Name, const ULearningNeuralNetworkData& Network) override final;
-
-		virtual bool ContainsNetwork(const FName& Name) const override final;
+		virtual int32 AddNetwork(const ULearningNeuralNetworkData& Network) override final;
 
 		virtual ETrainerResponse ReceiveNetwork(
-			const FName& Name,
+			const int32 NetworkId,
 			ULearningNeuralNetworkData& OutNetwork,
 			FRWLock* NetworkLock = nullptr,
 			const ELogSetting LogSettings = Trainer::DefaultLogSettings) override final;
 
 		virtual ETrainerResponse SendNetwork(
-			const FName& Name,
+			const int32 NetworkId,
 			const ULearningNeuralNetworkData& Network,
 			FRWLock* NetworkLock = nullptr,
 			const ELogSetting LogSettings = Trainer::DefaultLogSettings) override final;
 
-		virtual void AddReplayBuffer(const FName& Name,	const FReplayBuffer& ReplayBuffer) override final;
-
-		virtual bool ContainsReplayBuffer(const FName& Name) const override final;
+		virtual int32 AddReplayBuffer(const FReplayBuffer& ReplayBuffer) override final;
 
 		virtual ETrainerResponse SendReplayBuffer(
-			const FName& Name,
+			const int32 ReplayBufferId,
 			const FReplayBuffer& ReplayBuffer,
 			const ELogSetting LogSettings = Trainer::DefaultLogSettings) override final;
 
@@ -312,8 +314,8 @@ namespace UE::Learning
 
 		int32 ProcessIdx = INDEX_NONE;
 		TSharedMemoryArrayView<2, volatile int32> Controls;
-		TMap<FName, TSharedMemoryArrayView<1, uint8>> NeuralNetworkSharedMemoryArrayViews;
-		TMap<FName, FSharedMemoryExperienceContainer> SharedMemoryExperienceContainers;
+		TArray<TSharedMemoryArrayView<1, uint8>> NeuralNetworkSharedMemoryArrayViews;
+		TArray<FSharedMemoryExperienceContainer> SharedMemoryExperienceContainers;
 	};
 
 	/**
@@ -325,6 +327,8 @@ namespace UE::Learning
 		/**
 		* Creates a training server as a subprocess
 		*
+		* @param CustomTrainerPath			Path to check for custom trainer files
+		* @param TrainerFileName			The name of the training file to use
 		* @param PythonExecutablePath		Path to the python executable used for training. In general should be the python shipped with Unreal Editor.
 		* @param PythonContentPath			Path to the Python Content folder provided by the Learning plugin
 		* @param IntermediatePath			Path to the intermediate folder to write temporary files, logs, and snapshots to
@@ -334,6 +338,8 @@ namespace UE::Learning
 		* @param LogSettings				Logging settings to use
 		*/
 		FSocketTrainerServerProcess(
+			const FString& CustomTrainerPath,
+			const FString& TrainerFileName,
 			const FString& PythonExecutablePath,
 			const FString& PythonContentPath,
 			const FString& IntermediatePath,
@@ -404,39 +410,37 @@ namespace UE::Learning
 
 		virtual ETrainerResponse Wait() override final;
 
+		virtual bool HasNetworkOrCompleted() override final;
+
 		virtual ETrainerResponse SendConfig(
 			const TSharedRef<FJsonObject>& ConfigObject,
 			const ELogSetting LogSettings = Trainer::DefaultLogSettings) override final;
 
-		virtual void AddNetwork(const FName& Name, const ULearningNeuralNetworkData& Network) override final;
-
-		virtual bool ContainsNetwork(const FName& Name) const override final;
+		virtual int32 AddNetwork(const ULearningNeuralNetworkData& Network) override final;
 
 		virtual ETrainerResponse ReceiveNetwork(
-			const FName& Name,
+			const int32 NetworkId,
 			ULearningNeuralNetworkData& OutNetwork,
 			FRWLock* NetworkLock = nullptr,
 			const ELogSetting LogSettings = Trainer::DefaultLogSettings) override final;
 
 		virtual ETrainerResponse SendNetwork(
-			const FName& Name,
+			const int32 NetworkId,
 			const ULearningNeuralNetworkData& Network,
 			FRWLock* NetworkLock = nullptr,
 			const ELogSetting LogSettings = Trainer::DefaultLogSettings) override final;
 
-		virtual void AddReplayBuffer(const FName& Name, const FReplayBuffer& ReplayBuffer) override final;
-
-		virtual bool ContainsReplayBuffer(const FName& Name) const override final;
+		virtual int32 AddReplayBuffer(const FReplayBuffer& ReplayBuffer) override final;
 
 		virtual ETrainerResponse SendReplayBuffer(
-			const FName& Name,
+			const int32 ReplayBufferId,
 			const FReplayBuffer& ReplayBuffer,
 			const ELogSetting LogSettings = Trainer::DefaultLogSettings) override final;
 
 	private:
 
-		TMap<FName, TLearningArray<1, uint8>> NetworkBuffers;
-		TSet<FName> ExperienceBufferNames;
+		TArray<TLearningArray<1, uint8>> NetworkBuffers;
+		int32 LastReplayBufferId = -1;
 
 		float Timeout = Trainer::DefaultTimeout;
 
diff --git a/Engine/Plugins/Experimental/LearningAgents/Source/LearningTraining/Public/LearningImitationTrainer.h b/Engine/Plugins/Experimental/LearningAgents/Source/LearningTraining/Public/LearningImitationTrainer.h
deleted file mode 100644
diff --git a/Engine/Plugins/Experimental/LearningAgents/Source/LearningTraining/Public/LearningTrainer.h b/Engine/Plugins/Experimental/LearningAgents/Source/LearningTraining/Public/LearningTrainer.h
--- a/Engine/Plugins/Experimental/LearningAgents/Source/LearningTraining/Public/LearningTrainer.h
+++ b/Engine/Plugins/Experimental/LearningAgents/Source/LearningTraining/Public/LearningTrainer.h
@@ -184,6 +184,11 @@ namespace UE::Learning
 		*/
 		LEARNINGTRAINING_API FString GetPythonContentPath(const FString& EngineDir);
 
+		/**
+		* Gets the project's Python Content path.
+		*/
+		LEARNINGTRAINING_API FString GetProjectPythonContentPath();
+
 		/**
 		* Gets the LearningAgents Intermediate path from the intermediate directory.
 		*/
