Learning Agents: Remove process num from BP, fix misc issues

#rnx

diff --git a/Engine/Plugins/Experimental/LearningAgents/Source/LearningAgentsTraining/Private/LearningAgentsCommunicator.cpp b/Engine/Plugins/Experimental/LearningAgents/Source/LearningAgentsTraining/Private/LearningAgentsCommunicator.cpp
--- a/Engine/Plugins/Experimental/LearningAgents/Source/LearningAgentsTraining/Private/LearningAgentsCommunicator.cpp
+++ b/Engine/Plugins/Experimental/LearningAgents/Source/LearningAgentsTraining/Private/LearningAgentsCommunicator.cpp
@@ -36,7 +36,7 @@ FLearningAgentsTrainerProcess ULearningAgentsCommunicatorLibrary::SpawnSharedMem
 		PythonExecutablePath,
 		PythonContentPath,
 		IntermediatePath,
-		SharedMemorySettings.ProcessNum,
+		1, // ProcessNum hard-coded to 1 for now
 		SharedMemorySettings.Timeout);
 
 	return TrainerProcess;
@@ -57,7 +57,7 @@ FLearningAgentsCommunicator ULearningAgentsCommunicatorLibrary::MakeSharedMemory
 
 	Communicator.Trainer = MakeShared<UE::Learning::FSharedMemoryTrainer>(
 		SharedMemorySettings.TaskName,
-		SharedMemorySettings.ProcessNum,
+		1, // ProcessNum hard-coded to 1 for now
 		TrainerProcess.TrainerProcess,
 		SharedMemorySettings.Timeout);
 
diff --git a/Engine/Plugins/Experimental/LearningAgents/Source/LearningAgentsTraining/Private/LearningAgentsPPOTrainer.cpp b/Engine/Plugins/Experimental/LearningAgents/Source/LearningAgentsTraining/Private/LearningAgentsPPOTrainer.cpp
--- a/Engine/Plugins/Experimental/LearningAgents/Source/LearningAgentsTraining/Private/LearningAgentsPPOTrainer.cpp
+++ b/Engine/Plugins/Experimental/LearningAgents/Source/LearningAgentsTraining/Private/LearningAgentsPPOTrainer.cpp
@@ -455,7 +455,7 @@ TSharedRef<FJsonObject> ULearningAgentsPPOTrainer::CreateConfig(const FLearningA
 void ULearningAgentsPPOTrainer::SendConfig(const TSharedRef<FJsonObject>& ConfigObject)
 {
 	UE::Learning::ETrainerResponse Response = UE::Learning::ETrainerResponse::Success;
-	Response = Trainer->SendConfig(ConfigObject);
+	Response = Trainer->SendConfig(ConfigObject, UE::Learning::ELogSetting::Normal);
 	
 	if (Response != UE::Learning::ETrainerResponse::Success)
 	{
diff --git a/Engine/Plugins/Experimental/LearningAgents/Source/LearningAgentsTraining/Public/LearningAgentsCommunicator.h b/Engine/Plugins/Experimental/LearningAgents/Source/LearningAgentsTraining/Public/LearningAgentsCommunicator.h
--- a/Engine/Plugins/Experimental/LearningAgents/Source/LearningAgentsTraining/Public/LearningAgentsCommunicator.h
+++ b/Engine/Plugins/Experimental/LearningAgents/Source/LearningAgentsTraining/Public/LearningAgentsCommunicator.h
@@ -32,10 +32,6 @@ public:
 	UPROPERTY(EditAnywhere, Category = "LearningAgents")
 	FString TaskName = TEXT("Training");
 
-	/** Number of processes to use for multi-processed experience gathering. */
-	UPROPERTY(EditAnywhere, Category = "LearningAgents", meta = (ClampMin = "1", UIMin = "1"))
-	int32 ProcessNum = 1;
-
 	/** Time in seconds to wait for the training process before timing out. */
 	UPROPERTY(EditAnywhere, Category = "LearningAgents", meta = (ClampMin = "0.0", UIMin = "0.0"))
 	float Timeout = 10.0f;
diff --git a/Engine/Plugins/Experimental/LearningAgents/Source/LearningTraining/Private/LearningExternalTrainer.cpp b/Engine/Plugins/Experimental/LearningAgents/Source/LearningTraining/Private/LearningExternalTrainer.cpp
--- a/Engine/Plugins/Experimental/LearningAgents/Source/LearningTraining/Private/LearningExternalTrainer.cpp
+++ b/Engine/Plugins/Experimental/LearningAgents/Source/LearningTraining/Private/LearningExternalTrainer.cpp
@@ -88,11 +88,6 @@ namespace UE::Learning
 			TrainingProcessFlags);
 	}
 
-	FSharedMemoryTrainerServerProcess::~FSharedMemoryTrainerServerProcess()
-	{
-		Terminate();
-	}
-
 	bool FSharedMemoryTrainerServerProcess::IsRunning() const
 	{
 		return TrainingProcess.IsRunning();
@@ -721,11 +716,6 @@ namespace UE::Learning
 			TrainingProcessFlags);
 	}
 
-	FSocketTrainerServerProcess::~FSocketTrainerServerProcess()
-	{
-		Terminate();
-	}
-
 	bool FSocketTrainerServerProcess::IsRunning() const
 	{
 		return TrainingProcess.IsRunning();
diff --git a/Engine/Plugins/Experimental/LearningAgents/Source/LearningTraining/Private/LearningTrainer.cpp b/Engine/Plugins/Experimental/LearningAgents/Source/LearningTraining/Private/LearningTrainer.cpp
--- a/Engine/Plugins/Experimental/LearningAgents/Source/LearningTraining/Private/LearningTrainer.cpp
+++ b/Engine/Plugins/Experimental/LearningAgents/Source/LearningTraining/Private/LearningTrainer.cpp
@@ -44,7 +44,7 @@ namespace UE::Learning
 	{
 		if (IsRunning())
 		{
-			UE_LOG(LogLearning, Warning, TEXT("Terminating Subprocess..."));
+			UE_LOG(LogLearning, Display, TEXT("Terminating Subprocess..."));
 
 			FPlatformProcess::TerminateProc(ProcessHandle, true);
 		}
diff --git a/Engine/Plugins/Experimental/LearningAgents/Source/LearningTraining/Public/LearningExternalTrainer.h b/Engine/Plugins/Experimental/LearningAgents/Source/LearningTraining/Public/LearningExternalTrainer.h
--- a/Engine/Plugins/Experimental/LearningAgents/Source/LearningTraining/Public/LearningExternalTrainer.h
+++ b/Engine/Plugins/Experimental/LearningAgents/Source/LearningTraining/Public/LearningExternalTrainer.h
@@ -178,8 +178,6 @@ namespace UE::Learning
 			const float InTimeout = Trainer::DefaultTimeout,
 			const ESubprocessFlags TrainingProcessFlags = ESubprocessFlags::None);
 
-		~FSharedMemoryTrainerServerProcess();
-
 		/** Check if the server process is still running. */
 		virtual bool IsRunning() const override final;
 
@@ -349,8 +347,6 @@ namespace UE::Learning
 			const ESubprocessFlags TrainingProcessFlags = ESubprocessFlags::None,
 			const ELogSetting LogSettings = ELogSetting::Normal);
 
-		~FSocketTrainerServerProcess();
-
 		/**
 		* Check if the server process is still running
 		*/
