LearningAgents: Fix for PCA computation when no progress supplied.

diff --git a/Engine/Plugins/Experimental/LearningAgents/Source/Learning/Private/LearningPCA.cpp b/Engine/Plugins/Experimental/LearningAgents/Source/Learning/Private/LearningPCA.cpp
--- a/Engine/Plugins/Experimental/LearningAgents/Source/Learning/Private/LearningPCA.cpp
+++ b/Engine/Plugins/Experimental/LearningAgents/Source/Learning/Private/LearningPCA.cpp
@@ -185,7 +185,10 @@ namespace UE::Learning
 			OutEigenMatrix(OutData.Slice(RowStart, RowSliceNum)).noalias() =
 				(InEigenMatrix(Matrix).transpose() * InEigenMatrix(Data.Slice(RowStart, RowSliceNum)).transpose()).transpose();
 
-			Progress->Decrement(RowSliceNum);
+			if (Progress)
+			{
+				Progress->Decrement(RowSliceNum);
+			}
 		});
 
 		Array::Check(OutData);
