LearningAgents: Improve socket comms & support Mac and Linux training

diff --git a/Engine/Plugins/Experimental/LearningAgents/Source/LearningAgentsTraining/Private/LearningAgentsCommunicator.cpp b/Engine/Plugins/Experimental/LearningAgents/Source/LearningAgentsTraining/Private/LearningAgentsCommunicator.cpp
--- a/Engine/Plugins/Experimental/LearningAgents/Source/LearningAgentsTraining/Private/LearningAgentsCommunicator.cpp
+++ b/Engine/Plugins/Experimental/LearningAgents/Source/LearningAgentsTraining/Private/LearningAgentsCommunicator.cpp
@@ -13,6 +13,12 @@ FLearningAgentsTrainerProcess ULearningAgentsCommunicatorLibrary::SpawnSharedMem
 {
 	FLearningAgentsTrainerProcess TrainerProcess;
 
+	if (PLATFORM_MAC)
+	{
+		UE_LOG(LogLearning, Error, TEXT("SpawnSharedMemoryTrainingProcess: Shared Memory not supported on Mac. Switch to Socket Communicator instead."));
+		return TrainerProcess;
+	}
+
 	const FString PythonExecutablePath = UE::Learning::Trainer::GetPythonExecutablePath(TrainerProcessSettings.GetIntermediatePath());
 	if (!FPaths::FileExists(PythonExecutablePath))
 	{
@@ -49,6 +55,12 @@ FLearningAgentsCommunicator ULearningAgentsCommunicatorLibrary::MakeSharedMemory
 {
 	FLearningAgentsCommunicator Communicator;
 
+	if (PLATFORM_MAC)
+	{
+		UE_LOG(LogLearning, Error, TEXT("MakeSharedMemoryCommunicator: Shared Memory not supported on Mac. Switch to Socket Communicator instead."));
+		return Communicator;
+	}
+
 	if (!TrainerProcess.TrainerProcess)
 	{
 		UE_LOG(LogLearning, Error, TEXT("MakeSharedMemoryCommunicator: TrainerProcess is nullptr"));
diff --git a/Engine/Plugins/Experimental/LearningAgents/Source/LearningAgentsTraining/Private/LearningAgentsPPOTrainer.cpp b/Engine/Plugins/Experimental/LearningAgents/Source/LearningAgentsTraining/Private/LearningAgentsPPOTrainer.cpp
--- a/Engine/Plugins/Experimental/LearningAgents/Source/LearningAgentsTraining/Private/LearningAgentsPPOTrainer.cpp
+++ b/Engine/Plugins/Experimental/LearningAgents/Source/LearningAgentsTraining/Private/LearningAgentsPPOTrainer.cpp
@@ -258,12 +258,6 @@ void ULearningAgentsPPOTrainer::BeginTraining(
 	const FLearningAgentsTrainingGameSettings& TrainingGameSettings,
 	const bool bResetAgentsOnBegin)
 {
-	if (!PLATFORM_WINDOWS)
-	{
-		UE_LOG(LogLearning, Error, TEXT("%s: Training currently only supported on Windows."), *GetName());
-		return;
-	}
-
 	if (!IsSetup())
 	{
 		UE_LOG(LogLearning, Error, TEXT("%s: Setup not complete."), *GetName());
diff --git a/Engine/Plugins/Experimental/LearningAgents/Source/LearningTraining/Private/LearningExternalTrainer.cpp b/Engine/Plugins/Experimental/LearningAgents/Source/LearningTraining/Private/LearningExternalTrainer.cpp
--- a/Engine/Plugins/Experimental/LearningAgents/Source/LearningTraining/Private/LearningExternalTrainer.cpp
+++ b/Engine/Plugins/Experimental/LearningAgents/Source/LearningTraining/Private/LearningExternalTrainer.cpp
@@ -714,6 +714,12 @@ namespace UE::Learning
 			FileManager.ConvertToAbsolutePathForExternalAppForRead(*PythonExecutablePath),
 			CommandLineArguments, 
 			TrainingProcessFlags);
+
+		if (PLATFORM_MAC)
+		{
+			// TODO we seem to have to sleep on Mac so the trainer can start listening before we try to connect
+			FPlatformProcess::Sleep(1.0f);
+		}
 	}
 
 	bool FSocketTrainerServerProcess::IsRunning() const
@@ -780,7 +786,7 @@ namespace UE::Learning
 			return;
 		}
 
-		Socket = FTcpSocketBuilder(TEXT("LearningTrainerSocket")).AsNonBlocking().Build();
+		Socket = FTcpSocketBuilder(TEXT("LearningTrainerSocket")).AsBlocking().Build();
 
 		OutResponse = SocketTraining::WaitForConnection(*Socket, TrainingProcess, *Address, Timeout);
 	}
@@ -811,7 +817,11 @@ namespace UE::Learning
 
 	ETrainerResponse FSocketTrainer::SendStop()
 	{
-		checkf(Socket, TEXT("Training socket is nullptr"));
+		if (!Socket)
+		{
+			UE_LOG(LogLearning, Error, TEXT("Training socket is nullptr"));
+			return ETrainerResponse::Unexpected;
+		}
 
 		return SocketTraining::SendStop(*Socket, TrainingProcess, Timeout);
 	}
@@ -820,7 +830,11 @@ namespace UE::Learning
 		const TSharedRef<FJsonObject>& ConfigObject,
 		const ELogSetting LogSettings)
 	{
-		checkf(Socket, TEXT("Training socket is nullptr"));
+		if (!Socket)
+		{
+			UE_LOG(LogLearning, Error, TEXT("Training socket is nullptr"));
+			return ETrainerResponse::Unexpected;
+		}
 
 		FString ConfigString;
 		TSharedRef<TJsonWriter<>> JsonWriter = TJsonWriterFactory<>::Create(&ConfigString, 0);
@@ -843,7 +857,12 @@ namespace UE::Learning
 		FRWLock* NetworkLock,
 		const ELogSetting LogSettings)
 	{
-		checkf(Socket, TEXT("Training socket is nullptr"));
+		if (!Socket)
+		{
+			UE_LOG(LogLearning, Error, TEXT("Training socket is nullptr"));
+			return ETrainerResponse::Unexpected;
+		}
+
 		if (!ensureMsgf(NetworkBuffers.Num() >= NetworkId, TEXT("Network %d has not been added. Call AddNetwork prior to ReceiveNetwork."), NetworkId))
 		{
 			return ETrainerResponse::Unexpected;
@@ -858,7 +877,12 @@ namespace UE::Learning
 		FRWLock* NetworkLock,
 		const ELogSetting LogSettings)
 	{
-		checkf(Socket, TEXT("Training socket is nullptr"));
+		if (!Socket)
+		{
+			UE_LOG(LogLearning, Error, TEXT("Training socket is nullptr"));
+			return ETrainerResponse::Unexpected;
+		}
+
 		if (!ensureMsgf(NetworkBuffers.Num() >= NetworkId, TEXT("Network %d has not been added. Call AddNetwork prior to SendNetwork."), NetworkId))
 		{
 			return ETrainerResponse::Unexpected;
@@ -875,7 +899,12 @@ namespace UE::Learning
 
 	ETrainerResponse FSocketTrainer::SendReplayBuffer(const int32 ReplayBufferId, const FReplayBuffer& ReplayBuffer, const ELogSetting LogSettings)
 	{
-		checkf(Socket, TEXT("Training socket is nullptr"));
+		if (!Socket)
+		{
+			UE_LOG(LogLearning, Error, TEXT("Training socket is nullptr"));
+			return ETrainerResponse::Unexpected;
+		}
+
 		if (!ensureMsgf(ReplayBufferId <= LastReplayBufferId, TEXT("ReplayBuffer %d has not been added. Call AddReplayBuffer prior to SendReplayBuffer."), ReplayBufferId))
 		{
 			return ETrainerResponse::Unexpected;
diff --git a/Engine/Plugins/Experimental/LearningAgents/Source/LearningTraining/Private/LearningSocketTraining.cpp b/Engine/Plugins/Experimental/LearningAgents/Source/LearningTraining/Private/LearningSocketTraining.cpp
--- a/Engine/Plugins/Experimental/LearningAgents/Source/LearningTraining/Private/LearningSocketTraining.cpp
+++ b/Engine/Plugins/Experimental/LearningAgents/Source/LearningTraining/Private/LearningSocketTraining.cpp
@@ -14,21 +14,13 @@ namespace UE::Learning::SocketTraining
 	ETrainerResponse WaitForConnection(FSocket& Socket, FSubprocess* Process, const FInternetAddr& Addr, const float Timeout)
 	{
 		float WaitTime = 0.0f;
-		const float SleepTime = 0.001f;
+		const float SleepTime = 0.001f; // 1 millisecond
 
-		while (true)
+		while (!Process || Process->Update())
 		{
-			// If we're monitoring a process, then has it has exited?
-			if (Process && !Process->Update())
-			{
-				return ETrainerResponse::Unexpected;
-			}
-
-			Socket.Connect(Addr);
-
-			if (Socket.GetConnectionState() == ESocketConnectionState::SCS_Connected)
+			if (Socket.Connect(Addr))
 			{
-				break;
+				return ETrainerResponse::Success;
 			}
 			else
 			{
@@ -42,37 +34,32 @@ namespace UE::Learning::SocketTraining
 			}
 		}
 
-		return ETrainerResponse::Success;
+		return ETrainerResponse::Unexpected;
 	}
 
 	ETrainerResponse RecvWithTimeout(FSocket& Socket, FSubprocess* Process, uint8* Bytes, const int32 ByteNum, const float Timeout)
 	{
 		float WaitTime = 0.0f;
-		const float SleepTime = 0.001f;
 
 		int32 BytesRead = 0;
 		int32 TotalBytesRead = 0;
 
-		while (Socket.GetConnectionState() == SCS_Connected)
+		while (!Process || Process->Update())
 		{
-			// If we're monitoring a process, then has it has exited?
-			if (Process && !Process->Update())
+			if (Socket.Wait(ESocketWaitConditions::WaitForRead, FTimespan(1000))) // 1 millisecond
 			{
-				return ETrainerResponse::Unexpected;
-			}
-
-			if (Socket.Recv(Bytes + TotalBytesRead, ByteNum - TotalBytesRead, BytesRead))
-			{
-				TotalBytesRead += BytesRead;
-
-				if (TotalBytesRead == ByteNum)
+				if (Socket.Recv(Bytes + TotalBytesRead, ByteNum - TotalBytesRead, BytesRead))
 				{
-					return ETrainerResponse::Success;
+					TotalBytesRead += BytesRead;
+
+					if (TotalBytesRead == ByteNum)
+					{
+						return ETrainerResponse::Success;
+					}
 				}
 			}
 
-			FPlatformProcess::Sleep(SleepTime);
-			WaitTime += SleepTime;
+			WaitTime += 0.001f; // 1 millisecond
 
 			if (WaitTime > Timeout)
 			{
@@ -163,25 +150,26 @@ namespace UE::Learning::SocketTraining
 	ETrainerResponse SendWithTimeout(FSocket& Socket, FSubprocess* Process, const uint8* Bytes, const int32 ByteNum, const float Timeout)
 	{
 		float WaitTime = 0.0f;
-		const float SleepTime = 0.001f;
 
 		int32 BytesSent = 0;
 		int32 TotalBytesSent = 0;
 
-		while ((!Process || Process->Update()) && Socket.GetConnectionState() == SCS_Connected)
+		while (!Process || Process->Update())
 		{
-			if (Socket.Send(Bytes + TotalBytesSent, ByteNum - TotalBytesSent, BytesSent))
+			if (Socket.Wait(ESocketWaitConditions::WaitForWrite, FTimespan(1000))) // 1 millisecond
 			{
-				TotalBytesSent += BytesSent;
-
-				if (TotalBytesSent == ByteNum)
+				if (Socket.Send(Bytes + TotalBytesSent, ByteNum - TotalBytesSent, BytesSent))
 				{
-					return ETrainerResponse::Success;
+					TotalBytesSent += BytesSent;
+
+					if (TotalBytesSent == ByteNum)
+					{
+						return ETrainerResponse::Success;
+					}
 				}
 			}
 
-			FPlatformProcess::Sleep(SleepTime);
-			WaitTime += SleepTime;
+			WaitTime += 0.001f; // 1 millisecond
 
 			if (WaitTime > Timeout)
 			{
