LearningAgents: Fixed small bug in how PCA computation decides on the number of components to keep.

diff --git a/Engine/Plugins/Experimental/LearningAgents/Source/Learning/Private/LearningPCA.cpp b/Engine/Plugins/Experimental/LearningAgents/Source/Learning/Private/LearningPCA.cpp
--- a/Engine/Plugins/Experimental/LearningAgents/Source/Learning/Private/LearningPCA.cpp
+++ b/Engine/Plugins/Experimental/LearningAgents/Source/Learning/Private/LearningPCA.cpp
@@ -105,12 +105,12 @@ namespace UE::Learning
 
 			if (LogSettings != ELogSetting::Silent)
 			{
-				UE_LOG(LogLearning, Display, TEXT("--- PCA Variance ---"));
+				UE_LOG(LogLearning, Display, TEXT("--- PCA Cumulative Variance ---"));
 			}
 
-			for (int32 ValueIndex = 0; ValueIndex < FMath::Min((int32)VarianceRatio.size(), Settings.MaximumDimensions + 1); ValueIndex++)
+			for (int32 ValueIndex = 0; ValueIndex < FMath::Min((int32)VarianceRatio.size(), Settings.MaximumDimensions); ValueIndex++)
 			{
-				ComponentsToKeepNum = ValueIndex;
+				ComponentsToKeepNum = ValueIndex + 1;
 				KeptVarianceRatio += FMath::Abs(VarianceRatio(ValueIndex));
 
 				if (LogSettings != ELogSetting::Silent)
