LearningAgents: Added ClippedGaussianArray.

diff --git a/Engine/Plugins/Experimental/LearningAgents/Source/Learning/Private/LearningRandom.cpp b/Engine/Plugins/Experimental/LearningAgents/Source/Learning/Private/LearningRandom.cpp
--- a/Engine/Plugins/Experimental/LearningAgents/Source/Learning/Private/LearningRandom.cpp
+++ b/Engine/Plugins/Experimental/LearningAgents/Source/Learning/Private/LearningRandom.cpp
@@ -252,6 +252,23 @@ namespace UE::Learning::Random
 #endif
 	}
 
+	void ClippedGaussianArray(
+		TLearningArrayView<1, float> Output,
+		const uint32 State,
+		const float Mean,
+		const float Std,
+		const float Clip)
+	{
+		UE_LEARNING_TRACE_CPUPROFILER_EVENT_SCOPE(Learning::Random::ClippedGaussianArray);
+
+		const int32 ElementNum = Output.Num();
+
+		for (int32 ElementIdx = 0; ElementIdx < ElementNum; ElementIdx++)
+		{
+			Output[ElementIdx] = ClippedGaussian(State ^ 0x11a760f7 ^ Int(ElementIdx ^ 0x1dfc9b77), Mean, Std, Clip);
+		}
+	}
+
 	void PlanarClippedGaussianArray(
 		TLearningArrayView<1, FVector> Output,
 		const uint32 State,
@@ -558,6 +575,17 @@ namespace UE::Learning::Random
 		GaussianArray(Output, State ^ 0xf2381309, Mean, Std);
 	}
 
+	void SampleClippedGaussianArray(
+		TLearningArrayView<1, float> Output,
+		uint32& State,
+		const float Mean,
+		const float Std,
+		const float Clip)
+	{
+		State = Int(State ^ 0xb96fa3cf);
+		ClippedGaussianArray(Output, State ^ 0xfc1bc4a9, Mean, Std, Clip);
+	}
+
 	void SamplePlanarClippedGaussianArray(
 		TLearningArrayView<1, FVector> Output,
 		uint32& State,
@@ -624,5 +652,4 @@ namespace UE::Learning::Random
 			Logits,
 			Scale);
 	}
-
 }
\ No newline at end of file
diff --git a/Engine/Plugins/Experimental/LearningAgents/Source/Learning/Public/LearningRandom.h b/Engine/Plugins/Experimental/LearningAgents/Source/Learning/Public/LearningRandom.h
--- a/Engine/Plugins/Experimental/LearningAgents/Source/Learning/Public/LearningRandom.h
+++ b/Engine/Plugins/Experimental/LearningAgents/Source/Learning/Public/LearningRandom.h
@@ -144,6 +144,13 @@ namespace UE::Learning::Random
 		const float Mean = 0.0f,
 		const float Std = 1.0f);
 
+	LEARNING_API void ClippedGaussianArray(
+		TLearningArrayView<1, float> Output,
+		const uint32 State,
+		const float Mean = 0.0f,
+		const float Std = 1.0f,
+		const float Clip = 10.0f);
+
 	LEARNING_API void PlanarClippedGaussianArray(
 		TLearningArrayView<1, FVector> Output,
 		const uint32 State,
@@ -265,6 +272,13 @@ namespace UE::Learning::Random
 		uint32& State,
 		const float Mean = 0.0f,
 		const float Std = 1.0f);
+	
+	LEARNING_API void SampleClippedGaussianArray(
+		TLearningArrayView<1, float> Output,
+		uint32& State,
+		const float Mean = 0.0f,
+		const float Std = 1.0f,
+		const float Clip = 10.0f);
 
 	LEARNING_API void SamplePlanarClippedGaussianArray(
 		TLearningArrayView<1, FVector> Output,
