Learning: Updated MultiArray to have fixed int64 SizeType

diff --git a/Engine/Plugins/Experimental/LearningAgents/Source/Learning/Private/LearningNeuralNetwork.cpp b/Engine/Plugins/Experimental/LearningAgents/Source/Learning/Private/LearningNeuralNetwork.cpp
--- a/Engine/Plugins/Experimental/LearningAgents/Source/Learning/Private/LearningNeuralNetwork.cpp
+++ b/Engine/Plugins/Experimental/LearningAgents/Source/Learning/Private/LearningNeuralNetwork.cpp
@@ -49,7 +49,7 @@ bool ULearningNeuralNetworkData::LoadFromSnapshot(const TArrayView<const uint8>
 		return false;
 	}
 
-	int32 Offset = 0;
+	int64 Offset = 0;
 	
 	int32 MagicNumber = 0;
 	UE::Learning::DeserializeFromBytes(Offset, InBytes, MagicNumber);
@@ -79,7 +79,7 @@ bool ULearningNeuralNetworkData::LoadFromSnapshot(const TArrayView<const uint8>
 
 void ULearningNeuralNetworkData::SaveToSnapshot(TArrayView<uint8> OutBytes) const
 {
-	int32 Offset = 0;
+	int64 Offset = 0;
 	UE::Learning::SerializeToBytes(Offset, OutBytes, UE::Learning::NeuralNetwork::Private::MagicNumber);
 	UE::Learning::SerializeToBytes(Offset, OutBytes, UE::Learning::NeuralNetwork::Private::VersionNumber);
 	UE::Learning::SerializeToBytes(Offset, OutBytes, InputSize);
diff --git a/Engine/Plugins/Experimental/LearningAgents/Source/Learning/Public/LearningArray.h b/Engine/Plugins/Experimental/LearningAgents/Source/Learning/Public/LearningArray.h
--- a/Engine/Plugins/Experimental/LearningAgents/Source/Learning/Public/LearningArray.h
+++ b/Engine/Plugins/Experimental/LearningAgents/Source/Learning/Public/LearningArray.h
@@ -45,22 +45,6 @@
 #define UE_LEARNING_ARRAY_CHECK_SHAPE UE_LEARNING_DEVELOPMENT
 #endif
 
-// Enables/Disables index out-of-bounds checking in LearningArray and LearningArrayView.
-// Since so many of the editor tools use extensive numerical code having range checking
-// enabled can significantly slow things down so it is recommended to disable this in
-// editor mode when not debugging. Will be disabled automatically in Shipping.
-#ifndef UE_LEARNING_ARRAY_CHECK_INDEX
-#define UE_LEARNING_ARRAY_CHECK_INDEX UE_LEARNING_DEVELOPMENT
-#endif
-
-// Enables/Disables restrict in LearningArray and LearningArrayView. By convention functions
-// in Learning should not be called with aliasing arguments so it should be okay to enable
-// this but nonetheless it can be safer to leave it disabled in particular if you are
-// seeing odd behavior on certain platforms.
-#ifndef UE_LEARNING_ARRAY_RESTRICT
-#define UE_LEARNING_ARRAY_RESTRICT UE_LEARNING_DEVELOPMENT
-#endif
-
 
 // Profiling Macros
 
@@ -99,39 +83,19 @@
 #define UE_LEARNING_ARRAY_VALUE_CHECKF(...)
 #endif
 
-#if UE_LEARNING_ARRAY_RESTRICT == 1
-#define UE_LEARNING_ARRAY_RESTRICT_CHECK(PtrLhs, PtrRhs) check((PtrLhs == nullptr && PtrRhs == nullptr) || (PtrLhs != PtrRhs))
-#else
-#define UE_LEARNING_ARRAY_RESTRICT_CHECK(PtrLhs, PtrRhs)
-#endif
-
-
 // Learning Multi-Dimensional Array Types
 
 template<uint32 DimNum>
 using TLearningArrayShape = TMultiArrayShape<DimNum>;
 
 template<uint32 DimNum, typename ElementType>
-using TLearningArrayView = TMultiArrayView<
-	DimNum, 
-	ElementType, 
-	UE_LEARNING_ARRAY_CHECK_INDEX, 
-	UE_LEARNING_ARRAY_RESTRICT>;
+using TLearningArrayView = TMultiArrayView<DimNum, ElementType>;
 
 template<uint32 DimNum, typename ElementType>
-using TLearningConstArrayView = TConstMultiArrayView<
-	DimNum,
-	ElementType,
-	UE_LEARNING_ARRAY_CHECK_INDEX,
-	UE_LEARNING_ARRAY_RESTRICT>;
+using TLearningConstArrayView = TConstMultiArrayView<DimNum, ElementType>;
 
-template<uint32 DimNum, typename ElementType, typename Allocator = FDefaultAllocator>
-using TLearningArray = TMultiArray<
-	DimNum, 
-	ElementType, 
-	Allocator, 
-	UE_LEARNING_ARRAY_CHECK_INDEX, 
-	UE_LEARNING_ARRAY_RESTRICT>;
+template<uint32 DimNum, typename ElementType, typename Allocator = FDefaultAllocator64>
+using TLearningArray = TMultiArray<DimNum, ElementType, Allocator>;
 
 
 
@@ -293,29 +257,69 @@ namespace UE::Learning
 	/**
 	* Serialize an integer to bytes
 	*/
-	static inline void SerializeToBytes(int32& InOutOffset, TLearningArrayView<1, uint8> Bytes, const int32 InValue)
+	static inline void SerializeToBytes(int64& InOutOffset, TLearningArrayView<1, uint8> Bytes, const uint8 InValue)
 	{
-		UE_LEARNING_CHECK(InOutOffset + sizeof(int32) <= Bytes.Num());
+		UE_LEARNING_CHECK(InOutOffset + (int64)sizeof(uint8) <= Bytes.Num());
+		FMemory::Memcpy(&Bytes[InOutOffset], (uint8*)&InValue, sizeof(uint8));
+		InOutOffset += sizeof(uint8);
+	}
+
+	/**
+	* Serialize an integer to bytes
+	*/
+	static inline void SerializeToBytes(int64& InOutOffset, TLearningArrayView<1, uint8> Bytes, const int32 InValue)
+	{
+		UE_LEARNING_CHECK(InOutOffset + (int64)sizeof(int32) <= Bytes.Num());
 		FMemory::Memcpy(&Bytes[InOutOffset], (uint8*)&InValue, sizeof(int32));
 		InOutOffset += sizeof(int32);
 	}
 
 	/**
-	* Serialize an integer from bytes
+	* Serialize an integer to bytes
+	*/
+	static inline void SerializeToBytes(int64& InOutOffset, TLearningArrayView<1, uint8> Bytes, const int64 InValue)
+	{
+		UE_LEARNING_CHECK(InOutOffset + (int64)sizeof(int64) <= Bytes.Num());
+		FMemory::Memcpy(&Bytes[InOutOffset], (uint8*)&InValue, sizeof(int64));
+		InOutOffset += sizeof(int64);
+	}
+
+	/**
+	* Deserialize an integer from bytes
 	*/
-	static inline void DeserializeFromBytes(int32& InOutOffset, TLearningArrayView<1, const uint8> Bytes, int32& OutValue)
+	static inline void DeserializeFromBytes(int64& InOutOffset, TLearningArrayView<1, const uint8> Bytes, uint8& OutValue)
 	{
-		UE_LEARNING_CHECK(InOutOffset + sizeof(int32) <= Bytes.Num());
+		UE_LEARNING_CHECK(InOutOffset + (int64)sizeof(uint8) <= Bytes.Num());
+		FMemory::Memcpy((uint8*)&OutValue, &Bytes[InOutOffset], sizeof(uint8));
+		InOutOffset += sizeof(uint8);
+	}
+
+	/**
+	* Deserialize an integer from bytes
+	*/
+	static inline void DeserializeFromBytes(int64& InOutOffset, TLearningArrayView<1, const uint8> Bytes, int32& OutValue)
+	{
+		UE_LEARNING_CHECK(InOutOffset + (int64)sizeof(int32) <= Bytes.Num());
 		FMemory::Memcpy((uint8*)&OutValue, &Bytes[InOutOffset], sizeof(int32));
 		InOutOffset += sizeof(int32);
 	}
 
 	/**
-	* Serialize an integer from bytes
+	* Deserialize an integer from bytes
+	*/
+	static inline void DeserializeFromBytes(int64& InOutOffset, TLearningArrayView<1, const uint8> Bytes, int64& OutValue)
+	{
+		UE_LEARNING_CHECK(InOutOffset + (int64)sizeof(int64) <= Bytes.Num());
+		FMemory::Memcpy((uint8*)&OutValue, &Bytes[InOutOffset], sizeof(int64));
+		InOutOffset += sizeof(int64);
+	}
+
+	/**
+	* Deserialize an integer from bytes
 	*/
-	static inline void DeserializeFromBytes(int32& InOutOffset, TLearningArrayView<1, const uint8> Bytes, uint64& OutValue)
+	static inline void DeserializeFromBytes(int64& InOutOffset, TLearningArrayView<1, const uint8> Bytes, uint64& OutValue)
 	{
-		UE_LEARNING_CHECK(InOutOffset + sizeof(uint64) <= Bytes.Num());
+		UE_LEARNING_CHECK(InOutOffset + (int64)sizeof(uint64) <= Bytes.Num());
 		FMemory::Memcpy((uint8*)&OutValue, &Bytes[InOutOffset], sizeof(uint64));
 		InOutOffset += sizeof(uint64);
 	}
@@ -338,7 +342,7 @@ namespace UE::Learning
 			for (uint8 Idx = 0; Idx < InDimNum; Idx++)
 			{
 				UE_LEARNING_ARRAY_SHAPE_CHECKF(Lhs[Idx] == Rhs[Idx],
-					TEXT("Array Shapes don't match on dimension %i of %i (lhs: %i, rhs: %i)"), Idx + 1, InDimNum, Lhs[Idx], Rhs[Idx]);
+					TEXT("Array Shapes don't match on dimension %i of %i (lhs: %lld, rhs: %lld)"), Idx + 1, InDimNum, Lhs[Idx], Rhs[Idx]);
 			}
 #endif
 		}
@@ -350,11 +354,11 @@ namespace UE::Learning
 		inline void Check(const TLearningArrayView<InDimNum, InElementType> View)
 		{
 #if UE_LEARNING_ARRAY_CHECK_VALUES == 1
-			for (int32 Idx = 0; Idx < View.Num(); Idx++)
+			for (int64 Idx = 0; Idx < View.Num(); Idx++)
 			{
 				UE_LEARNING_ARRAY_VALUE_CHECKF(
 					FMath::IsFinite(View.GetData()[Idx]) && View.GetData()[Idx] != MAX_flt && View.GetData()[Idx] != -MAX_flt,
-					TEXT("Invalid value %f found at flat array index %i"), View.GetData()[Idx], Idx);
+					TEXT("Invalid value %f found at flat array index %lld"), View.GetData()[Idx], Idx);
 			}
 #endif
 		}
@@ -366,11 +370,11 @@ namespace UE::Learning
 		inline void Check(const TLearningArray<InDimNum, InElementType, Allocator>& View)
 		{
 #if UE_LEARNING_ARRAY_CHECK_VALUES == 1
-			for (int32 Idx = 0; Idx < View.Num(); Idx++)
+			for (int64 Idx = 0; Idx < View.Num(); Idx++)
 			{
 				UE_LEARNING_ARRAY_VALUE_CHECKF(
 					FMath::IsFinite(View.GetData()[Idx]) && View.GetData()[Idx] != MAX_flt && View.GetData()[Idx] != -MAX_flt,
-					TEXT("Invalid value %f found at flat array index %i"), View.GetData()[Idx], Idx);
+					TEXT("Invalid value %f found at flat array index %lld"), View.GetData()[Idx], Idx);
 			}
 #endif
 		}
@@ -383,11 +387,11 @@ namespace UE::Learning
 		inline void Check(const TLearningArrayView<InDimNum, const InElementType> View)
 		{
 #if UE_LEARNING_ARRAY_CHECK_VALUES == 1
-			for (int32 Idx = 0; Idx < View.Num(); Idx++)
+			for (int64 Idx = 0; Idx < View.Num(); Idx++)
 			{
 				UE_LEARNING_ARRAY_VALUE_CHECKF(
 					FMath::IsFinite(View.GetData()[Idx]) && View.GetData()[Idx] != MAX_flt && View.GetData()[Idx] != -MAX_flt,
-					TEXT("Invalid value %f found at flat array index %i"), View.GetData()[Idx], Idx);
+					TEXT("Invalid value %f found at flat array index %lld"), View.GetData()[Idx], Idx);
 			}
 #endif
 		}
@@ -412,7 +416,7 @@ namespace UE::Learning
 				{
 					UE_LEARNING_ARRAY_VALUE_CHECKF(
 						FMath::IsFinite(View.GetData()[Idx]) && View.GetData()[Idx] != MAX_flt && View.GetData()[Idx] != -MAX_flt,
-						TEXT("Invalid value %f found at flat array index %i"), View.GetData()[Idx], Idx);
+						TEXT("Invalid value %f found at flat array index %lld"), View.GetData()[Idx], Idx);
 				}
 			}
 #endif
@@ -438,7 +442,7 @@ namespace UE::Learning
 				{
 					UE_LEARNING_ARRAY_VALUE_CHECKF(
 						FMath::IsFinite(View.GetData()[Idx]) && View.GetData()[Idx] != MAX_flt && View.GetData()[Idx] != -MAX_flt,
-						TEXT("Invalid value %f found at flat array index %i"), View.GetData()[Idx], Idx);
+						TEXT("Invalid value %f found at flat array index %lld"), View.GetData()[Idx], Idx);
 				}
 			}
 #endif
@@ -464,7 +468,7 @@ namespace UE::Learning
 				{
 					UE_LEARNING_ARRAY_VALUE_CHECKF(
 						FMath::IsFinite(View.GetData()[Idx]) && View.GetData()[Idx] != MAX_flt && View.GetData()[Idx] != -MAX_flt,
-						TEXT("Invalid value %f found at flat array index %i"), View.GetData()[Idx], Idx);
+						TEXT("Invalid value %f found at flat array index %lld"), View.GetData()[Idx], Idx);
 				}
 			}
 #endif
@@ -539,7 +543,7 @@ namespace UE::Learning
 		template<uint8 InDimNum, typename InElementType>
 		inline void Set(TLearningArrayView<InDimNum, InElementType> View, const InElementType& Element)
 		{
-			for (int32 Idx = 0; Idx < View.Num(); Idx++)
+			for (int64 Idx = 0; Idx < View.Num(); Idx++)
 			{
 				View.GetData()[Idx] = Element;
 			}
@@ -551,7 +555,7 @@ namespace UE::Learning
 		template<uint8 InDimNum, typename InElementType, typename Allocator>
 		inline void Set(TLearningArray<InDimNum, InElementType, Allocator>& View, const InElementType& Element)
 		{
-			for (int32 Idx = 0; Idx < View.Num(); Idx++)
+			for (int64 Idx = 0; Idx < View.Num(); Idx++)
 			{
 				View.GetData()[Idx] = Element;
 			}
@@ -609,10 +613,9 @@ namespace UE::Learning
 			TLearningArrayView<InDimNum, InElementType> Dst,
 			TLearningArrayView<InDimNum, InElementType> Src)
 		{
-			UE_LEARNING_ARRAY_RESTRICT_CHECK(Dst.GetData(), Src.GetData());
 			CheckShapesEqual<InDimNum>(Dst.Shape(), Src.Shape());
 
-			for (int32 Idx = 0; Idx < Dst.Num(); Idx++)
+			for (int64 Idx = 0; Idx < Dst.Num(); Idx++)
 			{
 				Dst.GetData()[Idx] = Src.GetData()[Idx];
 			}
@@ -626,10 +629,9 @@ namespace UE::Learning
 			TLearningArrayView<InDimNum, InElementType> Dst,
 			const TLearningArrayView<InDimNum, const InElementType> Src)
 		{
-			UE_LEARNING_ARRAY_RESTRICT_CHECK(Dst.GetData(), Src.GetData());
 			CheckShapesEqual<InDimNum>(Dst.Shape(), Src.Shape());
 
-			for (int32 Idx = 0; Idx < Dst.Num(); Idx++)
+			for (int64 Idx = 0; Idx < Dst.Num(); Idx++)
 			{
 				Dst.GetData()[Idx] = Src.GetData()[Idx];
 			}
@@ -643,10 +645,9 @@ namespace UE::Learning
 			TLearningArray<InDimNum, InElementType, Allocator>& Dst,
 			const TLearningArrayView<InDimNum, const InElementType> Src)
 		{
-			UE_LEARNING_ARRAY_RESTRICT_CHECK(Dst.GetData(), Src.GetData());
 			CheckShapesEqual<InDimNum>(Dst.Shape(), Src.Shape());
 
-			for (int32 Idx = 0; Idx < Dst.Num(); Idx++)
+			for (int64 Idx = 0; Idx < Dst.Num(); Idx++)
 			{
 				Dst.GetData()[Idx] = Src.GetData()[Idx];
 			}
@@ -660,10 +661,9 @@ namespace UE::Learning
 			TLearningArray<InDimNum, InElementType, Allocator>& Dst,
 			const TLearningArrayView<InDimNum, InElementType> Src)
 		{
-			UE_LEARNING_ARRAY_RESTRICT_CHECK(Dst.GetData(), Src.GetData());
 			CheckShapesEqual<InDimNum>(Dst.Shape(), Src.Shape());
 
-			for (int32 Idx = 0; Idx < Dst.Num(); Idx++)
+			for (int64 Idx = 0; Idx < Dst.Num(); Idx++)
 			{
 				Dst.GetData()[Idx] = Src.GetData()[Idx];
 			}
@@ -677,10 +677,9 @@ namespace UE::Learning
 			TLearningArrayView<InDimNum, InElementType> Dst,
 			const TLearningArray<InDimNum, InElementType, Allocator>& Src)
 		{
-			UE_LEARNING_ARRAY_RESTRICT_CHECK(Dst.GetData(), Src.GetData());
 			CheckShapesEqual<InDimNum>(Dst.Shape(), Src.Shape());
 
-			for (int32 Idx = 0; Idx < Dst.Num(); Idx++)
+			for (int64 Idx = 0; Idx < Dst.Num(); Idx++)
 			{
 				Dst.GetData()[Idx] = Src.GetData()[Idx];
 			}
@@ -694,10 +693,9 @@ namespace UE::Learning
 			TLearningArray<InDimNum, InElementType, AllocatorLhs>& Dst,
 			const TLearningArray<InDimNum, InElementType, AllocatorRhs>& Src)
 		{
-			UE_LEARNING_ARRAY_RESTRICT_CHECK(Dst.GetData(), Src.GetData());
 			CheckShapesEqual<InDimNum>(Dst.Shape(), Src.Shape());
 
-			for (int32 Idx = 0; Idx < Dst.Num(); Idx++)
+			for (int64 Idx = 0; Idx < Dst.Num(); Idx++)
 			{
 				Dst.GetData()[Idx] = Src.GetData()[Idx];
 			}
@@ -712,7 +710,6 @@ namespace UE::Learning
 			const TLearningArrayView<InDimNum, const InElementType> Src,
 			const FIndexSet Indices)
 		{
-			UE_LEARNING_ARRAY_RESTRICT_CHECK(Dst.GetData(), Src.GetData());
 			CheckShapesEqual<InDimNum>(Dst.Shape(), Src.Shape());
 
 			if constexpr (InDimNum > 1)
@@ -740,7 +737,6 @@ namespace UE::Learning
 			const TLearningArray<InDimNum, InElementType, Allocator>& Src,
 			const FIndexSet Indices)
 		{
-			UE_LEARNING_ARRAY_RESTRICT_CHECK(Dst.GetData(), Src.GetData());
 			CheckShapesEqual<InDimNum>(Dst.Shape(), Src.Shape());
 
 			if constexpr (InDimNum > 1)
@@ -768,7 +764,6 @@ namespace UE::Learning
 			const TLearningArray<InDimNum, InElementType, AllocatorRhs>& Src,
 			const FIndexSet Indices)
 		{
-			UE_LEARNING_ARRAY_RESTRICT_CHECK(Dst.GetData(), Src.GetData());
 			CheckShapesEqual<InDimNum>(Dst.Shape(), Src.Shape());
 
 			if constexpr (InDimNum > 1)
@@ -796,7 +791,6 @@ namespace UE::Learning
 			const TLearningArrayView<InDimNum, const InElementType> Src,
 			const FIndexSet Indices)
 		{
-			UE_LEARNING_ARRAY_RESTRICT_CHECK(Dst.GetData(), Src.GetData());
 			CheckShapesEqual<InDimNum>(Dst.Shape(), Src.Shape());
 
 			if constexpr (InDimNum > 1)
@@ -823,10 +817,9 @@ namespace UE::Learning
 			const TLearningArray<InDimNum, InElementType, AllocatorLhs>& Lhs,
 			const TLearningArray<InDimNum, InElementType, AllocatorRhs>& Rhs)
 		{
-			UE_LEARNING_ARRAY_RESTRICT_CHECK(Lhs.GetData(), Rhs.GetData());
 			CheckShapesEqual<InDimNum>(Lhs.Shape(), Rhs.Shape());
 
-			for (int32 Idx = 0; Idx < Lhs.Num(); Idx++)
+			for (int64 Idx = 0; Idx < Lhs.Num(); Idx++)
 			{
 				if (Lhs.GetData()[Idx] != Rhs.GetData()[Idx])
 				{
@@ -845,10 +838,9 @@ namespace UE::Learning
 			const TLearningArrayView<InDimNum, const InElementType> Lhs,
 			const TLearningArray<InDimNum, InElementType, AllocatorRhs>& Rhs)
 		{
-			UE_LEARNING_ARRAY_RESTRICT_CHECK(Lhs.GetData(), Rhs.GetData());
 			CheckShapesEqual<InDimNum>(Lhs.Shape(), Rhs.Shape());
 
-			for (int32 Idx = 0; Idx < Lhs.Num(); Idx++)
+			for (int64 Idx = 0; Idx < Lhs.Num(); Idx++)
 			{
 				if (Lhs.GetData()[Idx] != Rhs.GetData()[Idx])
 				{
@@ -867,10 +859,9 @@ namespace UE::Learning
 			const TLearningArrayView<InDimNum, InElementType> Lhs,
 			const TLearningArray<InDimNum, InElementType, AllocatorRhs>& Rhs)
 		{
-			UE_LEARNING_ARRAY_RESTRICT_CHECK(Lhs.GetData(), Rhs.GetData());
 			CheckShapesEqual<InDimNum>(Lhs.Shape(), Rhs.Shape());
 
-			for (int32 Idx = 0; Idx < Lhs.Num(); Idx++)
+			for (int64 Idx = 0; Idx < Lhs.Num(); Idx++)
 			{
 				if (Lhs.GetData()[Idx] != Rhs.GetData()[Idx])
 				{
@@ -889,10 +880,9 @@ namespace UE::Learning
 			const TLearningArray<InDimNum, InElementType, AllocatorLhs>& Lhs,
 			const TLearningArrayView<InDimNum, const InElementType>& Rhs)
 		{
-			UE_LEARNING_ARRAY_RESTRICT_CHECK(Lhs.GetData(), Rhs.GetData());
 			CheckShapesEqual<InDimNum>(Lhs.Shape(), Rhs.Shape());
 
-			for (int32 Idx = 0; Idx < Lhs.Num(); Idx++)
+			for (int64 Idx = 0; Idx < Lhs.Num(); Idx++)
 			{
 				if (Lhs.GetData()[Idx] != Rhs.GetData()[Idx])
 				{
@@ -911,10 +901,9 @@ namespace UE::Learning
 			const TLearningArrayView<InDimNum, const InElementType> Lhs,
 			const TLearningArrayView<InDimNum, const InElementType> Rhs)
 		{
-			UE_LEARNING_ARRAY_RESTRICT_CHECK(Lhs.GetData(), Rhs.GetData());
 			CheckShapesEqual<InDimNum>(Lhs.Shape(), Rhs.Shape());
 
-			for (int32 Idx = 0; Idx < Lhs.Num(); Idx++)
+			for (int64 Idx = 0; Idx < Lhs.Num(); Idx++)
 			{
 				if (Lhs.GetData()[Idx] != Rhs.GetData()[Idx])
 				{
@@ -929,9 +918,9 @@ namespace UE::Learning
 		* Shift all elements in an array to the left. Data at the end of the array will remain the same.
 		*/
 		template<uint8 InDimNum, typename InElementType, typename Allocator>
-		inline void ShiftLeft(TLearningArray<InDimNum, InElementType, Allocator>& Array, const int32 ShiftNum)
+		inline void ShiftLeft(TLearningArray<InDimNum, InElementType, Allocator>& Array, const int64 ShiftNum)
 		{
-			for (int32 Idx = 0; Idx < (int32)Array.template Num<0>() - ShiftNum; Idx++)
+			for (int64 Idx = 0; Idx < (int64)Array.template Num<0>() - ShiftNum; Idx++)
 			{
 				Copy(Array[Idx], Array[Idx + ShiftNum]);
 			}
@@ -941,9 +930,9 @@ namespace UE::Learning
 		* Shift all elements in an array to the left. Data at the end of the array will remain the same.
 		*/
 		template<typename InElementType, typename Allocator>
-		inline void ShiftLeft(TLearningArray<1, InElementType, Allocator>& Array, const int32 ShiftNum)
+		inline void ShiftLeft(TLearningArray<1, InElementType, Allocator>& Array, const int64 ShiftNum)
 		{
-			for (int32 Idx = 0; Idx < (int32)Array.Num() - ShiftNum; Idx++)
+			for (int64 Idx = 0; Idx < (int64)Array.Num() - ShiftNum; Idx++)
 			{
 				Array.GetData()[Idx] = Array.GetData()[Idx + ShiftNum];
 			}
@@ -953,9 +942,9 @@ namespace UE::Learning
 		* Shift all elements in an array to the left. Data at the end of the array will remain the same.
 		*/
 		template<uint8 InDimNum, typename InElementType>
-		inline void ShiftLeft(TLearningArrayView<InDimNum, InElementType> Array, const int32 ShiftNum)
+		inline void ShiftLeft(TLearningArrayView<InDimNum, InElementType> Array, const int64 ShiftNum)
 		{
-			for (int32 Idx = 0; Idx < (int32)Array.template Num<0>() - ShiftNum; Idx++)
+			for (int64 Idx = 0; Idx < (int64)Array.template Num<0>() - ShiftNum; Idx++)
 			{
 				Copy(Array[Idx], Array[Idx + ShiftNum]);
 			}
@@ -965,9 +954,9 @@ namespace UE::Learning
 		* Shift all elements in an array to the left. Data at the end of the array will remain the same.
 		*/
 		template<typename InElementType>
-		inline void ShiftLeft(TLearningArrayView<1, InElementType> Array, const int32 ShiftNum)
+		inline void ShiftLeft(TLearningArrayView<1, InElementType> Array, const int64 ShiftNum)
 		{
-			for (int32 Idx = 0; Idx < (int32)Array.Num() - ShiftNum; Idx++)
+			for (int64 Idx = 0; Idx < (int64)Array.Num() - ShiftNum; Idx++)
 			{
 				Array.GetData()[Idx] = Array.GetData()[Idx + ShiftNum];
 			}
@@ -1004,7 +993,7 @@ namespace UE::Learning
 		* Serialization byte num
 		*/
 		template<uint8 InDimNum, typename InElementType>
-		inline int32 SerializationByteNum(const TLearningArrayShape<InDimNum> Shape)
+		inline int32 SerializationByteNum32(const TLearningArrayShape<InDimNum> Shape)
 		{
 			return sizeof(int32) + sizeof(int32) * InDimNum + sizeof(InElementType) * Shape.Total();
 		}
@@ -1013,15 +1002,15 @@ namespace UE::Learning
 		* Serialize an array to bytes
 		*/
 		template<uint8 InDimNum, typename InElementType, typename Allocator>
-		inline void SerializeToBytes(int32& InOutOffset, TLearningArrayView<1, uint8> Bytes, const TMultiArrayShape<InDimNum> Shape, const TArray<InElementType, Allocator>& InArray)
+		inline void SerializeToBytes32(int64& InOutOffset, TLearningArrayView<1, uint8> Bytes, const TMultiArrayShape<InDimNum> Shape, const TArray<InElementType, Allocator>& InArray)
 		{
 			Learning::SerializeToBytes(InOutOffset, Bytes, (int32)InDimNum);
 			for (uint8 ShapeIdx = 0; ShapeIdx < InDimNum; ShapeIdx++)
 			{
-				Learning::SerializeToBytes(InOutOffset, Bytes, Shape[ShapeIdx]);
+				Learning::SerializeToBytes(InOutOffset, Bytes, (int32)Shape[ShapeIdx]);
 			}
 
-			UE_LEARNING_CHECK(InOutOffset + sizeof(InElementType) * InArray.Num() <= Bytes.Num());
+			UE_LEARNING_CHECK(InOutOffset + (int64)sizeof(InElementType) * InArray.Num() <= Bytes.Num());
 			FMemory::Memcpy(&Bytes[InOutOffset], (uint8*)InArray.GetData(), sizeof(InElementType) * InArray.Num());
 			InOutOffset += sizeof(InElementType) * InArray.Num();
 		}
@@ -1030,15 +1019,15 @@ namespace UE::Learning
 		* Serialize an array to bytes
 		*/
 		template<uint8 InDimNum, typename InElementType, typename Allocator>
-		inline void SerializeToBytes(int32& InOutOffset, TLearningArrayView<1, uint8> Bytes, const TLearningArray<InDimNum, InElementType, Allocator>& InArray)
+		inline void SerializeToBytes32(int64& InOutOffset, TLearningArrayView<1, uint8> Bytes, const TLearningArray<InDimNum, InElementType, Allocator>& InArray)
 		{
 			Learning::SerializeToBytes(InOutOffset, Bytes, (int32)InDimNum);
 			for (uint8 ShapeIdx = 0; ShapeIdx < InDimNum; ShapeIdx++)
 			{
-				Learning::SerializeToBytes(InOutOffset, Bytes, InArray.Shape()[ShapeIdx]);
+				Learning::SerializeToBytes(InOutOffset, Bytes, (int32)InArray.Shape()[ShapeIdx]);
 			}
 
-			UE_LEARNING_CHECK(InOutOffset + sizeof(InElementType) * InArray.Num() <= Bytes.Num());
+			UE_LEARNING_CHECK(InOutOffset + (int64)sizeof(InElementType) * InArray.Num() <= Bytes.Num());
 			FMemory::Memcpy(&Bytes[InOutOffset], (uint8*)InArray.GetData(), sizeof(InElementType) * InArray.Num());
 			InOutOffset += sizeof(InElementType) * InArray.Num();
 		}
@@ -1047,7 +1036,7 @@ namespace UE::Learning
 		* Deserialize an array from bytes
 		*/
 		template<uint8 InDimNum, typename InElementType, typename Allocator>
-		inline void DeserializeFromBytes(int32& InOutOffset, TLearningArrayView<1, const uint8> Bytes, TArray<InElementType, Allocator>& OutArray)
+		inline void DeserializeFromBytes32(int64& InOutOffset, TLearningArrayView<1, const uint8> Bytes, TArray<InElementType, Allocator>& OutArray)
 		{
 			int32 DimNum = INDEX_NONE;
 			Learning::DeserializeFromBytes(InOutOffset, Bytes, DimNum);
@@ -1056,12 +1045,14 @@ namespace UE::Learning
 			TLearningArrayShape<InDimNum> Shape;
 			for (uint8 ShapeIdx = 0; ShapeIdx < InDimNum; ShapeIdx++)
 			{
-				Learning::DeserializeFromBytes(InOutOffset, Bytes, Shape[ShapeIdx]);
+				int32 Shape32 = 0;
+				Learning::DeserializeFromBytes(InOutOffset, Bytes, Shape32);
+				Shape[ShapeIdx] = Shape32;
 			}
 
 			OutArray.SetNumUninitialized(Shape.Total());
 
-			UE_LEARNING_CHECK(InOutOffset + sizeof(InElementType) * OutArray.Num() <= Bytes.Num());
+			UE_LEARNING_CHECK(InOutOffset + (int64)sizeof(InElementType) * OutArray.Num() <= Bytes.Num());
 			FMemory::Memcpy((uint8*)OutArray.GetData(), &Bytes[InOutOffset], sizeof(InElementType) * OutArray.Num());
 			InOutOffset += sizeof(InElementType) * OutArray.Num();
 		}
@@ -1070,7 +1061,7 @@ namespace UE::Learning
 		* Deserialize an array from bytes
 		*/
 		template<uint8 InDimNum, typename InElementType, typename Allocator>
-		inline void DeserializeFromBytes(int32& InOutOffset, TLearningArrayView<1, const uint8> Bytes, TLearningArray<InDimNum, InElementType, Allocator>& OutArray)
+		inline void DeserializeFromBytes32(int64& InOutOffset, TLearningArrayView<1, const uint8> Bytes, TLearningArray<InDimNum, InElementType, Allocator>& OutArray)
 		{
 			int32 DimNum = INDEX_NONE;
 			Learning::DeserializeFromBytes(InOutOffset, Bytes, DimNum);
@@ -1079,20 +1070,22 @@ namespace UE::Learning
 			TLearningArrayShape<InDimNum> Shape;
 			for (uint8 ShapeIdx = 0; ShapeIdx < InDimNum; ShapeIdx++)
 			{
-				Learning::DeserializeFromBytes(InOutOffset, Bytes, Shape[ShapeIdx]);
+				int32 Shape32 = 0;
+				Learning::DeserializeFromBytes(InOutOffset, Bytes, Shape32);
+				Shape[ShapeIdx] = Shape32;
 			}
 
 			OutArray.SetNumUninitialized(Shape);
 
-			UE_LEARNING_CHECK(InOutOffset + sizeof(InElementType) * OutArray.Num() <= Bytes.Num());
+			UE_LEARNING_CHECK(InOutOffset + (int64)sizeof(InElementType) * OutArray.Num() <= Bytes.Num());
 			FMemory::Memcpy((uint8*)OutArray.GetData(), &Bytes[InOutOffset],  sizeof(InElementType) * OutArray.Num());
 			InOutOffset += sizeof(InElementType) * OutArray.Num();
 		}
 
 		template<typename InElementType>
-		inline int32 IndexOf(const TLearningArrayView<1, const InElementType> Array, const InElementType& Element)
+		inline int64 IndexOf(const TLearningArrayView<1, const InElementType> Array, const InElementType& Element)
 		{
-			for (int32 Idx = 0; Idx < Array.Num(); Idx++)
+			for (int64 Idx = 0; Idx < Array.Num(); Idx++)
 			{
 				if (Array.GetData()[Idx] == Element)
 				{
@@ -1103,9 +1096,9 @@ namespace UE::Learning
 		}
 
 		template<typename InElementType>
-		inline int32 IndexOf(const TLearningArray<1, InElementType>& Array, const InElementType& Element)
+		inline int64 IndexOf(const TLearningArray<1, InElementType>& Array, const InElementType& Element)
 		{
-			for (int32 Idx = 0; Idx < Array.Num(); Idx++)
+			for (int64 Idx = 0; Idx < Array.Num(); Idx++)
 			{
 				if (Array.GetData()[Idx] == Element)
 				{
@@ -1118,13 +1111,13 @@ namespace UE::Learning
 		template<typename InElementType>
 		FString Format(const TLearningArrayView<1, const InElementType> Array, const TFunctionRef<FString(const InElementType&)> Formatter, const int32 MaxItemNum = 16)
 		{
-			const int32 ItemNum = Array.Num();
+			const int64 ItemNum = Array.Num();
 
 			FString Output = TEXT("[");
 
 			if (ItemNum <= MaxItemNum)
 			{
-				for (int32 Idx = 0; Idx < ItemNum; Idx++)
+				for (int64 Idx = 0; Idx < ItemNum; Idx++)
 				{
 					Output.Append(Formatter(Array[Idx]));
 					if (Idx != ItemNum - 1) { Output.Append(TEXT(" ")); }
@@ -1132,7 +1125,7 @@ namespace UE::Learning
 			}
 			else
 			{
-				for (int32 Idx = 0; Idx < MaxItemNum / 2; Idx++)
+				for (int64 Idx = 0; Idx < MaxItemNum / 2; Idx++)
 				{
 					Output.Append(Formatter(Array[Idx]));
 					Output.Append(TEXT(" "));
@@ -1140,7 +1133,7 @@ namespace UE::Learning
 
 				Output.Append(TEXT("... "));
 
-				for (int32 Idx = 0; Idx < MaxItemNum / 2; Idx++)
+				for (int64 Idx = 0; Idx < MaxItemNum / 2; Idx++)
 				{
 					Output.Append(Formatter(Array[ItemNum - Idx - 1]));
 					if (Idx != MaxItemNum / 2 - 1) { Output.Append(TEXT(" ")); }
diff --git a/Engine/Plugins/Experimental/LearningAgents/Source/Learning/Public/MultiArray.h b/Engine/Plugins/Experimental/LearningAgents/Source/Learning/Public/MultiArray.h
--- a/Engine/Plugins/Experimental/LearningAgents/Source/Learning/Public/MultiArray.h
+++ b/Engine/Plugins/Experimental/LearningAgents/Source/Learning/Public/MultiArray.h
@@ -54,21 +54,13 @@ namespace UE::MultiArray::Private
  * i.e. that they can be transparently moved to new memory without a copy constructor. 
  *
  **/
-template<uint8 InDimNum, typename InElementType, typename InAllocatorType = FDefaultAllocator, bool bInIsChecked = true, bool bInIsRestrict = false>
+template<uint8 InDimNum, typename InElementType, typename InAllocatorType = FDefaultAllocator>
 class TMultiArray
 {
 public:
 	static constexpr uint8 DimNum = InDimNum;
 	using ElementType = InElementType;
 	using AllocatorType = InAllocatorType;
-	static constexpr bool bIsChecked = bInIsChecked;
-	static constexpr bool bIsRestrict = bInIsRestrict;
-	using SizeType = typename InAllocatorType::SizeType;
-
-	using PointerType = std::conditional_t<
-		bIsRestrict,
-		ElementType* RESTRICT,
-		ElementType*>;
 
 	using ElementAllocatorType = std::conditional_t<
 		AllocatorType::NeedsElementType,
@@ -77,7 +69,6 @@ public:
 	>;
 
 	static_assert(DimNum > 1, "TMultiArray requires a positive, non-zero number of dimensions");
-	static_assert(std::is_signed_v<SizeType>, "TMultiArray only supports signed index types");
 
 public:
 
@@ -86,17 +77,17 @@ public:
 	 *
 	 * @returns Reference to indexed element.
 	 */
-	FORCEINLINE TMultiArrayView<DimNum - 1, ElementType, bIsChecked, bIsRestrict, SizeType> operator[](SizeType Index)
+	FORCEINLINE TMultiArrayView<DimNum - 1, ElementType> operator[](int64 Index)
 	{
 		RangeCheck(0, Index);
 
-		TMultiArrayShape<DimNum - 1, SizeType> NewShape;
+		TMultiArrayShape<DimNum - 1> NewShape;
 		for (uint8 Idx = 0; Idx < DimNum - 1; Idx++)
 		{
 			NewShape[Idx] = ArrayShape[Idx + 1];
 		}
 
-		return TMultiArrayView<DimNum - 1, ElementType, bIsChecked, bIsRestrict, SizeType>(GetData() + Index * Stride(0), NewShape);
+		return TMultiArrayView<DimNum - 1, ElementType>(GetData() + Index * Stride(0), NewShape);
 	}
 
 	/**
@@ -106,17 +97,17 @@ public:
 	 *
 	 * @returns Reference to indexed element.
 	 */
-	FORCEINLINE TMultiArrayView<DimNum - 1, const ElementType, bIsChecked, bIsRestrict, SizeType> operator[](SizeType Index) const
+	FORCEINLINE TMultiArrayView<DimNum - 1, const ElementType> operator[](int64 Index) const
 	{
 		RangeCheck(0, Index);
 
-		TMultiArrayShape<DimNum - 1, SizeType> NewShape;
+		TMultiArrayShape<DimNum - 1> NewShape;
 		for (uint8 Idx = 0; Idx < DimNum - 1; Idx++)
 		{
 			NewShape[Idx] = ArrayShape[Idx + 1];
 		}
 
-		return TMultiArrayView<DimNum - 1, const ElementType, bIsChecked, bIsRestrict, SizeType>(GetData() + Index * Stride(0), NewShape);
+		return TMultiArrayView<DimNum - 1, const ElementType>(GetData() + Index * Stride(0), NewShape);
 	}
 
 	/**
@@ -124,9 +115,9 @@ public:
 	 *
 	 * @returns Flattened array view.
 	 */
-	FORCEINLINE TMultiArrayView<1, ElementType, bIsChecked, bIsRestrict, SizeType> Flatten()
+	FORCEINLINE TMultiArrayView<1, ElementType> Flatten()
 	{
-		return TMultiArrayView<1, ElementType, bIsChecked, bIsRestrict, SizeType>(GetData(), Num());
+		return TMultiArrayView<1, ElementType>(GetData(), Num());
 	}
 
 	/**
@@ -134,9 +125,9 @@ public:
 	 *
 	 * @returns Flattened array view.
 	 */
-	FORCEINLINE TMultiArrayView<1, const ElementType, bIsChecked, bIsRestrict, SizeType> Flatten() const
+	FORCEINLINE TMultiArrayView<1, const ElementType> Flatten() const
 	{
-		return TMultiArrayView<1, const ElementType, bIsChecked, bIsRestrict, SizeType>(GetData(), Num());
+		return TMultiArrayView<1, const ElementType>(GetData(), Num());
 	}
 
 	/**
@@ -145,14 +136,11 @@ public:
 	 *
 	 * @returns View flattened on the given dimension.
 	 */
-	FORCEINLINE TMultiArrayView<DimNum - 1, ElementType, bIsChecked, bIsRestrict, SizeType> Flatten(uint8 Dimension)
+	FORCEINLINE TMultiArrayView<DimNum - 1, ElementType> Flatten(uint8 Dimension)
 	{
-		if (bIsChecked)
-		{
-			checkf((Dimension < DimNum - 1), TEXT("MultiArray flatten dimension out of bounds: %lld from a rank of %lld"), (long long)Dimension, (long long)DimNum); // & for one branch
-		}
+		checkf((Dimension < DimNum - 1), TEXT("MultiArray flatten dimension out of bounds: %lld from a rank of %lld"), (long long)Dimension, (long long)DimNum); // & for one branch
 
-		TMultiArrayShape<DimNum - 1, SizeType> NewShape;
+		TMultiArrayShape<DimNum - 1> NewShape;
 		uint8 SrcIdx = 0;
 		for (uint8 Idx = 0; Idx < DimNum - 1; Idx++)
 		{
@@ -168,7 +156,7 @@ public:
 			}
 		}
 
-		return TMultiArrayView<DimNum - 1, ElementType, bIsChecked, bIsRestrict, SizeType>(GetData(), NewShape);
+		return TMultiArrayView<DimNum - 1, ElementType>(GetData(), NewShape);
 	}
 
 	/**
@@ -177,14 +165,11 @@ public:
 	 *
 	 * @returns View flattened on the given dimension.
 	 */
-	FORCEINLINE TMultiArrayView<DimNum - 1, const ElementType, bIsChecked, bIsRestrict, SizeType> Flatten(uint8 Dimension) const
+	FORCEINLINE TMultiArrayView<DimNum - 1, const ElementType> Flatten(uint8 Dimension) const
 	{
-		if (bIsChecked)
-		{
-			checkf((Dimension < DimNum - 1), TEXT("MultiArray flatten dimension out of bounds: %lld from a rank of %lld"), (long long)Dimension, (long long)DimNum); // & for one branch
-		}
+		checkf((Dimension < DimNum - 1), TEXT("MultiArray flatten dimension out of bounds: %lld from a rank of %lld"), (long long)Dimension, (long long)DimNum); // & for one branch
 
-		TMultiArrayShape<DimNum - 1, SizeType> NewShape;
+		TMultiArrayShape<DimNum - 1> NewShape;
 		uint8 SrcIdx = 0;
 		for (uint8 Idx = 0; Idx < DimNum - 1; Idx++)
 		{
@@ -200,7 +185,7 @@ public:
 			}
 		}
 
-		return TMultiArrayView<DimNum - 1, const ElementType, bIsChecked, bIsRestrict, SizeType>(GetData(), NewShape);
+		return TMultiArrayView<DimNum - 1, const ElementType>(GetData(), NewShape);
 	}
 
 	/**
@@ -210,11 +195,11 @@ public:
 	 * @returns View flattened on the given dimension.
 	 */
 	template<uint8 Dimension>
-	FORCEINLINE TMultiArrayView<DimNum - 1, ElementType, bIsChecked, bIsRestrict, SizeType> Flatten()
+	FORCEINLINE TMultiArrayView<DimNum - 1, ElementType> Flatten()
 	{
 		static_assert(Dimension < DimNum - 1, "MultiArray flatten dimension out of bounds");
 
-		TMultiArrayShape<DimNum - 1, SizeType> NewShape;
+		TMultiArrayShape<DimNum - 1> NewShape;
 		uint8 SrcIdx = 0;
 		for (uint8 Idx = 0; Idx < DimNum - 1; Idx++)
 		{
@@ -230,7 +215,7 @@ public:
 			}
 		}
 
-		return TMultiArrayView<DimNum - 1, ElementType, bIsChecked, bIsRestrict, SizeType>(GetData(), NewShape);
+		return TMultiArrayView<DimNum - 1, ElementType>(GetData(), NewShape);
 	}
 
 	/**
@@ -240,11 +225,11 @@ public:
 	 * @returns View flattened on the given dimension.
 	 */
 	template<uint8 Dimension>
-	FORCEINLINE TMultiArrayView<DimNum - 1, const ElementType, bIsChecked, bIsRestrict, SizeType> Flatten() const
+	FORCEINLINE TMultiArrayView<DimNum - 1, const ElementType> Flatten() const
 	{
 		static_assert(Dimension < DimNum - 1, "MultiArray flatten dimension out of bounds");
 
-		TMultiArrayShape<DimNum - 1, SizeType> NewShape;
+		TMultiArrayShape<DimNum - 1> NewShape;
 		uint8 SrcIdx = 0;
 		for (uint8 Idx = 0; Idx < DimNum - 1; Idx++)
 		{
@@ -260,7 +245,7 @@ public:
 			}
 		}
 
-		return TMultiArrayView<DimNum - 1, const ElementType, bIsChecked, bIsRestrict, SizeType>(GetData(), NewShape);
+		return TMultiArrayView<DimNum - 1, const ElementType>(GetData(), NewShape);
 	}
 
 
@@ -277,15 +262,15 @@ public:
 	 * @param Ptr   A pointer to an array of elements to copy.
 	 * @param Nums The number of elements on each dimension.
 	 */
-	FORCEINLINE TMultiArray(const ElementType* Ptr, TMultiArrayShape<DimNum, SizeType> Shape)
+	FORCEINLINE TMultiArray(const ElementType* Ptr, TMultiArrayShape<DimNum> Shape)
 	{
 		check(Ptr != nullptr || Shape.Total() == 0);
 
 		CopyToEmpty(Ptr, Shape);
 	}
 
-	template <typename OtherElementType, bool bOtherIsChecked, bool bOtherIsRestrict, typename OtherSizeType>
-	FORCEINLINE explicit TMultiArray(const TMultiArrayView<DimNum, OtherElementType, bOtherIsChecked, bOtherIsRestrict, OtherSizeType>& Other)
+	template <typename OtherElementType>
+	FORCEINLINE explicit TMultiArray(const TMultiArrayView<DimNum, OtherElementType>& Other)
 		: TMultiArray(Other.GetData(), Other.Shape()) {}
 
 	/**
@@ -295,11 +280,10 @@ public:
 	 */
 	template <
 		typename OtherElementType,
-		typename OtherAllocator,
-		bool bOtherIsChecked
+		typename OtherAllocator
 		UE_REQUIRES(UE::MultiArray::Private::TMultiArrayElementsAreCompatible_V<ElementType, const OtherElementType&>)
 	>
-		FORCEINLINE explicit TMultiArray(const TMultiArray<DimNum, OtherElementType, OtherAllocator, bOtherIsChecked>& Other)
+		FORCEINLINE explicit TMultiArray(const TMultiArray<DimNum, OtherElementType, OtherAllocator>& Other)
 	{
 		CopyToEmpty(Other.GetData(), Other.Shape());
 	}
@@ -322,8 +306,8 @@ public:
 	 *
 	 * @param Other The source array to assign from.
 	 */
-	template<typename OtherAllocatorType, bool bOtherIsChecked>
-	TMultiArray& operator=(const TMultiArray<DimNum, ElementType, OtherAllocatorType, bOtherIsChecked>& Other)
+	template<typename OtherAllocatorType>
+	TMultiArray& operator=(const TMultiArray<DimNum, ElementType, OtherAllocatorType>& Other)
 	{
 		DestructItems(GetData(), Num());
 		CopyToEmpty(Other.GetData(), Other.Shape());
@@ -346,8 +330,8 @@ public:
 		return *this;
 	}
 
-	template <typename OtherElementType, bool bOtherIsChecked, bool bOtherIsRestrict, typename OtherSizeType>
-	TMultiArray& operator=(const TMultiArrayView<DimNum, OtherElementType, bOtherIsChecked, bOtherIsRestrict, OtherSizeType>& Other);
+	template <typename OtherElementType>
+	TMultiArray& operator=(const TMultiArrayView<DimNum, OtherElementType>& Other);
 
 private:
 
@@ -383,11 +367,10 @@ public:
 	 */
 	template <
 		typename OtherElementType,
-		typename OtherAllocator,
-		bool bOtherIsChecked
+		typename OtherAllocator
 		UE_REQUIRES(UE::MultiArray::Private::TMultiArrayElementsAreCompatible_V<ElementType, OtherElementType&&>)
 	>
-		FORCEINLINE explicit TMultiArray(TMultiArray<DimNum, OtherElementType, OtherAllocator, bOtherIsChecked>&& Other)
+		FORCEINLINE explicit TMultiArray(TMultiArray<DimNum, OtherElementType, OtherAllocator>&& Other)
 	{
 		MoveOrCopy(*this, Other);
 	}
@@ -421,9 +404,9 @@ public:
 	 *
 	 * @returns Pointer to first array entry or nullptr if empty
 	 */
-	FORCEINLINE PointerType GetData()
+	FORCEINLINE ElementType* GetData()
 	{
-		return (PointerType)AllocatorInstance.GetAllocation();
+		return (ElementType*)AllocatorInstance.GetAllocation();
 	}
 
 	/**
@@ -431,17 +414,24 @@ public:
 	 *
 	 * @returns Pointer to first array entry or nullptr if empty
 	 */
-	FORCEINLINE const PointerType GetData() const
+	FORCEINLINE const ElementType* GetData() const
 	{
-		return (const PointerType)AllocatorInstance.GetAllocation();
+		return (const ElementType*)AllocatorInstance.GetAllocation();
 	}
 
 	/**
 	 * DO NOT USE DIRECTLY
 	 * STL-like iterators to enable range-based for loop support.
 	 */
-	FORCEINLINE ElementType* begin() const { return GetData(); }
-	FORCEINLINE ElementType* end() const { return GetData() + Num(); }
+	FORCEINLINE ElementType* begin() { return GetData(); }
+	FORCEINLINE ElementType* end() { return GetData() + Num(); }
+
+	/**
+	 * DO NOT USE DIRECTLY
+	 * STL-like iterators to enable range-based for loop support.
+	 */
+	FORCEINLINE const ElementType* begin() const { return GetData(); }
+	FORCEINLINE const ElementType* end() const { return GetData() + Num(); }
 
 	/**
 	 * Helper function returning the size of the inner type.
@@ -459,12 +449,9 @@ public:
 	 */
 	FORCEINLINE void CheckInvariants() const
 	{
-		if (bIsChecked)
+		for (uint8 Idx = 0; Idx < DimNum; Idx++)
 		{
-			for (SizeType Idx = 0; Idx < DimNum; Idx++)
-			{
-				checkSlow(ArrayShape[Idx] >= 0);
-			}
+			checkSlow(ArrayShape[Idx] >= 0);
 		}
 	}
 
@@ -475,10 +462,7 @@ public:
 	 */
 	FORCEINLINE void DimensionCheck(uint8 Dimension) const
 	{
-		if (bIsChecked)
-		{
-			checkf((Dimension < DimNum), TEXT("MultiArray dimension out of bounds: %lld from a rank of %lld"), (long long)Dimension, (long long)DimNum); // & for one branch
-		}
+		checkf((Dimension < DimNum), TEXT("MultiArray dimension out of bounds: %lld from a rank of %lld"), (long long)Dimension, (long long)DimNum); // & for one branch
 	}
 
 	/**
@@ -487,15 +471,12 @@ public:
 	 * @param Dimension Dimension of the array.
 	 * @param Index Index to check.
 	 */
-	FORCEINLINE void RangeCheck(uint8 Dimension, SizeType Index) const
+	FORCEINLINE void RangeCheck(uint8 Dimension, int64 Index) const
 	{
 		DimensionCheck(Dimension);
 		CheckInvariants();
 
-		if (bIsChecked)
-		{
-			checkf((Index >= 0) & (Index < ArrayShape[Dimension]), TEXT("Array index out of bounds: %lld from a dimension of size %lld"), (long long)Index, (long long)ArrayShape[Dimension]); // & for one branch
-		}
+		checkf((Index >= 0) & (Index < ArrayShape[Dimension]), TEXT("Array index out of bounds: %lld from a dimension of size %lld"), (long long)Index, (long long)ArrayShape[Dimension]); // & for one branch
 	}
 
 	/**
@@ -506,16 +487,13 @@ public:
 	 * @param Index Starting index of the slice.
 	 * @param InNum Length of the slice.
 	 */
-	FORCEINLINE void SliceRangeCheck(uint8 Dimension, SizeType Index, SizeType InNum) const
+	FORCEINLINE void SliceRangeCheck(uint8 Dimension, int64 Index, int64 InNum) const
 	{
 		DimensionCheck(Dimension);
 
-		if (bIsChecked)
-		{
-			checkf(Index >= 0, TEXT("Invalid index (%lld)"), (long long)Index);
-			checkf(InNum >= 0, TEXT("Invalid count (%lld)"), (long long)InNum);
-			checkf(Index + InNum <= ArrayShape[Dimension], TEXT("Range (index: %lld, count: %lld) lies outside the view of %lld elements"), (long long)Index, (long long)InNum, (long long)ArrayShape[Dimension]);
-		}
+		checkf(Index >= 0, TEXT("Invalid index (%lld)"), (long long)Index);
+		checkf(InNum >= 0, TEXT("Invalid count (%lld)"), (long long)InNum);
+		checkf(Index + InNum <= ArrayShape[Dimension], TEXT("Range (index: %lld, count: %lld) lies outside the view of %lld elements"), (long long)Index, (long long)InNum, (long long)ArrayShape[Dimension]);
 	}
 
 	/**
@@ -548,7 +526,7 @@ public:
 	 *
 	 * @returns Number of dimensions in array.
 	 */
-	FORCEINLINE SizeType Rank() const
+	FORCEINLINE uint8 Rank() const
 	{
 		return DimNum;
 	}
@@ -558,7 +536,7 @@ public:
 	 *
 	 * @returns Total number of elements in array.
 	 */
-	FORCEINLINE SizeType Num() const
+	FORCEINLINE int64 Num() const
 	{
 		return ArrayShape.Total();
 	}
@@ -569,7 +547,7 @@ public:
 	 * @returns Number of elements in array.
 	 */
 	template<uint8 InDimIdx>
-	FORCEINLINE SizeType Num() const
+	FORCEINLINE int64 Num() const
 	{
 		static_assert(InDimIdx < DimNum);
 		return ArrayShape[InDimIdx];
@@ -581,7 +559,7 @@ public:
 	 * @param Dimension Dimension of the array.
 	 * @returns Number of elements in array.
 	 */
-	FORCEINLINE SizeType Num(uint8 Dimension) const
+	FORCEINLINE int64 Num(uint8 Dimension) const
 	{
 		DimensionCheck(Dimension);
 
@@ -593,7 +571,7 @@ public:
 	 *
 	 * @returns Number of elements in each dimension.
 	 */
-	FORCEINLINE TMultiArrayShape<DimNum, SizeType> Shape() const
+	FORCEINLINE TMultiArrayShape<DimNum> Shape() const
 	{
 		return ArrayShape;
 	}
@@ -604,11 +582,11 @@ public:
 	 * @param Dimension Dimension of the array.
 	 * @returns Stride of that dimension.
 	 */
-	FORCEINLINE SizeType Stride(uint8 Dimension) const
+	FORCEINLINE int64 Stride(uint8 Dimension) const
 	{
 		DimensionCheck(Dimension);
 
-		SizeType Total = 1;
+		int64 Total = 1;
 		for (uint8 Idx = Dimension + 1; Idx < DimNum; Idx++)
 		{
 			Total *= ArrayShape[Idx];
@@ -623,12 +601,12 @@ public:
 	 * @param InNum Number of elements in the new view
 	 * @returns Sliced view
 	 */
-	[[nodiscard]] FORCEINLINE TMultiArrayView<DimNum, ElementType, bIsChecked, bIsRestrict, SizeType> Slice(SizeType Index, SizeType InNum)
+	[[nodiscard]] FORCEINLINE TMultiArrayView<DimNum, ElementType> Slice(int64 Index, int64 InNum)
 	{
 		SliceRangeCheck(0, Index, InNum);
 		TMultiArrayShape NewShape = ArrayShape;
 		NewShape[0] = InNum;
-		return TMultiArrayView<DimNum, ElementType, bIsChecked, bIsRestrict, SizeType>(GetData() + Index * Stride(0), NewShape);
+		return TMultiArrayView<DimNum, ElementType>(GetData() + Index * Stride(0), NewShape);
 	}
 
 	/**
@@ -638,12 +616,12 @@ public:
 	 * @param InNum Number of elements in the new view
 	 * @returns Sliced view
 	 */
-	[[nodiscard]] FORCEINLINE TMultiArrayView<DimNum, const ElementType, bIsChecked, bIsRestrict, SizeType> Slice(SizeType Index, SizeType InNum) const
+	[[nodiscard]] FORCEINLINE TMultiArrayView<DimNum, const ElementType> Slice(int64 Index, int64 InNum) const
 	{
 		SliceRangeCheck(0, Index, InNum);
 		TMultiArrayShape NewShape = ArrayShape;
 		NewShape[0] = InNum;
-		return TMultiArrayView<DimNum, const ElementType, bIsChecked, bIsRestrict, SizeType>(GetData() + Index * Stride(0), NewShape);
+		return TMultiArrayView<DimNum, const ElementType>(GetData() + Index * Stride(0), NewShape);
 	}
 
 	/**
@@ -671,9 +649,9 @@ public:
 	 */
 	void Empty()
 	{
-		SizeType CurrTotal = Num();
+		int64 CurrTotal = Num();
 		DestructItems(GetData(), CurrTotal);
-		ArrayShape = TMultiArrayShape<DimNum, SizeType>();
+		ArrayShape = TMultiArrayShape<DimNum>();
 		AllocatorResizeAllocation(CurrTotal, 0);
 	}
 
@@ -682,10 +660,10 @@ public:
 	 *
 	 * @param NewShape New shape of the array.
 	 */
-	void SetNum(TMultiArrayShape<DimNum, SizeType> NewShape)
+	void SetNum(TMultiArrayShape<DimNum> NewShape)
 	{
-		SizeType CurrTotal = ArrayShape.Total();
-		SizeType NewTotal = NewShape.Total();
+		int64 CurrTotal = ArrayShape.Total();
+		int64 NewTotal = NewShape.Total();
 
 		if (NewTotal > CurrTotal)
 		{
@@ -706,10 +684,10 @@ public:
 	 *
 	 * @param NewShape New shape of the array.
 	 */
-	void SetNumZeroed(TMultiArrayShape<DimNum, SizeType> NewShape)
+	void SetNumZeroed(TMultiArrayShape<DimNum> NewShape)
 	{
-		SizeType CurrTotal = ArrayShape.Total();
-		SizeType NewTotal = NewShape.Total();
+		int64 CurrTotal = ArrayShape.Total();
+		int64 NewTotal = NewShape.Total();
 
 		if (NewTotal > CurrTotal)
 		{
@@ -730,10 +708,10 @@ public:
 	 *
 	 * @param NewShape New shape of the array.
 	 */
-	void SetNumUninitialized(TMultiArrayShape<DimNum, SizeType> NewShape)
+	void SetNumUninitialized(TMultiArrayShape<DimNum> NewShape)
 	{
-		SizeType CurrTotal = ArrayShape.Total();
-		SizeType NewTotal = NewShape.Total();
+		int64 CurrTotal = ArrayShape.Total();
+		int64 NewTotal = NewShape.Total();
 
 		if (NewTotal > CurrTotal)
 		{
@@ -749,7 +727,7 @@ public:
 	}
 
 private:
-	void AllocatorResizeAllocation(SizeType CurrentNum, SizeType NewNum)
+	void AllocatorResizeAllocation(int64 CurrentNum, int64 NewNum)
 	{
 		if constexpr (TAllocatorTraits<AllocatorType>::SupportsElementAlignment)
 		{
@@ -768,25 +746,21 @@ private:
 	 * @param OtherData A pointer to the data to copy
 	 * @param OtherNums The dimensions of the data to copy
 	 */
-	template <typename OtherElementType, typename OtherSizeType>
-	void CopyToEmpty(const OtherElementType* OtherData, TMultiArrayShape<DimNum, OtherSizeType> OtherShape)
+	template <typename OtherElementType>
+	void CopyToEmpty(const OtherElementType* OtherData, TMultiArrayShape<DimNum> OtherShape)
 	{
-		SizeType CurrTotal = ArrayShape.Total();
-		OtherSizeType OtherTotal = OtherShape.Total();
+		int64 CurrTotal = ArrayShape.Total();
+		int64 OtherTotal = OtherShape.Total();
 
 		for (uint8 Idx = 0; Idx < DimNum; Idx++)
 		{
-			SizeType NewNum = (SizeType)OtherShape[Idx];
-			
-			if (bIsChecked)
-			{
-				checkf((OtherSizeType)NewNum == OtherShape[Idx], TEXT("Invalid number of elements to add to this array type: %lld"), (long long)NewNum);
-			}
+			int64 NewNum = OtherShape[Idx];
+			checkf(NewNum == OtherShape[Idx], TEXT("Invalid number of elements to add to this array type: %lld"), (long long)NewNum);
 
 			ArrayShape[Idx] = NewNum;
 		}
 
-		SizeType NewTotal = ArrayShape.Total();
+		int64 NewTotal = ArrayShape.Total();
 
 		AllocatorResizeAllocation(CurrTotal, NewTotal);
 		ConstructItems<ElementType>(GetData(), OtherData, OtherTotal);
@@ -799,41 +773,33 @@ public:
 
 public:
 
-	FORCEINLINE operator TMultiArrayView<DimNum, ElementType, bIsChecked, bIsRestrict, SizeType>()
+	FORCEINLINE operator TMultiArrayView<DimNum, ElementType>()
 	{
-		return TMultiArrayView<DimNum, ElementType, bIsChecked, bIsRestrict, SizeType>(GetData(), ArrayShape);
+		return TMultiArrayView<DimNum, ElementType>(GetData(), ArrayShape);
 	}
 
-	FORCEINLINE operator TMultiArrayView<DimNum, const ElementType, bIsChecked, bIsRestrict, SizeType>() const
+	FORCEINLINE operator TMultiArrayView<DimNum, const ElementType>() const
 	{
-		return TMultiArrayView<DimNum, const ElementType, bIsChecked, bIsRestrict, SizeType>(GetData(), ArrayShape);
+		return TMultiArrayView<DimNum, const ElementType>(GetData(), ArrayShape);
 	}
 
 protected:
 
 	ElementAllocatorType AllocatorInstance;
-	TMultiArrayShape<DimNum, SizeType> ArrayShape;
+	TMultiArrayShape<DimNum> ArrayShape;
 };
 
 
 /**
 * Specialization for single dimensional MultiArray
 */
-template<typename InElementType, typename InAllocatorType, bool bInIsChecked, bool bInIsRestrict>
-class TMultiArray<1, InElementType, InAllocatorType, bInIsChecked, bInIsRestrict>
+template<typename InElementType, typename InAllocatorType>
+class TMultiArray<1, InElementType, InAllocatorType>
 {
 public:
 	static constexpr uint8 DimNum = 1;
 	using ElementType = InElementType;
 	using AllocatorType = InAllocatorType;
-	static constexpr bool bIsChecked = bInIsChecked;
-	static constexpr bool bIsRestrict = bInIsRestrict;
-	using SizeType = typename InAllocatorType::SizeType;
-
-	using PointerType = std::conditional_t<
-		bIsRestrict,
-		ElementType* RESTRICT,
-		ElementType*>;
 
 	using ElementAllocatorType = std::conditional_t<
 		AllocatorType::NeedsElementType,
@@ -841,8 +807,6 @@ public:
 		typename AllocatorType::ForAnyElementType
 	>;
 
-	static_assert(std::is_signed_v<SizeType>, "TMultiArray only supports signed index types");
-
 	/**
 	 * Initializer list constructor
 	 */
@@ -851,7 +815,7 @@ public:
 		// This is not strictly legal, as std::initializer_list's iterators are not guaranteed to be pointers, but
 		// this appears to be the case on all of our implementations.  Also, if it's not true on a new implementation,
 		// it will fail to compile rather than behave badly.
-		CopyToEmpty<InElementType, SizeType>(InitList.begin(), { (SizeType)InitList.size() });
+		CopyToEmpty<InElementType>(InitList.begin(), { (int64)InitList.size() });
 	}
 
 	/**
@@ -866,7 +830,7 @@ public:
 		// This is not strictly legal, as std::initializer_list's iterators are not guaranteed to be pointers, but
 		// this appears to be the case on all of our implementations.  Also, if it's not true on a new implementation,
 		// it will fail to compile rather than behave badly.
-		CopyToEmpty<InElementType, SizeType>(InitList.begin(), { (SizeType)InitList.size() });
+		CopyToEmpty<InElementType>(InitList.begin(), { (int64)InitList.size() });
 		return *this;
 	}
 
@@ -875,7 +839,7 @@ public:
 	 *
 	 * @returns Reference to indexed element.
 	 */
-	FORCEINLINE ElementType& operator[](SizeType Index)
+	FORCEINLINE ElementType& operator[](int64 Index)
 	{
 		RangeCheck(0, Index);
 		return GetData()[Index];
@@ -888,7 +852,7 @@ public:
 	 *
 	 * @returns Reference to indexed element.
 	 */
-	FORCEINLINE const ElementType& operator[](SizeType Index) const
+	FORCEINLINE const ElementType& operator[](int64 Index) const
 	{
 		RangeCheck(0, Index);
 		return GetData()[Index];
@@ -899,9 +863,9 @@ public:
 	 *
 	 * @returns Flattened array view.
 	 */
-	FORCEINLINE TMultiArrayView<1, ElementType, bIsChecked, bIsRestrict, SizeType> Flatten()
+	FORCEINLINE TMultiArrayView<1, ElementType> Flatten()
 	{
-		return TMultiArrayView<1, ElementType, bIsChecked, bIsRestrict, SizeType>(GetData(), Num());
+		return TMultiArrayView<1, ElementType>(GetData(), Num());
 	}
 
 	/**
@@ -909,33 +873,33 @@ public:
 	 *
 	 * @returns Flattened array view.
 	 */
-	FORCEINLINE TMultiArrayView<1, const ElementType, bIsChecked, bIsRestrict, SizeType> Flatten() const
+	FORCEINLINE TMultiArrayView<1, const ElementType> Flatten() const
 	{
-		return TMultiArrayView<1, const ElementType, bIsChecked, bIsRestrict, SizeType>(GetData(), Num());
+		return TMultiArrayView<1, const ElementType>(GetData(), Num());
 	}
 
 	/** Explicit conversion to TArrayView. */
-	FORCEINLINE TArrayView<ElementType, SizeType> ArrayView()
+	FORCEINLINE TArrayView<ElementType> ArrayView()
 	{
-		return TArrayView<ElementType, SizeType>(GetData(), Num());
+		return TArrayView<ElementType>(GetData(), Num());
 	}
 
 	/** Explicit conversion to TArrayView. */
-	FORCEINLINE TArrayView<const ElementType, SizeType> ArrayView() const
+	FORCEINLINE TArrayView<const ElementType> ArrayView() const
 	{
-		return TArrayView<const ElementType, SizeType>(GetData(), Num());
+		return TArrayView<const ElementType>(GetData(), Num());
 	}
 
 	/** Implicit cast to TArrayView. */
-	FORCEINLINE operator TArrayView<ElementType, SizeType>()
+	FORCEINLINE operator TArrayView<ElementType>()
 	{
-		return TArrayView<ElementType, SizeType>(GetData(), Num());
+		return TArrayView<ElementType>(GetData(), Num());
 	}
 
 	/** Implicit cast to TArrayView. */
-	FORCEINLINE operator TArrayView<const ElementType, SizeType>() const
+	FORCEINLINE operator TArrayView<const ElementType>() const
 	{
-		return TArrayView<const ElementType, SizeType>(GetData(), Num());
+		return TArrayView<const ElementType>(GetData(), Num());
 	}
 
 public:
@@ -951,15 +915,15 @@ public:
 	 * @param Ptr   A pointer to an array of elements to copy.
 	 * @param Nums The number of elements on each dimension.
 	 */
-	FORCEINLINE TMultiArray(const ElementType* Ptr, TMultiArrayShape<DimNum, SizeType> Shape)
+	FORCEINLINE TMultiArray(const ElementType* Ptr, TMultiArrayShape<DimNum> Shape)
 	{
 		check(Ptr != nullptr || Shape.Total() == 0);
 
 		CopyToEmpty(Ptr, Shape);
 	}
 
-	template <typename OtherElementType, bool bOtherIsChecked, bool bOtherIsRestrict, typename OtherSizeType>
-	explicit TMultiArray(const TMultiArrayView<DimNum, OtherElementType, bOtherIsChecked, bOtherIsRestrict, OtherSizeType>& Other)
+	template <typename OtherElementType>
+	explicit TMultiArray(const TMultiArrayView<DimNum, OtherElementType>& Other)
 		: TMultiArray(Other.GetData(), Other.Shape()) {}
 
 	/**
@@ -969,11 +933,10 @@ public:
 	 */
 	template <
 		typename OtherElementType,
-		typename OtherAllocator,
-		bool bOtherIsChecked
+		typename OtherAllocator
 		UE_REQUIRES(UE::MultiArray::Private::TMultiArrayElementsAreCompatible_V<ElementType, const OtherElementType&>)
 	>
-		FORCEINLINE explicit TMultiArray(const TMultiArray<DimNum, OtherElementType, OtherAllocator, bOtherIsChecked>& Other)
+		FORCEINLINE explicit TMultiArray(const TMultiArray<DimNum, OtherElementType, OtherAllocator>& Other)
 	{
 		CopyToEmpty(Other.GetData(), Other.Shape());
 	}
@@ -996,8 +959,8 @@ public:
 	 *
 	 * @param Other The source array to assign from.
 	 */
-	template<typename OtherAllocatorType, bool bOtherIsChecked>
-	TMultiArray& operator=(const TMultiArray<DimNum, ElementType, OtherAllocatorType, bOtherIsChecked>& Other)
+	template<typename OtherAllocatorType>
+	TMultiArray& operator=(const TMultiArray<DimNum, ElementType, OtherAllocatorType>& Other)
 	{
 		DestructItems(GetData(), Num());
 		CopyToEmpty(Other.GetData(), Other.Shape());
@@ -1020,8 +983,8 @@ public:
 		return *this;
 	}
 
-	template <typename OtherElementType, bool bOtherIsChecked, bool bOtherIsRestrict, typename OtherSizeType>
-	TMultiArray& operator=(const TMultiArrayView<DimNum, OtherElementType, bOtherIsChecked, bOtherIsRestrict, OtherSizeType>& Other);
+	template <typename OtherElementType>
+	TMultiArray& operator=(const TMultiArrayView<DimNum, OtherElementType>& Other);
 
 private:
 
@@ -1057,11 +1020,10 @@ public:
 	 */
 	template <
 		typename OtherElementType,
-		typename OtherAllocator,
-		bool bOtherIsChecked
+		typename OtherAllocator
 		UE_REQUIRES(UE::MultiArray::Private::TMultiArrayElementsAreCompatible_V<ElementType, OtherElementType&&>)
 	>
-		FORCEINLINE explicit TMultiArray(TMultiArray<DimNum, OtherElementType, OtherAllocator, bOtherIsChecked>&& Other)
+		FORCEINLINE explicit TMultiArray(TMultiArray<DimNum, OtherElementType, OtherAllocator>&& Other)
 	{
 		MoveOrCopy(*this, Other);
 	}
@@ -1095,9 +1057,9 @@ public:
 	 *
 	 * @returns Pointer to first array entry or nullptr if empty
 	 */
-	FORCEINLINE PointerType GetData()
+	FORCEINLINE ElementType* GetData()
 	{
-		return (PointerType)AllocatorInstance.GetAllocation();
+		return (ElementType*)AllocatorInstance.GetAllocation();
 	}
 
 	/**
@@ -1105,17 +1067,25 @@ public:
 	 *
 	 * @returns Pointer to first array entry or nullptr if empty
 	 */
-	FORCEINLINE const PointerType GetData() const
+	FORCEINLINE const ElementType* GetData() const
 	{
-		return (const PointerType)AllocatorInstance.GetAllocation();
+		return (const ElementType*)AllocatorInstance.GetAllocation();
 	}
 
 	/**
 	 * DO NOT USE DIRECTLY
 	 * STL-like iterators to enable range-based for loop support.
 	 */
-	FORCEINLINE ElementType* begin() const { return GetData(); }
-	FORCEINLINE ElementType* end() const { return GetData() + Num(); }
+	FORCEINLINE ElementType* begin() { return GetData(); }
+	FORCEINLINE ElementType* end() { return GetData() + Num(); }
+
+
+	/**
+	 * DO NOT USE DIRECTLY
+	 * STL-like iterators to enable range-based for loop support.
+	 */
+	FORCEINLINE const ElementType* begin() const { return GetData(); }
+	FORCEINLINE const ElementType* end() const { return GetData() + Num(); }
 
 	/**
 	 * Helper function returning the size of the inner type.
@@ -1133,12 +1103,9 @@ public:
 	 */
 	FORCEINLINE void CheckInvariants() const
 	{
-		if (bIsChecked)
+		for (uint8 Idx = 0; Idx < DimNum; Idx++)
 		{
-			for (SizeType Idx = 0; Idx < DimNum; Idx++)
-			{
-				checkSlow(ArrayShape[Idx] >= 0);
-			}
+			checkSlow(ArrayShape[Idx] >= 0);
 		}
 	}
 
@@ -1149,10 +1116,7 @@ public:
 	 */
 	FORCEINLINE void DimensionCheck(uint8 Dimension) const
 	{
-		if (bIsChecked)
-		{
-			checkf((Dimension < DimNum), TEXT("MultiArray dimension out of bounds: %lld from a rank of %lld"), (long long)Dimension, (long long)DimNum); // & for one branch
-		}
+		checkf((Dimension < DimNum), TEXT("MultiArray dimension out of bounds: %lld from a rank of %lld"), (long long)Dimension, (long long)DimNum); // & for one branch
 	}
 
 	/**
@@ -1161,15 +1125,12 @@ public:
 	 * @param Dimension Dimension of the array.
 	 * @param Index Index to check.
 	 */
-	FORCEINLINE void RangeCheck(uint8 Dimension, SizeType Index) const
+	FORCEINLINE void RangeCheck(uint8 Dimension, int64 Index) const
 	{
 		DimensionCheck(Dimension);
 		CheckInvariants();
 
-		if (bIsChecked)
-		{
-			checkf((Index >= 0) & (Index < ArrayShape[Dimension]), TEXT("Array index out of bounds: %lld from a dimension of size %lld"), (long long)Index, (long long)ArrayShape[Dimension]); // & for one branch
-		}
+		checkf((Index >= 0) & (Index < ArrayShape[Dimension]), TEXT("Array index out of bounds: %lld from a dimension of size %lld"), (long long)Index, (long long)ArrayShape[Dimension]); // & for one branch
 	}
 
 	/**
@@ -1180,16 +1141,13 @@ public:
 	 * @param Index Starting index of the slice.
 	 * @param InNum Length of the slice.
 	 */
-	FORCEINLINE void SliceRangeCheck(uint8 Dimension, SizeType Index, SizeType InNum) const
+	FORCEINLINE void SliceRangeCheck(uint8 Dimension, int64 Index, int64 InNum) const
 	{
 		DimensionCheck(Dimension);
 
-		if (bIsChecked)
-		{
-			checkf(Index >= 0, TEXT("Invalid index (%lld)"), (long long)Index);
-			checkf(InNum >= 0, TEXT("Invalid count (%lld)"), (long long)InNum);
-			checkf(Index + InNum <= ArrayShape[Dimension], TEXT("Range (index: %lld, count: %lld) lies outside the view of %lld elements"), (long long)Index, (long long)InNum, (long long)ArrayShape[Dimension]);
-		}
+		checkf(Index >= 0, TEXT("Invalid index (%lld)"), (long long)Index);
+		checkf(InNum >= 0, TEXT("Invalid count (%lld)"), (long long)InNum);
+		checkf(Index + InNum <= ArrayShape[Dimension], TEXT("Range (index: %lld, count: %lld) lies outside the view of %lld elements"), (long long)Index, (long long)InNum, (long long)ArrayShape[Dimension]);
 	}
 
 	/**
@@ -1222,7 +1180,7 @@ public:
 	 *
 	 * @returns Number of dimensions in array.
 	 */
-	FORCEINLINE SizeType Rank() const
+	FORCEINLINE uint8 Rank() const
 	{
 		return DimNum;
 	}
@@ -1232,7 +1190,7 @@ public:
 	 *
 	 * @returns Total number of elements in array.
 	 */
-	FORCEINLINE SizeType Num() const
+	FORCEINLINE int64 Num() const
 	{
 		return ArrayShape.Total();
 	}
@@ -1243,7 +1201,7 @@ public:
 	 * @returns Number of elements in array.
 	 */
 	template<uint8 InDimIdx>
-	FORCEINLINE SizeType Num() const
+	FORCEINLINE int64 Num() const
 	{
 		static_assert(InDimIdx < DimNum);
 		return ArrayShape[InDimIdx];
@@ -1255,7 +1213,7 @@ public:
 	 * @param Dimension Dimension of the array.
 	 * @returns Number of elements in array.
 	 */
-	FORCEINLINE SizeType Num(uint8 Dimension) const
+	FORCEINLINE int64 Num(uint8 Dimension) const
 	{
 		DimensionCheck(Dimension);
 
@@ -1267,7 +1225,7 @@ public:
 	 *
 	 * @returns Number of elements in each dimension.
 	 */
-	FORCEINLINE TMultiArrayShape<DimNum, SizeType> Shape() const
+	FORCEINLINE TMultiArrayShape<DimNum> Shape() const
 	{
 		return ArrayShape;
 	}
@@ -1278,11 +1236,11 @@ public:
 	 * @param Dimension Dimension of the array.
 	 * @returns Stride of that dimension.
 	 */
-	FORCEINLINE SizeType Stride(uint8 Dimension) const
+	FORCEINLINE int64 Stride(uint8 Dimension) const
 	{
 		DimensionCheck(Dimension);
 
-		SizeType Total = 1;
+		int64 Total = 1;
 		for (uint8 Idx = Dimension + 1; Idx < DimNum; Idx++)
 		{
 			Total *= ArrayShape[Idx];
@@ -1297,12 +1255,12 @@ public:
 	 * @param InNum Number of elements in the new view
 	 * @returns Sliced view
 	 */
-	[[nodiscard]] FORCEINLINE TMultiArrayView<DimNum, ElementType, bIsChecked, bIsRestrict, SizeType> Slice(SizeType Index, SizeType InNum)
+	[[nodiscard]] FORCEINLINE TMultiArrayView<DimNum, ElementType> Slice(int64 Index, int64 InNum)
 	{
 		SliceRangeCheck(0, Index, InNum);
 		TMultiArrayShape NewShape = ArrayShape;
 		NewShape[0] = InNum;
-		return TMultiArrayView<DimNum, ElementType, bIsChecked, bIsRestrict, SizeType>(GetData() + Index * Stride(0), NewShape);
+		return TMultiArrayView<DimNum, ElementType>(GetData() + Index * Stride(0), NewShape);
 	}
 
 	/**
@@ -1312,12 +1270,12 @@ public:
 	 * @param InNum Number of elements in the new view
 	 * @returns Sliced view
 	 */
-	[[nodiscard]] FORCEINLINE TMultiArrayView<DimNum, const ElementType, bIsChecked, bIsRestrict, SizeType> Slice(SizeType Index, SizeType InNum) const
+	[[nodiscard]] FORCEINLINE TMultiArrayView<DimNum, const ElementType> Slice(int64 Index, int64 InNum) const
 	{
 		SliceRangeCheck(0, Index, InNum);
 		TMultiArrayShape NewShape = ArrayShape;
 		NewShape[0] = InNum;
-		return TMultiArrayView<DimNum, const ElementType, bIsChecked, bIsRestrict, SizeType>(GetData() + Index * Stride(0), NewShape);
+		return TMultiArrayView<DimNum, const ElementType>(GetData() + Index * Stride(0), NewShape);
 	}
 
 	/**
@@ -1345,9 +1303,9 @@ public:
 	 */
 	void Empty()
 	{
-		SizeType CurrTotal = Num();
+		int64 CurrTotal = Num();
 		DestructItems(GetData(), CurrTotal);
-		ArrayShape = TMultiArrayShape<DimNum, SizeType>();
+		ArrayShape = TMultiArrayShape<DimNum>();
 		AllocatorResizeAllocation(CurrTotal, 0);
 	}
 
@@ -1356,10 +1314,10 @@ public:
 	 *
 	 * @param NewShape New shape of the array.
 	 */
-	void SetNum(TMultiArrayShape<DimNum, SizeType> NewShape)
+	void SetNum(TMultiArrayShape<DimNum> NewShape)
 	{
-		SizeType CurrTotal = ArrayShape.Total();
-		SizeType NewTotal = NewShape.Total();
+		int64 CurrTotal = ArrayShape.Total();
+		int64 NewTotal = NewShape.Total();
 
 		if (NewTotal > CurrTotal)
 		{
@@ -1380,10 +1338,10 @@ public:
 	 *
 	 * @param NewShape New shape of the array.
 	 */
-	void SetNumZeroed(TMultiArrayShape<DimNum, SizeType> NewShape)
+	void SetNumZeroed(TMultiArrayShape<DimNum> NewShape)
 	{
-		SizeType CurrTotal = ArrayShape.Total();
-		SizeType NewTotal = NewShape.Total();
+		int64 CurrTotal = ArrayShape.Total();
+		int64 NewTotal = NewShape.Total();
 
 		if (NewTotal > CurrTotal)
 		{
@@ -1404,10 +1362,10 @@ public:
 	 *
 	 * @param NewShape New shape of the array.
 	 */
-	void SetNumUninitialized(TMultiArrayShape<DimNum, SizeType> NewShape)
+	void SetNumUninitialized(TMultiArrayShape<DimNum> NewShape)
 	{
-		SizeType CurrTotal = ArrayShape.Total();
-		SizeType NewTotal = NewShape.Total();
+		int64 CurrTotal = ArrayShape.Total();
+		int64 NewTotal = NewShape.Total();
 
 		if (NewTotal > CurrTotal)
 		{
@@ -1423,7 +1381,7 @@ public:
 	}
 
 private:
-	void AllocatorResizeAllocation(SizeType CurrentNum, SizeType NewNum)
+	void AllocatorResizeAllocation(int64 CurrentNum, int64 NewNum)
 	{
 		if constexpr (TAllocatorTraits<AllocatorType>::SupportsElementAlignment)
 		{
@@ -1442,25 +1400,22 @@ private:
 	 * @param OtherData A pointer to the data to copy
 	 * @param OtherNums The dimensions of the data to copy
 	 */
-	template <typename OtherElementType, typename OtherSizeType>
-	void CopyToEmpty(const OtherElementType* OtherData, TMultiArrayShape<DimNum, OtherSizeType> OtherShape)
+	template <typename OtherElementType>
+	void CopyToEmpty(const OtherElementType* OtherData, TMultiArrayShape<DimNum> OtherShape)
 	{
-		SizeType CurrTotal = ArrayShape.Total();
-		OtherSizeType OtherTotal = OtherShape.Total();
+		int64 CurrTotal = ArrayShape.Total();
+		int64 OtherTotal = OtherShape.Total();
 
 		for (uint8 Idx = 0; Idx < DimNum; Idx++)
 		{
-			SizeType NewNum = (SizeType)OtherShape[Idx];
+			int64 NewNum = OtherShape[Idx];
 
-			if (bIsChecked)
-			{
-				checkf((OtherSizeType)NewNum == OtherShape[Idx], TEXT("Invalid number of elements to add to this array type: %lld"), (long long)NewNum);
-			}
+			checkf(NewNum == OtherShape[Idx], TEXT("Invalid number of elements to add to this array type: %lld"), (long long)NewNum);
 
 			ArrayShape[Idx] = NewNum;
 		}
 
-		SizeType NewTotal = ArrayShape.Total();
+		int64 NewTotal = ArrayShape.Total();
 
 		AllocatorResizeAllocation(CurrTotal, NewTotal);
 		ConstructItems<ElementType>(GetData(), OtherData, OtherTotal);
@@ -1473,18 +1428,18 @@ public:
 
 public:
 
-	FORCEINLINE operator TMultiArrayView<DimNum, ElementType, bIsChecked, bIsRestrict, SizeType>()
+	FORCEINLINE operator TMultiArrayView<DimNum, ElementType>()
 	{
-		return TMultiArrayView<DimNum, ElementType, bIsChecked, bIsRestrict, SizeType>(GetData(), ArrayShape);
+		return TMultiArrayView<DimNum, ElementType>(GetData(), ArrayShape);
 	}
 
-	FORCEINLINE operator TMultiArrayView<DimNum, const ElementType, bIsChecked, bIsRestrict, SizeType>() const
+	FORCEINLINE operator TMultiArrayView<DimNum, const ElementType>() const
 	{
-		return TMultiArrayView<DimNum, const ElementType, bIsChecked, bIsRestrict, SizeType>(GetData(), ArrayShape);
+		return TMultiArrayView<DimNum, const ElementType>(GetData(), ArrayShape);
 	}
 
 protected:
 
 	ElementAllocatorType AllocatorInstance;
-	TMultiArrayShape<DimNum, SizeType> ArrayShape;
+	TMultiArrayShape<DimNum> ArrayShape;
 };
diff --git a/Engine/Plugins/Experimental/LearningAgents/Source/Learning/Public/MultiArrayView.h b/Engine/Plugins/Experimental/LearningAgents/Source/Learning/Public/MultiArrayView.h
--- a/Engine/Plugins/Experimental/LearningAgents/Source/Learning/Public/MultiArrayView.h
+++ b/Engine/Plugins/Experimental/LearningAgents/Source/Learning/Public/MultiArrayView.h
@@ -91,33 +91,31 @@ namespace UE::MultiArrayView::Private
 	};
 }
 
-template<uint8 InDimNum, typename InSizeType = FDefaultAllocator::SizeType>
+template<uint8 InDimNum>
 class TMultiArrayShape
 {
 public:
 	static constexpr uint8 DimNum = InDimNum;
-	using SizeType = InSizeType;
 
 	static_assert(DimNum >= 1, "TMultiArrayShape requires a positive, non-zero number of dimensions");
-	static_assert(std::is_signed_v<SizeType>, "TMultiArrayShape only supports signed index types");
 
 public:
 
 	TMultiArrayShape() = default;
 
-	TMultiArrayShape(std::initializer_list<SizeType> InNums)
+	TMultiArrayShape(std::initializer_list<int64> InNums)
 	{
 		checkf(InNums.size() == DimNum, TEXT("Wrong number of elements in constructor"));
 
 		uint8 Idx = 0;
-		for (SizeType InNum : InNums)
+		for (int64 InNum : InNums)
 		{
 			Nums[Idx] = InNum;
 			Idx++;
 		}
 	}
 
-	TMultiArrayShape(const SizeType* InNums)
+	TMultiArrayShape(const int64* InNums)
 	{
 		for (uint8 Idx = 0; Idx < DimNum; Idx++)
 		{
@@ -125,25 +123,25 @@ public:
 		}
 	}
 
-	SizeType& operator[](uint8 Dimension)
+	int64& operator[](uint8 Dimension)
 	{
 		return Nums[Dimension];
 	}
 
-	const SizeType& operator[](uint8 Dimension) const
+	const int64& operator[](uint8 Dimension) const
 	{
 		return Nums[Dimension];
 	}
 
-	SIZE_T Total() const
+	int64 Total() const
 	{
 		if constexpr (DimNum == 1)
 		{
-			return (SIZE_T)Nums[0];
+			return Nums[0];
 		}
 		else
 		{
-			SIZE_T Total = (SIZE_T)Nums[0];
+			int64 Total = Nums[0];
 			for (uint8 Idx = 1; Idx < DimNum; Idx++)
 			{
 				Total *= Nums[Idx];
@@ -153,7 +151,7 @@ public:
 	}
 
 private:
-	SizeType Nums[DimNum] = { 0 };
+	int64 Nums[DimNum] = { 0 };
 };
 
 /**
@@ -163,23 +161,14 @@ private:
  *
  * @see TArrayView
  */
-template<uint8 InDimNum, typename InElementType, bool bInIsChecked = true, bool bInIsRestrict = false, typename InSizeType = FDefaultAllocator::SizeType>
+template<uint8 InDimNum, typename InElementType>
 class TMultiArrayView
 {
 public:
 	static constexpr uint8 DimNum = InDimNum;
 	using ElementType = InElementType;
-	static constexpr bool bIsChecked = bInIsChecked;
-	static constexpr bool bIsRestrict = bInIsRestrict;
-	using SizeType = InSizeType;
-
-	using PointerType = std::conditional_t<
-		bIsRestrict,
-		ElementType* RESTRICT,
-		ElementType*>;
 
 	static_assert(DimNum > 1, "TMultiArrayView requires a positive, non-zero number of dimensions");
-	static_assert(std::is_signed_v<SizeType>, "TMultiArrayView only supports signed index types");
 
 public:
 
@@ -194,7 +183,7 @@ public:
 	 * @param InData	The data to view
 	 * @param InShape	The number of elements on each dimension
 	 */
-	FORCEINLINE TMultiArrayView(ElementType* InData, TMultiArrayShape<DimNum, SizeType> InShape)
+	FORCEINLINE TMultiArrayView(ElementType* InData, TMultiArrayShape<DimNum> InShape)
 		: DataPtr(InData)
 		, ArrayShape(InShape)
 	{
@@ -206,17 +195,17 @@ public:
 	 *
 	 * @returns Reference to indexed element.
 	 */
-	FORCEINLINE TMultiArrayView<DimNum - 1, ElementType, bIsChecked, bIsRestrict, SizeType> operator[](SizeType Index) const
+	FORCEINLINE TMultiArrayView<DimNum - 1, ElementType> operator[](int64 Index) const
 	{
 		RangeCheck(0, Index);
 
-		TMultiArrayShape<DimNum - 1, SizeType> NewShape;
+		TMultiArrayShape<DimNum - 1> NewShape;
 		for (uint8 Idx = 0; Idx < DimNum - 1; Idx++)
 		{
 			NewShape[Idx] = ArrayShape[Idx + 1];
 		}
 
-		return TMultiArrayView<DimNum - 1, ElementType, bIsChecked, bIsRestrict, SizeType>(DataPtr + Index * Stride(0), NewShape);
+		return TMultiArrayView<DimNum - 1, ElementType>(DataPtr + Index * Stride(0), NewShape);
 	}
 
 	/**
@@ -224,9 +213,9 @@ public:
 	 *
 	 * @returns Flattened array view.
 	 */
-	FORCEINLINE TMultiArrayView<1, ElementType, bIsChecked, bIsRestrict, SizeType> Flatten() const
+	FORCEINLINE TMultiArrayView<1, ElementType> Flatten() const
 	{
-		return TMultiArrayView<1, ElementType, bIsChecked, bIsRestrict, SizeType>(DataPtr, Num());
+		return TMultiArrayView<1, ElementType>(DataPtr, Num());
 	}
 
 	/**
@@ -235,14 +224,11 @@ public:
 	 *
 	 * @returns View flattened on the given dimension.
 	 */
-	FORCEINLINE TMultiArrayView<DimNum - 1, ElementType, bIsChecked, bIsRestrict, SizeType> Flatten(uint8 Dimension) const
+	FORCEINLINE TMultiArrayView<DimNum - 1, ElementType> Flatten(uint8 Dimension) const
 	{
-		if (bIsChecked)
-		{
-			checkf((Dimension < DimNum - 1), TEXT("MultiArray flatten dimension out of bounds: %lld from a rank of %lld"), (long long)Dimension, (long long)DimNum); // & for one branch
-		}
+		checkf((Dimension < DimNum - 1), TEXT("MultiArray flatten dimension out of bounds: %lld from a rank of %lld"), (long long)Dimension, (long long)DimNum); // & for one branch
 
-		TMultiArrayShape<DimNum - 1, SizeType> NewShape;
+		TMultiArrayShape<DimNum - 1> NewShape;
 		uint8 SrcIdx = 0;
 		for (uint8 Idx = 0; Idx < DimNum - 1; Idx++)
 		{
@@ -258,7 +244,7 @@ public:
 			}
 		}
 
-		return TMultiArrayView<DimNum - 1, ElementType, bIsChecked, bIsRestrict, SizeType>(DataPtr, NewShape);
+		return TMultiArrayView<DimNum - 1, ElementType>(DataPtr, NewShape);
 	}
 
 	/**
@@ -268,11 +254,11 @@ public:
 	 * @returns View flattened on the given dimension.
 	 */
 	template<uint8 Dimension>
-	FORCEINLINE TMultiArrayView<DimNum - 1, ElementType, bIsChecked, bIsRestrict, SizeType> Flatten() const
+	FORCEINLINE TMultiArrayView<DimNum - 1, ElementType> Flatten() const
 	{
 		static_assert(Dimension < DimNum - 1, "MultiArray flatten dimension out of bounds");
 
-		TMultiArrayShape<DimNum - 1, SizeType> NewShape;
+		TMultiArrayShape<DimNum - 1> NewShape;
 		uint8 SrcIdx = 0;
 		for (uint8 Idx = 0; Idx < DimNum - 1; Idx++)
 		{
@@ -288,7 +274,7 @@ public:
 			}
 		}
 
-		return TMultiArrayView<DimNum - 1, ElementType, bIsChecked, bIsRestrict, SizeType>(DataPtr, NewShape);
+		return TMultiArrayView<DimNum - 1, ElementType>(DataPtr, NewShape);
 	}
 
 public:
@@ -298,17 +284,7 @@ public:
 	 *
 	 * @returns Pointer to first array entry or nullptr if empty
 	 */
-	FORCEINLINE PointerType GetData()
-	{
-		return DataPtr;
-	}
-
-	/**
-	 * Helper function for returning a typed pointer to the first array entry.
-	 *
-	 * @returns Pointer to first array entry or nullptr if empty
-	 */
-	FORCEINLINE const PointerType GetData() const
+	FORCEINLINE ElementType* GetData() const
 	{
 		return DataPtr;
 	}
@@ -344,12 +320,9 @@ public:
 	 */
 	FORCEINLINE void CheckInvariants() const
 	{
-		if (bIsChecked)
+		for (uint8 Idx = 0; Idx < DimNum; Idx++)
 		{
-			for (SizeType Idx = 0; Idx < DimNum; Idx++)
-			{
-				checkSlow(ArrayShape[Idx] >= 0);
-			}
+			checkSlow(ArrayShape[Idx] >= 0);
 		}
 	}
 
@@ -360,10 +333,7 @@ public:
 	 */
 	FORCEINLINE void DimensionCheck(uint8 Dimension) const
 	{
-		if (bIsChecked)
-		{
-			checkf((Dimension < DimNum), TEXT("MultiArray dimension out of bounds: %lld from a rank of %lld"), (long long)Dimension, (long long)DimNum); // & for one branch
-		}
+		checkf((Dimension < DimNum), TEXT("MultiArray dimension out of bounds: %lld from a rank of %lld"), (long long)Dimension, (long long)DimNum); // & for one branch
 	}
 
 	/**
@@ -372,15 +342,12 @@ public:
 	 * @param Dimension Dimension of the array.
 	 * @param Index Index to check.
 	 */
-	FORCEINLINE void RangeCheck(uint8 Dimension, SizeType Index) const
+	FORCEINLINE void RangeCheck(uint8 Dimension, int64 Index) const
 	{
 		DimensionCheck(Dimension);
 		CheckInvariants();
 
-		if (bIsChecked)
-		{
-			checkf((Index >= 0) & (Index < ArrayShape[Dimension]), TEXT("MultiArray index out of bounds: %lld from a dimension of size %lld"), (long long)Index, (long long)ArrayShape[Dimension]); // & for one branch
-		}
+		checkf((Index >= 0) & (Index < ArrayShape[Dimension]), TEXT("MultiArray index out of bounds: %lld from a dimension of size %lld"), (long long)Index, (long long)ArrayShape[Dimension]); // & for one branch
 	}
 
 	/**
@@ -391,16 +358,13 @@ public:
 	 * @param Index Starting index of the slice.
 	 * @param InNum Length of the slice.
 	 */
-	FORCEINLINE void SliceRangeCheck(uint8 Dimension, SizeType Index, SizeType InNum) const
+	FORCEINLINE void SliceRangeCheck(uint8 Dimension, int64 Index, int64 InNum) const
 	{
 		DimensionCheck(Dimension);
 
-		if (bIsChecked)
-		{
-			checkf(Index >= 0, TEXT("Invalid index (%lld)"), (long long)Index);
-			checkf(InNum >= 0, TEXT("Invalid count (%lld)"), (long long)InNum);
-			checkf(Index + InNum <= ArrayShape[Dimension], TEXT("Range (index: %lld, count: %lld) lies outside the view of %lld elements"), (long long)Index, (long long)InNum, (long long)ArrayShape[Dimension]);
-		}
+		checkf(Index >= 0, TEXT("Invalid index (%lld)"), (long long)Index);
+		checkf(InNum >= 0, TEXT("Invalid count (%lld)"), (long long)InNum);
+		checkf(Index + InNum <= ArrayShape[Dimension], TEXT("Range (index: %lld, count: %lld) lies outside the view of %lld elements"), (long long)Index, (long long)InNum, (long long)ArrayShape[Dimension]);
 	}
 
 	/**
@@ -433,7 +397,7 @@ public:
 	 *
 	 * @returns Number of dimensions in array.
 	 */
-	FORCEINLINE SizeType Rank() const
+	FORCEINLINE uint8 Rank() const
 	{
 		return DimNum;
 	}
@@ -443,7 +407,7 @@ public:
 	 *
 	 * @returns Total number of elements in array.
 	 */
-	FORCEINLINE SizeType Num() const
+	FORCEINLINE int64 Num() const
 	{
 		return ArrayShape.Total();
 	}
@@ -454,7 +418,7 @@ public:
 	 * @returns Number of elements in array.
 	 */
 	template<uint8 InDimIdx>
-	FORCEINLINE SizeType Num() const
+	FORCEINLINE int64 Num() const
 	{
 		static_assert(InDimIdx < DimNum);
 		return ArrayShape[InDimIdx];
@@ -466,7 +430,7 @@ public:
 	 * @param Dimension Dimension of the array.
 	 * @returns Number of elements in array.
 	 */
-	FORCEINLINE SizeType Num(uint8 Dimension) const
+	FORCEINLINE int64 Num(uint8 Dimension) const
 	{
 		DimensionCheck(Dimension);
 
@@ -478,7 +442,7 @@ public:
 	 *
 	 * @returns Number of elements in each dimension.
 	 */
-	FORCEINLINE TMultiArrayShape<DimNum, SizeType> Shape() const
+	FORCEINLINE TMultiArrayShape<DimNum> Shape() const
 	{
 		return ArrayShape;
 	}
@@ -489,11 +453,11 @@ public:
 	 * @param Dimension Dimension of the array.
 	 * @returns Stride of that dimension.
 	 */
-	FORCEINLINE SizeType Stride(uint8 Dimension) const
+	FORCEINLINE int64 Stride(uint8 Dimension) const
 	{
 		DimensionCheck(Dimension);
 
-		SizeType Total = 1;
+		int64 Total = 1;
 		for (uint8 Idx = Dimension + 1; Idx < DimNum; Idx++)
 		{
 			Total *= ArrayShape[Idx];
@@ -508,7 +472,7 @@ public:
 	 * @param InNum Number of elements in the new view
 	 * @returns Sliced view
 	 */
-	[[nodiscard]] FORCEINLINE TMultiArrayView Slice(SizeType Index, SizeType InNum) const
+	[[nodiscard]] FORCEINLINE TMultiArrayView Slice(int64 Index, int64 InNum) const
 	{
 		SliceRangeCheck(0, Index, InNum);
 		TMultiArrayShape NewShape = ArrayShape;
@@ -534,83 +498,28 @@ public:
 		return false;
 	}
 
-public:
-
-	/** Implicit cast for Checked and Restrict. */
-	FORCEINLINE operator TMultiArrayView<DimNum, ElementType, true, false, SizeType>()
-	{
-		return TMultiArrayView<DimNum, ElementType, true, false, SizeType>(DataPtr, ArrayShape);
-	}
-
-	/** Implicit cast for Checked and Restrict. */
-	FORCEINLINE operator TMultiArrayView<DimNum, ElementType, false, false, SizeType>()
-	{
-		return TMultiArrayView<DimNum, ElementType, false, false, SizeType>(DataPtr, ArrayShape);
-	}
-
-	/** Implicit cast for Checked and Restrict. */
-	FORCEINLINE operator TMultiArrayView<DimNum, ElementType, true, true, SizeType>()
-	{
-		return TMultiArrayView<DimNum, ElementType, true, true, SizeType>(DataPtr, ArrayShape);
-	}
-
-	/** Implicit cast for Checked and Restrict. */
-	FORCEINLINE operator TMultiArrayView<DimNum, ElementType, false, true, SizeType>()
-	{
-		return TMultiArrayView<DimNum, ElementType, false, true, SizeType>(DataPtr, ArrayShape);
-	}
-
-	/** Implicit cast for Checked and Restrict. */
-	FORCEINLINE operator TMultiArrayView<DimNum, const ElementType, true, false, SizeType>() const
-	{
-		return TMultiArrayView<DimNum, const ElementType, true, false, SizeType>(DataPtr, ArrayShape);
-	}
-
-	/** Implicit cast for Checked and Restrict. */
-	FORCEINLINE operator TMultiArrayView<DimNum, const ElementType, false, false, SizeType>() const
-	{
-		return TMultiArrayView<DimNum, const ElementType, false, false, SizeType>(DataPtr, ArrayShape);
-	}
-
-	/** Implicit cast for Checked and Restrict. */
-	FORCEINLINE operator TMultiArrayView<DimNum, const ElementType, true, true, SizeType>() const
-	{
-		return TMultiArrayView<DimNum, const ElementType, true, true, SizeType>(DataPtr, ArrayShape);
-	}
-
-	/** Implicit cast for Checked and Restrict. */
-	FORCEINLINE operator TMultiArrayView<DimNum, const ElementType, false, true, SizeType>() const
+	/** Implicit cast for const. */
+	FORCEINLINE operator TMultiArrayView<DimNum, const ElementType>() const
 	{
-		return TMultiArrayView<DimNum, const ElementType, false, true, SizeType>(DataPtr, ArrayShape);
+		return TMultiArrayView<DimNum, const ElementType>(DataPtr, ArrayShape);
 	}
 
-
 private:
 
-	PointerType DataPtr;
-	TMultiArrayShape<DimNum, SizeType> ArrayShape;
+	ElementType* DataPtr;
+	TMultiArrayShape<DimNum> ArrayShape;
 };
 
 
 /**
 * Specialization for single dimensional MultiArrayView
 */
-template<typename InElementType, bool bInIsChecked, bool bInIsRestrict, typename InSizeType>
-class TMultiArrayView<1, InElementType, bInIsChecked, bInIsRestrict, InSizeType>
+template<typename InElementType>
+class TMultiArrayView<1, InElementType>
 {
 public:
 	static constexpr uint8 DimNum = 1;
 	using ElementType = InElementType;
-	static constexpr bool bIsChecked = bInIsChecked;
-	static constexpr bool bIsRestrict = bInIsRestrict;
-	using SizeType = InSizeType;
-
-	using PointerType = std::conditional_t<
-		bIsRestrict,
-		ElementType* RESTRICT,
-		ElementType*>;
-
-	static_assert(std::is_signed_v<SizeType>, "TMultiArrayView only supports signed index types");
 
 	/**
 	 * Constructor.
@@ -651,8 +560,8 @@ public:
 		>::GetData(Forward<OtherRangeType>(Other)))
 	{
 		const auto InCount = GetNum(Forward<OtherRangeType>(Other));
-		check((InCount >= 0) && ((sizeof(InCount) < sizeof(SizeType)) || (InCount <= static_cast<decltype(InCount)>(TNumericLimits<SizeType>::Max()))));
-		ArrayShape[0] = (SizeType)InCount;
+		check((InCount >= 0) && ((sizeof(InCount) < sizeof(int64)) || (InCount <= static_cast<decltype(InCount)>(TNumericLimits<int64>::Max()))));
+		ArrayShape[0] = (int64)InCount;
 	}
 
 	/**
@@ -661,7 +570,7 @@ public:
 	 * @param InData	The data to view
 	 * @param InNums	The number of elements on each dimension
 	 */
-	FORCEINLINE TMultiArrayView(ElementType* InData, TMultiArrayShape<DimNum, SizeType> InShape)
+	FORCEINLINE TMultiArrayView(ElementType* InData, TMultiArrayShape<DimNum> InShape)
 		: DataPtr(InData)
 		, ArrayShape(InShape)
 	{
@@ -678,7 +587,7 @@ public:
 		typename OtherElementType
 		UE_REQUIRES(UE::MultiArrayView::Private::TIsCompatibleElementType_V<OtherElementType, ElementType>)
 	>
-	FORCEINLINE TMultiArrayView(OtherElementType* InData, SizeType InNum)
+	FORCEINLINE TMultiArrayView(OtherElementType* InData, int64 InNum)
 		: DataPtr(InData)
 	{
 		ArrayShape[0] = InNum;
@@ -704,34 +613,34 @@ public:
 	 *
 	 * @returns Reference to indexed element.
 	 */
-	FORCEINLINE ElementType& operator[](SizeType Index) const
+	FORCEINLINE ElementType& operator[](int64 Index) const
 	{
 		RangeCheck(0, Index);
 		return DataPtr[Index];
 	}
 
 	/** Explicit conversion to TArrayView. */
-	FORCEINLINE TArrayView<ElementType, SizeType> ArrayView()
+	FORCEINLINE TArrayView<ElementType> ArrayView()
 	{
-		return TArrayView<ElementType, SizeType>(GetData(), Num());
+		return TArrayView<ElementType>(GetData(), Num());
 	}
 
 	/** Explicit conversion to TArrayView. */
-	FORCEINLINE TArrayView<const ElementType, SizeType> ArrayView() const
+	FORCEINLINE TArrayView<const ElementType> ArrayView() const
 	{
-		return TArrayView<const ElementType, SizeType>(GetData(), Num());
+		return TArrayView<const ElementType>(GetData(), Num());
 	}
 
 	/** Implicit cast to TArrayView. */
-	FORCEINLINE operator TArrayView<ElementType, SizeType>()
+	FORCEINLINE operator TArrayView<ElementType>()
 	{
-		return TArrayView<ElementType, SizeType>(DataPtr, ArrayShape[0]);
+		return TArrayView<ElementType>(DataPtr, ArrayShape[0]);
 	}
 
 	/** Implicit cast to TArrayView. */
-	FORCEINLINE operator TArrayView<const ElementType, SizeType>() const
+	FORCEINLINE operator TArrayView<const ElementType>() const
 	{
-		return TArrayView<const ElementType, SizeType>(DataPtr, ArrayShape[0]);
+		return TArrayView<const ElementType>(DataPtr, ArrayShape[0]);
 	}
 
 public:
@@ -741,17 +650,7 @@ public:
 	 *
 	 * @returns Pointer to first array entry or nullptr if empty
 	 */
-	FORCEINLINE PointerType GetData()
-	{
-		return DataPtr;
-	}
-
-	/**
-	 * Helper function for returning a typed pointer to the first array entry.
-	 *
-	 * @returns Pointer to first array entry or nullptr if empty
-	 */
-	FORCEINLINE const PointerType GetData() const
+	FORCEINLINE ElementType* GetData() const
 	{
 		return DataPtr;
 	}
@@ -786,12 +685,9 @@ public:
 	 */
 	FORCEINLINE void CheckInvariants() const
 	{
-		if (bIsChecked)
+		for (uint8 Idx = 0; Idx < DimNum; Idx++)
 		{
-			for (SizeType Idx = 0; Idx < DimNum; Idx++)
-			{
-				checkSlow(ArrayShape[Idx] >= 0);
-			}
+			checkSlow(ArrayShape[Idx] >= 0);
 		}
 	}
 
@@ -802,10 +698,7 @@ public:
 	 */
 	FORCEINLINE void DimensionCheck(uint8 Dimension) const
 	{
-		if (bIsChecked)
-		{
-			checkf((Dimension < DimNum), TEXT("MultiArray dimension out of bounds: %lld from a rank of %lld"), (long long)Dimension, (long long)DimNum); // & for one branch
-		}
+		checkf((Dimension < DimNum), TEXT("MultiArray dimension out of bounds: %lld from a rank of %lld"), (long long)Dimension, (long long)DimNum); // & for one branch
 	}
 
 	/**
@@ -814,15 +707,12 @@ public:
 	 * @param Dimension Dimension of the array.
 	 * @param Index Index to check.
 	 */
-	FORCEINLINE void RangeCheck(uint8 Dimension, SizeType Index) const
+	FORCEINLINE void RangeCheck(uint8 Dimension, int64 Index) const
 	{
 		DimensionCheck(Dimension);
 		CheckInvariants();
 
-		if (bIsChecked)
-		{
-			checkf((Index >= 0) & (Index < ArrayShape[Dimension]), TEXT("MultiArray index out of bounds: %lld from a dimension of size %lld"), (long long)Index, (long long)ArrayShape[Dimension]); // & for one branch
-		}
+		checkf((Index >= 0) & (Index < ArrayShape[Dimension]), TEXT("MultiArray index out of bounds: %lld from a dimension of size %lld"), (long long)Index, (long long)ArrayShape[Dimension]); // & for one branch
 	}
 
 	/**
@@ -833,16 +723,13 @@ public:
 	 * @param Index Starting index of the slice.
 	 * @param InNum Length of the slice.
 	 */
-	FORCEINLINE void SliceRangeCheck(uint8 Dimension, SizeType Index, SizeType InNum) const
+	FORCEINLINE void SliceRangeCheck(uint8 Dimension, int64 Index, int64 InNum) const
 	{
 		DimensionCheck(Dimension);
 
-		if (bIsChecked)
-		{
-			checkf(Index >= 0, TEXT("Invalid index (%lld)"), (long long)Index);
-			checkf(InNum >= 0, TEXT("Invalid count (%lld)"), (long long)InNum);
-			checkf(Index + InNum <= ArrayShape[Dimension], TEXT("Range (index: %lld, count: %lld) lies outside the view of %lld elements"), (long long)Index, (long long)InNum, (long long)ArrayShape[Dimension]);
-		}
+		checkf(Index >= 0, TEXT("Invalid index (%lld)"), (long long)Index);
+		checkf(InNum >= 0, TEXT("Invalid count (%lld)"), (long long)InNum);
+		checkf(Index + InNum <= ArrayShape[Dimension], TEXT("Range (index: %lld, count: %lld) lies outside the view of %lld elements"), (long long)Index, (long long)InNum, (long long)ArrayShape[Dimension]);
 	}
 
 	/**
@@ -875,7 +762,7 @@ public:
 	 *
 	 * @returns Number of dimensions in array.
 	 */
-	FORCEINLINE SizeType Rank() const
+	FORCEINLINE uint8 Rank() const
 	{
 		return DimNum;
 	}
@@ -885,7 +772,7 @@ public:
 	 *
 	 * @returns Total number of elements in array.
 	 */
-	FORCEINLINE SizeType Num() const
+	FORCEINLINE int64 Num() const
 	{
 		return ArrayShape.Total();
 	}
@@ -896,7 +783,7 @@ public:
 	 * @returns Number of elements in array.
 	 */
 	template<uint8 InDimIdx>
-	FORCEINLINE SizeType Num() const
+	FORCEINLINE int64 Num() const
 	{
 		static_assert(InDimIdx < DimNum);
 		return ArrayShape[InDimIdx];
@@ -908,7 +795,7 @@ public:
 	 * @param Dimension Dimension of the array.
 	 * @returns Number of elements in array.
 	 */
-	FORCEINLINE SizeType Num(uint8 Dimension) const
+	FORCEINLINE int64 Num(uint8 Dimension) const
 	{
 		DimensionCheck(Dimension);
 
@@ -920,7 +807,7 @@ public:
 	 *
 	 * @returns Number of elements in each dimension.
 	 */
-	FORCEINLINE TMultiArrayShape<DimNum, SizeType> Shape() const
+	FORCEINLINE TMultiArrayShape<DimNum> Shape() const
 	{
 		return ArrayShape;
 	}
@@ -931,11 +818,11 @@ public:
 	 * @param Dimension Dimension of the array.
 	 * @returns Stride of that dimension.
 	 */
-	FORCEINLINE SizeType Stride(uint8 Dimension) const
+	FORCEINLINE int64 Stride(uint8 Dimension) const
 	{
 		DimensionCheck(Dimension);
 
-		SizeType Total = 1;
+		int64 Total = 1;
 		for (uint8 Idx = Dimension + 1; Idx < DimNum; Idx++)
 		{
 			Total *= ArrayShape[Idx];
@@ -950,7 +837,7 @@ public:
 	 * @param InNum Number of elements in the new view
 	 * @returns Sliced view
 	 */
-	[[nodiscard]] FORCEINLINE TMultiArrayView Slice(SizeType Index, SizeType InNum) const
+	[[nodiscard]] FORCEINLINE TMultiArrayView Slice(int64 Index, int64 InNum) const
 	{
 		SliceRangeCheck(0, Index, InNum);
 		TMultiArrayShape NewShape = ArrayShape;
@@ -976,64 +863,20 @@ public:
 		return false;
 	}
 
-public:
-
-	/** Implicit cast for Checked and Restrict. */
-	FORCEINLINE operator TMultiArrayView<DimNum, ElementType, true, false, SizeType>()
-	{
-		return TMultiArrayView<DimNum, ElementType, true, false, SizeType>(DataPtr, ArrayShape);
-	}
-
-	/** Implicit cast for Checked and Restrict. */
-	FORCEINLINE operator TMultiArrayView<DimNum, ElementType, false, false, SizeType>()
-	{
-		return TMultiArrayView<DimNum, ElementType, false, false, SizeType>(DataPtr, ArrayShape);
-	}
-
-	/** Implicit cast for Checked and Restrict. */
-	FORCEINLINE operator TMultiArrayView<DimNum, ElementType, true, true, SizeType>()
-	{
-		return TMultiArrayView<DimNum, ElementType, true, true, SizeType>(DataPtr, ArrayShape);
-	}
-
-	/** Implicit cast for Checked and Restrict. */
-	FORCEINLINE operator TMultiArrayView<DimNum, ElementType, false, true, SizeType>()
-	{
-		return TMultiArrayView<DimNum, ElementType, false, true, SizeType>(DataPtr, ArrayShape);
-	}
-
-	/** Implicit cast for Checked and Restrict. */
-	FORCEINLINE operator TMultiArrayView<DimNum, const ElementType, true, false, SizeType>() const
-	{
-		return TMultiArrayView<DimNum, const ElementType, true, false, SizeType>(DataPtr, ArrayShape);
-	}
-
-	/** Implicit cast for Checked and Restrict. */
-	FORCEINLINE operator TMultiArrayView<DimNum, const ElementType, false, false, SizeType>() const
-	{
-		return TMultiArrayView<DimNum, const ElementType, false, false, SizeType>(DataPtr, ArrayShape);
-	}
-
-	/** Implicit cast for Checked and Restrict. */
-	FORCEINLINE operator TMultiArrayView<DimNum, const ElementType, true, true, SizeType>() const
-	{
-		return TMultiArrayView<DimNum, const ElementType, true, true, SizeType>(DataPtr, ArrayShape);
-	}
-
-	/** Implicit cast for Checked and Restrict. */
-	FORCEINLINE operator TMultiArrayView<DimNum, const ElementType, false, true, SizeType>() const
+	/** Implicit cast for const. */
+	FORCEINLINE operator TMultiArrayView<DimNum, const ElementType>() const
 	{
-		return TMultiArrayView<DimNum, const ElementType, false, true, SizeType>(DataPtr, ArrayShape);
+		return TMultiArrayView<DimNum, const ElementType>(DataPtr, ArrayShape);
 	}
 
 private:
 
-	PointerType DataPtr;
-	TMultiArrayShape<DimNum, SizeType> ArrayShape;
+	ElementType* DataPtr;
+	TMultiArrayShape<DimNum> ArrayShape;
 };
 
-template<uint8 InDimNum, typename InElementType, bool bInIsChecked = true, bool bInIsRestrict = false, typename InSizeType = FDefaultAllocator::SizeType>
-using TConstMultiArrayView = TMultiArrayView<InDimNum, const InElementType, bInIsChecked, bInIsRestrict, InSizeType>;
+template<uint8 InDimNum, typename InElementType>
+using TConstMultiArrayView = TMultiArrayView<InDimNum, const InElementType>;
 
 #if UE_ENABLE_INCLUDE_ORDER_DEPRECATED_IN_5_4
 #include "Templates/IsConst.h"
diff --git a/Engine/Plugins/Experimental/LearningAgents/Source/LearningAgentsTraining/Private/LearningAgentsRecording.cpp b/Engine/Plugins/Experimental/LearningAgents/Source/LearningAgentsTraining/Private/LearningAgentsRecording.cpp
--- a/Engine/Plugins/Experimental/LearningAgents/Source/LearningAgentsTraining/Private/LearningAgentsRecording.cpp
+++ b/Engine/Plugins/Experimental/LearningAgents/Source/LearningAgentsTraining/Private/LearningAgentsRecording.cpp
@@ -35,7 +35,7 @@ void ULearningAgentsRecording::LoadRecordingFromFile(const FFilePath& File)
 			return;
 		}
 
-		int32 Offset = 0;
+		int64 Offset = 0;
 
 		int32 MagicNumber;
 		UE::Learning::DeserializeFromBytes(Offset, RecordingData, MagicNumber);
@@ -67,8 +67,8 @@ void ULearningAgentsRecording::LoadRecordingFromFile(const FFilePath& File)
 			UE::Learning::DeserializeFromBytes(Offset, RecordingData, Records[RecordIdx].ActionDimNum);
 			UE::Learning::DeserializeFromBytes(Offset, RecordingData, Records[RecordIdx].ObservationCompatibilityHash);
 			UE::Learning::DeserializeFromBytes(Offset, RecordingData, Records[RecordIdx].ActionCompatibilityHash);
-			UE::Learning::Array::DeserializeFromBytes<2, float>(Offset, RecordingData, Records[RecordIdx].ObservationData);
-			UE::Learning::Array::DeserializeFromBytes<2, float>(Offset, RecordingData, Records[RecordIdx].ActionData);
+			UE::Learning::Array::DeserializeFromBytes32<2, float>(Offset, RecordingData, Records[RecordIdx].ObservationData);
+			UE::Learning::Array::DeserializeFromBytes32<2, float>(Offset, RecordingData, Records[RecordIdx].ActionData);
 		}
 
 		UE_LEARNING_CHECK(Offset == RecordingData.Num());
@@ -98,13 +98,13 @@ void ULearningAgentsRecording::SaveRecordingToFile(const FFilePath& File) const
 			sizeof(int32) + // ActionDimNum
 			sizeof(int32) + // ObservationCompatibilityHash
 			sizeof(int32) + // ActionCompatibilityHash
-			UE::Learning::Array::SerializationByteNum<2, float>({ Records[RecordIdx].StepNum, Records[RecordIdx].ObservationDimNum }) + // Observations
-			UE::Learning::Array::SerializationByteNum<2, float>({ Records[RecordIdx].StepNum, Records[RecordIdx].ActionDimNum });	   // Actions
+			UE::Learning::Array::SerializationByteNum32<2, float>({ Records[RecordIdx].StepNum, Records[RecordIdx].ObservationDimNum }) + // Observations
+			UE::Learning::Array::SerializationByteNum32<2, float>({ Records[RecordIdx].StepNum, Records[RecordIdx].ActionDimNum });	   // Actions
 	}
 
 	RecordingData.SetNumUninitialized(TotalByteNum);
 
-	int32 Offset = 0;
+	int64 Offset = 0;
 	UE::Learning::SerializeToBytes(Offset, RecordingData, UE::Learning::Agents::Recording::Private::MagicNumber);
 	UE::Learning::SerializeToBytes(Offset, RecordingData, UE::Learning::Agents::Recording::Private::VersionNumber);
 	UE::Learning::SerializeToBytes(Offset, RecordingData, Records.Num());
@@ -116,8 +116,8 @@ void ULearningAgentsRecording::SaveRecordingToFile(const FFilePath& File) const
 		UE::Learning::SerializeToBytes(Offset, RecordingData, Records[RecordIdx].ActionDimNum);
 		UE::Learning::SerializeToBytes(Offset, RecordingData, Records[RecordIdx].ObservationCompatibilityHash);
 		UE::Learning::SerializeToBytes(Offset, RecordingData, Records[RecordIdx].ActionCompatibilityHash);
-		UE::Learning::Array::SerializeToBytes<2, float>(Offset, RecordingData, { Records[RecordIdx].StepNum, Records[RecordIdx].ObservationDimNum }, Records[RecordIdx].ObservationData);
-		UE::Learning::Array::SerializeToBytes<2, float>(Offset, RecordingData, { Records[RecordIdx].StepNum, Records[RecordIdx].ActionDimNum }, Records[RecordIdx].ActionData);
+		UE::Learning::Array::SerializeToBytes32<2, float>(Offset, RecordingData, { Records[RecordIdx].StepNum, Records[RecordIdx].ObservationDimNum }, Records[RecordIdx].ObservationData);
+		UE::Learning::Array::SerializeToBytes32<2, float>(Offset, RecordingData, { Records[RecordIdx].StepNum, Records[RecordIdx].ActionDimNum }, Records[RecordIdx].ActionData);
 	}
 
 	UE_LEARNING_CHECK(Offset == RecordingData.Num());
