LearningAgents: Make trainer command-line and logger nicely pythonic

diff --git a/Engine/Plugins/Experimental/LearningAgents/Source/LearningTraining/Private/LearningExternalTrainer.cpp b/Engine/Plugins/Experimental/LearningAgents/Source/LearningTraining/Private/LearningExternalTrainer.cpp
--- a/Engine/Plugins/Experimental/LearningAgents/Source/LearningTraining/Private/LearningExternalTrainer.cpp
+++ b/Engine/Plugins/Experimental/LearningAgents/Source/LearningTraining/Private/LearningExternalTrainer.cpp
@@ -31,7 +31,8 @@ namespace UE::Learning
 		const FString& InIntermediatePath,
 		const int32 ProcessNum,
 		const float InTimeout,
-		const ESubprocessFlags TrainingProcessFlags)
+		const ESubprocessFlags TrainingProcessFlags,
+		const ELogSetting LogSettings)
 	{
 		check(ProcessNum > 0);
 
@@ -74,10 +75,11 @@ namespace UE::Learning
 		ConfigPath = InIntermediatePath / TEXT("Configs") / FString::Printf(TEXT("%s_%s_%s_%s.json"), *TaskName, *TrainerFileName, *TrainerType, *TimeStamp);
 
 		IFileManager& FileManager = IFileManager::Get();
-		const FString CommandLineArguments = FString::Printf(TEXT("\"%s\" \"%s\" \"%s\" SharedMemory \"%s\" %i \"%s\""),
+		const FString CommandLineArguments = FString::Printf(TEXT("\"%s\" -p \"%s\" -m \"%s\" \"%s\" SharedMemory \"%s\" %i \"%s\""),
 			*FileManager.ConvertToAbsolutePathForExternalAppForRead(*(PythonContentPath / TEXT("train.py"))),
 			*FileManager.ConvertToAbsolutePathForExternalAppForRead(*CustomTrainerPath),
 			*TrainerFileName,
+			LogSettings == ELogSetting::Normal ? TEXT("-l") : TEXT(""),
 			*Controls.Guid.ToString(EGuidFormats::DigitsWithHyphensInBraces),
 			ProcessNum,
 			*FileManager.ConvertToAbsolutePathForExternalAppForRead(*ConfigPath));
@@ -252,7 +254,6 @@ namespace UE::Learning
 
 		IFileManager& FileManager = IFileManager::Get();
 		ConfigObject->SetStringField(TEXT("IntermediatePath"), *FileManager.ConvertToAbsolutePathForExternalAppForRead(*IntermediatePath));
-		ConfigObject->SetBoolField(TEXT("LoggingEnabled"), LogSettings == ELogSetting::Silent ? false : true);
 
 		TSharedPtr<FJsonObject> SharedMemoryObject = MakeShared<FJsonObject>();
 
@@ -701,14 +702,14 @@ namespace UE::Learning
 		UE_LEARNING_CHECK(FPaths::DirectoryExists(PythonContentPath));
 
 		IFileManager& FileManager = IFileManager::Get();
-		const FString CommandLineArguments = FString::Printf(TEXT("\"%s\" \"%s\" \"%s\" Socket \"%s:%i\" \"%s\" %i"),
+		const FString CommandLineArguments = FString::Printf(TEXT("\"%s\" -p \"%s\" -m \"%s\" \"%s\" Socket \"%s:%i\" \"%s\""),
 			*FileManager.ConvertToAbsolutePathForExternalAppForRead(*(PythonContentPath / TEXT("train.py"))),
 			*FileManager.ConvertToAbsolutePathForExternalAppForRead(*CustomTrainerPath),
 			*TrainerFileName,
+			LogSettings == ELogSetting::Normal ? TEXT("-l") : TEXT(""),
 			IpAddress,
 			Port,
-			*FileManager.ConvertToAbsolutePathForExternalAppForRead(*IntermediatePath),
-			LogSettings == ELogSetting::Normal ? 1 : 0);
+			*FileManager.ConvertToAbsolutePathForExternalAppForRead(*IntermediatePath));
 
 		TrainingProcess.Launch(
 			FileManager.ConvertToAbsolutePathForExternalAppForRead(*PythonExecutablePath),
diff --git a/Engine/Plugins/Experimental/LearningAgents/Source/LearningTraining/Public/LearningExternalTrainer.h b/Engine/Plugins/Experimental/LearningAgents/Source/LearningTraining/Public/LearningExternalTrainer.h
--- a/Engine/Plugins/Experimental/LearningAgents/Source/LearningTraining/Public/LearningExternalTrainer.h
+++ b/Engine/Plugins/Experimental/LearningAgents/Source/LearningTraining/Public/LearningExternalTrainer.h
@@ -166,6 +166,7 @@ namespace UE::Learning
 		* @param ProcessNum					Number of processes to use for multi-processed experience gathering (used
 		*									to allocate enough shared memory)
 		* @param TrainingProcessFlags		Training server subprocess flags
+		* @param LogSettings				Logging settings to use
 		*/
 		FSharedMemoryTrainerServerProcess(
 			const FString& TaskName,
@@ -176,7 +177,8 @@ namespace UE::Learning
 			const FString& IntermediatePath,
 			const int32 ProcessNum,
 			const float InTimeout = Trainer::DefaultTimeout,
-			const ESubprocessFlags TrainingProcessFlags = ESubprocessFlags::None);
+			const ESubprocessFlags TrainingProcessFlags = ESubprocessFlags::None,
+			const ELogSetting LogSettings = ELogSetting::Normal);
 
 		/** Check if the server process is still running. */
 		virtual bool IsRunning() const override final;
