LearningAgents: Fixes an issue with cooked game unable to find custom python trainer path
#rnx

diff --git a/Engine/Plugins/Experimental/LearningAgents/Source/LearningAgentsTraining/Private/LearningAgentsCommunicator.cpp b/Engine/Plugins/Experimental/LearningAgents/Source/LearningAgentsTraining/Private/LearningAgentsCommunicator.cpp
--- a/Engine/Plugins/Experimental/LearningAgents/Source/LearningAgentsTraining/Private/LearningAgentsCommunicator.cpp
+++ b/Engine/Plugins/Experimental/LearningAgents/Source/LearningAgentsTraining/Private/LearningAgentsCommunicator.cpp
@@ -35,9 +35,16 @@ FLearningAgentsTrainerProcess ULearningAgentsCommunicatorLibrary::SpawnSharedMem
 
 	const FString IntermediatePath = UE::Learning::Trainer::GetIntermediatePath(TrainerProcessSettings.GetIntermediatePath());
 
+	const FString CustomTrainerModulePath = TrainerProcessSettings.GetCustomTrainerModulePath();
+	if (!CustomTrainerModulePath.IsEmpty() and !FPaths::DirectoryExists(CustomTrainerModulePath))
+	{
+		UE_LOG(LogLearning, Error, TEXT("SpawnSharedMemoryTrainingProcess: Can't find custom trainer module \"%s\"."), *CustomTrainerModulePath);
+		return TrainerProcess;
+	}
+
 	TrainerProcess.TrainerProcess = MakeShared<UE::Learning::FSharedMemoryTrainerServerProcess>(
 		SharedMemorySettings.TaskName,
-		UE::Learning::Trainer::GetProjectPythonContentPath(),
+		CustomTrainerModulePath,
 		TrainerProcessSettings.TrainerFileName,
 		PythonExecutablePath,
 		PythonContentPath,
@@ -50,7 +57,6 @@ FLearningAgentsTrainerProcess ULearningAgentsCommunicatorLibrary::SpawnSharedMem
 
 FLearningAgentsCommunicator ULearningAgentsCommunicatorLibrary::MakeSharedMemoryCommunicator(
 	const FLearningAgentsTrainerProcess& TrainerProcess,
-	const FLearningAgentsTrainerProcessSettings& TrainerProcessSettings,
 	const FLearningAgentsSharedMemoryCommunicatorSettings& SharedMemorySettings)
 {
 	FLearningAgentsCommunicator Communicator;
@@ -98,8 +104,15 @@ FLearningAgentsTrainerProcess ULearningAgentsCommunicatorLibrary::SpawnSocketTra
 
 	const FString IntermediatePath = UE::Learning::Trainer::GetIntermediatePath(TrainerProcessSettings.GetIntermediatePath());
 
+	const FString CustomTrainerModulePath = TrainerProcessSettings.GetCustomTrainerModulePath();
+	if (!CustomTrainerModulePath.IsEmpty() and !FPaths::DirectoryExists(CustomTrainerModulePath))
+	{
+		UE_LOG(LogLearning, Error, TEXT("SpawnSocketTrainingProcess: Can't find custom trainer module \"%s\"."), *CustomTrainerModulePath);
+		return TrainerProcess;
+	}
+
 	TrainerProcess.TrainerProcess = MakeShared<UE::Learning::FSocketTrainerServerProcess>(
-		UE::Learning::Trainer::GetProjectPythonContentPath(),
+		CustomTrainerModulePath,
 		TrainerProcessSettings.TrainerFileName,
 		PythonExecutablePath,
 		PythonContentPath,
diff --git a/Engine/Plugins/Experimental/LearningAgents/Source/LearningAgentsTraining/Private/LearningAgentsTrainer.cpp b/Engine/Plugins/Experimental/LearningAgents/Source/LearningAgentsTraining/Private/LearningAgentsTrainer.cpp
--- a/Engine/Plugins/Experimental/LearningAgents/Source/LearningAgentsTraining/Private/LearningAgentsTrainer.cpp
+++ b/Engine/Plugins/Experimental/LearningAgents/Source/LearningAgentsTraining/Private/LearningAgentsTrainer.cpp
@@ -197,3 +197,8 @@ FString FLearningAgentsTrainerProcessSettings::GetIntermediatePath() const
 	return NonEditorIntermediateRelativePath;
 #endif
 }
+
+FString FLearningAgentsTrainerProcessSettings::GetCustomTrainerModulePath() const
+{
+	return CustomTrainerModulePath.Path;
+}
diff --git a/Engine/Plugins/Experimental/LearningAgents/Source/LearningAgentsTraining/Public/LearningAgentsCommunicator.h b/Engine/Plugins/Experimental/LearningAgents/Source/LearningAgentsTraining/Public/LearningAgentsCommunicator.h
--- a/Engine/Plugins/Experimental/LearningAgents/Source/LearningAgentsTraining/Public/LearningAgentsCommunicator.h
+++ b/Engine/Plugins/Experimental/LearningAgents/Source/LearningAgentsTraining/Public/LearningAgentsCommunicator.h
@@ -108,7 +108,6 @@ public:
 	UFUNCTION(BlueprintCallable, Category = "LearningAgents", meta = (AutoCreateRefTerm = "SharedMemorySettings"))
 	static FLearningAgentsCommunicator MakeSharedMemoryCommunicator(
 		const FLearningAgentsTrainerProcess& TrainerProcess,
-		const FLearningAgentsTrainerProcessSettings& TrainerProcessSettings = FLearningAgentsTrainerProcessSettings(),
 		const FLearningAgentsSharedMemoryCommunicatorSettings& SharedMemorySettings = FLearningAgentsSharedMemoryCommunicatorSettings());
 
 	/**
diff --git a/Engine/Plugins/Experimental/LearningAgents/Source/LearningAgentsTraining/Public/LearningAgentsTrainer.h b/Engine/Plugins/Experimental/LearningAgents/Source/LearningAgentsTraining/Public/LearningAgentsTrainer.h
--- a/Engine/Plugins/Experimental/LearningAgents/Source/LearningAgentsTraining/Public/LearningAgentsTrainer.h
+++ b/Engine/Plugins/Experimental/LearningAgents/Source/LearningAgentsTraining/Public/LearningAgentsTrainer.h
@@ -140,6 +140,10 @@ public:
 	UPROPERTY(EditAnywhere, Category = "LearningAgents")
 	FString NonEditorIntermediateRelativePath;
 
+	/** The complete path to a custom trainer module when training with a custom trainer. */
+	UPROPERTY(EditAnywhere, Category = "LearningAgents")
+	FDirectoryPath CustomTrainerModulePath;
+
 	/** Trainer file name. The name of the python file to use for training. Do NOT include the '.py' file extension. */
 	UPROPERTY(EditAnywhere, Category = "LearningAgents")
 	FString TrainerFileName = TEXT("train_ppo");
@@ -151,4 +155,7 @@ public:
 
 	/** Gets the Relative Intermediate Path  */
 	FString GetIntermediatePath() const;
+
+	/** Gets the Custom Trainer Module Path  */
+	FString GetCustomTrainerModulePath() const;
 };
