LearningAgents: Action Masking

#virtualized

diff --git a/Engine/Plugins/Experimental/LearningAgents/Source/Learning/Private/Learning.ispc b/Engine/Plugins/Experimental/LearningAgents/Source/Learning/Private/Learning.ispc
--- a/Engine/Plugins/Experimental/LearningAgents/Source/Learning/Private/Learning.ispc
+++ b/Engine/Plugins/Experimental/LearningAgents/Source/Learning/Private/Learning.ispc
@@ -156,6 +156,32 @@ export void LearningRandomDistributionIndependantNormal(
 	}
 }
 
+export void LearningRandomDistributionIndependantNormalMasked(
+	uniform float Output[],
+	uniform const float Mean[],
+	uniform const float LogStd[],
+	uniform const bool Masked[],
+	uniform const float MaskedValues[],
+	uniform const int DimNum,
+	uniform const uint32 State,
+	uniform const float Scale)
+{
+	if (Scale < UE_SMALL_NUMBER)
+	{
+		foreach(DimIdx = 0 ... DimNum)
+		{
+			Output[DimIdx] = select(Masked[DimIdx], MaskedValues[DimIdx], Mean[DimIdx]);
+		}
+	}
+	else
+	{
+		foreach(DimIdx = 0 ... DimNum)
+		{
+			Output[DimIdx] = select(Masked[DimIdx], MaskedValues[DimIdx], LearningRandomGaussian(State ^ 0x0a006c6b ^ LearningRandomInt(DimIdx ^ 0x6fe1d121), Mean[DimIdx], Scale * exp(min(LogStd[DimIdx], 10.0f))));
+		}
+	}
+}
+
 export void LearningRandomDistributionBernoulli(
 	uniform float Output[],
 	uniform const float Logits[],
@@ -179,6 +205,30 @@ export void LearningRandomDistributionBernoulli(
 	}
 }
 
+export void LearningRandomDistributionBernoulliMasked(
+	uniform float Output[],
+	uniform const float Logits[],
+	uniform const bool Masked[],
+	uniform const int DimNum,
+	uniform const uint32 State,
+	uniform const float Scale)
+{
+	if (Scale < UE_SMALL_NUMBER)
+	{
+		foreach(DimIdx = 0 ... DimNum)
+		{
+			Output[DimIdx] = select(!Masked[DimIdx] && Logits[DimIdx] > 0.0f, 1.0f, 0.0f);
+		}
+	}
+	else
+	{
+		foreach(DimIdx = 0 ... DimNum)
+		{
+			Output[DimIdx] = select(!Masked[DimIdx] && LearningSigmoid(Logits[DimIdx] / Scale) > LearningRandomUniform(State ^ 0x8c992b7f ^ LearningRandomInt(DimIdx ^ 0x1953d0fc)), 1.0f, 0.0f);
+		}
+	}
+}
+
 /////////////////////////////////////////////////////////
 // 
 // Functions used by various different Optimizers
diff --git a/Engine/Plugins/Experimental/LearningAgents/Source/Learning/Private/LearningAction.cpp b/Engine/Plugins/Experimental/LearningAgents/Source/Learning/Private/LearningAction.cpp
--- a/Engine/Plugins/Experimental/LearningAgents/Source/Learning/Private/LearningAction.cpp
+++ b/Engine/Plugins/Experimental/LearningAgents/Source/Learning/Private/LearningAction.cpp
@@ -66,6 +66,16 @@
 			return Size;
 		}
 
+		static inline int32 GetTotalActionModifierVectorSize(const FSchema& Schema, const TArrayView<const FSchemaElement> Elements)
+		{
+			int32 Size = 0;
+			for (const FSchemaElement SubElement : Elements)
+			{
+				Size += Schema.GetActionModifierVectorSize(SubElement);
+			}
+			return Size;
+		}
+
 		static inline bool CheckAllValid(const FObject& Object, const TArrayView<const FObjectElement> Elements)
 		{
 			for (const FObjectElement SubElement : Elements)
@@ -111,6 +121,24 @@
 			return true;
 		}
 
+		static inline bool CheckAllValid(const FModifier& Object, const TArrayView<const FModifierElement> Elements)
+		{
+			for (const FModifierElement SubElement : Elements)
+			{
+				if (!Object.IsValid(SubElement)) { return false; }
+			}
+			return true;
+		}
+
+		static inline bool CheckExclusiveMaskValid(const TArrayView<const bool> Mask)
+		{
+			for (int32 MaskIdx = 0; MaskIdx < Mask.Num(); MaskIdx++)
+			{
+				if (!Mask[MaskIdx]) { return true; }
+			}
+			return false;
+		}
+
 		static inline float Logit(const float X)
 		{
 			return FMath::Loge(FMath::Max(X / FMath::Max(1.0f - X, FLT_MIN), FLT_MIN));
@@ -124,6 +152,7 @@
 		EncodedVectorSizes.Add(0);
 		ActionVectorSizes.Add(0);
 		ActionDistributionVectorSizes.Add(0);
+		ActionModifierVectorSizes.Add(1);
 		TypeDataIndices.Add(INDEX_NONE);
 
 		return { Index, Generation };
@@ -143,6 +172,7 @@
 		EncodedVectorSizes.Add(2 * Parameters.Num);
 		ActionVectorSizes.Add(Parameters.Num);
 		ActionDistributionVectorSizes.Add(2 * Parameters.Num);
+		ActionModifierVectorSizes.Add(1 + 2 * Parameters.Num);
 		TypeDataIndices.Add(ContinuousData.Add(ElementData));
 
 		return { Index, Generation };
@@ -164,6 +194,7 @@
 		EncodedVectorSizes.Add(Parameters.Num);
 		ActionVectorSizes.Add(Parameters.Num);
 		ActionDistributionVectorSizes.Add(Parameters.Num);
+		ActionModifierVectorSizes.Add(1 + Parameters.Num);
 		TypeDataIndices.Add(DiscreteExclusiveData.Add(ElementData));
 
 		return { Index, Generation };
@@ -185,6 +216,7 @@
 		EncodedVectorSizes.Add(Parameters.Num);
 		ActionVectorSizes.Add(Parameters.Num);
 		ActionDistributionVectorSizes.Add(Parameters.Num);
+		ActionModifierVectorSizes.Add(1 + Parameters.Num);
 		TypeDataIndices.Add(DiscreteInclusiveData.Add(ElementData));
 
 		return { Index, Generation };
@@ -208,6 +240,7 @@
 		EncodedVectorSizes.Add(Private::GetTotalEncodedActionVectorSize(*this, Parameters.Elements));
 		ActionVectorSizes.Add(Private::GetTotalActionVectorSize(*this, Parameters.Elements));
 		ActionDistributionVectorSizes.Add(Private::GetTotalActionDistributionVectorSize(*this, Parameters.Elements));
+		ActionModifierVectorSizes.Add(1 + Private::GetTotalActionModifierVectorSize(*this, Parameters.Elements));
 		TypeDataIndices.Add(AndData.Add(ElementData));
 
 		return { Index, Generation };
@@ -235,6 +268,7 @@
 		EncodedVectorSizes.Add(Private::GetTotalEncodedActionVectorSize(*this, Parameters.Elements) + Parameters.Elements.Num());
 		ActionVectorSizes.Add(Private::GetMaxActionVectorSize(*this, Parameters.Elements) + Parameters.Elements.Num());
 		ActionDistributionVectorSizes.Add(Private::GetTotalActionDistributionVectorSize(*this, Parameters.Elements) + Parameters.Elements.Num());
+		ActionModifierVectorSizes.Add(1 + Parameters.Elements.Num() + Private::GetTotalActionModifierVectorSize(*this, Parameters.Elements));
 		TypeDataIndices.Add(OrExclusiveData.Add(ElementData));
 
 		return { Index, Generation };
@@ -262,6 +296,7 @@
 		EncodedVectorSizes.Add(Private::GetTotalEncodedActionVectorSize(*this, Parameters.Elements) + Parameters.Elements.Num());
 		ActionVectorSizes.Add(Private::GetTotalActionVectorSize(*this, Parameters.Elements) + Parameters.Elements.Num());
 		ActionDistributionVectorSizes.Add(Private::GetTotalActionDistributionVectorSize(*this, Parameters.Elements) + Parameters.Elements.Num());
+		ActionModifierVectorSizes.Add(1 + Parameters.Elements.Num() + Private::GetTotalActionModifierVectorSize(*this, Parameters.Elements));
 		TypeDataIndices.Add(OrInclusiveData.Add(ElementData));
 
 		return { Index, Generation };
@@ -284,6 +319,7 @@
 		EncodedVectorSizes.Add(GetEncodedVectorSize(Parameters.Element) * Parameters.Num);
 		ActionVectorSizes.Add(GetActionVectorSize(Parameters.Element) * Parameters.Num);
 		ActionDistributionVectorSizes.Add(GetActionDistributionVectorSize(Parameters.Element) * Parameters.Num);
+		ActionModifierVectorSizes.Add(1 + GetActionModifierVectorSize(Parameters.Element) * Parameters.Num);
 		TypeDataIndices.Add(ArrayData.Add(ElementData));
 
 		return { Index, Generation };
@@ -307,6 +343,7 @@
 		EncodedVectorSizes.Add(ElementData.EncodingSize);
 		ActionVectorSizes.Add(GetActionVectorSize(Parameters.Element));
 		ActionDistributionVectorSizes.Add(GetActionDistributionVectorSize(Parameters.Element));
+		ActionModifierVectorSizes.Add(1 + GetActionModifierVectorSize(Parameters.Element));
 		TypeDataIndices.Add(EncodingData.Add(ElementData));
 
 		return { Index, Generation };
@@ -347,6 +384,12 @@
 		return ActionDistributionVectorSizes[Element.Index];
 	}
 
+	int32 FSchema::GetActionModifierVectorSize(const FSchemaElement Element) const
+	{
+		UE_LEARNING_CHECK(IsValid(Element));
+		return ActionModifierVectorSizes[Element.Index];
+	}
+
 	FSchemaContinuousParameters FSchema::GetContinuous(const FSchemaElement Element) const
 	{
 		UE_LEARNING_CHECK(IsValid(Element) && GetType(Element) == EType::Continuous);
@@ -830,6 +873,386 @@
 		Generation++;
 	}
 
+
+	FModifierElement FModifier::CreateNull(const FName Tag)
+	{
+		const int32 Index = Types.Add(EType::Null);
+		Tags.Add(Tag);
+
+		ContinuousDataOffsets.Add(ContinuousMaskeds.Num());
+		ContinuousDataNums.Add(0);
+
+		DiscreteDataOffsets.Add(DiscreteValues.Num());
+		DiscreteDataNums.Add(0);
+
+		ElementDataOffsets.Add(SubElementModifiers.Num());
+		ElementDataNums.Add(0);
+
+		MaskedDataOffsets.Add(MaskedElementNames.Num());
+		MaskedDataNums.Add(0);
+
+		return { Index, Generation };
+	}
+
+	FModifierElement FModifier::CreateContinuous(const FModifierContinuousParameters Parameters, const FName Tag)
+	{
+		const int32 Index = Types.Add(EType::Continuous);
+		Tags.Add(Tag);
+
+		ContinuousDataOffsets.Add(ContinuousMaskeds.Num());
+		ContinuousDataNums.Add(Parameters.MaskedValues.Num());
+
+		DiscreteDataOffsets.Add(DiscreteValues.Num());
+		DiscreteDataNums.Add(0);
+
+		ElementDataOffsets.Add(SubElementModifiers.Num());
+		ElementDataNums.Add(0);
+
+		MaskedDataOffsets.Add(MaskedElementNames.Num());
+		MaskedDataNums.Add(0);
+
+		ContinuousMaskeds.Append(Parameters.Masked);
+		ContinuousMaskedValues.Append(Parameters.MaskedValues);
+
+		return { Index, Generation };
+	}
+
+	FModifierElement FModifier::CreateDiscreteExclusive(const FModifierDiscreteExclusiveParameters Parameters, const FName Tag)
+	{
+		const int32 Index = Types.Add(EType::DiscreteExclusive);
+		Tags.Add(Tag);
+
+		ContinuousDataOffsets.Add(ContinuousMaskeds.Num());
+		ContinuousDataNums.Add(0);
+
+		DiscreteDataOffsets.Add(DiscreteValues.Num());
+		DiscreteDataNums.Add(Parameters.MaskedIndices.Num());
+
+		ElementDataOffsets.Add(SubElementModifiers.Num());
+		ElementDataNums.Add(0);
+
+		MaskedDataOffsets.Add(MaskedElementNames.Num());
+		MaskedDataNums.Add(0);
+
+		DiscreteValues.Append(Parameters.MaskedIndices);
+
+		return { Index, Generation };
+	}
+
+	FModifierElement FModifier::CreateDiscreteInclusive(const FModifierDiscreteInclusiveParameters Parameters, const FName Tag)
+	{
+		const int32 Index = Types.Add(EType::DiscreteInclusive);
+		Tags.Add(Tag);
+
+		ContinuousDataOffsets.Add(ContinuousMaskeds.Num());
+		ContinuousDataNums.Add(0);
+
+		DiscreteDataOffsets.Add(DiscreteValues.Num());
+		DiscreteDataNums.Add(Parameters.MaskedIndices.Num());
+
+		ElementDataOffsets.Add(SubElementModifiers.Num());
+		ElementDataNums.Add(0);
+
+		MaskedDataOffsets.Add(MaskedElementNames.Num());
+		MaskedDataNums.Add(0);
+
+		DiscreteValues.Append(Parameters.MaskedIndices);
+
+		return { Index, Generation };
+	}
+
+	FModifierElement FModifier::CreateAnd(const FModifierAndParameters Parameters, const FName Tag)
+	{
+		UE_LEARNING_CHECK(Parameters.Elements.Num() == Parameters.ElementNames.Num());
+		UE_LEARNING_CHECK(!Private::ContainsDuplicates(Parameters.ElementNames));
+		UE_LEARNING_CHECK(Private::CheckAllValid(*this, Parameters.Elements));
+
+		const int32 Index = Types.Add(EType::And);
+		Tags.Add(Tag);
+
+		ContinuousDataOffsets.Add(ContinuousMaskeds.Num());
+		ContinuousDataNums.Add(0);
+
+		DiscreteDataOffsets.Add(DiscreteValues.Num());
+		DiscreteDataNums.Add(0);
+
+		ElementDataOffsets.Add(SubElementModifiers.Num());
+		ElementDataNums.Add(Parameters.Elements.Num());
+
+		MaskedDataOffsets.Add(MaskedElementNames.Num());
+		MaskedDataNums.Add(0);
+
+		SubElementModifiers.Append(Parameters.Elements);
+		SubElementNames.Append(Parameters.ElementNames);
+
+		return { Index, Generation };
+	}
+
+	FModifierElement FModifier::CreateOrExclusive(const FModifierOrExclusiveParameters Parameters, const FName Tag)
+	{
+		UE_LEARNING_CHECK(Parameters.Elements.Num() == Parameters.ElementNames.Num());
+		UE_LEARNING_CHECK(!Private::ContainsDuplicates(Parameters.ElementNames));
+		UE_LEARNING_CHECK(!Private::ContainsDuplicates(Parameters.MaskedElements));
+		UE_LEARNING_CHECK(Private::CheckAllValid(*this, Parameters.Elements));
+
+		const int32 Index = Types.Add(EType::OrExclusive);
+		Tags.Add(Tag);
+
+		ContinuousDataOffsets.Add(ContinuousMaskeds.Num());
+		ContinuousDataNums.Add(0);
+
+		DiscreteDataOffsets.Add(DiscreteValues.Num());
+		DiscreteDataNums.Add(0);
+
+		ElementDataOffsets.Add(SubElementModifiers.Num());
+		ElementDataNums.Add(Parameters.Elements.Num());
+
+		MaskedDataOffsets.Add(MaskedElementNames.Num());
+		MaskedDataNums.Add(Parameters.MaskedElements.Num());
+
+		SubElementModifiers.Append(Parameters.Elements);
+		SubElementNames.Append(Parameters.ElementNames);
+
+		MaskedElementNames.Append(Parameters.MaskedElements);
+
+		return { Index, Generation };
+	}
+
+	FModifierElement FModifier::CreateOrInclusive(const FModifierOrInclusiveParameters Parameters, const FName Tag)
+	{
+		UE_LEARNING_CHECK(Parameters.Elements.Num() == Parameters.ElementNames.Num());
+		UE_LEARNING_CHECK(!Private::ContainsDuplicates(Parameters.ElementNames));
+		UE_LEARNING_CHECK(!Private::ContainsDuplicates(Parameters.MaskedElements));
+		UE_LEARNING_CHECK(Private::CheckAllValid(*this, Parameters.Elements));
+
+		const int32 Index = Types.Add(EType::OrInclusive);
+		Tags.Add(Tag);
+
+		ContinuousDataOffsets.Add(ContinuousMaskeds.Num());
+		ContinuousDataNums.Add(0);
+
+		DiscreteDataOffsets.Add(DiscreteValues.Num());
+		DiscreteDataNums.Add(0);
+
+		ElementDataOffsets.Add(SubElementModifiers.Num());
+		ElementDataNums.Add(Parameters.Elements.Num());
+
+		MaskedDataOffsets.Add(MaskedElementNames.Num());
+		MaskedDataNums.Add(Parameters.MaskedElements.Num());
+
+		SubElementModifiers.Append(Parameters.Elements);
+		SubElementNames.Append(Parameters.ElementNames);
+
+		MaskedElementNames.Append(Parameters.MaskedElements);
+
+		return { Index, Generation };
+	}
+
+	FModifierElement FModifier::CreateArray(const FModifierArrayParameters Parameters, const FName Tag)
+	{
+		UE_LEARNING_CHECK(Private::CheckAllValid(*this, Parameters.Elements));
+
+		const int32 Index = Types.Add(EType::Array);
+		Tags.Add(Tag);
+
+		ContinuousDataOffsets.Add(ContinuousMaskeds.Num());
+		ContinuousDataNums.Add(0);
+
+		DiscreteDataOffsets.Add(DiscreteValues.Num());
+		DiscreteDataNums.Add(0);
+
+		ElementDataOffsets.Add(SubElementModifiers.Num());
+		ElementDataNums.Add(Parameters.Elements.Num());
+
+		MaskedDataOffsets.Add(MaskedElementNames.Num());
+		MaskedDataNums.Add(0);
+
+		for (int32 ElementIdx = 0; ElementIdx < Parameters.Elements.Num(); ElementIdx++)
+		{
+			SubElementNames.Add(NAME_None);
+		}
+		SubElementModifiers.Append(Parameters.Elements);
+
+		return { Index, Generation };
+	}
+
+	FModifierElement FModifier::CreateEncoding(const FModifierEncodingParameters Parameters, const FName Tag)
+	{
+		UE_LEARNING_CHECK(IsValid(Parameters.Element));
+
+		const int32 Index = Types.Add(EType::Encoding);
+		Tags.Add(Tag);
+
+		ContinuousDataOffsets.Add(ContinuousMaskeds.Num());
+		ContinuousDataNums.Add(0);
+
+		DiscreteDataOffsets.Add(DiscreteValues.Num());
+		DiscreteDataNums.Add(0);
+
+		ElementDataOffsets.Add(SubElementModifiers.Num());
+		ElementDataNums.Add(1);
+
+		MaskedDataOffsets.Add(MaskedElementNames.Num());
+		MaskedDataNums.Add(0);
+
+		SubElementNames.Add(NAME_None);
+		SubElementModifiers.Add(Parameters.Element);
+
+		return { Index, Generation };
+	}
+
+	bool FModifier::IsValid(const FModifierElement Element) const
+	{
+		return Element.Generation == Generation && Element.Index != INDEX_NONE;
+	}
+
+	EType FModifier::GetType(const FModifierElement Element) const
+	{
+		UE_LEARNING_CHECK(IsValid(Element));
+		return Types[Element.Index];
+	}
+
+	FName FModifier::GetTag(const FModifierElement Element) const
+	{
+		UE_LEARNING_CHECK(IsValid(Element));
+		return Tags[Element.Index];
+	}
+
+	FModifierContinuousParameters FModifier::GetContinuous(const FModifierElement Element) const
+	{
+		UE_LEARNING_CHECK(IsValid(Element) && GetType(Element) == EType::Continuous);
+
+		FModifierContinuousParameters Parameters;
+		Parameters.Masked = TArrayView<const bool>(ContinuousMaskeds.GetData() + ContinuousDataOffsets[Element.Index], ContinuousDataNums[Element.Index]);
+		Parameters.MaskedValues = TArrayView<const float>(ContinuousMaskedValues.GetData() + ContinuousDataOffsets[Element.Index], ContinuousDataNums[Element.Index]);
+		return Parameters;
+	}
+
+	FModifierDiscreteExclusiveParameters FModifier::GetDiscreteExclusive(const FModifierElement Element) const
+	{
+		UE_LEARNING_CHECK(IsValid(Element) && GetType(Element) == EType::DiscreteExclusive);
+
+		FModifierDiscreteExclusiveParameters Parameters;
+		Parameters.MaskedIndices = TArrayView<const int32>(DiscreteValues.GetData() + DiscreteDataOffsets[Element.Index], DiscreteDataNums[Element.Index]);
+		return Parameters;
+	}
+
+	FModifierDiscreteInclusiveParameters FModifier::GetDiscreteInclusive(const FModifierElement Element) const
+	{
+		UE_LEARNING_CHECK(IsValid(Element) && GetType(Element) == EType::DiscreteInclusive);
+
+		FModifierDiscreteInclusiveParameters Parameters;
+		Parameters.MaskedIndices = TArrayView<const int32>(DiscreteValues.GetData() + DiscreteDataOffsets[Element.Index], DiscreteDataNums[Element.Index]);
+		return Parameters;
+	}
+
+	FModifierAndParameters FModifier::GetAnd(const FModifierElement Element) const
+	{
+		UE_LEARNING_CHECK(IsValid(Element) && GetType(Element) == EType::And);
+
+		FModifierAndParameters Parameters;
+		Parameters.ElementNames = TArrayView<const FName>(SubElementNames.GetData() + ElementDataOffsets[Element.Index], ElementDataNums[Element.Index]);
+		Parameters.Elements = TArrayView<const FModifierElement>(SubElementModifiers.GetData() + ElementDataOffsets[Element.Index], ElementDataNums[Element.Index]);
+		return Parameters;
+	}
+
+	FModifierOrExclusiveParameters FModifier::GetOrExclusive(const FModifierElement Element) const
+	{
+		UE_LEARNING_CHECK(IsValid(Element) && GetType(Element) == EType::OrExclusive);
+
+		FModifierOrExclusiveParameters Parameters;
+		Parameters.ElementNames = TArrayView<const FName>(SubElementNames.GetData() + ElementDataOffsets[Element.Index], ElementDataNums[Element.Index]);
+		Parameters.Elements = TArrayView<const FModifierElement>(SubElementModifiers.GetData() + ElementDataOffsets[Element.Index], ElementDataNums[Element.Index]);
+		Parameters.MaskedElements = TArrayView<const FName>(MaskedElementNames.GetData() + MaskedDataOffsets[Element.Index], MaskedDataNums[Element.Index]);
+		return Parameters;
+	}
+
+	FModifierOrInclusiveParameters FModifier::GetOrInclusive(const FModifierElement Element) const
+	{
+		UE_LEARNING_CHECK(IsValid(Element) && GetType(Element) == EType::OrInclusive);
+
+		FModifierOrInclusiveParameters Parameters;
+		Parameters.ElementNames = TArrayView<const FName>(SubElementNames.GetData() + ElementDataOffsets[Element.Index], ElementDataNums[Element.Index]);
+		Parameters.Elements = TArrayView<const FModifierElement>(SubElementModifiers.GetData() + ElementDataOffsets[Element.Index], ElementDataNums[Element.Index]);
+		Parameters.MaskedElements = TArrayView<const FName>(MaskedElementNames.GetData() + MaskedDataOffsets[Element.Index], MaskedDataNums[Element.Index]);
+		return Parameters;
+	}
+
+	FModifierArrayParameters FModifier::GetArray(const FModifierElement Element) const
+	{
+		UE_LEARNING_CHECK(IsValid(Element) && GetType(Element) == EType::Array);
+
+		FModifierArrayParameters Parameters;
+		Parameters.Elements = TArrayView<const FModifierElement>(SubElementModifiers.GetData() + ElementDataOffsets[Element.Index], ElementDataNums[Element.Index]);
+		return Parameters;
+	}
+
+	FModifierEncodingParameters FModifier::GetEncoding(const FModifierElement Element) const
+	{
+		UE_LEARNING_CHECK(IsValid(Element) && GetType(Element) == EType::Encoding);
+
+		FModifierEncodingParameters Parameters;
+		Parameters.Element = SubElementModifiers[ElementDataOffsets[Element.Index]];
+		return Parameters;
+	}
+
+	uint32 FModifier::GetGeneration() const
+	{
+		return Generation;
+	}
+
+	void FModifier::Empty()
+	{
+		Types.Empty();
+		Tags.Empty();
+		ContinuousDataOffsets.Empty();
+		ContinuousDataNums.Empty();
+		DiscreteDataOffsets.Empty();
+		DiscreteDataNums.Empty();
+		ElementDataOffsets.Empty();
+		ElementDataNums.Empty();
+		MaskedDataOffsets.Empty();
+		MaskedDataNums.Empty();
+
+		ContinuousMaskeds.Empty();
+		ContinuousMaskedValues.Empty();
+		DiscreteValues.Empty();
+		SubElementModifiers.Empty();
+		SubElementNames.Empty();
+		MaskedElementNames.Empty();
+
+		Generation++;
+	}
+
+	bool FModifier::IsEmpty() const
+	{
+		return Types.IsEmpty();
+	}
+
+	void FModifier::Reset()
+	{
+		Types.Reset();
+		Tags.Reset();
+		ContinuousDataOffsets.Reset();
+		ContinuousDataNums.Reset();
+		DiscreteDataOffsets.Reset();
+		DiscreteDataNums.Reset();
+		ElementDataOffsets.Reset();
+		ElementDataNums.Reset();
+		MaskedDataOffsets.Reset();
+		MaskedDataNums.Reset();
+
+		ContinuousMaskeds.Reset();
+		ContinuousMaskedValues.Reset();
+		DiscreteValues.Reset();
+		SubElementModifiers.Reset();
+		SubElementNames.Reset();
+		MaskedElementNames.Reset();
+
+		Generation++;
+	}
+
+
 	namespace Private
 	{
 		static inline NNE::RuntimeBasic::FModelBuilder::EActivationFunction GetNNEActivationFunction(const EEncodingActivationFunction ActivationFunction)
@@ -1266,6 +1689,7 @@
 		uint32& InOutRandomState,
 		TLearningArrayView<1, float> OutActionVector,
 		const TLearningArrayView<1, const float> ActionDistributionVector,
+		const TLearningArrayView<1, const float> ActionModifierVector,
 		const FSchema& Schema,
 		const FSchemaElement SchemaElement,
 		const float ActionNoiseScale)
@@ -1283,13 +1707,35 @@
 			const int32 ValueNum = Schema.GetContinuous(SchemaElement).Num;
 			UE_LEARNING_CHECK(ValueNum == OutActionVector.Num());
 			UE_LEARNING_CHECK(ValueNum * 2 == ActionDistributionVector.Num());
+			UE_LEARNING_CHECK(1 + ValueNum * 2 == ActionModifierVector.Num());
 
-			Random::SampleDistributionIndependantNormal(
-				OutActionVector,
-				InOutRandomState,
-				ActionDistributionVector.Slice(0, ValueNum),
-				ActionDistributionVector.Slice(ValueNum, ValueNum),
-				ActionNoiseScale);
+			if (ActionModifierVector[0])
+			{
+				TLearningArray<1, bool, TInlineAllocator<32>> Masked;
+				Masked.SetNumUninitialized({ ValueNum });
+				for (int32 ValueIdx = 0; ValueIdx < ValueNum; ValueIdx++)
+				{
+					Masked[ValueIdx] = ActionModifierVector[1 + ValueIdx] == 1.0f;
+				}
+
+				Random::SampleDistributionIndependantNormalMasked(
+					OutActionVector,
+					InOutRandomState,
+					ActionDistributionVector.Slice(0, ValueNum),
+					ActionDistributionVector.Slice(ValueNum, ValueNum),
+					Masked,
+					ActionModifierVector.Slice(1 + ValueNum, ValueNum),
+					ActionNoiseScale);
+			}
+			else
+			{
+				Random::SampleDistributionIndependantNormal(
+					OutActionVector,
+					InOutRandomState,
+					ActionDistributionVector.Slice(0, ValueNum),
+					ActionDistributionVector.Slice(ValueNum, ValueNum),
+					ActionNoiseScale);
+			}
 
 			break;
 		}
@@ -1299,12 +1745,33 @@
 			const int32 ValueNum = Schema.GetDiscreteExclusive(SchemaElement).Num;
 			UE_LEARNING_CHECK(ValueNum == OutActionVector.Num());
 			UE_LEARNING_CHECK(ValueNum == ActionDistributionVector.Num());
+			UE_LEARNING_CHECK(1 + ValueNum == ActionModifierVector.Num());
 
-			Random::SampleDistributionMultinoulli(
-				OutActionVector,
-				InOutRandomState,
-				ActionDistributionVector,
-				ActionNoiseScale);
+			if (ActionModifierVector[0])
+			{
+				TLearningArray<1, bool, TInlineAllocator<32>> Masked;
+				Masked.SetNumUninitialized({ ValueNum });
+				for (int32 ValueIdx = 0; ValueIdx < ValueNum; ValueIdx++)
+				{
+					Masked[ValueIdx] = ActionModifierVector[1 + ValueIdx] == 1.0f;
+				}
+				UE_LEARNING_CHECK(Private::CheckExclusiveMaskValid(Masked));
+
+				Random::SampleDistributionMultinoulliMasked(
+					OutActionVector,
+					InOutRandomState,
+					ActionDistributionVector,
+					Masked,
+					ActionNoiseScale);
+			}
+			else
+			{
+				Random::SampleDistributionMultinoulli(
+					OutActionVector,
+					InOutRandomState,
+					ActionDistributionVector,
+					ActionNoiseScale);
+			}
 
 			break;
 		}
@@ -1314,12 +1781,32 @@
 			const int32 ValueNum = Schema.GetDiscreteInclusive(SchemaElement).Num;
 			UE_LEARNING_CHECK(ValueNum == OutActionVector.Num());
 			UE_LEARNING_CHECK(ValueNum == ActionDistributionVector.Num());
+			UE_LEARNING_CHECK(1 + ValueNum == ActionModifierVector.Num());
 
-			Random::SampleDistributionBernoulli(
-				OutActionVector,
-				InOutRandomState,
-				ActionDistributionVector,
-				ActionNoiseScale);
+			if (ActionModifierVector[0])
+			{
+				TLearningArray<1, bool, TInlineAllocator<32>> Masked;
+				Masked.SetNumUninitialized({ ValueNum });
+				for (int32 ValueIdx = 0; ValueIdx < ValueNum; ValueIdx++)
+				{
+					Masked[ValueIdx] = ActionModifierVector[1 + ValueIdx] == 1.0f;
+				}
+
+				Random::SampleDistributionBernoulliMasked(
+					OutActionVector,
+					InOutRandomState,
+					ActionDistributionVector,
+					Masked,
+					ActionNoiseScale);
+			}
+			else
+			{
+				Random::SampleDistributionBernoulli(
+					OutActionVector,
+					InOutRandomState,
+					ActionDistributionVector,
+					ActionNoiseScale);
+			}
 
 			break;
 		}
@@ -1330,26 +1817,31 @@
 
 			int32 SubElementActionVectorOffset = 0;
 			int32 SubElementActionDistributionVectorOffset = 0;
+			int32 SubElementActionModifierVectorOffset = 1;
 
 			for (const FSchemaElement SubElement : Parameters.Elements)
 			{
 				const int32 SubElementActionVectorSize = Schema.GetActionVectorSize(SubElement);
 				const int32 SubElementActionDistributionVectorSize = Schema.GetActionDistributionVectorSize(SubElement);
+				const int32 SubElementActionModifierVectorSize = Schema.GetActionModifierVectorSize(SubElement);
 
 				SampleVectorFromDistributionVector(
 					InOutRandomState,
 					OutActionVector.Slice(SubElementActionVectorOffset, SubElementActionVectorSize),
 					ActionDistributionVector.Slice(SubElementActionDistributionVectorOffset, SubElementActionDistributionVectorSize),
+					ActionModifierVector.Slice(SubElementActionModifierVectorOffset, SubElementActionModifierVectorSize),
 					Schema,
 					SubElement,
 					ActionNoiseScale);
 
 				SubElementActionVectorOffset += SubElementActionVectorSize;
 				SubElementActionDistributionVectorOffset += SubElementActionDistributionVectorSize;
+				SubElementActionModifierVectorOffset += SubElementActionModifierVectorSize;
 			}
 
 			UE_LEARNING_CHECK(SubElementActionVectorOffset == OutActionVector.Num());
 			UE_LEARNING_CHECK(SubElementActionDistributionVectorOffset == ActionDistributionVector.Num());
+			UE_LEARNING_CHECK(SubElementActionModifierVectorOffset == ActionModifierVector.Num());
 
 			break;
 		}
@@ -1360,42 +1852,79 @@
 
 			const int32 SubElementActionVectorMax = Private::GetMaxActionVectorSize(Schema, Parameters.Elements);
 			const int32 SubElementActionDistributionVectorTotal = Private::GetTotalActionDistributionVectorSize(Schema, Parameters.Elements);
+			const int32 SubElementActionModifierVectorTotal = Private::GetTotalActionModifierVectorSize(Schema, Parameters.Elements);
+			const int32 ElementNum = Parameters.Elements.Num();
+
+			UE_LEARNING_CHECK(SubElementActionVectorMax + ElementNum == OutActionVector.Num());
+			UE_LEARNING_CHECK(SubElementActionDistributionVectorTotal + ElementNum == ActionDistributionVector.Num());
+			UE_LEARNING_CHECK(1 + ElementNum + SubElementActionModifierVectorTotal == ActionModifierVector.Num());
 
 			// Zero main part of vector
 			Array::Zero(OutActionVector.Slice(0, SubElementActionVectorMax));
 
-			// Sample which sub-element to generate
-			Random::SampleDistributionMultinoulli(
-				OutActionVector.Slice(SubElementActionVectorMax, Parameters.Elements.Num()),
-				InOutRandomState,
-				ActionDistributionVector.Slice(SubElementActionDistributionVectorTotal, Parameters.Elements.Num()),
-				ActionNoiseScale);
+			if (ActionModifierVector[0])
+			{
+				TLearningArray<1, bool, TInlineAllocator<32>> Masked;
+				Masked.SetNumUninitialized({ ElementNum });
+				for (int32 ElementIdx = 0; ElementIdx < ElementNum; ElementIdx++)
+				{
+					Masked[ElementIdx] = ActionModifierVector[1 + ElementIdx] == 1.0f;
+				}
+				UE_LEARNING_CHECK(Private::CheckExclusiveMaskValid(Masked));
 
+				// Sample which sub-element to generate
+				Random::SampleDistributionMultinoulliMasked(
+					OutActionVector.Slice(SubElementActionVectorMax, ElementNum),
+					InOutRandomState,
+					ActionDistributionVector.Slice(SubElementActionDistributionVectorTotal, ElementNum),
+					Masked,
+					ActionNoiseScale);
+			}
+			else
+			{
+				// Sample which sub-element to generate
+				Random::SampleDistributionMultinoulli(
+					OutActionVector.Slice(SubElementActionVectorMax, ElementNum),
+					InOutRandomState,
+					ActionDistributionVector.Slice(SubElementActionDistributionVectorTotal, ElementNum),
+					ActionNoiseScale);
+			}
+
+			int32 SubElementsSampled = 0;
 			int32 SubElementActionDistributionVectorOffset = 0;
+			int32 SubElementActionModifierVectorOffset = 1 + ElementNum;
 
-			for (int32 SubElementIdx = 0; SubElementIdx < Parameters.Elements.Num(); SubElementIdx++)
+			for (int32 SubElementIdx = 0; SubElementIdx < ElementNum; SubElementIdx++)
 			{
+				const FSchemaElement SubElement = Parameters.Elements[SubElementIdx];
+				const int32 SubElementActionVectorSize = Schema.GetActionVectorSize(SubElement);
+				const int32 SubElementActionDistributionVectorSize = Schema.GetActionDistributionVectorSize(SubElement);
+				const int32 SubElementActionModifierVectorSize = Schema.GetActionModifierVectorSize(SubElement);
+
+				UE_LEARNING_CHECK(SubElementActionVectorSize <= SubElementActionVectorMax);
+
 				if (OutActionVector[SubElementActionVectorMax + SubElementIdx])
 				{
-					const FSchemaElement SubElement = Parameters.Elements[SubElementIdx];
-					const int32 SubElementActionVectorSize = Schema.GetActionVectorSize(SubElement);
-					const int32 SubElementActionDistributionVectorSize = Schema.GetActionDistributionVectorSize(SubElement);
-
 					// Sample Sub-Element
 					SampleVectorFromDistributionVector(
 						InOutRandomState,
 						OutActionVector.Slice(0, SubElementActionVectorSize),
 						ActionDistributionVector.Slice(SubElementActionDistributionVectorOffset, SubElementActionDistributionVectorSize),
+						ActionModifierVector.Slice(SubElementActionModifierVectorOffset, SubElementActionModifierVectorSize),
 						Schema,
 						SubElement,
 						ActionNoiseScale);
 
-					SubElementActionDistributionVectorOffset += SubElementActionDistributionVectorSize;
+					SubElementsSampled++;
+				}
 
-					// We can break early because this exclusively samples one action
-					break;
-				}
+				SubElementActionDistributionVectorOffset += SubElementActionDistributionVectorSize;
+				SubElementActionModifierVectorOffset += SubElementActionModifierVectorSize;
 			}
+
+			UE_LEARNING_CHECK(SubElementsSampled == 1); // Exactly one sub-element should have been sampled
+			UE_LEARNING_CHECK(SubElementActionDistributionVectorOffset == SubElementActionDistributionVectorTotal);
+			UE_LEARNING_CHECK(SubElementActionModifierVectorOffset == 1 + ElementNum + SubElementActionModifierVectorTotal);
 
 			break;
 		}
@@ -1406,41 +1935,75 @@
 
 			const int32 SubElementActionVectorTotal = Private::GetTotalActionVectorSize(Schema, Parameters.Elements);
 			const int32 SubElementActionDistributionVectorTotal = Private::GetTotalActionDistributionVectorSize(Schema, Parameters.Elements);
+			const int32 SubElementActionModifierVectorTotal = Private::GetTotalActionModifierVectorSize(Schema, Parameters.Elements);
+			const int32 ElementNum = Parameters.Elements.Num();
+
+			UE_LEARNING_CHECK(SubElementActionVectorTotal + ElementNum == OutActionVector.Num());
+			UE_LEARNING_CHECK(SubElementActionDistributionVectorTotal + ElementNum == ActionDistributionVector.Num());
+			UE_LEARNING_CHECK(1 + ElementNum + SubElementActionModifierVectorTotal == ActionModifierVector.Num());
 
 			// Zero main part of vector
 			Array::Zero(OutActionVector.Slice(0, SubElementActionVectorTotal));
 
-			// Sample which sub-elements to generate
-			Random::SampleDistributionBernoulli(
-				OutActionVector.Slice(SubElementActionVectorTotal, Parameters.Elements.Num()),
-				InOutRandomState,
-				ActionDistributionVector.Slice(SubElementActionDistributionVectorTotal, Parameters.Elements.Num()),
-				ActionNoiseScale);
+			if (ActionModifierVector[0])
+			{
+				TLearningArray<1, bool, TInlineAllocator<32>> Masked;
+				Masked.SetNumUninitialized({ ElementNum });
+				for (int32 ElementIdx = 0; ElementIdx < ElementNum; ElementIdx++)
+				{
+					Masked[ElementIdx] = ActionModifierVector[1 + ElementIdx] == 1.0f;
+				}
+
+				// Sample which sub-elements to generate
+				Random::SampleDistributionBernoulliMasked(
+					OutActionVector.Slice(SubElementActionVectorTotal, ElementNum),
+					InOutRandomState,
+					ActionDistributionVector.Slice(SubElementActionDistributionVectorTotal, ElementNum),
+					Masked,
+					ActionNoiseScale);
+			}
+			else
+			{
+				// Sample which sub-elements to generate
+				Random::SampleDistributionBernoulli(
+					OutActionVector.Slice(SubElementActionVectorTotal, ElementNum),
+					InOutRandomState,
+					ActionDistributionVector.Slice(SubElementActionDistributionVectorTotal, ElementNum),
+					ActionNoiseScale);
+			}
 
 			int32 SubElementActionVectorOffset = 0;
 			int32 SubElementActionDistributionVectorOffset = 0;
+			int32 SubElementActionModifierVectorOffset = 1 + ElementNum;
 
-			for (int32 SubElementIdx = 0; SubElementIdx < Parameters.Elements.Num(); SubElementIdx++)
+			for (int32 SubElementIdx = 0; SubElementIdx < ElementNum; SubElementIdx++)
 			{
+				const FSchemaElement SubElement = Parameters.Elements[SubElementIdx];
+				const int32 SubElementActionVectorSize = Schema.GetActionVectorSize(SubElement);
+				const int32 SubElementActionDistributionVectorSize = Schema.GetActionDistributionVectorSize(SubElement);
+				const int32 SubElementActionModifierVectorSize = Schema.GetActionModifierVectorSize(SubElement);
+
 				if (OutActionVector[SubElementActionVectorTotal + SubElementIdx])
 				{
-					const FSchemaElement SubElement = Parameters.Elements[SubElementIdx];
-					const int32 SubElementActionVectorSize = Schema.GetActionVectorSize(SubElement);
-					const int32 SubElementActionDistributionVectorSize = Schema.GetActionDistributionVectorSize(SubElement);
-
 					// Sample sub-elements
 					SampleVectorFromDistributionVector(
 						InOutRandomState,
 						OutActionVector.Slice(SubElementActionVectorOffset, SubElementActionVectorSize),
 						ActionDistributionVector.Slice(SubElementActionDistributionVectorOffset, SubElementActionDistributionVectorSize),
+						ActionModifierVector.Slice(SubElementActionModifierVectorOffset, SubElementActionModifierVectorSize),
 						Schema,
 						SubElement,
 						ActionNoiseScale);
+				}
 
-					SubElementActionVectorOffset += SubElementActionVectorSize;
-					SubElementActionDistributionVectorOffset += SubElementActionDistributionVectorSize;
-				}
+				SubElementActionVectorOffset += SubElementActionVectorSize;
+				SubElementActionDistributionVectorOffset += SubElementActionDistributionVectorSize;
+				SubElementActionModifierVectorOffset += SubElementActionModifierVectorSize;
 			}
+
+			UE_LEARNING_CHECK(SubElementActionVectorOffset == SubElementActionVectorTotal);
+			UE_LEARNING_CHECK(SubElementActionDistributionVectorOffset == SubElementActionDistributionVectorTotal);
+			UE_LEARNING_CHECK(SubElementActionModifierVectorOffset == 1 + ElementNum + SubElementActionModifierVectorTotal);
 
 			break;
 		}
@@ -1451,6 +2014,11 @@
 
 			const int32 SubElementActionVectorSize = Schema.GetActionVectorSize(Parameters.Element);
 			const int32 SubElementActionDistributionVectorSize = Schema.GetActionDistributionVectorSize(Parameters.Element);
+			const int32 SubElementActionModifierVectorSize = Schema.GetActionModifierVectorSize(Parameters.Element);
+
+			UE_LEARNING_CHECK(SubElementActionVectorSize * Parameters.Num == OutActionVector.Num());
+			UE_LEARNING_CHECK(SubElementActionDistributionVectorSize * Parameters.Num == ActionDistributionVector.Num());
+			UE_LEARNING_CHECK(1 + SubElementActionModifierVectorSize * Parameters.Num == ActionModifierVector.Num());
 
 			for (int32 ElementIdx = 0; ElementIdx < Parameters.Num; ElementIdx++)
 			{
@@ -1458,6 +2026,7 @@
 					InOutRandomState,
 					OutActionVector.Slice(ElementIdx * SubElementActionVectorSize, SubElementActionVectorSize),
 					ActionDistributionVector.Slice(ElementIdx * SubElementActionDistributionVectorSize, SubElementActionDistributionVectorSize),
+					ActionModifierVector.Slice(1 + ElementIdx * SubElementActionModifierVectorSize, SubElementActionModifierVectorSize),
 					Schema,
 					Parameters.Element,
 					ActionNoiseScale);
@@ -1469,11 +2038,13 @@
 		case EType::Encoding:
 		{
 			const FSchemaEncodingParameters Parameters = Schema.GetEncoding(SchemaElement);
+			const int32 SubElementActionModifierVectorSize = Schema.GetActionModifierVectorSize(Parameters.Element);
 
 			SampleVectorFromDistributionVector(
 				InOutRandomState,
 				OutActionVector,
 				ActionDistributionVector,
+				ActionModifierVector.Slice(1, SubElementActionModifierVectorSize),
 				Schema,
 				Parameters.Element,
 				ActionNoiseScale);
@@ -1534,3 +2105,501 @@
 			}
 
 			return;
+		}
+
+		case EType::DiscreteExclusive:
+		{
+			int32 ActionValue = Object.GetDiscreteExclusive(ObjectElement).DiscreteIndex;
+			UE_LEARNING_CHECK(Schema.GetActionVectorSize(SchemaElement) > ActionValue && ActionValue >= 0);
+			UE_LEARNING_CHECK(Schema.GetActionVectorSize(SchemaElement) == OutActionVector.Num());
+
+			// Set the single value in the action vector
+
+			OutActionVector[ActionValue] = 1.0f;
+			return;
+		}
+
+		case EType::DiscreteInclusive:
+		{
+			TArrayView<const int32> ActionValues = Object.GetDiscreteInclusive(ObjectElement).DiscreteIndices;
+			UE_LEARNING_CHECK(Schema.GetActionVectorSize(SchemaElement) >= ActionValues.Num());
+			UE_LEARNING_CHECK(Schema.GetActionVectorSize(SchemaElement) == OutActionVector.Num());
+
+			// Set values in the action vector
+
+			for (int32 ActionValueIdx = 0; ActionValueIdx < ActionValues.Num(); ActionValueIdx++)
+			{
+				UE_LEARNING_CHECK(Schema.GetActionVectorSize(SchemaElement) > ActionValues[ActionValueIdx] && ActionValues[ActionValueIdx] >= 0);
+				OutActionVector[ActionValues[ActionValueIdx]] = 1.0f;
+			}
+
+			return;
+		}
+
+		case EType::And:
+		{
+			// Check the number of sub-elements match
+
+			const FSchemaAndParameters SchemaParameters = Schema.GetAnd(SchemaElement);
+			const FObjectAndParameters ObjectParameters = Object.GetAnd(ObjectElement);
+			UE_LEARNING_CHECK(SchemaParameters.Elements.Num() == ObjectParameters.Elements.Num());
+
+			// Set the Sub-elements
+
+			int32 SubElementOffset = 0;
+
+			for (int32 SchemaElementIdx = 0; SchemaElementIdx < SchemaParameters.Elements.Num(); SchemaElementIdx++)
+			{
+				const int32 ObjectElementIndex = ObjectParameters.ElementNames.Find(SchemaParameters.ElementNames[SchemaElementIdx]);
+				UE_LEARNING_CHECK(ObjectElementIndex != INDEX_NONE);
+
+				const int32 SubElementSize = Schema.GetActionVectorSize(SchemaParameters.Elements[SchemaElementIdx]);
+
+				SetVectorFromObject(
+					OutActionVector.Slice(SubElementOffset, SubElementSize),
+					Schema,
+					SchemaParameters.Elements[SchemaElementIdx],
+					Object,
+					ObjectParameters.Elements[ObjectElementIndex]);
+
+				SubElementOffset += SubElementSize;
+			}
+
+			UE_LEARNING_CHECK(SubElementOffset == OutActionVector.Num());
+			return;
+		}
+
+		case EType::OrExclusive:
+		{
+			// Check only one sub-element is given and index is valid
+
+			const FSchemaOrExclusiveParameters SchemaParameters = Schema.GetOrExclusive(SchemaElement);
+			const FObjectOrExclusiveParameters ObjectParameters = Object.GetOrExclusive(ObjectElement);
+
+			const int32 SchemaElementIndex = SchemaParameters.ElementNames.Find(ObjectParameters.ElementName);
+			UE_LEARNING_CHECK(SchemaElementIndex != INDEX_NONE);
+
+			// Set the sub-element
+
+			const int32 SubElementSize = Schema.GetActionVectorSize(SchemaParameters.Elements[SchemaElementIndex]);
+
+			SetVectorFromObject(
+				OutActionVector.Slice(0, SubElementSize),
+				Schema,
+				SchemaParameters.Elements[SchemaElementIndex],
+				Object,
+				ObjectParameters.Element);
+
+			// Set Mask
+
+			const int32 MaxSubElementSize = Private::GetMaxActionVectorSize(Schema, SchemaParameters.Elements);
+
+			OutActionVector[MaxSubElementSize + SchemaElementIndex] = 1.0f;
+
+			UE_LEARNING_CHECK(OutActionVector.Num() == MaxSubElementSize + SchemaParameters.Elements.Num());
+			return;
+		}
+
+		case EType::OrInclusive:
+		{
+			// Check all indices are in range
+
+			const FSchemaOrInclusiveParameters SchemaParameters = Schema.GetOrInclusive(SchemaElement);
+			const FObjectOrInclusiveParameters ObjectParameters = Object.GetOrInclusive(ObjectElement);
+			UE_LEARNING_CHECK(ObjectParameters.Elements.Num() <= SchemaParameters.Elements.Num());
+
+			// Update sub-elements
+
+			int32 SubElementOffset = 0;
+
+			for (int32 ObjectElementIdx = 0; ObjectElementIdx < ObjectParameters.Elements.Num(); ObjectElementIdx++)
+			{
+				const int32 SchemaElementIdx = SchemaParameters.ElementNames.Find(ObjectParameters.ElementNames[ObjectElementIdx]);
+				UE_LEARNING_CHECK(SchemaElementIdx != INDEX_NONE);
+
+				const int32 SubElementSize = Schema.GetActionVectorSize(SchemaParameters.Elements[SchemaElementIdx]);
+
+				SetVectorFromObject(
+					OutActionVector.Slice(SubElementOffset, SubElementSize),
+					Schema,
+					SchemaParameters.Elements[SchemaElementIdx],
+					Object,
+					ObjectParameters.Elements[ObjectElementIdx]);
+
+				SubElementOffset += SubElementSize;
+			}
+
+			// Set Mask
+
+			UE_LEARNING_CHECK(SubElementOffset + SchemaParameters.Elements.Num() == OutActionVector.Num());
+
+			for (int32 ObjectElementIdx = 0; ObjectElementIdx < ObjectParameters.Elements.Num(); ObjectElementIdx++)
+			{
+				const int32 SchemaElementIdx = SchemaParameters.ElementNames.Find(ObjectParameters.ElementNames[ObjectElementIdx]);
+				UE_LEARNING_CHECK(SchemaElementIdx != INDEX_NONE);
+
+				OutActionVector[SubElementOffset + SchemaElementIdx] = 1.0f;
+			}
+
+			return;
+		}
+
+		case EType::Array:
+		{
+			// Check number of array elements is correct
+
+			const FSchemaArrayParameters SchemaParameters = Schema.GetArray(SchemaElement);
+			const FObjectArrayParameters ObjectParameters = Object.GetArray(ObjectElement);
+			UE_LEARNING_CHECK(SchemaParameters.Num == ObjectParameters.Elements.Num());
+
+			// Update sub-elements
+
+			const int32 SubElementSize = Schema.GetActionVectorSize(SchemaParameters.Element);
+
+			for (int32 ElementIdx = 0; ElementIdx < SchemaParameters.Num; ElementIdx++)
+			{
+				SetVectorFromObject(
+					OutActionVector.Slice(ElementIdx * SubElementSize, SubElementSize),
+					Schema,
+					SchemaParameters.Element,
+					Object,
+					ObjectParameters.Elements[ElementIdx]);
+			}
+
+			return;
+		}
+
+		case EType::Encoding:
+		{
+			const FSchemaEncodingParameters SchemaParameters = Schema.GetEncoding(SchemaElement);
+			const FObjectEncodingParameters ObjectParameters = Object.GetEncoding(ObjectElement);
+
+			SetVectorFromObject(
+				OutActionVector,
+				Schema,
+				SchemaParameters.Element,
+				Object,
+				ObjectParameters.Element);
+
+			return;
+		}
+
+		default:
+		{
+			UE_LEARNING_NOT_IMPLEMENTED();
+			return;
+		}
+
+		}
+	}
+
+	void GetObjectFromVector(
+		FObject& OutObject,
+		FObjectElement& OutObjectElement,
+		const FSchema& Schema,
+		const FSchemaElement SchemaElement,
+		const TLearningArrayView<1, const float> ActionVector)
+	{
+		UE_LEARNING_CHECK(Schema.IsValid(SchemaElement));
+
+		// Check that the types match
+
+		const EType SchemaElementType = Schema.GetType(SchemaElement);
+		const FName SchemaElementTag = Schema.GetTag(SchemaElement);
+
+		// Get Action Vector Size
+
+		const int32 ActionVectorSize = ActionVector.Num();
+		UE_LEARNING_CHECK(ActionVectorSize == Schema.GetActionVectorSize(SchemaElement));
+
+		// Logic for each specific element type
+
+		switch (SchemaElementType)
+		{
+
+		case EType::Null:
+		{
+			OutObjectElement = OutObject.CreateNull(SchemaElementTag);
+			return;
+		}
+
+		case EType::Continuous:
+		{
+			const FSchemaContinuousParameters SchemaParameters = Schema.GetContinuous(SchemaElement);
+			UE_LEARNING_CHECK(ActionVectorSize == SchemaParameters.Num);
+
+			const int32 ValueNum = SchemaParameters.Num;
+			const float ValueScale = FMath::Max(SchemaParameters.Scale, UE_SMALL_NUMBER);
+
+			TLearningArray<1, float, TInlineAllocator<32>> ActionValues;
+			ActionValues.SetNumUninitialized({ ValueNum });
+			for (int32 ValueIdx = 0; ValueIdx < ValueNum; ValueIdx++)
+			{
+				ActionValues[ValueIdx] = ValueScale * ActionVector[ValueIdx];
+			}
+
+			OutObjectElement = OutObject.CreateContinuous({ MakeArrayView(ActionValues.GetData(), ActionValues.Num()) }, SchemaElementTag);
+			return;
+		}
+
+		case EType::DiscreteExclusive:
+		{
+			UE_LEARNING_CHECK(ActionVectorSize == Schema.GetDiscreteExclusive(SchemaElement).Num);
+
+			// Find Index
+			int32 ExclusiveIndex = INDEX_NONE;
+			for (int32 Idx = 0; Idx < ActionVectorSize; Idx++)
+			{
+				UE_LEARNING_CHECK(ActionVector[Idx] == 0.0f || ActionVector[Idx] == 1.0f);
+				if (ActionVector[Idx])
+				{
+					ExclusiveIndex = Idx;
+					break;
+				}
+			}
+			UE_LEARNING_CHECK(ExclusiveIndex != INDEX_NONE);
+
+			OutObjectElement = OutObject.CreateDiscreteExclusive({ ExclusiveIndex }, SchemaElementTag);
+			return;
+		}
+
+		case EType::DiscreteInclusive:
+		{
+			UE_LEARNING_CHECK(ActionVectorSize == Schema.GetDiscreteInclusive(SchemaElement).Num);
+
+			// Find Indices
+			TArray<int32, TInlineAllocator<8>> InclusiveIndices;
+			InclusiveIndices.Reserve(ActionVectorSize);
+			for (int32 Idx = 0; Idx < ActionVectorSize; Idx++)
+			{
+				UE_LEARNING_CHECK(ActionVector[Idx] == 0.0f || ActionVector[Idx] == 1.0f);
+				if (ActionVector[Idx])
+				{
+					InclusiveIndices.Add(Idx);
+				}
+			}
+
+			OutObjectElement = OutObject.CreateDiscreteInclusive({ InclusiveIndices }, SchemaElementTag);
+			return;
+		}
+
+		case EType::And:
+		{
+			const FSchemaAndParameters Parameters = Schema.GetAnd(SchemaElement);
+			
+			// Create Sub-elements
+
+			TArray<FObjectElement, TInlineAllocator<8>> SubElements;
+			SubElements.SetNumUninitialized(Parameters.Elements.Num());
+
+			int32 SubElementOffset = 0;
+			for (int32 SchemaElementIdx = 0; SchemaElementIdx < Parameters.Elements.Num(); SchemaElementIdx++)
+			{
+				const int32 SubElementSize = Schema.GetActionVectorSize(Parameters.Elements[SchemaElementIdx]);
+
+				GetObjectFromVector(
+					OutObject,
+					SubElements[SchemaElementIdx],
+					Schema,
+					Parameters.Elements[SchemaElementIdx],
+					ActionVector.Slice(SubElementOffset, SubElementSize));
+
+				SubElementOffset += SubElementSize;
+			}
+			UE_LEARNING_CHECK(SubElementOffset == ActionVectorSize);
+
+			OutObjectElement = OutObject.CreateAnd({ Parameters.ElementNames, SubElements }, SchemaElementTag);
+			return;
+		}
+
+		case EType::OrExclusive:
+		{
+			const FSchemaOrExclusiveParameters Parameters = Schema.GetOrExclusive(SchemaElement);
+
+			// Find active element
+
+			const int32 MaxSubElementSize = Private::GetMaxActionVectorSize(Schema, Parameters.Elements);
+
+			int32 SchemaElementIndex = INDEX_NONE;
+			for (int32 SubElementIdx = 0; SubElementIdx < Parameters.Elements.Num(); SubElementIdx++)
+			{
+				UE_LEARNING_CHECK(ActionVector[MaxSubElementSize + SubElementIdx] == 0.0f || ActionVector[MaxSubElementSize + SubElementIdx] == 1.0f);
+				if (ActionVector[MaxSubElementSize + SubElementIdx])
+				{
+					SchemaElementIndex = SubElementIdx;
+					break;
+				}
+			}
+			UE_LEARNING_CHECK(SchemaElementIndex != INDEX_NONE);
+
+			// Create sub-element
+
+			const int32 SubElementSize = Schema.GetActionVectorSize(Parameters.Elements[SchemaElementIndex]);
+
+			FObjectElement SubElement;
+			GetObjectFromVector(
+				OutObject,
+				SubElement,
+				Schema,
+				Parameters.Elements[SchemaElementIndex],
+				ActionVector.Slice(0, SubElementSize));
+
+			OutObjectElement = OutObject.CreateOrExclusive({ Parameters.ElementNames[SchemaElementIndex], SubElement }, SchemaElementTag);
+			return;
+		}
+
+		case EType::OrInclusive:
+		{
+			const FSchemaOrInclusiveParameters Parameters = Schema.GetOrInclusive(SchemaElement);
+
+			// Find total sub-element size
+
+			const int32 TotalSubElementSize = Private::GetTotalActionVectorSize(Schema, Parameters.Elements);
+
+			// Create sub-elements
+
+			TArray<FName, TInlineAllocator<8>> SubElementNames;
+			TArray<FObjectElement, TInlineAllocator<8>> SubElements;
+			SubElementNames.Reserve(Parameters.Elements.Num());
+			SubElements.Reserve(Parameters.Elements.Num());
+
+			int32 SubElementOffset = 0;
+			for (int32 SubElementIdx = 0; SubElementIdx < Parameters.Elements.Num(); SubElementIdx++)
+			{
+				const int32 SubElementSize = Schema.GetActionVectorSize(Parameters.Elements[SubElementIdx]);
+
+				UE_LEARNING_CHECK(ActionVector[TotalSubElementSize + SubElementIdx] == 0.0f || ActionVector[TotalSubElementSize + SubElementIdx] == 1.0f);
+				if (ActionVector[TotalSubElementSize + SubElementIdx])
+				{
+					FObjectElement SubElement;
+					GetObjectFromVector(
+						OutObject,
+						SubElement,
+						Schema,
+						Parameters.Elements[SubElementIdx],
+						ActionVector.Slice(SubElementOffset, SubElementSize));
+
+					SubElementNames.Add(Parameters.ElementNames[SubElementIdx]);
+					SubElements.Add(SubElement);
+				}
+
+				SubElementOffset += SubElementSize;
+			}
+			UE_LEARNING_CHECK(SubElementOffset + Parameters.Elements.Num() == ActionVectorSize);
+
+			OutObjectElement = OutObject.CreateOrInclusive({ SubElementNames, SubElements }, SchemaElementTag);
+			return;
+		}
+
+		case EType::Array:
+		{
+			const FSchemaArrayParameters Parameters = Schema.GetArray(SchemaElement);
+
+			TArray<FObjectElement, TInlineAllocator<8>> SubElements;
+			SubElements.SetNumUninitialized(Parameters.Num);
+
+			// Create sub-elements
+
+			const int32 SubElementSize = Schema.GetActionVectorSize(Parameters.Element);
+
+			for (int32 ElementIdx = 0; ElementIdx < Parameters.Num; ElementIdx++)
+			{
+				GetObjectFromVector(
+					OutObject,
+					SubElements[ElementIdx],
+					Schema,
+					Parameters.Element,
+					ActionVector.Slice(ElementIdx * SubElementSize, SubElementSize));
+			}
+
+			OutObjectElement = OutObject.CreateArray({ SubElements }, SchemaElementTag);
+			return;
+		}
+
+		case EType::Encoding:
+		{
+			const FSchemaEncodingParameters Parameters = Schema.GetEncoding(SchemaElement);
+
+			FObjectElement SubElement;
+			GetObjectFromVector(
+				OutObject,
+				SubElement,
+				Schema,
+				Parameters.Element,
+				ActionVector);
+
+			OutObjectElement = OutObject.CreateEncoding({ SubElement }, SchemaElementTag);
+			return;
+		}
+
+		default:
+		{
+			UE_LEARNING_NOT_IMPLEMENTED();
+			OutObjectElement = FObjectElement();
+			return;
+		}
+		}
+	}
+
+	void SetVectorFromModifier(
+		TLearningArrayView<1, float> OutActionModifierVector,
+		const FSchema& Schema,
+		const FSchemaElement SchemaElement,
+		const FModifier& Modifier,
+		const FModifierElement ModifierElement)
+	{
+		UE_LEARNING_CHECK(Schema.IsValid(SchemaElement));
+		UE_LEARNING_CHECK(Modifier.IsValid(ModifierElement));
+		UE_LEARNING_CHECK(OutActionModifierVector.Num() == Schema.GetActionModifierVectorSize(SchemaElement));
+
+		// Check that the types match
+
+		const EType SchemaElementType = Schema.GetType(SchemaElement);
+		const EType ModifierElementType = Modifier.GetType(ModifierElement);
+		UE_LEARNING_CHECK(ModifierElementType == EType::Null || ModifierElementType == SchemaElementType);
+
+		// Zero Action Modifier Vector and return early if we have a null Type
+
+		Array::Zero(OutActionModifierVector);
+
+		if (ModifierElementType == EType::Null)
+		{
+			return;
+		}
+
+		// Indicate we have a modifier by setting the first element in the vector to 1.0f
+
+		OutActionModifierVector[0] = 1.0f;
+
+		// Logic for each specific modifier type
+
+		switch (SchemaElementType)
+		{
+
+		case EType::Null:
+		{
+			// This should never be reached
+			UE_LEARNING_NOT_IMPLEMENTED();
+			break;
+		}
+
+		case EType::Continuous:
+		{
+			// Check the input sizes match
+
+			const FSchemaContinuousParameters SchemaParameters = Schema.GetContinuous(SchemaElement);
+			const int32 ValueNum = SchemaParameters.Num;
+
+			const TArrayView<const bool> Masked = Modifier.GetContinuous(ModifierElement).Masked;
+			const TArrayView<const float> MaskedValues = Modifier.GetContinuous(ModifierElement).MaskedValues;
+			UE_LEARNING_CHECK(Masked.Num() == ValueNum);
+			UE_LEARNING_CHECK(MaskedValues.Num() == ValueNum);
+			UE_LEARNING_CHECK(Schema.GetActionModifierVectorSize(SchemaElement) == 1 + Masked.Num() + MaskedValues.Num());
+
+			for (int32 ValueIdx = 0; ValueIdx < ValueNum; ValueIdx++)
+			{
+				OutActionModifierVector[1 + ValueIdx] = Masked[ValueIdx] ? 1.0f : 0.0f;
+				OutActionModifierVector[1 + ValueNum + ValueIdx] = MaskedValues[ValueIdx];
+			}
+
+			return;
@@ -1538,28 +2607,27 @@
 
 		case EType::DiscreteExclusive:
 		{
-			int32 ActionValue = Object.GetDiscreteExclusive(ObjectElement).DiscreteIndex;
-			UE_LEARNING_CHECK(Schema.GetActionVectorSize(SchemaElement) > ActionValue && ActionValue >= 0);
-			UE_LEARNING_CHECK(Schema.GetActionVectorSize(SchemaElement) == OutActionVector.Num());
+			const TArrayView<const int32> MaskIndices = Modifier.GetDiscreteExclusive(ModifierElement).MaskedIndices;
+			UE_LEARNING_CHECK(Schema.GetDiscreteExclusive(SchemaElement).Num >= MaskIndices.Num());
 
-			// Set the single value in the action vector
+			for (int32 MaskIndicesIdx = 0; MaskIndicesIdx < MaskIndices.Num(); MaskIndicesIdx++)
+			{
+				UE_LEARNING_CHECK(Schema.GetDiscreteExclusive(SchemaElement).Num > MaskIndices[MaskIndicesIdx] && MaskIndices[MaskIndicesIdx] >= 0);
+				OutActionModifierVector[1 + MaskIndices[MaskIndicesIdx]] = 1.0f;
+			}
 
-			OutActionVector[ActionValue] = 1.0f;
 			return;
 		}
 
 		case EType::DiscreteInclusive:
 		{
-			TArrayView<const int32> ActionValues = Object.GetDiscreteInclusive(ObjectElement).DiscreteIndices;
-			UE_LEARNING_CHECK(Schema.GetActionVectorSize(SchemaElement) >= ActionValues.Num());
-			UE_LEARNING_CHECK(Schema.GetActionVectorSize(SchemaElement) == OutActionVector.Num());
+			const TArrayView<const int32> MaskIndices = Modifier.GetDiscreteInclusive(ModifierElement).MaskedIndices;
+			UE_LEARNING_CHECK(Schema.GetDiscreteInclusive(SchemaElement).Num >= MaskIndices.Num());
 
-			// Set values in the action vector
-
-			for (int32 ActionValueIdx = 0; ActionValueIdx < ActionValues.Num(); ActionValueIdx++)
+			for (int32 MaskIndicesIdx = 0; MaskIndicesIdx < MaskIndices.Num(); MaskIndicesIdx++)
 			{
-				UE_LEARNING_CHECK(Schema.GetActionVectorSize(SchemaElement) > ActionValues[ActionValueIdx] && ActionValues[ActionValueIdx] >= 0);
-				OutActionVector[ActionValues[ActionValueIdx]] = 1.0f;
+				UE_LEARNING_CHECK(Schema.GetDiscreteInclusive(SchemaElement).Num > MaskIndices[MaskIndicesIdx] && MaskIndices[MaskIndicesIdx] >= 0);
+				OutActionModifierVector[1 + MaskIndices[MaskIndicesIdx]] = 1.0f;
 			}
 
 			return;
@@ -1567,109 +2635,118 @@
 
 		case EType::And:
 		{
-			// Check the number of sub-elements match
-
 			const FSchemaAndParameters SchemaParameters = Schema.GetAnd(SchemaElement);
-			const FObjectAndParameters ObjectParameters = Object.GetAnd(ObjectElement);
-			UE_LEARNING_CHECK(SchemaParameters.Elements.Num() == ObjectParameters.Elements.Num());
+			const FModifierAndParameters ModifierParameters = Modifier.GetAnd(ModifierElement);
+
+			UE_LEARNING_CHECK(OutActionModifierVector.Num() == 1 + Private::GetTotalActionModifierVectorSize(Schema, SchemaParameters.Elements));
 
 			// Set the Sub-elements
 
-			int32 SubElementOffset = 0;
+			int32 SubElementOffset = 1;
 
 			for (int32 SchemaElementIdx = 0; SchemaElementIdx < SchemaParameters.Elements.Num(); SchemaElementIdx++)
 			{
-				const int32 ObjectElementIndex = ObjectParameters.ElementNames.Find(SchemaParameters.ElementNames[SchemaElementIdx]);
-				UE_LEARNING_CHECK(ObjectElementIndex != INDEX_NONE);
+				const int32 SubElementSize = Schema.GetActionModifierVectorSize(SchemaParameters.Elements[SchemaElementIdx]);
 
-				const int32 SubElementSize = Schema.GetActionVectorSize(SchemaParameters.Elements[SchemaElementIdx]);
-
-				SetVectorFromObject(
-					OutActionVector.Slice(SubElementOffset, SubElementSize),
-					Schema,
-					SchemaParameters.Elements[SchemaElementIdx],
-					Object,
-					ObjectParameters.Elements[ObjectElementIndex]);
+				const int32 ModifierElementIndex = ModifierParameters.ElementNames.Find(SchemaParameters.ElementNames[SchemaElementIdx]);
+				if (ModifierElementIndex != INDEX_NONE)
+				{
+					SetVectorFromModifier(
+						OutActionModifierVector.Slice(SubElementOffset, SubElementSize),
+						Schema,
+						SchemaParameters.Elements[SchemaElementIdx],
+						Modifier,
+						ModifierParameters.Elements[ModifierElementIndex]);
+				}
 
 				SubElementOffset += SubElementSize;
 			}
 
-			UE_LEARNING_CHECK(SubElementOffset == OutActionVector.Num());
+			UE_LEARNING_CHECK(SubElementOffset == OutActionModifierVector.Num());
 			return;
 		}
 
 		case EType::OrExclusive:
 		{
-			// Check only one sub-element is given and index is valid
-
 			const FSchemaOrExclusiveParameters SchemaParameters = Schema.GetOrExclusive(SchemaElement);
-			const FObjectOrExclusiveParameters ObjectParameters = Object.GetOrExclusive(ObjectElement);
+			const FModifierOrExclusiveParameters ModifierParameters = Modifier.GetOrExclusive(ModifierElement);
 
-			const int32 SchemaElementIndex = SchemaParameters.ElementNames.Find(ObjectParameters.ElementName);
-			UE_LEARNING_CHECK(SchemaElementIndex != INDEX_NONE);
+			UE_LEARNING_CHECK(OutActionModifierVector.Num() == 1 + SchemaParameters.Elements.Num() + Private::GetTotalActionModifierVectorSize(Schema, SchemaParameters.Elements));
 
-			// Set the sub-element
+			// Set the Mask
 
-			const int32 SubElementSize = Schema.GetActionVectorSize(SchemaParameters.Elements[SchemaElementIndex]);
+			for (int32 MaskElementIdx = 0; MaskElementIdx < ModifierParameters.MaskedElements.Num(); MaskElementIdx++)
+			{
+				const int32 SchemaMaskElementIdx = SchemaParameters.ElementNames.Find(ModifierParameters.MaskedElements[MaskElementIdx]);
+				UE_LEARNING_CHECK(SchemaMaskElementIdx != INDEX_NONE);
+				OutActionModifierVector[1 + SchemaMaskElementIdx] = 1.0f;
+			}
 
-			SetVectorFromObject(
-				OutActionVector.Slice(0, SubElementSize),
-				Schema,
-				SchemaParameters.Elements[SchemaElementIndex],
-				Object,
-				ObjectParameters.Element);
+			// Set the Sub-elements
 
-			// Set Mask
+			int32 SubElementOffset = 1 + SchemaParameters.Elements.Num();
 
-			const int32 MaxSubElementSize = Private::GetMaxActionVectorSize(Schema, SchemaParameters.Elements);
+			for (int32 SchemaElementIdx = 0; SchemaElementIdx < SchemaParameters.Elements.Num(); SchemaElementIdx++)
+			{
+				const int32 SubElementSize = Schema.GetActionModifierVectorSize(SchemaParameters.Elements[SchemaElementIdx]);
 
-			OutActionVector[MaxSubElementSize + SchemaElementIndex] = 1.0f;
+				const int32 ModifierElementIndex = ModifierParameters.ElementNames.Find(SchemaParameters.ElementNames[SchemaElementIdx]);
+				if (ModifierElementIndex != INDEX_NONE)
+				{
+					SetVectorFromModifier(
+						OutActionModifierVector.Slice(SubElementOffset, SubElementSize),
+						Schema,
+						SchemaParameters.Elements[SchemaElementIdx],
+						Modifier,
+						ModifierParameters.Elements[ModifierElementIndex]);
+				}
 
-			UE_LEARNING_CHECK(OutActionVector.Num() == MaxSubElementSize + SchemaParameters.Elements.Num());
+				SubElementOffset += SubElementSize;
+			}
+
+			UE_LEARNING_CHECK(SubElementOffset == OutActionModifierVector.Num());
 			return;
 		}
 
 		case EType::OrInclusive:
 		{
-			// Check all indices are in range
-
 			const FSchemaOrInclusiveParameters SchemaParameters = Schema.GetOrInclusive(SchemaElement);
-			const FObjectOrInclusiveParameters ObjectParameters = Object.GetOrInclusive(ObjectElement);
-			UE_LEARNING_CHECK(ObjectParameters.Elements.Num() <= SchemaParameters.Elements.Num());
+			const FModifierOrInclusiveParameters ModifierParameters = Modifier.GetOrInclusive(ModifierElement);
 
-			// Update sub-elements
+			UE_LEARNING_CHECK(OutActionModifierVector.Num() == 1 + SchemaParameters.Elements.Num() + Private::GetTotalActionModifierVectorSize(Schema, SchemaParameters.Elements));
 
-			int32 SubElementOffset = 0;
+			// Set the Mask
 
-			for (int32 ObjectElementIdx = 0; ObjectElementIdx < ObjectParameters.Elements.Num(); ObjectElementIdx++)
+			for (int32 MaskElementIdx = 0; MaskElementIdx < ModifierParameters.MaskedElements.Num(); MaskElementIdx++)
 			{
-				const int32 SchemaElementIdx = SchemaParameters.ElementNames.Find(ObjectParameters.ElementNames[ObjectElementIdx]);
-				UE_LEARNING_CHECK(SchemaElementIdx != INDEX_NONE);
+				const int32 SchemaMaskElementIdx = SchemaParameters.ElementNames.Find(ModifierParameters.MaskedElements[MaskElementIdx]);
+				UE_LEARNING_CHECK(SchemaMaskElementIdx != INDEX_NONE);
+				OutActionModifierVector[1 + SchemaMaskElementIdx] = 1.0f;
+			}
 
-				const int32 SubElementSize = Schema.GetActionVectorSize(SchemaParameters.Elements[SchemaElementIdx]);
+			// Set the Sub-elements
 
-				SetVectorFromObject(
-					OutActionVector.Slice(SubElementOffset, SubElementSize),
-					Schema,
-					SchemaParameters.Elements[SchemaElementIdx],
-					Object,
-					ObjectParameters.Elements[ObjectElementIdx]);
+			int32 SubElementOffset = 1 + SchemaParameters.Elements.Num();
+
+			for (int32 SchemaElementIdx = 0; SchemaElementIdx < SchemaParameters.Elements.Num(); SchemaElementIdx++)
+			{
+				const int32 SubElementSize = Schema.GetActionModifierVectorSize(SchemaParameters.Elements[SchemaElementIdx]);
+
+				const int32 ModifierElementIndex = ModifierParameters.ElementNames.Find(SchemaParameters.ElementNames[SchemaElementIdx]);
+				if (ModifierElementIndex != INDEX_NONE)
+				{
+					SetVectorFromModifier(
+						OutActionModifierVector.Slice(SubElementOffset, SubElementSize),
+						Schema,
+						SchemaParameters.Elements[SchemaElementIdx],
+						Modifier,
+						ModifierParameters.Elements[ModifierElementIndex]);
+				}
 
 				SubElementOffset += SubElementSize;
 			}
 
-			// Set Mask
-
-			UE_LEARNING_CHECK(SubElementOffset + SchemaParameters.Elements.Num() == OutActionVector.Num());
-
-			for (int32 ObjectElementIdx = 0; ObjectElementIdx < ObjectParameters.Elements.Num(); ObjectElementIdx++)
-			{
-				const int32 SchemaElementIdx = SchemaParameters.ElementNames.Find(ObjectParameters.ElementNames[ObjectElementIdx]);
-				UE_LEARNING_CHECK(SchemaElementIdx != INDEX_NONE);
-
-				OutActionVector[SubElementOffset + SchemaElementIdx] = 1.0f;
-			}
-
+			UE_LEARNING_CHECK(SubElementOffset == OutActionModifierVector.Num());
 			return;
 		}
 
@@ -1678,21 +2755,21 @@
 			// Check number of array elements is correct
 
 			const FSchemaArrayParameters SchemaParameters = Schema.GetArray(SchemaElement);
-			const FObjectArrayParameters ObjectParameters = Object.GetArray(ObjectElement);
-			UE_LEARNING_CHECK(SchemaParameters.Num == ObjectParameters.Elements.Num());
+			const FModifierArrayParameters ModifierParameters = Modifier.GetArray(ModifierElement);
+			UE_LEARNING_CHECK(SchemaParameters.Num == ModifierParameters.Elements.Num());
 
 			// Update sub-elements
 
-			const int32 SubElementSize = Schema.GetActionVectorSize(SchemaParameters.Element);
+			const int32 SubElementSize = Schema.GetActionModifierVectorSize(SchemaParameters.Element);
 
 			for (int32 ElementIdx = 0; ElementIdx < SchemaParameters.Num; ElementIdx++)
 			{
-				SetVectorFromObject(
-					OutActionVector.Slice(ElementIdx * SubElementSize, SubElementSize),
+				SetVectorFromModifier(
+					OutActionModifierVector.Slice(1 + ElementIdx * SubElementSize, SubElementSize),
 					Schema,
 					SchemaParameters.Element,
-					Object,
-					ObjectParameters.Elements[ElementIdx]);
+					Modifier,
+					ModifierParameters.Elements[ElementIdx]);
 			}
 
 			return;
@@ -1701,14 +2778,16 @@
 		case EType::Encoding:
 		{
 			const FSchemaEncodingParameters SchemaParameters = Schema.GetEncoding(SchemaElement);
-			const FObjectEncodingParameters ObjectParameters = Object.GetEncoding(ObjectElement);
+			const FModifierEncodingParameters ModifierParameters = Modifier.GetEncoding(ModifierElement);
 
-			SetVectorFromObject(
-				OutActionVector,
+			const int32 SubElementSize = Schema.GetActionModifierVectorSize(SchemaParameters.Element);
+
+			SetVectorFromModifier(
+				OutActionModifierVector.Slice(1, SubElementSize),
 				Schema,
 				SchemaParameters.Element,
-				Object,
-				ObjectParameters.Element);
+				Modifier,
+				ModifierParameters.Element);
 
 			return;
 		}
@@ -1722,24 +2801,39 @@
 		}
 	}
 
-	void GetObjectFromVector(
-		FObject& OutObject,
-		FObjectElement& OutObjectElement,
+	void GetModifierFromVector(
+		FModifier& OutModifier,
+		FModifierElement& OutModifierElement,
 		const FSchema& Schema,
 		const FSchemaElement SchemaElement,
-		const TLearningArrayView<1, const float> ActionVector)
+		const TLearningArrayView<1, const float> ActionModifierVector)
 	{
 		UE_LEARNING_CHECK(Schema.IsValid(SchemaElement));
 
-		// Check that the types match
+		// Get Type and Tag
 
 		const EType SchemaElementType = Schema.GetType(SchemaElement);
 		const FName SchemaElementTag = Schema.GetTag(SchemaElement);
 
-		// Get Action Vector Size
+		// Get Action Modifier Vector Size
 
-		const int32 ActionVectorSize = ActionVector.Num();
-		UE_LEARNING_CHECK(ActionVectorSize == Schema.GetActionVectorSize(SchemaElement));
+		const int32 ActionModifierVectorSize = ActionModifierVector.Num();
+		UE_LEARNING_CHECK(ActionModifierVectorSize == Schema.GetActionModifierVectorSize(SchemaElement));
 
+		// We always have at least one element in the ActionModifierVector which says if the element is provided
+		// if this first value is zero then it means nothing below is masked and we always just return the null element
+
+		UE_LEARNING_CHECK(ActionModifierVectorSize > 0);
+
+		if (ActionModifierVector[0] == 0.0f)
+		{
+			OutModifierElement = OutModifier.CreateNull(SchemaElementTag);
+			return;
+		}
+		else
+		{
+			UE_LEARNING_CHECK(ActionModifierVector[0] == 1.0f);
+		}
+
 		// Logic for each specific element type
 
@@ -1748,175 +2842,193 @@
 
 		case EType::Null:
 		{
-			OutObjectElement = OutObject.CreateNull(SchemaElementTag);
+			OutModifierElement = OutModifier.CreateNull(SchemaElementTag);
 			return;
 		}
 
 		case EType::Continuous:
 		{
 			const FSchemaContinuousParameters SchemaParameters = Schema.GetContinuous(SchemaElement);
-			UE_LEARNING_CHECK(ActionVectorSize == SchemaParameters.Num);
+			UE_LEARNING_CHECK(ActionModifierVectorSize == 1 + 2 * SchemaParameters.Num);
 
 			const int32 ValueNum = SchemaParameters.Num;
-			const float ValueScale = FMath::Max(SchemaParameters.Scale, UE_SMALL_NUMBER);
 
-			TLearningArray<1, float, TInlineAllocator<32>> ActionValues;
-			ActionValues.SetNumUninitialized({ ValueNum });
+			TLearningArray<1, bool, TInlineAllocator<32>> ActionMasked;
+			TLearningArray<1, float, TInlineAllocator<32>> ActionMaskedValues;
+			ActionMasked.SetNumUninitialized({ ValueNum });
+			ActionMaskedValues.SetNumUninitialized({ ValueNum });
+
 			for (int32 ValueIdx = 0; ValueIdx < ValueNum; ValueIdx++)
 			{
-				ActionValues[ValueIdx] = ValueScale * ActionVector[ValueIdx];
+				ActionMasked[ValueIdx] = ActionModifierVector[1 + ValueIdx] == 1.0f;
+				ActionMaskedValues[ValueIdx] = ActionModifierVector[1 + ValueNum + ValueIdx];
 			}
 
-			OutObjectElement = OutObject.CreateContinuous({ MakeArrayView(ActionValues.GetData(), ActionValues.Num()) }, SchemaElementTag);
+			OutModifierElement = OutModifier.CreateContinuous({ 
+				MakeArrayView(ActionMasked.GetData(), ActionMasked.Num()),
+				MakeArrayView(ActionMaskedValues.GetData(), ActionMaskedValues.Num()) }, SchemaElementTag);
+
 			return;
 		}
 
 		case EType::DiscreteExclusive:
 		{
-			UE_LEARNING_CHECK(ActionVectorSize == Schema.GetDiscreteExclusive(SchemaElement).Num);
+			const int32 ValueNum = Schema.GetDiscreteExclusive(SchemaElement).Num;
 
-			// Find Index
-			int32 ExclusiveIndex = INDEX_NONE;
-			for (int32 Idx = 0; Idx < ActionVectorSize; Idx++)
+			UE_LEARNING_CHECK(ActionModifierVectorSize == 1 + ValueNum);
+
+			// Find Indices
+			TArray<int32, TInlineAllocator<8>> MaskedIndices;
+			MaskedIndices.Reserve(ValueNum);
+			for (int32 Idx = 0; Idx < ValueNum; Idx++)
 			{
-				UE_LEARNING_CHECK(ActionVector[Idx] == 0.0f || ActionVector[Idx] == 1.0f);
-				if (ActionVector[Idx])
+				UE_LEARNING_CHECK(ActionModifierVector[1 + Idx] == 0.0f || ActionModifierVector[1 + Idx] == 1.0f);
+				if (ActionModifierVector[1 + Idx] == 1.0f)
 				{
-					ExclusiveIndex = Idx;
-					break;
+					MaskedIndices.Add(Idx);
 				}
 			}
-			UE_LEARNING_CHECK(ExclusiveIndex != INDEX_NONE);
 
-			OutObjectElement = OutObject.CreateDiscreteExclusive({ ExclusiveIndex }, SchemaElementTag);
+			OutModifierElement = OutModifier.CreateDiscreteExclusive({ MaskedIndices }, SchemaElementTag);
 			return;
 		}
 
 		case EType::DiscreteInclusive:
 		{
-			UE_LEARNING_CHECK(ActionVectorSize == Schema.GetDiscreteInclusive(SchemaElement).Num);
+			const int32 ValueNum = Schema.GetDiscreteInclusive(SchemaElement).Num;
+
+			UE_LEARNING_CHECK(ActionModifierVectorSize == 1 + ValueNum);
 
 			// Find Indices
-			TArray<int32, TInlineAllocator<8>> InclusiveIndices;
-			InclusiveIndices.Reserve(ActionVectorSize);
-			for (int32 Idx = 0; Idx < ActionVectorSize; Idx++)
+			TArray<int32, TInlineAllocator<8>> MaskedIndices;
+			MaskedIndices.Reserve(ValueNum);
+			for (int32 Idx = 0; Idx < ValueNum; Idx++)
 			{
-				UE_LEARNING_CHECK(ActionVector[Idx] == 0.0f || ActionVector[Idx] == 1.0f);
-				if (ActionVector[Idx])
+				UE_LEARNING_CHECK(ActionModifierVector[1 + Idx] == 0.0f || ActionModifierVector[1 + Idx] == 1.0f);
+				if (ActionModifierVector[1 + Idx] == 1.0f)
 				{
-					InclusiveIndices.Add(Idx);
+					MaskedIndices.Add(Idx);
 				}
 			}
 
-			OutObjectElement = OutObject.CreateDiscreteInclusive({ InclusiveIndices }, SchemaElementTag);
+			OutModifierElement = OutModifier.CreateDiscreteInclusive({ MaskedIndices }, SchemaElementTag);
 			return;
 		}
 
 		case EType::And:
 		{
 			const FSchemaAndParameters Parameters = Schema.GetAnd(SchemaElement);
-			
+
 			// Create Sub-elements
 
-			TArray<FObjectElement, TInlineAllocator<8>> SubElements;
+			TArray<FModifierElement, TInlineAllocator<8>> SubElements;
 			SubElements.SetNumUninitialized(Parameters.Elements.Num());
 
-			int32 SubElementOffset = 0;
+			int32 SubElementOffset = 1;
 			for (int32 SchemaElementIdx = 0; SchemaElementIdx < Parameters.Elements.Num(); SchemaElementIdx++)
 			{
-				const int32 SubElementSize = Schema.GetActionVectorSize(Parameters.Elements[SchemaElementIdx]);
+				const int32 SubElementSize = Schema.GetActionModifierVectorSize(Parameters.Elements[SchemaElementIdx]);
 
-				GetObjectFromVector(
-					OutObject,
+				GetModifierFromVector(
+					OutModifier,
 					SubElements[SchemaElementIdx],
 					Schema,
 					Parameters.Elements[SchemaElementIdx],
-					ActionVector.Slice(SubElementOffset, SubElementSize));
+					ActionModifierVector.Slice(SubElementOffset, SubElementSize));
 
 				SubElementOffset += SubElementSize;
 			}
-			UE_LEARNING_CHECK(SubElementOffset == ActionVectorSize);
+			UE_LEARNING_CHECK(SubElementOffset == ActionModifierVectorSize);
 
-			OutObjectElement = OutObject.CreateAnd({ Parameters.ElementNames, SubElements }, SchemaElementTag);
+			OutModifierElement = OutModifier.CreateAnd({ Parameters.ElementNames, SubElements }, SchemaElementTag);
 			return;
 		}
 
 		case EType::OrExclusive:
 		{
 			const FSchemaOrExclusiveParameters Parameters = Schema.GetOrExclusive(SchemaElement);
+			const int32 SubElementNum = Parameters.Elements.Num();
 
-			// Find active element
+			// Extract Mask Elements
 
-			const int32 MaxSubElementSize = Private::GetMaxActionVectorSize(Schema, Parameters.Elements);
+			TArray<FName, TInlineAllocator<8>> MaskedElements;
+			MaskedElements.Reserve(SubElementNum);
 
-			int32 SchemaElementIndex = INDEX_NONE;
-			for (int32 SubElementIdx = 0; SubElementIdx < Parameters.Elements.Num(); SubElementIdx++)
+			for (int32 Idx = 0; Idx < SubElementNum; Idx++)
 			{
-				UE_LEARNING_CHECK(ActionVector[MaxSubElementSize + SubElementIdx] == 0.0f || ActionVector[MaxSubElementSize + SubElementIdx] == 1.0f);
-				if (ActionVector[MaxSubElementSize + SubElementIdx])
+				UE_LEARNING_CHECK(ActionModifierVector[1 + Idx] == 0.0f || ActionModifierVector[1 + Idx] == 1.0f);
+				if (ActionModifierVector[1 + Idx] == 1.0f)
 				{
-					SchemaElementIndex = SubElementIdx;
-					break;
+					MaskedElements.Add(Parameters.ElementNames[Idx]);
 				}
 			}
-			UE_LEARNING_CHECK(SchemaElementIndex != INDEX_NONE);
 
-			// Create sub-element
+			// Create Sub-elements
 
-			const int32 SubElementSize = Schema.GetActionVectorSize(Parameters.Elements[SchemaElementIndex]);
+			TArray<FModifierElement, TInlineAllocator<8>> SubElements;
+			SubElements.SetNumUninitialized(SubElementNum);
 
-			FObjectElement SubElement;
-			GetObjectFromVector(
-				OutObject,
-				SubElement,
-				Schema,
-				Parameters.Elements[SchemaElementIndex],
-				ActionVector.Slice(0, SubElementSize));
+			int32 SubElementOffset = 1 + SubElementNum;
+			for (int32 SchemaElementIdx = 0; SchemaElementIdx < SubElementNum; SchemaElementIdx++)
+			{
+				const int32 SubElementSize = Schema.GetActionModifierVectorSize(Parameters.Elements[SchemaElementIdx]);
 
-			OutObjectElement = OutObject.CreateOrExclusive({ Parameters.ElementNames[SchemaElementIndex], SubElement }, SchemaElementTag);
+				GetModifierFromVector(
+					OutModifier,
+					SubElements[SchemaElementIdx],
+					Schema,
+					Parameters.Elements[SchemaElementIdx],
+					ActionModifierVector.Slice(SubElementOffset, SubElementSize));
+
+				SubElementOffset += SubElementSize;
+			}
+			UE_LEARNING_CHECK(SubElementOffset == ActionModifierVectorSize);
+
+			OutModifierElement = OutModifier.CreateOrExclusive({ Parameters.ElementNames, SubElements, MaskedElements }, SchemaElementTag);
 			return;
 		}
 
 		case EType::OrInclusive:
 		{
 			const FSchemaOrInclusiveParameters Parameters = Schema.GetOrInclusive(SchemaElement);
+			const int32 SubElementNum = Parameters.Elements.Num();
 
-			// Find total sub-element size
+			// Extract Mask Elements
 
-			const int32 TotalSubElementSize = Private::GetTotalActionVectorSize(Schema, Parameters.Elements);
+			TArray<FName, TInlineAllocator<8>> MaskedElements;
+			MaskedElements.Reserve(SubElementNum);
 
-			// Create sub-elements
+			for (int32 Idx = 0; Idx < SubElementNum; Idx++)
+			{
+				UE_LEARNING_CHECK(ActionModifierVector[1 + Idx] == 0.0f || ActionModifierVector[1 + Idx] == 1.0f);
+				if (ActionModifierVector[1 + Idx] == 1.0f)
+				{
+					MaskedElements.Add(Parameters.ElementNames[Idx]);
+				}
+			}
 
-			TArray<FName, TInlineAllocator<8>> SubElementNames;
-			TArray<FObjectElement, TInlineAllocator<8>> SubElements;
-			SubElementNames.Reserve(Parameters.Elements.Num());
-			SubElements.Reserve(Parameters.Elements.Num());
+			// Create Sub-elements
 
-			int32 SubElementOffset = 0;
-			for (int32 SubElementIdx = 0; SubElementIdx < Parameters.Elements.Num(); SubElementIdx++)
+			TArray<FModifierElement, TInlineAllocator<8>> SubElements;
+			SubElements.SetNumUninitialized(SubElementNum);
+
+			int32 SubElementOffset = 1 + SubElementNum;
+			for (int32 SchemaElementIdx = 0; SchemaElementIdx < SubElementNum; SchemaElementIdx++)
 			{
-				const int32 SubElementSize = Schema.GetActionVectorSize(Parameters.Elements[SubElementIdx]);
+				const int32 SubElementSize = Schema.GetActionModifierVectorSize(Parameters.Elements[SchemaElementIdx]);
 
-				UE_LEARNING_CHECK(ActionVector[TotalSubElementSize + SubElementIdx] == 0.0f || ActionVector[TotalSubElementSize + SubElementIdx] == 1.0f);
-				if (ActionVector[TotalSubElementSize + SubElementIdx])
-				{
-					FObjectElement SubElement;
-					GetObjectFromVector(
-						OutObject,
-						SubElement,
-						Schema,
-						Parameters.Elements[SubElementIdx],
-						ActionVector.Slice(SubElementOffset, SubElementSize));
-
-					SubElementNames.Add(Parameters.ElementNames[SubElementIdx]);
-					SubElements.Add(SubElement);
-				}
+				GetModifierFromVector(
+					OutModifier,
+					SubElements[SchemaElementIdx],
+					Schema,
+					Parameters.Elements[SchemaElementIdx],
+					ActionModifierVector.Slice(SubElementOffset, SubElementSize));
 
 				SubElementOffset += SubElementSize;
 			}
-			UE_LEARNING_CHECK(SubElementOffset + Parameters.Elements.Num() == ActionVectorSize);
+			UE_LEARNING_CHECK(SubElementOffset == ActionModifierVectorSize);
 
-			OutObjectElement = OutObject.CreateOrInclusive({ SubElementNames, SubElements }, SchemaElementTag);
+			OutModifierElement = OutModifier.CreateOrInclusive({ Parameters.ElementNames, SubElements, MaskedElements }, SchemaElementTag);
 			return;
 		}
 
@@ -1924,49 +3036,51 @@
 		{
 			const FSchemaArrayParameters Parameters = Schema.GetArray(SchemaElement);
 
-			TArray<FObjectElement, TInlineAllocator<8>> SubElements;
+			TArray<FModifierElement, TInlineAllocator<8>> SubElements;
 			SubElements.SetNumUninitialized(Parameters.Num);
 
 			// Create sub-elements
 
-			const int32 SubElementSize = Schema.GetActionVectorSize(Parameters.Element);
+			const int32 SubElementSize = Schema.GetActionModifierVectorSize(Parameters.Element);
 
 			for (int32 ElementIdx = 0; ElementIdx < Parameters.Num; ElementIdx++)
 			{
-				GetObjectFromVector(
-					OutObject,
+				GetModifierFromVector(
+					OutModifier,
 					SubElements[ElementIdx],
 					Schema,
 					Parameters.Element,
-					ActionVector.Slice(ElementIdx * SubElementSize, SubElementSize));
+					ActionModifierVector.Slice(1 + ElementIdx * SubElementSize, SubElementSize));
 			}
 
-			OutObjectElement = OutObject.CreateArray({ SubElements }, SchemaElementTag);
+			OutModifierElement = OutModifier.CreateArray({ SubElements }, SchemaElementTag);
 			return;
 		}
 
 		case EType::Encoding:
 		{
 			const FSchemaEncodingParameters Parameters = Schema.GetEncoding(SchemaElement);
+			const int32 SubElementSize = Schema.GetActionModifierVectorSize(Parameters.Element);
 
-			FObjectElement SubElement;
-			GetObjectFromVector(
-				OutObject,
+			FModifierElement SubElement;
+			GetModifierFromVector(
+				OutModifier,
 				SubElement,
 				Schema,
 				Parameters.Element,
-				ActionVector);
+				ActionModifierVector.Slice(1, SubElementSize));
 
-			OutObjectElement = OutObject.CreateEncoding({ SubElement }, SchemaElementTag);
+			OutModifierElement = OutModifier.CreateEncoding({ SubElement }, SchemaElementTag);
 			return;
 		}
 
 		default:
 		{
 			UE_LEARNING_NOT_IMPLEMENTED();
-			OutObjectElement = FObjectElement();
+			OutModifierElement = FModifierElement();
 			return;
 		}
 		}
 	}
+
 }
diff --git a/Engine/Plugins/Experimental/LearningAgents/Source/Learning/Private/LearningRandom.cpp b/Engine/Plugins/Experimental/LearningAgents/Source/Learning/Private/LearningRandom.cpp
--- a/Engine/Plugins/Experimental/LearningAgents/Source/Learning/Private/LearningRandom.cpp
+++ b/Engine/Plugins/Experimental/LearningAgents/Source/Learning/Private/LearningRandom.cpp
@@ -341,6 +341,47 @@ namespace UE::Learning::Random
 #endif
 	}
 
+	void DistributionIndependantNormalMasked(
+		TLearningArrayView<1, float> Output,
+		const uint32 State,
+		const TLearningArrayView<1, const float> Mean,
+		const TLearningArrayView<1, const float> LogStd,
+		const TLearningArrayView<1, const bool> Masked,
+		const TLearningArrayView<1, const float> MaskedValues,
+		const float Scale)
+	{
+		UE_LEARNING_TRACE_CPUPROFILER_EVENT_SCOPE(Learning::Random::DistributionIndependantNormalMasked);
+
+		const int32 ElementNum = Output.Num();
+
+#if UE_LEARNING_ISPC
+		ispc::LearningRandomDistributionIndependantNormalMasked(
+			Output.GetData(),
+			Mean.GetData(),
+			LogStd.GetData(),
+			Masked.GetData(),
+			MaskedValues.GetData(),
+			ElementNum,
+			State,
+			Scale);
+#else
+		if (Scale < UE_SMALL_NUMBER)
+		{
+			for (int32 ElementIdx = 0; ElementIdx < ElementNum; ElementIdx++)
+			{
+				Output[ElementIdx] = Masked[ElementIdx] ? MaskedValues[ElementIdx] : Mean[ElementIdx];
+			}
+		}
+		else
+		{
+			for (int32 ElementIdx = 0; ElementIdx < ElementNum; ElementIdx++)
+			{
+				Output[ElementIdx] = Masked[ElementIdx] ? MaskedValues[ElementIdx] : Gaussian(State ^ 0x0a006c6b ^ Int(ElementIdx ^ 0x6fe1d121), Mean[ElementIdx], Scale * FMath::Exp(FMath::Min(LogStd[ElementIdx], 10.0f)));
+			}
+		}
+#endif
+	}
+
 	void DistributionMultinoulli(
 		TLearningArrayView<1, float> Output,
 		const uint32 State,
@@ -349,13 +390,15 @@ namespace UE::Learning::Random
 	{
 		UE_LEARNING_TRACE_CPUPROFILER_EVENT_SCOPE(Learning::Random::DistributionMultinoulli);
 
+		UE_LEARNING_CHECK(Output.Num() == Logits.Num());
+
 		const int32 ElementNum = Output.Num();
 		if (ElementNum == 0)
 		{
 			return;
 		}
 
-		float MaxValue = -FLT_MAX;
+		float MaxValue = -UE_MAX_FLT;
 		uint32 MaxIdx = INDEX_NONE;
 
 		if (Scale < UE_SMALL_NUMBER)
@@ -389,6 +432,61 @@ namespace UE::Learning::Random
 		Output[MaxIdx] = 1.0f;
 	}
 
+	void DistributionMultinoulliMasked(
+		TLearningArrayView<1, float> Output,
+		const uint32 State,
+		const TLearningArrayView<1, const float> Logits,
+		const TLearningArrayView<1, const bool> Masked,
+		const float Scale)
+	{
+		UE_LEARNING_TRACE_CPUPROFILER_EVENT_SCOPE(Learning::Random::DistributionMultinoulliMasked);
+
+		UE_LEARNING_CHECK(Output.Num() == Logits.Num());
+		UE_LEARNING_CHECK(Logits.Num() == Masked.Num());
+
+		const int32 ElementNum = Output.Num();
+		if (ElementNum == 0)
+		{
+			return;
+		}
+
+		float MaxValue = -UE_MAX_FLT;
+		uint32 MaxIdx = INDEX_NONE;
+
+		if (Scale < UE_SMALL_NUMBER)
+		{
+			for (int32 ElementIdx = 0; ElementIdx < ElementNum; ElementIdx++)
+			{
+				if (!Masked[ElementIdx] && Logits[ElementIdx] > MaxValue)
+				{
+					MaxValue = Logits[ElementIdx];
+					MaxIdx = ElementIdx;
+				}
+			}
+		}
+		else
+		{
+			for (int32 ElementIdx = 0; ElementIdx < ElementNum; ElementIdx++)
+			{
+				if (!Masked[ElementIdx])
+				{
+					const float ElementValue = Logits[ElementIdx] / Scale - FMath::Loge(-FMath::Loge(Uniform(State ^ 0x93575b09 ^ Int(ElementIdx ^ 0xa0c6ac79))));
+
+					if (ElementValue > MaxValue)
+					{
+						MaxValue = ElementValue;
+						MaxIdx = ElementIdx;
+					}
+				}
+			}
+		}
+
+		UE_LEARNING_CHECK(MaxIdx != INDEX_NONE);
+
+		Array::Zero(Output);
+		Output[MaxIdx] = 1.0f;
+	}
+
 	void DistributionBernoulli(
 		TLearningArrayView<1, float> Output,
 		const uint32 State,
@@ -397,6 +495,8 @@ namespace UE::Learning::Random
 	{
 		UE_LEARNING_TRACE_CPUPROFILER_EVENT_SCOPE(Learning::Random::DistributionBernoulli);
 
+		UE_LEARNING_CHECK(Output.Num() == Logits.Num());
+
 		const int32 ElementNum = Output.Num();
 
 #if UE_LEARNING_ISPC
@@ -424,6 +524,47 @@ namespace UE::Learning::Random
 #endif
 	}
 
+
+	void DistributionBernoulliMasked(
+		TLearningArrayView<1, float> Output,
+		const uint32 State,
+		const TLearningArrayView<1, const float> Logits,
+		const TLearningArrayView<1, const bool> Masked,
+		const float Scale)
+	{
+		UE_LEARNING_TRACE_CPUPROFILER_EVENT_SCOPE(Learning::Random::DistributionBernoulliMasked);
+
+		UE_LEARNING_CHECK(Output.Num() == Logits.Num());
+		UE_LEARNING_CHECK(Logits.Num() == Masked.Num());
+
+		const int32 ElementNum = Output.Num();
+
+#if UE_LEARNING_ISPC
+		ispc::LearningRandomDistributionBernoulliMasked(
+			Output.GetData(),
+			Logits.GetData(),
+			Masked.GetData(),
+			ElementNum,
+			State,
+			Scale);
+#else
+		if (Scale < UE_SMALL_NUMBER)
+		{
+			for (int32 ElementIdx = 0; ElementIdx < ElementNum; ElementIdx++)
+			{
+				Output[ElementIdx] = !Masked[ElementIdx] && Logits[ElementIdx] > 0.0f ? 1.0f : 0.0f;
+			}
+		}
+		else
+		{
+			for (int32 ElementIdx = 0; ElementIdx < ElementNum; ElementIdx++)
+			{
+				Output[ElementIdx] = !Masked[ElementIdx] && Sigmoid(Logits[ElementIdx] / Scale) > Uniform(State ^ 0x8c992b7f ^ Int(ElementIdx ^ 0x1953d0fc)) ? 1.0f : 0.0f;
+			}
+		}
+#endif
+	}
+
 	////////////
 
 	uint32 SampleInt(uint32& State)
@@ -625,6 +766,26 @@ namespace UE::Learning::Random
 			Scale);
 	}
 
+	void SampleDistributionIndependantNormalMasked(
+		TLearningArrayView<1, float> Output,
+		uint32& State,
+		const TLearningArrayView<1, const float> Mean,
+		const TLearningArrayView<1, const float> LogStd,
+		const TLearningArrayView<1, const bool> Masked,
+		const TLearningArrayView<1, const float> MaskedValues,
+		const float Scale)
+	{
+		State = Int(State ^ 0xc58f03a8);
+		DistributionIndependantNormalMasked(
+			Output,
+			State,
+			Mean,
+			LogStd,
+			Masked,
+			MaskedValues,
+			Scale);
+	}
+
 	void SampleDistributionMultinoulli(
 		TLearningArrayView<1, float> Output,
 		uint32& State,
@@ -639,6 +800,22 @@ namespace UE::Learning::Random
 			Scale);
 	}
 
+	void SampleDistributionMultinoulliMasked(
+		TLearningArrayView<1, float> Output,
+		uint32& State,
+		const TLearningArrayView<1, const float> Logits,
+		const TLearningArrayView<1, const bool> Masked,
+		const float Scale)
+	{
+		State = Int(State ^ 0xb744ea38);
+		DistributionMultinoulliMasked(
+			Output,
+			State,
+			Logits,
+			Masked,
+			Scale);
+	}
+
 	void SampleDistributionBernoulli(
 		TLearningArrayView<1, float> Output,
 		uint32& State,
@@ -652,4 +829,20 @@ namespace UE::Learning::Random
 			Logits,
 			Scale);
 	}
+
+	void SampleDistributionBernoulliMasked(
+		TLearningArrayView<1, float> Output,
+		uint32& State,
+		const TLearningArrayView<1, const float> Logits,
+		const TLearningArrayView<1, const bool> Masked,
+		const float Scale)
+	{
+		State = Int(State ^ 0x21610d4a);
+		DistributionBernoulliMasked(
+			Output,
+			State,
+			Logits,
+			Masked,
+			Scale);
+	}
 }
\ No newline at end of file
diff --git a/Engine/Plugins/Experimental/LearningAgents/Source/Learning/Public/LearningAction.h b/Engine/Plugins/Experimental/LearningAgents/Source/Learning/Public/LearningAction.h
--- a/Engine/Plugins/Experimental/LearningAgents/Source/Learning/Public/LearningAction.h
+++ b/Engine/Plugins/Experimental/LearningAgents/Source/Learning/Public/LearningAction.h
@@ -195,6 +195,9 @@ namespace UE::Learning::Action
 		// Get the action distribution vector size of the given element
 		int32 GetActionDistributionVectorSize(const FSchemaElement Element) const;
 
+		// Get the action modifier vector size of the given element
+		int32 GetActionModifierVectorSize(const FSchemaElement Element) const;
+
 		// Get the current generation
 		uint32 GetGeneration() const;
 		
@@ -269,6 +272,7 @@ namespace UE::Learning::Action
 		TArray<int32> EncodedVectorSizes;
 		TArray<int32> ActionVectorSizes;
 		TArray<int32> ActionDistributionVectorSizes;
+		TArray<int32> ActionModifierVectorSizes;
 		TArray<int32> TypeDataIndices;
 
 		/** These are indexed based on the type and the index found in TypeDataIndices */
@@ -422,6 +426,154 @@ namespace UE::Learning::Action
 		TArray<FObjectElement> SubElementObjects;
 	};
 
+
+	/**
+	 * Action Modifier Element
+	 *
+	 * A single element representing a modification to an action performed during sampling (such as masking). Internally this consists of a index 
+	 * used by the object to look up the associated action modification data and a generation id which can be used to check when this index is no 
+	 * longer valid.
+	 */
+	struct FModifierElement { int32 Index = INDEX_NONE; uint32 Generation = INDEX_NONE; };
+
+	struct FModifierContinuousParameters
+	{
+		// Indicates which dimensions are masked
+		TArrayView<const bool> Masked;
+
+		// Values used to replace masked dimensions 
+		TArrayView<const float> MaskedValues;
+	};
+
+	struct FModifierDiscreteExclusiveParameters
+	{
+		// Disallowed indices
+		TArrayView<const int32> MaskedIndices;
+	};
+
+	struct FModifierDiscreteInclusiveParameters
+	{
+		// Disallowed indices
+		TArrayView<const int32> MaskedIndices;
+	};
+
+	struct FModifierAndParameters
+	{
+		// Names of the modified sub-actions.
+		TArrayView<const FName> ElementNames;
+
+		// The modifiers of the associated sub-elements.
+		TArrayView<const FModifierElement> Elements;
+	};
+
+	struct FModifierOrExclusiveParameters
+	{
+		// Names of the modified sub-action.
+		TArrayView<const FName> ElementNames;
+
+		// The modifiers of the associated sub-element.
+		TArrayView<const FModifierElement> Elements;
+
+		// Disallowed sub-elements
+		TArrayView<const FName> MaskedElements;
+	};
+
+	struct FModifierOrInclusiveParameters
+	{
+		// Names of the modified sub-actions.
+		TArrayView<const FName> ElementNames;
+
+		// The modifiers of the associated sub-elements.
+		TArrayView<const FModifierElement> Elements;
+
+		// Disallowed sub-elements
+		TArrayView<const FName> MaskedElements;
+	};
+
+	struct FModifierArrayParameters
+	{
+		// Array of sub-element modifiers.
+		TArrayView<const FModifierElement> Elements;
+	};
+
+	struct FModifierEncodingParameters
+	{
+		// Encoded sub-element modifier.
+		FModifierElement Element;
+	};
+
+	/**
+	 * Action Modifier
+	 *
+	 * This object allows you to describe modifications you want to apply to actions during the sampling stage. Internally this object contains a 
+	 * pool of individual elements. This allows them to be constructed performantly and in a cache efficient way. This object is therefore
+	 * required to access any data about the individual action modifier elements that are created.
+	 */
+	struct LEARNING_API FModifier
+	{
+		FModifierElement CreateNull(const FName Tag = NAME_None);
+		FModifierElement CreateContinuous(const FModifierContinuousParameters Parameters, const FName Tag = NAME_None);
+		FModifierElement CreateDiscreteExclusive(const FModifierDiscreteExclusiveParameters Parameters, const FName Tag = NAME_None);
+		FModifierElement CreateDiscreteInclusive(const FModifierDiscreteInclusiveParameters Parameters, const FName Tag = NAME_None);
+		FModifierElement CreateAnd(const FModifierAndParameters Parameters, const FName Tag = NAME_None);
+		FModifierElement CreateOrExclusive(const FModifierOrExclusiveParameters Parameters, const FName Tag = NAME_None);
+		FModifierElement CreateOrInclusive(const FModifierOrInclusiveParameters Parameters, const FName Tag = NAME_None);
+		FModifierElement CreateArray(const FModifierArrayParameters Parameters, const FName Tag = NAME_None);
+		FModifierElement CreateEncoding(const FModifierEncodingParameters Parameters, const FName Tag = NAME_None);
+
+		FModifierContinuousParameters GetContinuous(const FModifierElement Element) const;
+		FModifierDiscreteExclusiveParameters GetDiscreteExclusive(const FModifierElement Element) const;
+		FModifierDiscreteInclusiveParameters GetDiscreteInclusive(const FModifierElement Element) const;
+		FModifierAndParameters GetAnd(const FModifierElement Element) const;
+		FModifierOrExclusiveParameters GetOrExclusive(const FModifierElement Element) const;
+		FModifierOrInclusiveParameters GetOrInclusive(const FModifierElement Element) const;
+		FModifierArrayParameters GetArray(const FModifierElement Element) const;
+		FModifierEncodingParameters GetEncoding(const FModifierElement Element) const;
+
+		// Checks if the given element is valid
+		bool IsValid(const FModifierElement Element) const;
+
+		// Gets the type of the given element
+		EType GetType(const FModifierElement Element) const;
+
+		// Gets the tag of the given element
+		FName GetTag(const FModifierElement Element) const;
+
+		// Get the current generation
+		uint32 GetGeneration() const;
+
+		// Checks if the given object is empty of elements.
+		bool IsEmpty() const;
+
+		// Empty all internal buffers of elements. This invalidates all existing elements.
+		void Empty();
+
+		// Reset all internal buffers (without freeing memory). This invalidates all existing elements.
+		void Reset();
+
+	private:
+
+		uint32 Generation = 0;
+
+		TArray<EType> Types;
+		TArray<FName> Tags;
+		TArray<int32> ContinuousDataOffsets;
+		TArray<int32> ContinuousDataNums;
+		TArray<int32> DiscreteDataOffsets;
+		TArray<int32> DiscreteDataNums;
+		TArray<int32> ElementDataOffsets;
+		TArray<int32> ElementDataNums;
+		TArray<int32> MaskedDataOffsets;
+		TArray<int32> MaskedDataNums;
+
+		TArray<bool> ContinuousMaskeds;
+		TArray<float> ContinuousMaskedValues;
+		TArray<int32> DiscreteValues;
+		TArray<FName> SubElementNames;
+		TArray<FModifierElement> SubElementModifiers;
+		TArray<FName> MaskedElementNames;
+	};
+
 	/**
 	 * Gets a hash value representing object compatibility between schemas i.e. if objects from one schema can be used by objects expecting another 
 	 * schema. This is not a cryptographic hash, and so `AreSchemaObjectsCompatible` should still be used as the ultimate source of truth.
@@ -504,6 +656,7 @@ namespace UE::Learning::Action
 	 * @param InOutRandomState			Random State/Seed to use when sampling the action vector
 	 * @param OutActionVector			Output action vector
 	 * @param ActionDistributionVector	Input Action Distribution vector
+	 * @param ActionModifierVector		Input Action Modifier vector
 	 * @param Schema					Action Schema
 	 * @param SchemaElement				Action Schema Element
 	 * @param ActionNoiseScale			Scale factor for the noise. Set this to zero to always choose the mean (expected) action.
@@ -512,6 +665,7 @@ namespace UE::Learning::Action
 		uint32& InOutRandomState,
 		TLearningArrayView<1, float> OutActionVector,
 		const TLearningArrayView<1, const float> ActionDistributionVector,
+		const TLearningArrayView<1, const float> ActionModifierVector,
 		const FSchema& Schema,
 		const FSchemaElement SchemaElement,
 		const float ActionNoiseScale = 1.0f);
@@ -547,4 +701,36 @@ namespace UE::Learning::Action
 		const FSchema& Schema,
 		const FSchemaElement SchemaElement,
 		const TLearningArrayView<1, const float> ActionVector);
+
+	/**
+	 * Convert an action modifier into an action modifier vector.
+	 *
+	 * @param OutActionModifierVector	Output action modifier vector
+	 * @param Schema					Action Schema
+	 * @param SchemaElement				Action Schema Element
+	 * @param Modifier					Action Modifier
+	 * @param ModifierElement			Action Modifier Element
+	 */
+	LEARNING_API void SetVectorFromModifier(
+		TLearningArrayView<1, float> OutActionModifierVector,
+		const FSchema& Schema,
+		const FSchemaElement SchemaElement,
+		const FModifier& Modifier,
+		const FModifierElement ModifierElement);
+
+	/**
+	 * Convert an action modifier vector into an action modifier
+	 *
+	 * @param OutModifier				Output Action Modifier
+	 * @param OutModifierElement		Output Action Modifier Element
+	 * @param Schema					Action Schema
+	 * @param SchemaElement				Action Schema Element
+	 * @param ActionModifierVector		Input Action Modifier vector
+	 */
+	LEARNING_API void GetModifierFromVector(
+		FModifier& OutModifier,
+		FModifierElement& OutModifierElement,
+		const FSchema& Schema,
+		const FSchemaElement SchemaElement,
+		const TLearningArrayView<1, const float> ActionModifierVector);
 }
\ No newline at end of file
diff --git a/Engine/Plugins/Experimental/LearningAgents/Source/Learning/Public/LearningRandom.h b/Engine/Plugins/Experimental/LearningAgents/Source/Learning/Public/LearningRandom.h
--- a/Engine/Plugins/Experimental/LearningAgents/Source/Learning/Public/LearningRandom.h
+++ b/Engine/Plugins/Experimental/LearningAgents/Source/Learning/Public/LearningRandom.h
@@ -173,18 +173,41 @@ namespace UE::Learning::Random
 		const TLearningArrayView<1, const float> LogStd,
 		const float Scale);
 
+	LEARNING_API void DistributionIndependantNormalMasked(
+		TLearningArrayView<1, float> Output,
+		const uint32 State,
+		const TLearningArrayView<1, const float> Mean,
+		const TLearningArrayView<1, const float> LogStd,
+		const TLearningArrayView<1, const bool> Masked,
+		const TLearningArrayView<1, const float> MaskedValues,
+		const float Scale);
+
 	LEARNING_API void DistributionMultinoulli(
 		TLearningArrayView<1, float> Output,
 		const uint32 State,
 		const TLearningArrayView<1, const float> Logits,
 		const float Scale);
 
+	LEARNING_API void DistributionMultinoulliMasked(
+		TLearningArrayView<1, float> Output,
+		const uint32 State,
+		const TLearningArrayView<1, const float> Logits,
+		const TLearningArrayView<1, const bool> Masked,
+		const float Scale);
+
 	LEARNING_API void DistributionBernoulli(
 		TLearningArrayView<1, float> Output,
 		const uint32 State,
 		const TLearningArrayView<1, const float> Logits,
 		const float Scale);
 
+	LEARNING_API void DistributionBernoulliMasked(
+		TLearningArrayView<1, float> Output,
+		const uint32 State,
+		const TLearningArrayView<1, const float> Logits,
+		const TLearningArrayView<1, const bool> Masked,
+		const float Scale);
+
 	/////////////////////////////////////////////////////////
 	//
 	// Here we provide a more normal interface which updates
@@ -302,15 +325,38 @@ namespace UE::Learning::Random
 		const TLearningArrayView<1, const float> LogStd,
 		const float Scale);
 
+	LEARNING_API void SampleDistributionIndependantNormalMasked(
+		TLearningArrayView<1, float> Output,
+		uint32& State,
+		const TLearningArrayView<1, const float> Mean,
+		const TLearningArrayView<1, const float> LogStd,
+		const TLearningArrayView<1, const bool> Masked,
+		const TLearningArrayView<1, const float> MaskedValues,
+		const float Scale);
+
 	LEARNING_API void SampleDistributionMultinoulli(
 		TLearningArrayView<1, float> Output,
 		uint32& State,
 		const TLearningArrayView<1, const float> Logits,
 		const float Scale);
 
+	LEARNING_API void SampleDistributionMultinoulliMasked(
+		TLearningArrayView<1, float> Output,
+		uint32& State,
+		const TLearningArrayView<1, const float> Logits,
+		const TLearningArrayView<1, const bool> Masked,
+		const float Scale);
+
 	LEARNING_API void SampleDistributionBernoulli(
 		TLearningArrayView<1, float> Output,
 		uint32& State,
 		const TLearningArrayView<1, const float> Logits,
 		const float Scale);
+
+	LEARNING_API void SampleDistributionBernoulliMasked(
+		TLearningArrayView<1, float> Output,
+		uint32& State,
+		const TLearningArrayView<1, const float> Logits,
+		const TLearningArrayView<1, const bool> Masked,
+		const float Scale);
 }
\ No newline at end of file
diff --git a/Engine/Plugins/Experimental/LearningAgents/Source/LearningAgents/Private/LearningAgentsActions.cpp b/Engine/Plugins/Experimental/LearningAgents/Source/LearningAgents/Private/LearningAgentsActions.cpp
--- a/Engine/Plugins/Experimental/LearningAgents/Source/LearningAgents/Private/LearningAgentsActions.cpp
+++ b/Engine/Plugins/Experimental/LearningAgents/Source/LearningAgents/Private/LearningAgentsActions.cpp
@@ -10,16 +10,36 @@
 
 #include "Containers/StaticArray.h"
 
+bool operator==(const FLearningAgentsActionSchemaElement& Lhs, const FLearningAgentsActionSchemaElement& Rhs)
+{
+	return Lhs.SchemaElement.Index == Rhs.SchemaElement.Index;
+}
+
 bool operator==(const FLearningAgentsActionObjectElement& Lhs, const FLearningAgentsActionObjectElement& Rhs)
 {
 	return Lhs.ObjectElement.Index == Rhs.ObjectElement.Index;
 }
 
-uint32 GetTypeHash(const FLearningAgentsActionObjectElement& Element)
+bool operator==(const FLearningAgentsActionModifierElement& Lhs, const FLearningAgentsActionModifierElement& Rhs)
+{
+	return Lhs.ModifierElement.Index == Rhs.ModifierElement.Index;
+}
+
+uint32 GetTypeHash(const FLearningAgentsActionSchemaElement& Element)
+{
+	return (uint32)Element.SchemaElement.Index;
+}
+
+uint32 GetTypeHash(const FLearningAgentsActionObjectElement & Element)
 {
 	return (uint32)Element.ObjectElement.Index;
 }
 
+uint32 GetTypeHash(const FLearningAgentsActionModifierElement& Element)
+{
+	return (uint32)Element.ModifierElement.Index;
+}
+
 namespace UE::Learning::Agents::Action::Private
 {
 	static inline bool ContainsDuplicates(const TArrayView<const int32> Indices)
@@ -321,6 +341,358 @@ namespace UE::Learning::Agents::Action::Private
 		}
 	}
 
+
+	static bool ValidateActionModifierMatchesSchema(
+		const Learning::Action::FSchema& Schema,
+		const Learning::Action::FSchemaElement SchemaElement,
+		const Learning::Action::FModifier& Modifier,
+		const Learning::Action::FModifierElement ModifierElement)
+	{
+		// Check Elements are Valid
+
+		if (!Schema.IsValid(SchemaElement))
+		{
+			UE_LOG(LogLearning, Error, TEXT("ValidateActionModifierMatchesSchema: Invalid Action Schema Element."));
+			return false;
+		}
+
+		if (!Modifier.IsValid(ModifierElement))
+		{
+			UE_LOG(LogLearning, Error, TEXT("ValidateActionModifierMatchesSchema: Invalid Action Modifier Element."));
+			return false;
+		}
+
+		// Null Modifiers Match Everything
+
+		const Learning::Action::EType ActionSchemaElementType = Schema.GetType(SchemaElement);
+		const Learning::Action::EType ActionModifierElementType = Modifier.GetType(ModifierElement);
+
+		if (ActionModifierElementType == Learning::Action::EType::Null)
+		{
+			return true;
+		}
+
+		// Check Names Match
+
+		const FName ActionSchemaElementTag = Schema.GetTag(SchemaElement);
+		const FName ActionModifierElementTag = Modifier.GetTag(ModifierElement);
+
+		if (ActionSchemaElementTag != ActionModifierElementTag)
+		{
+			UE_LOG(LogLearning, Warning, TEXT("ValidateActionModifierMatchesSchema: Action Modifier tag does not match Schema. Expected '%s', got '%s'."),
+				*ActionSchemaElementTag.ToString(), *ActionModifierElementTag.ToString());
+		}
+
+		// Check Types Match
+
+		if (ActionSchemaElementType != ActionModifierElementType)
+		{
+			UE_LOG(LogLearning, Error, TEXT("ValidateActionModifierMatchesSchema: Action Modifier '%s' type does not match Schema. Expected type '%s', got type '%s'."),
+				*ActionSchemaElementTag.ToString(),
+				GetActionTypeString(ActionSchemaElementType),
+				GetActionTypeString(ActionModifierElementType));
+			return false;
+		}
+
+		// Type Specific Checks
+
+		switch (ActionSchemaElementType)
+		{
+		case Learning::Action::EType::Null: return true;
+
+		case Learning::Action::EType::Continuous:
+		{
+			const int32 SchemaElementSize = Schema.GetContinuous(SchemaElement).Num;
+			const int32 ModifierMaskedElementSize = Modifier.GetContinuous(ModifierElement).Masked.Num();
+			const int32 ModifierValueElementSize = Modifier.GetContinuous(ModifierElement).MaskedValues.Num();
+
+			if (SchemaElementSize != ModifierMaskedElementSize)
+			{
+				UE_LOG(LogLearning, Error, TEXT("ValidateActionModifierMatchesSchema: Action Modifier '%s' size does not match Schema. Expected '%i', got '%i'."),
+					*ActionSchemaElementTag.ToString(),
+					SchemaElementSize,
+					ModifierMaskedElementSize);
+				return false;
+			}
+
+			if (SchemaElementSize != ModifierValueElementSize)
+			{
+				UE_LOG(LogLearning, Error, TEXT("ValidateActionModifierMatchesSchema: Action Modifier '%s' size does not match Schema. Expected '%i', got '%i'."),
+					*ActionSchemaElementTag.ToString(),
+					SchemaElementSize,
+					ModifierValueElementSize);
+				return false;
+			}
+
+			return true;
+		}
+
+		case Learning::Action::EType::DiscreteExclusive:
+		{
+			const int32 SchemaElementSize = Schema.GetDiscreteExclusive(SchemaElement).Num;
+			const TArrayView<const int32> ModifierMaskedIndices = Modifier.GetDiscreteExclusive(ModifierElement).MaskedIndices;
+
+			// For exclusive action we must have one less than the schema size or it means all are masked.
+			if (ModifierMaskedIndices.Num() > SchemaElementSize - 1)
+			{
+				UE_LOG(LogLearning, Error, TEXT("ValidateActionModifierMatchesSchema: Action Modifier '%s' too many indices provided. Expected at most '%i', got '%i'."),
+					*ActionSchemaElementTag.ToString(),
+					SchemaElementSize - 1,
+					ModifierMaskedIndices.Num());
+				return false;
+			}
+
+			for (int32 MaskIdx = 0; MaskIdx < ModifierMaskedIndices.Num(); MaskIdx++)
+			{
+				if (ModifierMaskedIndices[MaskIdx] < 0 || ModifierMaskedIndices[MaskIdx] >= SchemaElementSize)
+				{
+					UE_LOG(LogLearning, Error, TEXT("ValidateActionModifierMatchesSchema: Action Modifier '%s' masked index out of range for Schema. Expected '<%i', got '%i'."),
+						*ActionSchemaElementTag.ToString(),
+						SchemaElementSize,
+						ModifierMaskedIndices[MaskIdx]);
+					return false;
+				}
+			}
+
+			return true;
+		}
+
+		case Learning::Action::EType::DiscreteInclusive:
+		{
+			const int32 SchemaElementSize = Schema.GetDiscreteInclusive(SchemaElement).Num;
+			const TArrayView<const int32> ModifierMaskedIndices = Modifier.GetDiscreteInclusive(ModifierElement).MaskedIndices;
+
+			if (ModifierMaskedIndices.Num() > SchemaElementSize)
+			{
+				UE_LOG(LogLearning, Error, TEXT("ValidateActionModifierMatchesSchema: Action Modifier '%s' too many indices provided. Expected at most '%i', got '%i'."),
+					*ActionSchemaElementTag.ToString(),
+					SchemaElementSize,
+					ModifierMaskedIndices.Num());
+				return false;
+			}
+
+			for (int32 MaskIdx = 0; MaskIdx < ModifierMaskedIndices.Num(); MaskIdx++)
+			{
+				if (ModifierMaskedIndices[MaskIdx] < 0 || ModifierMaskedIndices[MaskIdx] >= SchemaElementSize)
+				{
+					UE_LOG(LogLearning, Error, TEXT("ValidateActionModifierMatchesSchema: Action Modifier '%s' masked index out of range for Schema. Expected '<%i', got '%i'."),
+						*ActionSchemaElementTag.ToString(),
+						SchemaElementSize,
+						ModifierMaskedIndices[MaskIdx]);
+					return false;
+				}
+			}
+
+			return true;
+		}
+
+		case Learning::Action::EType::And:
+		{
+			const Learning::Action::FSchemaAndParameters SchemaParameters = Schema.GetAnd(SchemaElement);
+			const Learning::Action::FModifierAndParameters ModifierParameters = Modifier.GetAnd(ModifierElement);
+			UE_LEARNING_CHECK(SchemaParameters.Elements.Num() == SchemaParameters.ElementNames.Num());
+			UE_LEARNING_CHECK(ModifierParameters.Elements.Num() == ModifierParameters.ElementNames.Num());
+
+			if (SchemaParameters.Elements.Num() < ModifierParameters.Elements.Num())
+			{
+				UE_LOG(LogLearning, Error, TEXT("ValidateActionModifierMatchesSchema: Action Modifier '%s' too many sub-elements. Expected '<=%i', got '%i'."),
+					*ActionSchemaElementTag.ToString(),
+					SchemaParameters.Elements.Num(),
+					ModifierParameters.Elements.Num());
+				return false;
+			}
+
+			for (int32 ModifierElementIdx = 0; ModifierElementIdx < ModifierParameters.Elements.Num(); ModifierElementIdx++)
+			{
+				const int32 SchemaElementIdx = SchemaParameters.ElementNames.Find(ModifierParameters.ElementNames[ModifierElementIdx]);
+
+				if (SchemaElementIdx == INDEX_NONE)
+				{
+					UE_LOG(LogLearning, Error, TEXT("ValidateActionModifierMatchesSchema: Action '%s' does not include '%s' action given by Modifier."),
+						*ActionSchemaElementTag.ToString(),
+						*ModifierParameters.ElementNames[ModifierElementIdx].ToString());
+					return false;
+				}
+
+				if (!ValidateActionModifierMatchesSchema(
+					Schema,
+					SchemaParameters.Elements[SchemaElementIdx],
+					Modifier,
+					ModifierParameters.Elements[ModifierElementIdx]))
+				{
+					return false;
+				}
+			}
+
+			return true;
+		}
+
+		case Learning::Action::EType::OrExclusive:
+		{
+			const Learning::Action::FSchemaOrExclusiveParameters SchemaParameters = Schema.GetOrExclusive(SchemaElement);
+			const Learning::Action::FModifierOrExclusiveParameters ModifierParameters = Modifier.GetOrExclusive(ModifierElement);
+			UE_LEARNING_CHECK(SchemaParameters.Elements.Num() == SchemaParameters.ElementNames.Num());
+
+			// For exclusive action we must have one less than the schema size or it means all are masked.
+			if (ModifierParameters.MaskedElements.Num() > SchemaParameters.Elements.Num() - 1)
+			{
+				UE_LOG(LogLearning, Error, TEXT("ValidateActionModifierMatchesSchema: Action Modifier '%s' too many masked elements. Expected at most '%i', got '%i'."),
+					*ActionSchemaElementTag.ToString(),
+					SchemaParameters.Elements.Num() - 1,
+					ModifierParameters.MaskedElements.Num());
+				return false;
+			}
+
+			for (int32 ModifierMaskIdx = 0; ModifierMaskIdx < ModifierParameters.MaskedElements.Num(); ModifierMaskIdx++)
+			{
+				const FName MaskedElement = ModifierParameters.MaskedElements[ModifierMaskIdx];
+
+				if (SchemaParameters.ElementNames.Find(MaskedElement) == INDEX_NONE)
+				{
+					UE_LOG(LogLearning, Error, TEXT("ValidateActionModifierMatchesSchema: Action '%s' does not include '%s' action given by Modifier."),
+						*ActionSchemaElementTag.ToString(),
+						*MaskedElement.ToString());
+					return false;
+				}
+			}
+
+			// TODO: Check Mask does not contain all elements
+
+			if (SchemaParameters.Elements.Num() < ModifierParameters.Elements.Num())
+			{
+				UE_LOG(LogLearning, Error, TEXT("ValidateActionModifierMatchesSchema: Action Modifier '%s' too many sub-elements. Expected '<=%i', got '%i'."),
+					*ActionSchemaElementTag.ToString(),
+					SchemaParameters.Elements.Num(),
+					ModifierParameters.Elements.Num());
+				return false;
+			}
+
+			for (int32 ModifierElementIdx = 0; ModifierElementIdx < ModifierParameters.Elements.Num(); ModifierElementIdx++)
+			{
+				const int32 SchemaElementIdx = SchemaParameters.ElementNames.Find(ModifierParameters.ElementNames[ModifierElementIdx]);
+
+				if (SchemaElementIdx == INDEX_NONE)
+				{
+					UE_LOG(LogLearning, Error, TEXT("ValidateActionModifierMatchesSchema: Action '%s' does not include '%s' action given by Modifier."),
+						*ActionSchemaElementTag.ToString(),
+						*ModifierParameters.ElementNames[ModifierElementIdx].ToString());
+					return false;
+				}
+
+				if (!ValidateActionModifierMatchesSchema(
+					Schema,
+					SchemaParameters.Elements[SchemaElementIdx],
+					Modifier,
+					ModifierParameters.Elements[ModifierElementIdx]))
+				{
+					return false;
+				}
+			}
+
+			return true;
+		}
+
+		case Learning::Action::EType::OrInclusive:
+		{
+			const Learning::Action::FSchemaOrInclusiveParameters SchemaParameters = Schema.GetOrInclusive(SchemaElement);
+			const Learning::Action::FModifierOrInclusiveParameters ModifierParameters = Modifier.GetOrInclusive(ModifierElement);
+
+			for (int32 ModifierMaskIdx = 0; ModifierMaskIdx < ModifierParameters.MaskedElements.Num(); ModifierMaskIdx++)
+			{
+				const FName MaskedElement = ModifierParameters.MaskedElements[ModifierMaskIdx];
+
+				if (SchemaParameters.ElementNames.Find(MaskedElement) == INDEX_NONE)
+				{
+					UE_LOG(LogLearning, Error, TEXT("ValidateActionModifierMatchesSchema: Action '%s' does not include '%s' action given by Modifier."),
+						*ActionSchemaElementTag.ToString(),
+						*MaskedElement.ToString());
+					return false;
+				}
+			}
+
+			if (SchemaParameters.Elements.Num() < ModifierParameters.Elements.Num())
+			{
+				UE_LOG(LogLearning, Error, TEXT("ValidateActionModifierMatchesSchema: Action Modifier '%s' too many sub-elements. Expected '<=%i', got '%i'."),
+					*ActionSchemaElementTag.ToString(),
+					SchemaParameters.Elements.Num(),
+					ModifierParameters.Elements.Num());
+				return false;
+			}
+
+			for (int32 ModifierElementIdx = 0; ModifierElementIdx < ModifierParameters.Elements.Num(); ModifierElementIdx++)
+			{
+				const int32 SchemaElementIdx = SchemaParameters.ElementNames.Find(ModifierParameters.ElementNames[ModifierElementIdx]);
+
+				if (SchemaElementIdx == INDEX_NONE)
+				{
+					UE_LOG(LogLearning, Error, TEXT("ValidateActionModifierMatchesSchema: Action '%s' does not include '%s' action given by Modifier."),
+						*ActionSchemaElementTag.ToString(),
+						*ModifierParameters.ElementNames[ModifierElementIdx].ToString());
+					return false;
+				}
+
+				if (!ValidateActionModifierMatchesSchema(
+					Schema,
+					SchemaParameters.Elements[SchemaElementIdx],
+					Modifier,
+					ModifierParameters.Elements[ModifierElementIdx]))
+				{
+					return false;
+				}
+			}
+
+			return true;
+		}
+
+		case Learning::Action::EType::Array:
+		{
+			const Learning::Action::FSchemaArrayParameters SchemaParameters = Schema.GetArray(SchemaElement);
+			const Learning::Action::FModifierArrayParameters ModifierParameters = Modifier.GetArray(ModifierElement);
+
+			if (ModifierParameters.Elements.Num() != SchemaParameters.Num)
+			{
+				UE_LOG(LogLearning, Error, TEXT("ValidateActionModifierMatchesSchema: Action '%s' array incorrect size. Expected '%i' elements, got '%i'."),
+					*ActionSchemaElementTag.ToString(),
+					SchemaParameters.Num,
+					ModifierParameters.Elements.Num());
+				return false;
+			}
+
+			for (int32 ElementIdx = 0; ElementIdx < ModifierParameters.Elements.Num(); ElementIdx++)
+			{
+				if (!ValidateActionModifierMatchesSchema(
+					Schema,
+					SchemaParameters.Element,
+					Modifier,
+					ModifierParameters.Elements[ElementIdx]))
+				{
+					return false;
+				}
+			}
+
+			return true;
+		}
+
+		case Learning::Action::EType::Encoding:
+		{
+			const Learning::Action::FSchemaEncodingParameters SchemaParameters = Schema.GetEncoding(SchemaElement);
+			const Learning::Action::FModifierEncodingParameters ModifierParameters = Modifier.GetEncoding(ModifierElement);
+
+			return ValidateActionModifierMatchesSchema(
+				Schema,
+				SchemaParameters.Element,
+				Modifier,
+				ModifierParameters.Element);
+		}
+
+		default:
+		{
+			UE_LEARNING_NOT_IMPLEMENTED();
+			return false;
+		}
+		}
+	}
+
 	static void LogAction(
 		const UE::Learning::Action::FObject& Object,
 		const UE::Learning::Action::FObjectElement ObjectElement,
@@ -527,6 +899,31 @@ bool ULearningAgentsActions::ValidateActionObjectMatchesSchema(
 		ObjectElement.ObjectElement);
 }
 
+bool ULearningAgentsActions::ValidateActionModifierMatchesSchema(
+	const ULearningAgentsActionSchema* Schema,
+	const FLearningAgentsActionSchemaElement SchemaElement,
+	const ULearningAgentsActionModifier* Modifier,
+	const FLearningAgentsActionModifierElement ModifierElement)
+{
+	if (!Schema)
+	{
+		UE_LOG(LogLearning, Error, TEXT("ValidateActionModifierMatchesSchema: Schema is nullptr."));
+		return false;
+	}
+
+	if (!Modifier)
+	{
+		UE_LOG(LogLearning, Error, TEXT("ValidateActionModifierMatchesSchema: Modifier is nullptr."));
+		return false;
+	}
+
+	return UE::Learning::Agents::Action::Private::ValidateActionModifierMatchesSchema(
+		Schema->ActionSchema,
+		SchemaElement.SchemaElement,
+		Modifier->ActionModifier,
+		ModifierElement.ModifierElement);
+}
+
 FLearningAgentsActionSchemaElement ULearningAgentsActions::SpecifyNullAction(ULearningAgentsActionSchema* Schema, const FName Tag)
 {
 	if (!Schema)
@@ -1638,6 +2035,493 @@ FLearningAgentsActionObjectElement ULearningAgentsActions::MakeDirectionAction(U
 		}, Tag);
 }
 
+FLearningAgentsActionModifierElement ULearningAgentsActions::MakeNullActionModifier(ULearningAgentsActionModifier* Modifier, const FName Tag)
+{
+	if (!Modifier)
+	{
+		UE_LOG(LogLearning, Error, TEXT("MakeNullActionModifier: Modifier is nullptr."));
+		return FLearningAgentsActionModifierElement();
+	}
+
+	return { Modifier->ActionModifier.CreateNull(Tag) };
+}
+
+FLearningAgentsActionModifierElement ULearningAgentsActions::MakeContinuousActionModifier(ULearningAgentsActionModifier* Modifier, const TArray<bool>& Masked, const TArray<float>& MaskedValues, const FName Tag)
+{
+	return MakeContinuousActionModifierFromArrayView(Modifier, Masked, MaskedValues, Tag);
+}
+
+FLearningAgentsActionModifierElement ULearningAgentsActions::MakeContinuousActionModifierFromArrayView(ULearningAgentsActionModifier* Modifier, const TArrayView<const bool> Masked, const TArrayView<const float> MaskedValues, const FName Tag)
+{
+	if (!Modifier)
+	{
+		UE_LOG(LogLearning, Error, TEXT("MakeContinuousActionModifierFromArrayView: Modifier is nullptr."));
+		return FLearningAgentsActionModifierElement();
+	}
+
+	if (MaskedValues.Num() != Masked.Num())
+	{
+		UE_LOG(LogLearning, Error, TEXT("MakeContinuousActionModifierFromArrayView: Masked and MaskedValues sizes don't match. Got %i and %i."), Masked.Num(), MaskedValues.Num());
+		return FLearningAgentsActionModifierElement();
+	}
+
+	return { Modifier->ActionModifier.CreateContinuous({ Masked, MaskedValues }, Tag) };
+}
+
+FLearningAgentsActionModifierElement ULearningAgentsActions::MakeExclusiveDiscreteActionModifier(ULearningAgentsActionModifier* Modifier, const TArray<int32>& MaskedIndices, const FName Tag)
+{
+	return MakeExclusiveDiscreteActionModifierFromArrayView(Modifier, MaskedIndices, Tag);
+}
+
+FLearningAgentsActionModifierElement ULearningAgentsActions::MakeExclusiveDiscreteActionModifierFromArrayView(ULearningAgentsActionModifier* Modifier, const TArrayView<const int32> MaskedIndices, const FName Tag)
+{
+	if (!Modifier)
+	{
+		UE_LOG(LogLearning, Error, TEXT("MakeExclusiveDiscreteActionModifierFromArrayView: Modifier is nullptr."));
+		return FLearningAgentsActionModifierElement();
+	}
+
+	return { Modifier->ActionModifier.CreateDiscreteExclusive({ MaskedIndices }, Tag) };
+}
+
+FLearningAgentsActionModifierElement ULearningAgentsActions::MakeInclusiveDiscreteActionModifier(ULearningAgentsActionModifier* Modifier, const TArray<int32>& MaskedIndices, const FName Tag)
+{
+	return MakeInclusiveDiscreteActionModifierFromArrayView(Modifier, MaskedIndices, Tag);
+}
+
+FLearningAgentsActionModifierElement ULearningAgentsActions::MakeInclusiveDiscreteActionModifierFromArrayView(ULearningAgentsActionModifier* Modifier, const TArrayView<const int32> MaskedIndices, const FName Tag)
+{
+	if (!Modifier)
+	{
+		UE_LOG(LogLearning, Error, TEXT("MakeInclusiveDiscreteActionModifierFromArrayView: Modifier is nullptr."));
+		return FLearningAgentsActionModifierElement();
+	}
+
+	return { Modifier->ActionModifier.CreateDiscreteInclusive({ MaskedIndices }, Tag) };
+}
+
+FLearningAgentsActionModifierElement ULearningAgentsActions::MakeStructActionModifier(ULearningAgentsActionModifier* Modifier, const TMap<FName, FLearningAgentsActionModifierElement>& Elements, const FName Tag)
+{
+	if (Elements.Num() == 0)
+	{
+		UE_LOG(LogLearning, Warning, TEXT("MakeStructActionModifier: Creating zero-sized Struct Action Modifier."));
+	}
+
+	const int32 SubElementNum = Elements.Num();
+
+	TArray<FLearningAgentsActionModifierElement, TInlineAllocator<16>> SubElements;
+	TArray<FName, TInlineAllocator<16>> SubElementNames;
+	SubElements.Empty(SubElementNum);
+	SubElementNames.Empty(SubElementNum);
+
+	for (const TPair<FName, FLearningAgentsActionModifierElement>& Element : Elements)
+	{
+		SubElements.Add(Element.Value);
+		SubElementNames.Add(Element.Key);
+	}
+
+	return MakeStructActionModifierFromArrayViews(Modifier, SubElementNames, SubElements, Tag);
+}
+
+FLearningAgentsActionModifierElement ULearningAgentsActions::MakeStructActionModifierFromArrays(ULearningAgentsActionModifier* Modifier, const TArray<FName>& ElementNames, const TArray<FLearningAgentsActionModifierElement>& Elements, const FName Tag)
+{
+	return MakeStructActionModifierFromArrayViews(Modifier, ElementNames, Elements, Tag);
+}
+
+FLearningAgentsActionModifierElement ULearningAgentsActions::MakeStructActionModifierFromArrayViews(ULearningAgentsActionModifier* Modifier, const TArrayView<const FName> ElementNames, const TArrayView<const FLearningAgentsActionModifierElement> Elements, const FName Tag)
+{
+	if (!Modifier)
+	{
+		UE_LOG(LogLearning, Error, TEXT("MakeStructActionModifierFromArrayViews: Modifier is nullptr."));
+		return FLearningAgentsActionModifierElement();
+	}
+
+	if (Elements.Num() == 0)
+	{
+		UE_LOG(LogLearning, Warning, TEXT("MakeStructActionModifierFromArrayViews: Creating zero-sized Struct Action Modifier."));
+	}
+
+	if (Elements.Num() != ElementNames.Num())
+	{
+		UE_LOG(LogLearning, Error, TEXT("MakeStructActionModifierFromArrayViews: Number of elements (%i) must match number of names (%i)."), Elements.Num(), ElementNames.Num());
+		return FLearningAgentsActionModifierElement();
+	}
+
+	if (UE::Learning::Agents::Action::Private::ContainsDuplicates(ElementNames))
+	{
+		UE_LOG(LogLearning, Error, TEXT("MakeStructActionModifierFromArrayViews: Element Names contain duplicates."));
+		return FLearningAgentsActionModifierElement();
+	}
+
+	TArray<UE::Learning::Action::FModifierElement, TInlineAllocator<16>> SubElements;
+	SubElements.Empty(Elements.Num());
+
+	for (const FLearningAgentsActionModifierElement& Element : Elements)
+	{
+		if (!Modifier->ActionModifier.IsValid(Element.ModifierElement))
+		{
+			UE_LOG(LogLearning, Error, TEXT("MakeStructActionModifierFromArrayViews: Invalid Action Modifier."));
+			return FLearningAgentsActionModifierElement();
+		}
+
+		SubElements.Add(Element.ModifierElement);
+	}
+
+	return { Modifier->ActionModifier.CreateAnd({ ElementNames, SubElements }, Tag) };
+}
+
+FLearningAgentsActionModifierElement ULearningAgentsActions::MakeExclusiveUnionActionModifier(ULearningAgentsActionModifier* Modifier, const TMap<FName, FLearningAgentsActionModifierElement>& Elements, const TArray<FName>& MaskedElements, const FName Tag)
+{
+	const int32 SubElementNum = Elements.Num();
+
+	TArray<FLearningAgentsActionModifierElement, TInlineAllocator<16>> SubElements;
+	TArray<FName, TInlineAllocator<16>> SubElementNames;
+	SubElements.Empty(SubElementNum);
+	SubElementNames.Empty(SubElementNum);
+
+	for (const TPair<FName, FLearningAgentsActionModifierElement>& Element : Elements)
+	{
+		SubElements.Add(Element.Value);
+		SubElementNames.Add(Element.Key);
+	}
+
+	return MakeExclusiveUnionActionModifierFromArrayViews(Modifier, SubElementNames, SubElements, MaskedElements, Tag);
+}
+
+FLearningAgentsActionModifierElement ULearningAgentsActions::MakeExclusiveUnionActionModifierFromArrays(ULearningAgentsActionModifier* Modifier, const TArray<FName>& ElementNames, const TArray<FLearningAgentsActionModifierElement>& Elements, const TArray<FName>& MaskedElements, const FName Tag)
+{
+	return MakeExclusiveUnionActionModifierFromArrayViews(Modifier, ElementNames, Elements, MaskedElements, Tag);
+}
+
+FLearningAgentsActionModifierElement ULearningAgentsActions::MakeExclusiveUnionActionModifierFromArrayViews(ULearningAgentsActionModifier* Modifier, const TArrayView<const FName> ElementNames, const TArrayView<const FLearningAgentsActionModifierElement> Elements, const TArrayView<const FName> MaskedElements, const FName Tag)
+{
+	if (!Modifier)
+	{
+		UE_LOG(LogLearning, Error, TEXT("MakeExclusiveUnionActionModifierFromArrayViews: Modifier is nullptr."));
+		return FLearningAgentsActionModifierElement();
+	}
+
+	if (Elements.Num() != ElementNames.Num())
+	{
+		UE_LOG(LogLearning, Error, TEXT("MakeExclusiveUnionActionModifierFromArrayViews: Number of elements (%i) must match number of names (%i)."), Elements.Num(), ElementNames.Num());
+		return FLearningAgentsActionModifierElement();
+	}
+
+	if (UE::Learning::Agents::Action::Private::ContainsDuplicates(ElementNames))
+	{
+		UE_LOG(LogLearning, Error, TEXT("MakeExclusiveUnionActionModifierFromArrayViews: Element Names contain duplicates."));
+		return FLearningAgentsActionModifierElement();
+	}
+
+	TArray<UE::Learning::Action::FModifierElement, TInlineAllocator<16>> SubElements;
+	SubElements.Empty(Elements.Num());
+
+	for (const FLearningAgentsActionModifierElement& Element : Elements)
+	{
+		if (!Modifier->ActionModifier.IsValid(Element.ModifierElement))
+		{
+			UE_LOG(LogLearning, Error, TEXT("MakeExclusiveUnionActionModifierFromArrayViews: Invalid Action Modifier."));
+			return FLearningAgentsActionModifierElement();
+		}
+
+		SubElements.Add(Element.ModifierElement);
+	}
+
+	return { Modifier->ActionModifier.CreateOrExclusive({ ElementNames, SubElements, MaskedElements }, Tag) };
+}
+
+FLearningAgentsActionModifierElement ULearningAgentsActions::MakeInclusiveUnionActionModifier(ULearningAgentsActionModifier* Modifier, const TMap<FName, FLearningAgentsActionModifierElement>& Elements, const TArray<FName>& MaskedElements, const FName Tag)
+{
+	const int32 SubElementNum = Elements.Num();
+
+	TArray<FLearningAgentsActionModifierElement, TInlineAllocator<16>> SubElements;
+	TArray<FName, TInlineAllocator<16>> SubElementNames;
+	SubElements.Empty(SubElementNum);
+	SubElementNames.Empty(SubElementNum);
+
+	for (const TPair<FName, FLearningAgentsActionModifierElement>& Element : Elements)
+	{
+		SubElements.Add(Element.Value);
+		SubElementNames.Add(Element.Key);
+	}
+
+	return MakeInclusiveUnionActionModifierFromArrayViews(Modifier, SubElementNames, SubElements, MaskedElements, Tag);
+}
+
+FLearningAgentsActionModifierElement ULearningAgentsActions::MakeInclusiveUnionActionModifierFromArrays(ULearningAgentsActionModifier* Modifier, const TArray<FName>& ElementNames, const TArray<FLearningAgentsActionModifierElement>& Elements, const TArray<FName>& MaskedElements, const FName Tag)
+{
+	return MakeInclusiveUnionActionModifierFromArrayViews(Modifier, ElementNames, Elements, MaskedElements, Tag);
+}
+
+FLearningAgentsActionModifierElement ULearningAgentsActions::MakeInclusiveUnionActionModifierFromArrayViews(ULearningAgentsActionModifier* Modifier, const TArrayView<const FName> ElementNames, const TArrayView<const FLearningAgentsActionModifierElement> Elements, const TArrayView<const FName> MaskedElements, const FName Tag)
+{
+	if (!Modifier)
+	{
+		UE_LOG(LogLearning, Error, TEXT("MakeInclusiveUnionActionModifierFromArrayViews: Modifier is nullptr."));
+		return FLearningAgentsActionModifierElement();
+	}
+
+	if (Elements.Num() != ElementNames.Num())
+	{
+		UE_LOG(LogLearning, Error, TEXT("MakeInclusiveUnionActionModifierFromArrayViews: Number of elements (%i) must match number of names (%i)."), Elements.Num(), ElementNames.Num());
+		return FLearningAgentsActionModifierElement();
+	}
+
+	if (UE::Learning::Agents::Action::Private::ContainsDuplicates(ElementNames))
+	{
+		UE_LOG(LogLearning, Error, TEXT("MakeInclusiveUnionActionModifierFromArrayViews: Element Names contain duplicates."));
+		return FLearningAgentsActionModifierElement();
+	}
+
+	TArray<UE::Learning::Action::FModifierElement, TInlineAllocator<16>> SubElements;
+	SubElements.Empty(Elements.Num());
+
+	for (const FLearningAgentsActionModifierElement& Element : Elements)
+	{
+		if (!Modifier->ActionModifier.IsValid(Element.ModifierElement))
+		{
+			UE_LOG(LogLearning, Error, TEXT("MakeInclusiveUnionActionModifierFromArrayViews: Invalid Action Modifier."));
+			return FLearningAgentsActionModifierElement();
+		}
+
+		SubElements.Add(Element.ModifierElement);
+	}
+
+	return { Modifier->ActionModifier.CreateOrInclusive({ ElementNames, SubElements, MaskedElements }, Tag) };
+}
+
+FLearningAgentsActionModifierElement ULearningAgentsActions::MakeStaticArrayActionModifier(ULearningAgentsActionModifier* Modifier, const TArray<FLearningAgentsActionModifierElement>& Elements, const FName Tag)
+{
+	return MakeStaticArrayActionModifierFromArrayView(Modifier, Elements, Tag);
+}
+
+FLearningAgentsActionModifierElement ULearningAgentsActions::MakeStaticArrayActionModifierFromArrayView(ULearningAgentsActionModifier* Modifier, const TArrayView<const FLearningAgentsActionModifierElement> Elements, const FName Tag)
+{
+	if (!Modifier)
+	{
+		UE_LOG(LogLearning, Error, TEXT("MakeStaticArrayActionModifierFromArrayView: Modifier is nullptr."));
+		return FLearningAgentsActionModifierElement();
+	}
+
+	if (Elements.Num() == 0)
+	{
+		UE_LOG(LogLearning, Warning, TEXT("MakeStaticArrayActionModifierFromArrayView: Creating zero-sized Static Array Action."));
+	}
+
+	TArray<UE::Learning::Action::FModifierElement, TInlineAllocator<16>> SubElements;
+	SubElements.Empty(Elements.Num());
+
+	for (const FLearningAgentsActionModifierElement& Element : Elements)
+	{
+		if (!Modifier->ActionModifier.IsValid(Element.ModifierElement))
+		{
+			UE_LOG(LogLearning, Error, TEXT("MakeStaticArrayActionModifierFromArrayView: Invalid Action Modifier."));
+			return FLearningAgentsActionModifierElement();
+		}
+
+		SubElements.Add(Element.ModifierElement);
+	}
+
+	return { Modifier->ActionModifier.CreateArray({ SubElements }, Tag) };
+}
+
+FLearningAgentsActionModifierElement ULearningAgentsActions::MakePairActionModifier(ULearningAgentsActionModifier* Modifier, const FLearningAgentsActionModifierElement Key, const FLearningAgentsActionModifierElement Value, const FName Tag)
+{
+	return MakeStructActionModifierFromArrayViews(Modifier, { TEXT("Key"), TEXT("Value") }, { Key, Value }, Tag);
+}
+
+FLearningAgentsActionModifierElement ULearningAgentsActions::MakeEnumActionModifier(ULearningAgentsActionModifier* Modifier, const UEnum* Enum, const TArray<uint8>& EnumMaskedValues, const FName Tag)
+{
+	return MakeEnumActionModifierFromArrayView(Modifier, Enum, EnumMaskedValues, Tag);
+}
+
+FLearningAgentsActionModifierElement ULearningAgentsActions::MakeEnumActionModifierFromArrayView(ULearningAgentsActionModifier* Modifier, const UEnum* Enum, const TArrayView<const uint8> EnumMaskedValues, const FName Tag)
+{
+	if (!Enum)
+	{
+		UE_LOG(LogLearning, Error, TEXT("MakeEnumActionModifierFromArrayView: Enum is nullptr."));
+		return FLearningAgentsActionModifierElement();
+	}
+
+	TArray<int32, TInlineAllocator<32>> EnumValueIndices;
+
+	for (const uint8 EnumValue : EnumMaskedValues)
+	{
+		const int32 EnumValueIndex = Enum->GetIndexByValue(EnumValue);
+
+		if (EnumValueIndex == INDEX_NONE || EnumValueIndex < 0 || EnumValueIndex >= Enum->NumEnums() - 1)
+		{
+			UE_LOG(LogLearning, Error, TEXT("MakeEnumActionModifierFromArrayView: EnumValue %i not valid for Enum '%s'."), EnumValue, *Enum->GetName());
+			return FLearningAgentsActionModifierElement();
+		}
+	}
+
+	return MakeExclusiveDiscreteActionModifierFromArrayView(Modifier, EnumValueIndices, Tag);
+}
+
+FLearningAgentsActionModifierElement ULearningAgentsActions::MakeBitmaskActionModifier(ULearningAgentsActionModifier* Modifier, const UEnum* Enum, const int32 MaskedBitmask, const FName Tag)
+{
+	if (!Enum)
+	{
+		UE_LOG(LogLearning, Error, TEXT("MakeBitmaskActionModifier: Enum is nullptr."));
+		return FLearningAgentsActionModifierElement();
+	}
+
+	if (Enum->NumEnums() - 1 > 32)
+	{
+		UE_LOG(LogLearning, Error, TEXT("MakeBitmaskActionModifier: Too many values in Enum to use as Bitmask (%i)."), Enum->NumEnums() - 1);
+		return FLearningAgentsActionModifierElement();
+	}
+
+	TArray<int32, TInlineAllocator<32>> BitmaskIndices;
+	BitmaskIndices.Empty(Enum->NumEnums() - 1);
+
+	for (int32 BitmaskIdx = 0; BitmaskIdx < Enum->NumEnums() - 1; BitmaskIdx++)
+	{
+		if (MaskedBitmask & (1 << BitmaskIdx))
+		{
+			BitmaskIndices.Add(BitmaskIdx);
+		}
+	}
+
+	return MakeInclusiveDiscreteActionModifierFromArrayView(Modifier, BitmaskIndices, Tag);
+}
+
+FLearningAgentsActionModifierElement ULearningAgentsActions::MakeOptionalActionModifier(ULearningAgentsActionModifier* Modifier, const FLearningAgentsActionModifierElement Element, const bool bAllowOnlyValid, const bool bAllowOnlyNull, const FName Tag)
+{
+	if (bAllowOnlyValid && bAllowOnlyNull)
+	{
+		UE_LOG(LogLearning, Warning, TEXT("MakeOptionalActionModifier: Must either set bAllowOnlyValid or bAllowOnlyNull."));
+	}
+
+	if (bAllowOnlyValid)
+	{
+		return MakeExclusiveUnionActionModifierFromArrayViews(Modifier, { TEXT("Valid") }, { Element }, { TEXT("Null") }, Tag);
+	}
+
+	if (bAllowOnlyNull)
+	{
+		return MakeExclusiveUnionActionModifierFromArrayViews(Modifier, { TEXT("Valid") }, { Element }, { TEXT("Valid") }, Tag);
+	}
+
+	return MakeExclusiveUnionActionModifierFromArrayViews(Modifier, { TEXT("Valid") }, { Element }, {}, Tag);
+}
+
+FLearningAgentsActionModifierElement ULearningAgentsActions::MakeEitherActionModifier(ULearningAgentsActionModifier* Modifier, const FLearningAgentsActionModifierElement A, const FLearningAgentsActionModifierElement B, const bool bAllowOnlyA, const bool bAllowOnlyB, const FName Tag)
+{
+	if (bAllowOnlyA && bAllowOnlyB)
+	{
+		UE_LOG(LogLearning, Warning, TEXT("MakeEitherActionModifier: Must either set bAllowOnlyA or bAllowOnlyB."));
+	}
+
+	if (bAllowOnlyA)
+	{
+		return MakeExclusiveUnionActionModifierFromArrayViews(Modifier, { TEXT("A"), TEXT("B") }, { A, B }, { TEXT("B") }, Tag);
+	}
+
+	if (bAllowOnlyB)
+	{
+		return MakeExclusiveUnionActionModifierFromArrayViews(Modifier, { TEXT("A"), TEXT("B") }, { A, B }, { TEXT("A") }, Tag);
+	}
+
+	return MakeExclusiveUnionActionModifierFromArrayViews(Modifier, { TEXT("A"), TEXT("B") }, { A, B }, {}, Tag);
+}
+
+FLearningAgentsActionModifierElement ULearningAgentsActions::MakeEncodingActionModifier(ULearningAgentsActionModifier* Modifier, const FLearningAgentsActionModifierElement Element, const FName Tag)
+{
+	if (!Modifier)
+	{
+		UE_LOG(LogLearning, Error, TEXT("MakeEncodingActionModifier: Modifier is nullptr."));
+		return FLearningAgentsActionModifierElement();
+	}
+
+	if (!Modifier->ActionModifier.IsValid(Element.ModifierElement))
+	{
+		UE_LOG(LogLearning, Error, TEXT("MakeEncodingActionModifier: Invalid Action Modifier."));
+		return FLearningAgentsActionModifierElement();
+	}
+
+	return { Modifier->ActionModifier.CreateEncoding({ Element.ModifierElement }, Tag) };
+}
+
+FLearningAgentsActionModifierElement ULearningAgentsActions::MakeBoolActionModifier(ULearningAgentsActionModifier* Modifier, const bool bValue, const FName Tag)
+{
+	return MakeExclusiveDiscreteActionModifierFromArrayView(Modifier, { bValue ? 0 : 1 }, Tag);
+}
+
+FLearningAgentsActionModifierElement ULearningAgentsActions::MakeFloatActionModifier(ULearningAgentsActionModifier* Modifier, const float MaskedValue, const bool bMasked, const FName Tag)
+{
+	return MakeContinuousActionModifierFromArrayView(Modifier, { bMasked }, { MaskedValue }, Tag);
+}
+
+FLearningAgentsActionModifierElement ULearningAgentsActions::MakeLocationActionModifier(ULearningAgentsActionModifier* Modifier, const FVector MaskedLocation, const bool bMaskedX, const bool bMaskedY, const bool bMaskedZ, const FTransform RelativeTransform, const FName Tag)
+{
+	const FVector LocalLocation = RelativeTransform.InverseTransformPosition(MaskedLocation);
+
+	return MakeContinuousActionModifierFromArrayView(Modifier, {
+		bMaskedX,
+		bMaskedY,
+		bMaskedZ
+		}, {
+		(float)LocalLocation.X, 
+		(float)LocalLocation.Y, 
+		(float)LocalLocation.Z }, Tag);
+}
+
+FLearningAgentsActionModifierElement ULearningAgentsActions::MakeScaleActionModifier(ULearningAgentsActionModifier* Modifier, const FVector MaskedScale, const bool bMaskedX, const bool bMaskedY, const bool bMaskedZ, const FVector RelativeScale, const FName Tag)
+{
+	const FVector LocalLogScale =
+		UE::Learning::Agents::Action::Private::VectorLogSafe(MaskedScale) -
+		UE::Learning::Agents::Action::Private::VectorLogSafe(RelativeScale);
+
+	return MakeContinuousActionModifierFromArrayView(Modifier, {
+		bMaskedX,
+		bMaskedY,
+		bMaskedZ
+		}, {
+		(float)LocalLogScale.X, 
+		(float)LocalLogScale.Y, 
+		(float)LocalLogScale.Z } , Tag);
+}
+
+FLearningAgentsActionModifierElement ULearningAgentsActions::MakeAngleActionModifier(ULearningAgentsActionModifier* Modifier, const float MaskedAngle, const bool bMask, const float RelativeAngle, const FName Tag)
+{
+	return MakeContinuousActionModifierFromArrayView(Modifier, { bMask }, { FMath::FindDeltaAngleDegrees(RelativeAngle, MaskedAngle) }, Tag);
+}
+
+FLearningAgentsActionModifierElement ULearningAgentsActions::MakeAngleActionModifierRadians(ULearningAgentsActionModifier* Modifier, const float MaskedAngle, const bool bMask, const float RelativeAngle, const FName Tag)
+{
+	return MakeAngleActionModifier(Modifier, FMath::RadiansToDegrees(MaskedAngle), bMask, FMath::RadiansToDegrees(RelativeAngle), Tag);
+}
+
+FLearningAgentsActionModifierElement ULearningAgentsActions::MakeVelocityActionModifier(ULearningAgentsActionModifier* Modifier, const FVector MaskedVelocity, const bool bMaskedX, const bool bMaskedY, const bool bMaskedZ, const FTransform RelativeTransform, const FName Tag)
+{
+	const FVector LocalVelocity = RelativeTransform.InverseTransformVectorNoScale(MaskedVelocity);
+
+	return MakeContinuousActionModifierFromArrayView(Modifier, {
+		bMaskedX,
+		bMaskedY,
+		bMaskedZ }, {
+		(float)LocalVelocity.X,
+		(float)LocalVelocity.Y,
+		(float)LocalVelocity.Z,
+		}, Tag);
+}
+
+FLearningAgentsActionModifierElement ULearningAgentsActions::MakeDirectionActionModifier(ULearningAgentsActionModifier* Modifier, const FVector MaskedDirection, const bool bMaskedX, const bool bMaskedY, const bool bMaskedZ, const FTransform RelativeTransform, const FName Tag)
+{
+	const FVector LocalDirection = RelativeTransform.InverseTransformVectorNoScale(MaskedDirection).GetSafeNormal(UE_SMALL_NUMBER, FVector::ForwardVector);
+
+	return MakeContinuousActionModifierFromArrayView(Modifier, {
+		bMaskedX,
+		bMaskedY,
+		bMaskedZ }, {
+		(float)LocalDirection.X,
+		(float)LocalDirection.Y,
+		(float)LocalDirection.Z,
+		}, Tag);
+}
+
 bool ULearningAgentsActions::GetNullAction(const ULearningAgentsActionObject* Object, const FLearningAgentsActionObjectElement Element, const FName Tag)
 {
 	if (!Object)
diff --git a/Engine/Plugins/Experimental/LearningAgents/Source/LearningAgents/Private/LearningAgentsInteractor.cpp b/Engine/Plugins/Experimental/LearningAgents/Source/LearningAgents/Private/LearningAgentsInteractor.cpp
--- a/Engine/Plugins/Experimental/LearningAgents/Source/LearningAgents/Private/LearningAgentsInteractor.cpp
+++ b/Engine/Plugins/Experimental/LearningAgents/Source/LearningAgents/Private/LearningAgentsInteractor.cpp
@@ -127,6 +127,7 @@ void ULearningAgentsInteractor::SetupInteractor(ULearningAgentsManager*& InManag
 	}
 
 	const int32 ActionDistributionVectorSize = ActionSchema->ActionSchema.GetActionDistributionVectorSize(ActionSchemaElement.SchemaElement);
+	const int32 ActionModifierVectorSize = ActionSchema->ActionSchema.GetActionModifierVectorSize(ActionSchemaElement.SchemaElement);
 
 	ActionVectors.SetNumUninitialized({ Manager->GetMaxAgentNum(), ActionVectorSize });
 	ActionCompatibilityHash = UE::Learning::Action::GetSchemaObjectsCompatibilityHash(ActionSchema->ActionSchema, ActionSchemaElement.SchemaElement);
@@ -136,12 +137,21 @@ void ULearningAgentsInteractor::SetupInteractor(ULearningAgentsManager*& InManag
 	ActionObject = NewObject<ULearningAgentsActionObject>(this, ActionObjectUniqueName);
 	ActionObjectElements.Empty(Manager->GetMaxAgentNum());
 
+	ActionModifierVectors.SetNumUninitialized({ Manager->GetMaxAgentNum(), ActionModifierVectorSize });
+	
+	const FName ActionModifierUniqueName = MakeUniqueObjectName(this, ULearningAgentsActionModifier::StaticClass(), TEXT("ActionModifier"), EUniqueObjectNameOptions::GloballyUnique);
+
+	ActionModifier = NewObject<ULearningAgentsActionModifier>(this, ActionModifierUniqueName);
+	ActionModifierElements.Empty(Manager->GetMaxAgentNum());
+
 	// Reset Agent iteration
 
 	ObservationVectorIteration.SetNumUninitialized({ Manager->GetMaxAgentNum() });
+	ActionModifierVectorIteration.SetNumUninitialized({ Manager->GetMaxAgentNum() });
 	ActionVectorIteration.SetNumUninitialized({ Manager->GetMaxAgentNum() });
 
 	UE::Learning::Array::Set<1, uint64>(ObservationVectorIteration, INDEX_NONE);
+	UE::Learning::Array::Set<1, uint64>(ActionModifierVectorIteration, INDEX_NONE);
 	UE::Learning::Array::Set<1, uint64>(ActionVectorIteration, INDEX_NONE);
 
 	bIsSetup = true;
@@ -158,6 +168,7 @@ void ULearningAgentsInteractor::OnAgentsAdded_Implementation(const TArray<int32>
 	}
 
 	UE::Learning::Array::Set<1, uint64>(ObservationVectorIteration, 0, AgentIds);
+	UE::Learning::Array::Set<1, uint64>(ActionModifierVectorIteration, 0, AgentIds);
 	UE::Learning::Array::Set<1, uint64>(ActionVectorIteration, 0, AgentIds);
 }
 
@@ -170,6 +181,7 @@ void ULearningAgentsInteractor::OnAgentsRemoved_Implementation(const TArray<int3
 	}
 
 	UE::Learning::Array::Set<1, uint64>(ObservationVectorIteration, INDEX_NONE, AgentIds);
+	UE::Learning::Array::Set<1, uint64>(ActionModifierVectorIteration, INDEX_NONE, AgentIds);
 	UE::Learning::Array::Set<1, uint64>(ActionVectorIteration, INDEX_NONE, AgentIds);
 }
 
@@ -182,6 +194,7 @@ void ULearningAgentsInteractor::OnAgentsReset_Implementation(const TArray<int32>
 	}
 
 	UE::Learning::Array::Set<1, uint64>(ObservationVectorIteration, 0, AgentIds);
+	UE::Learning::Array::Set<1, uint64>(ActionModifierVectorIteration, 0, AgentIds);
 	UE::Learning::Array::Set<1, uint64>(ActionVectorIteration, 0, AgentIds);
 }
 
@@ -225,7 +238,7 @@ void ULearningAgentsInteractor::PerformAgentActions_Implementation(const ULearni
 
 	if (AgentNum != InActionObjectElements.Num())
 	{
-		UE_LOG(LogLearning, Error, TEXT("%s: Not enough Action Objects. Expected %i, Got %i."), *GetName(), AgentNum, InActionObjectElements.Num());
+		UE_LOG(LogLearning, Error, TEXT("%s: PerformAgentActions: Not enough Action Objects. Expected %i, Got %i."), *GetName(), AgentNum, InActionObjectElements.Num());
 		return;
 	}
 
@@ -235,6 +248,40 @@ void ULearningAgentsInteractor::PerformAgentActions_Implementation(const ULearni
 	}
 }
 
+void ULearningAgentsInteractor::MakeAgentActionModifier_Implementation(
+	FLearningAgentsActionModifierElement& OutActionModifierElement, 
+	ULearningAgentsActionModifier* InActionModifier, 
+	const ULearningAgentsObservationObject* InObservationObject,
+	const FLearningAgentsObservationObjectElement& InObservationObjectElement,
+	const int32 AgentId)
+{
+	OutActionModifierElement = ULearningAgentsActions::MakeNullActionModifier(InActionModifier);
+}
+
+void ULearningAgentsInteractor::MakeAgentActionModifiers_Implementation(
+	TArray<FLearningAgentsActionModifierElement>& OutActionModifierElements, 
+	ULearningAgentsActionModifier* InActionModifier, 
+	const ULearningAgentsObservationObject* InObservationObject,
+	const TArray<FLearningAgentsObservationObjectElement>& InObservationObjectElements,
+	const TArray<int32>& AgentIds)
+{
+	const int32 AgentNum = AgentIds.Num();
+
+	if (AgentNum != InObservationObjectElements.Num())
+	{
+		UE_LOG(LogLearning, Error, TEXT("%s: MakeAgentActionModifiers: Not enough Observation Objects. Expected %i, Got %i."), *GetName(), AgentNum, InObservationObjectElements.Num());
+		return;
+	}
+
+	OutActionModifierElements.Empty(AgentNum);
+	for (int32 AgentIdx = 0; AgentIdx < AgentNum; AgentIdx++)
+	{
+		FLearningAgentsActionModifierElement OutElement;
+		MakeAgentActionModifier(OutElement, InActionModifier, InObservationObject, InObservationObjectElements[AgentIdx], AgentIds[AgentIdx]);
+		OutActionModifierElements.Add(OutElement);
+	}
+}
+
 void ULearningAgentsInteractor::GatherObservations(const UE::Learning::FIndexSet AgentSet, bool bIncrementIteration)
 {
 	UE_LEARNING_TRACE_CPUPROFILER_EVENT_SCOPE(ULearningAgentsInteractor::GatherObservations);
@@ -288,6 +335,82 @@ void ULearningAgentsInteractor::GatherObservations(const UE::Learning::FIndexSet
 	}
 }
 
+void ULearningAgentsInteractor::MakeActionModifiers(const UE::Learning::FIndexSet AgentSet, bool bIncrementIteration)
+{
+	UE_LEARNING_TRACE_CPUPROFILER_EVENT_SCOPE(ULearningAgentsInteractor::MakeActionModifiers);
+
+	if (!IsSetup())
+	{
+		UE_LOG(LogLearning, Error, TEXT("%s: Setup not complete."), *GetName());
+		return;
+	}
+
+	// Check which agents have had their observation vector set
+
+	ValidAgentIds.Empty(Manager->GetMaxAgentNum());
+
+	for (const int32 AgentId : AgentSet)
+	{
+		if (ObservationVectorIteration[AgentId] == 0)
+		{
+			UE_LOG(LogLearning, Warning, TEXT("%s: Agent with id %i does not have an observation vector so action modifiers will not be created for it. Was GatherObservations run without error?"), *GetName(), AgentId);
+			continue;
+		}
+
+		ValidAgentIds.Add(AgentId);
+	}
+
+	ValidAgentSet = ValidAgentIds;
+	ValidAgentSet.TryMakeSlice();
+
+	// Run MakeAgentActionModifiers Callback
+
+	ValidAgentIds.Empty(Manager->GetMaxAgentNum());
+	for (const int32 AgentId : AgentSet)
+	{
+		ValidAgentIds.Add(AgentId);
+	}
+
+	ActionModifier->ActionModifier.Reset();
+	ActionModifierElements.Empty(Manager->GetMaxAgentNum());
+	MakeAgentActionModifiers(
+		ActionModifierElements,
+		ActionModifier,
+		ObservationObject,
+		ObservationObjectElements, 
+		ValidAgentIds);
+
+	if (ValidAgentIds.Num() != ActionModifierElements.Num())
+	{
+		UE_LOG(LogLearning, Error, TEXT("%s: Not enough Action Modifiers added by MakeAgentActionModifiers. Expected %i, Got %i."), *GetName(), ValidAgentIds.Num(), ActionModifierElements.Num());
+		return;
+	}
+
+	// Check Action Modifiers are Valid and if so convert to action modifier vectors
+
+	for (int32 AgentIdx = 0; AgentIdx < AgentSet.Num(); AgentIdx++)
+	{
+		if (ULearningAgentsActions::ValidateActionModifierMatchesSchema(
+			ActionSchema,
+			ActionSchemaElement,
+			ActionModifier,
+			ActionModifierElements[AgentIdx]))
+		{
+			UE::Learning::Action::SetVectorFromModifier(
+				ActionModifierVectors[AgentSet[AgentIdx]],
+				ActionSchema->ActionSchema,
+				ActionSchemaElement.SchemaElement,
+				ActionModifier->ActionModifier,
+				ActionModifierElements[AgentIdx].ModifierElement);
+
+			if (bIncrementIteration)
+			{
+				ActionModifierVectorIteration[AgentSet[AgentIdx]]++;
+			}
+		}
+	}
+}
+
 void ULearningAgentsInteractor::PerformActions(const UE::Learning::FIndexSet AgentSet)
 {
 	UE_LEARNING_TRACE_CPUPROFILER_EVENT_SCOPE(ULearningAgentsInteractor::PerformActions);
@@ -350,6 +473,16 @@ void ULearningAgentsInteractor::GatherObservations()
 	GatherObservations(Manager->GetAllAgentSet());
 }
 
+void ULearningAgentsInteractor::MakeActionModifiers()
+{
+	if (Manager->GetAgentNum() == 0)
+	{
+		UE_LOG(LogLearning, Warning, TEXT("%s: No agents added to Manager."), *GetName());
+	}
+
+	MakeActionModifiers(Manager->GetAllAgentSet());
+}
+
 void ULearningAgentsInteractor::PerformActions()
 {
 	if (Manager->GetAgentNum() == 0)
@@ -392,6 +525,37 @@ void ULearningAgentsInteractor::GetObservationVector(TArray<float>& OutObservati
 	OutObservationCompatibilityHash = ObservationCompatibilityHash;
 }
 
+void ULearningAgentsInteractor::GetActionModifierVector(TArray<float>& OutActionModifierVector, int32& OutActionCompatibilityHash, const int32 AgentId)
+{
+	if (!IsSetup())
+	{
+		UE_LOG(LogLearning, Error, TEXT("%s: Setup not complete."), *GetName());
+		OutActionModifierVector.Empty();
+		OutActionCompatibilityHash = 0;
+		return;
+	}
+
+	if (!HasAgent(AgentId))
+	{
+		UE_LOG(LogLearning, Error, TEXT("%s: AgentId %d not found in the agents set."), *GetName(), AgentId);
+		OutActionModifierVector.Empty();
+		OutActionCompatibilityHash = 0;
+		return;
+	}
+
+	if (ActionModifierVectorIteration[AgentId] == 0)
+	{
+		UE_LOG(LogLearning, Error, TEXT("%s: Action Modifier vector not set for agent %i."), *GetName(), AgentId);
+		OutActionModifierVector.Empty();
+		OutActionCompatibilityHash = 0;
+		return;
+	}
+
+	OutActionModifierVector.SetNumUninitialized(GetActionModifierVectorSize());
+	UE::Learning::Array::Copy<1, float>(OutActionModifierVector, ActionModifierVectors[AgentId]);
+	OutActionCompatibilityHash = ActionCompatibilityHash;
+}
+
 void ULearningAgentsInteractor::GetActionVector(TArray<float>& OutActionVector, int32& OutActionCompatibilityHash, const int32 AgentId)
 {
 	if (!IsSetup())
@@ -453,6 +617,36 @@ void ULearningAgentsInteractor::SetObservationVector(const TArray<float>& Observ
 	if (bIncrementIteration) { ObservationVectorIteration[AgentId]++; }
 }
 
+void ULearningAgentsInteractor::SetActionModifierVector(const TArray<float>& ActionModifierVector, const int32 InActionCompatibilityHash, const int32 AgentId, bool bIncrementIteration)
+{
+	if (!IsSetup())
+	{
+		UE_LOG(LogLearning, Error, TEXT("%s: Setup not complete."), *GetName());
+		return;
+	}
+
+	if (!HasAgent(AgentId))
+	{
+		UE_LOG(LogLearning, Error, TEXT("%s: AgentId %d not found in the agents set."), *GetName(), AgentId);
+		return;
+	}
+
+	if (InActionCompatibilityHash != ActionCompatibilityHash)
+	{
+		UE_LOG(LogLearning, Error, TEXT("%s: Action Compatibility hash incompatible. Got %i, expected %i."), *GetName(), InActionCompatibilityHash, ActionCompatibilityHash);
+		return;
+	}
+
+	if (ActionModifierVector.Num() != GetActionModifierVectorSize())
+	{
+		UE_LOG(LogLearning, Error, TEXT("%s: Action Modifier Vector size incompatible. Got %i, expected %i."), *GetName(), ActionModifierVector.Num(), GetActionModifierVectorSize());
+		return;
+	}
+
+	UE::Learning::Array::Copy<1, float>(ActionModifierVectors[AgentId], ActionModifierVector);
+	if (bIncrementIteration) { ActionModifierVectorIteration[AgentId]++; }
+}
+
 void ULearningAgentsInteractor::SetActionVector(const TArray<float>& ActionVector, const int32 InActionCompatibilityHash, const int32 AgentId, bool bIncrementIteration)
 {
 	if (!IsSetup())
@@ -500,6 +694,23 @@ bool ULearningAgentsInteractor::HasObservationVector(const int32 AgentId) const
 	return ObservationVectorIteration[AgentId] > 0;
 }
 
+bool ULearningAgentsInteractor::HasActionModifierVector(const int32 AgentId) const
+{
+	if (!IsSetup())
+	{
+		UE_LOG(LogLearning, Error, TEXT("%s: Setup not complete."), *GetName());
+		return false;
+	}
+
+	if (!HasAgent(AgentId))
+	{
+		UE_LOG(LogLearning, Error, TEXT("%s: AgentId %d not found in the agents set."), *GetName(), AgentId);
+		return false;
+	}
+
+	return ActionModifierVectorIteration[AgentId] > 0;
+}
+
 bool ULearningAgentsInteractor::HasActionVector(const int32 AgentId) const
 {
 	if (!IsSetup())
@@ -561,6 +772,17 @@ int32 ULearningAgentsInteractor::GetActionDistributionVectorSize() const
 	return ActionSchema->ActionSchema.GetActionDistributionVectorSize(ActionSchemaElement.SchemaElement);
 }
 
+int32 ULearningAgentsInteractor::GetActionModifierVectorSize() const
+{
+	if (!IsSetup())
+	{
+		UE_LOG(LogLearning, Error, TEXT("%s: Setup not complete."), *GetName());
+		return 0;
+	}
+
+	return ActionSchema->ActionSchema.GetActionModifierVectorSize(ActionSchemaElement.SchemaElement);
+}
+
 int32 ULearningAgentsInteractor::GetActionEncodedVectorSize() const
 {
 	if (!IsSetup())
@@ -592,7 +814,7 @@ const FLearningAgentsActionSchemaElement ULearningAgentsInteractor::GetActionSch
 	return ActionSchemaElement;
 }
 
-TLearningArrayView<2, const float> ULearningAgentsInteractor::GetObservationVectorArrayView() const
+TLearningArrayView<2, const float> ULearningAgentsInteractor::GetObservationVectorsArrayView() const
 {
 	return ObservationVectors;
 }
@@ -602,7 +824,17 @@ uint64 ULearningAgentsInteractor::GetObservationIteration(const int32 AgentId) c
 	return ObservationVectorIteration[AgentId];
 }
 
-TLearningArrayView<2, const float> ULearningAgentsInteractor::GetActionVectorArrayView() const
+TLearningArrayView<2, const float> ULearningAgentsInteractor::GetActionModifierVectorsArrayView() const
+{
+	return ActionModifierVectors;
+}
+
+uint64 ULearningAgentsInteractor::GetActionModifierIteration(const int32 AgentId) const
+{
+	return ActionModifierVectorIteration[AgentId];
+}
+
+TLearningArrayView<2, const float> ULearningAgentsInteractor::GetActionVectorsArrayView() const
 {
 	return ActionVectors;
 }
@@ -622,6 +854,16 @@ const TArray<FLearningAgentsObservationObjectElement>& ULearningAgentsInteractor
 	return ObservationObjectElements;
 }
 
+const ULearningAgentsActionModifier* ULearningAgentsInteractor::GetActionModifier() const
+{
+	return ActionModifier;
+}
+
+const TArray<FLearningAgentsActionModifierElement>& ULearningAgentsInteractor::GetActionModifierElements() const
+{
+	return ActionModifierElements;
+}
+
 ULearningAgentsActionObject* ULearningAgentsInteractor::GetActionObject()
 {
 	return ActionObject;
diff --git a/Engine/Plugins/Experimental/LearningAgents/Source/LearningAgents/Private/LearningAgentsObservations.cpp b/Engine/Plugins/Experimental/LearningAgents/Source/LearningAgents/Private/LearningAgentsObservations.cpp
--- a/Engine/Plugins/Experimental/LearningAgents/Source/LearningAgents/Private/LearningAgentsObservations.cpp
+++ b/Engine/Plugins/Experimental/LearningAgents/Source/LearningAgents/Private/LearningAgentsObservations.cpp
@@ -14,11 +14,21 @@
 #include "Engine/World.h"
 #include "Components/SplineComponent.h"
 
+bool operator==(const FLearningAgentsObservationSchemaElement& Lhs, const FLearningAgentsObservationSchemaElement& Rhs)
+{
+	return Lhs.SchemaElement.Index == Rhs.SchemaElement.Index;
+}
+
 bool operator==(const FLearningAgentsObservationObjectElement& Lhs, const FLearningAgentsObservationObjectElement& Rhs)
 {
 	return Lhs.ObjectElement.Index == Rhs.ObjectElement.Index;
 }
 
+uint32 GetTypeHash(const FLearningAgentsObservationSchemaElement& Element)
+{
+	return (uint32)Element.SchemaElement.Index;
+}
+
 uint32 GetTypeHash(const FLearningAgentsObservationObjectElement& Element)
 {
 	return (uint32)Element.ObjectElement.Index;
diff --git a/Engine/Plugins/Experimental/LearningAgents/Source/LearningAgents/Private/LearningAgentsPolicy.cpp b/Engine/Plugins/Experimental/LearningAgents/Source/LearningAgents/Private/LearningAgentsPolicy.cpp
--- a/Engine/Plugins/Experimental/LearningAgents/Source/LearningAgents/Private/LearningAgentsPolicy.cpp
+++ b/Engine/Plugins/Experimental/LearningAgents/Source/LearningAgents/Private/LearningAgentsPolicy.cpp
@@ -569,10 +569,10 @@ void ULearningAgentsPolicy::EncodeObservations()
 
 	// Encode Observations
 
-	if (EncoderObject->GetNeuralNetwork()->GetInputSize() != Interactor->GetObservationVectorArrayView().Num<1>())
+	if (EncoderObject->GetNeuralNetwork()->GetInputSize() != Interactor->GetObservationVectorSize())
 	{
 		UE_LOG(LogLearning, Error, TEXT("%s: Encoder Network Input size doesn't match. Network input size is %i but Encoder expects %i."), *GetName(),
-			EncoderObject->GetNeuralNetwork()->GetInputSize(), Interactor->GetObservationVectorArrayView().Num<1>());
+			EncoderObject->GetNeuralNetwork()->GetInputSize(), Interactor->GetObservationVectorSize());
 		return;
 	}
 
@@ -583,7 +583,7 @@ void ULearningAgentsPolicy::EncodeObservations()
 		return;
 	}
 
-	EncoderObject->Evaluate(ObservationVectorsEncoded, Interactor->GetObservationVectorArrayView(), ValidAgentSet);
+	EncoderObject->Evaluate(ObservationVectorsEncoded, Interactor->GetObservationVectorsArrayView(), ValidAgentSet);
 
 	for (const int32 AgentId : ValidAgentSet)
 	{
@@ -675,7 +675,7 @@ void ULearningAgentsPolicy::DecodeAndSampleActions(const float ActionNoiseScale)
 		UE_LOG(LogLearning, Warning, TEXT("%s: No agents added to Manager."), *GetName());
 	}
 
-	// Check which agents have had their encoded action vector set
+	// Check which agents have had their encoded action vector set and action modifiers made
 
 	ValidAgentIds.Empty(Manager->GetMaxAgentNum());
 
@@ -687,6 +687,12 @@ void ULearningAgentsPolicy::DecodeAndSampleActions(const float ActionNoiseScale)
 			continue;
 		}
 
+		if (Interactor->GetActionModifierIteration(AgentId) == 0)
+		{
+			UE_LOG(LogLearning, Warning, TEXT("%s: Agent with id %i does not have an action modifier made so actions will not be decoded for it. Was MakeActionModifiers run without error?"), *GetName(), AgentId);
+			continue;
+		}
+
 		ValidAgentIds.Add(AgentId);
 	}
 
@@ -710,7 +716,7 @@ void ULearningAgentsPolicy::DecodeAndSampleActions(const float ActionNoiseScale)
 	}
 
 	DecoderObject->Evaluate(ActionDistributionVectors, ActionVectorsEncoded, ValidAgentSet);
-
+	
 	// Sample actions from action distribution vectors
 
 	for (const int32 AgentId : ValidAgentSet)
@@ -719,6 +725,7 @@ void ULearningAgentsPolicy::DecodeAndSampleActions(const float ActionNoiseScale)
 			Seeds[AgentId],
 			Interactor->GetActionVectorsArrayView()[AgentId],
 			ActionDistributionVectors[AgentId],
+			Interactor->GetActionModifierVectorsArrayView()[AgentId],
 			Interactor->GetActionSchema()->ActionSchema,
 			Interactor->GetActionSchemaElement().SchemaElement,
 			ActionNoiseScale);
@@ -741,6 +748,7 @@ void ULearningAgentsPolicy::RunInference(const float ActionNoiseScale)
 	}
 
 	Interactor->GatherObservations();
+	Interactor->MakeActionModifiers();
 	EncodeObservations();
 	EvaluatePolicy();
 	DecodeAndSampleActions(ActionNoiseScale);
diff --git a/Engine/Plugins/Experimental/LearningAgents/Source/LearningAgents/Public/LearningAgentsActions.h b/Engine/Plugins/Experimental/LearningAgents/Source/LearningAgents/Public/LearningAgentsActions.h
--- a/Engine/Plugins/Experimental/LearningAgents/Source/LearningAgents/Public/LearningAgentsActions.h
+++ b/Engine/Plugins/Experimental/LearningAgents/Source/LearningAgents/Public/LearningAgentsActions.h
@@ -15,8 +15,10 @@
 
 class ULearningAgentsActionSchema;
 class ULearningAgentsActionObject;
+class ULearningAgentsActionModifier;
 struct FLearningAgentsActionSchemaElement;
 struct FLearningAgentsActionObjectElement;
+struct FLearningAgentsActionModifierElement;
 
 /** An element of an Action Schema */
 USTRUCT(BlueprintType)
@@ -36,11 +38,33 @@ struct LEARNINGAGENTS_API FLearningAgentsActionObjectElement
 	UE::Learning::Action::FObjectElement ObjectElement;
 };
 
-/** Comparison operator for Action Object Elements */
+/** An element of an Action Modifier */
+USTRUCT(BlueprintType)
+struct LEARNINGAGENTS_API FLearningAgentsActionModifierElement
+{
+	GENERATED_BODY()
+
+	UE::Learning::Action::FModifierElement ModifierElement;
+};
+
+/** Comparison and Hashing operators for Action Elements */
+
+LEARNINGAGENTS_API bool operator==(const FLearningAgentsActionSchemaElement& Lhs, const FLearningAgentsActionSchemaElement& Rhs);
 LEARNINGAGENTS_API bool operator==(const FLearningAgentsActionObjectElement& Lhs, const FLearningAgentsActionObjectElement& Rhs);
+LEARNINGAGENTS_API bool operator==(const FLearningAgentsActionModifierElement& Lhs, const FLearningAgentsActionModifierElement& Rhs);
 
-/** Hashing operator for Action Object Elements */
+LEARNINGAGENTS_API uint32 GetTypeHash(const FLearningAgentsActionSchemaElement& Element);
 LEARNINGAGENTS_API uint32 GetTypeHash(const FLearningAgentsActionObjectElement& Element);
+LEARNINGAGENTS_API uint32 GetTypeHash(const FLearningAgentsActionModifierElement& Element);
+
+template<>
+struct TStructOpsTypeTraits<FLearningAgentsActionSchemaElement> : public TStructOpsTypeTraitsBase2<FLearningAgentsActionSchemaElement>
+{
+	enum
+	{
+		WithIdenticalViaEquality = true,
+	};
+};
 
 template<>
 struct TStructOpsTypeTraits<FLearningAgentsActionObjectElement> : public TStructOpsTypeTraitsBase2<FLearningAgentsActionObjectElement>
@@ -51,6 +75,15 @@ struct TStructOpsTypeTraits<FLearningAgentsActionObjectElement> : public TStruct
 	};
 };
 
+template<>
+struct TStructOpsTypeTraits<FLearningAgentsActionModifierElement> : public TStructOpsTypeTraitsBase2<FLearningAgentsActionModifierElement>
+{
+	enum
+	{
+		WithIdenticalViaEquality = true,
+	};
+};
+
 /**
  * Action Schema
  *
@@ -81,6 +114,21 @@ public:
 	UE::Learning::Action::FObject ActionObject;
 };
 
+/**
+ * Action Modifier
+ *
+ * This object is used to construct or get the values of action modifiers.
+ */
+UCLASS(BlueprintType)
+class LEARNINGAGENTS_API ULearningAgentsActionModifier : public UObject
+{
+	GENERATED_BODY()
+
+public:
+
+	UE::Learning::Action::FModifier ActionModifier;
+};
+
 /** Enum Type representing either action A or action B */
 UENUM(BlueprintType)
 enum class ELearningAgentsEitherAction : uint8
@@ -120,6 +168,22 @@ public:
 		const ULearningAgentsActionObject* Object,
 		const FLearningAgentsActionObjectElement ObjectElement);
 
+	/**
+	 * Validates that the given action modifier matches the schema. Will log errors on modifiers that don't match.
+	 *
+	 * @param Schema				Action Schema
+	 * @param SchemaElement			Action Schema Element
+	 * @param Modifier				Action Modifier
+	 * @param ModifierElement		Action Modifier Element
+	 * @returns						true if the modifier matches the schema
+	 */
+	UFUNCTION(BlueprintPure = false, Category = "LearningAgents")
+	static bool ValidateActionModifierMatchesSchema(
+		const ULearningAgentsActionSchema* Schema,
+		const FLearningAgentsActionSchemaElement SchemaElement,
+		const ULearningAgentsActionModifier* Modifier,
+		const FLearningAgentsActionModifierElement ModifierElement);
+
 	/**
 	 * Logs an Action Object Element. Useful for debugging.
 	 *
@@ -659,6 +723,91 @@ public:
 	UFUNCTION(BlueprintPure, Category = "LearningAgents", meta = (AdvancedDisplay = 3))
 	static FLearningAgentsActionObjectElement MakeDirectionAction(ULearningAgentsActionObject* Object, const FVector Direction, const FTransform RelativeTransform = FTransform(), const FName Tag = TEXT("DirectionAction"));
 
+public:
+
+	UFUNCTION(BlueprintPure, Category = "LearningAgents", meta = (AdvancedDisplay = 1))
+	static FLearningAgentsActionModifierElement MakeNullActionModifier(ULearningAgentsActionModifier* Modifier, const FName Tag = TEXT("NullAction"));
+
+	UFUNCTION(BlueprintPure, Category = "LearningAgents", meta = (AdvancedDisplay = 3))
+	static FLearningAgentsActionModifierElement MakeContinuousActionModifier(ULearningAgentsActionModifier* Modifier, const TArray<bool>& Masked, const TArray<float>& MaskedValues, const FName Tag = TEXT("ContinuousAction"));
+	static FLearningAgentsActionModifierElement MakeContinuousActionModifierFromArrayView(ULearningAgentsActionModifier* Modifier, const TArrayView<const bool> Masked, const TArrayView<const float> MaskedValues, const FName Tag = TEXT("ContinuousAction"));
+
+	UFUNCTION(BlueprintPure, Category = "LearningAgents", meta = (AdvancedDisplay = 2))
+	static FLearningAgentsActionModifierElement MakeExclusiveDiscreteActionModifier(ULearningAgentsActionModifier* Modifier, const TArray<int32>& MaskedIndices, const FName Tag = TEXT("DiscreteExclusiveAction"));
+	static FLearningAgentsActionModifierElement MakeExclusiveDiscreteActionModifierFromArrayView(ULearningAgentsActionModifier* Modifier, const TArrayView<const int32> MaskedIndices, const FName Tag = TEXT("DiscreteExclusiveAction"));
+
+	UFUNCTION(BlueprintPure, Category = "LearningAgents", meta = (AdvancedDisplay = 2))
+	static FLearningAgentsActionModifierElement MakeInclusiveDiscreteActionModifier(ULearningAgentsActionModifier* Modifier, const TArray<int32>& MaskedIndices, const FName Tag = TEXT("DiscreteInclusiveAction"));
+	static FLearningAgentsActionModifierElement MakeInclusiveDiscreteActionModifierFromArrayView(ULearningAgentsActionModifier* Modifier, const TArrayView<const int32> MaskedIndices, const FName Tag = TEXT("DiscreteInclusiveAction"));
+
+	UFUNCTION(BlueprintPure, Category = "LearningAgents", meta = (AdvancedDisplay = 2))
+	static FLearningAgentsActionModifierElement MakeStructActionModifier(ULearningAgentsActionModifier* Modifier, const TMap<FName, FLearningAgentsActionModifierElement>& Elements, const FName Tag = TEXT("StructAction"));
+
+	UFUNCTION(BlueprintPure, Category = "LearningAgents", meta = (AdvancedDisplay = 3))
+	static FLearningAgentsActionModifierElement MakeStructActionModifierFromArrays(ULearningAgentsActionModifier* Modifier, const TArray<FName>& ElementNames, const TArray<FLearningAgentsActionModifierElement>& Elements, const FName Tag = TEXT("StructAction"));
+	static FLearningAgentsActionModifierElement MakeStructActionModifierFromArrayViews(ULearningAgentsActionModifier* Modifier, const TArrayView<const FName> ElementNames, const TArrayView<const FLearningAgentsActionModifierElement> Elements, const FName Tag = TEXT("StructAction"));
+
+	UFUNCTION(BlueprintPure, Category = "LearningAgents", meta = (AdvancedDisplay = 3))
+	static FLearningAgentsActionModifierElement MakeExclusiveUnionActionModifier(ULearningAgentsActionModifier* Modifier, const TMap<FName, FLearningAgentsActionModifierElement>& Elements, const TArray<FName>& MaskedElements, const FName Tag = TEXT("ExclusiveUnionAction"));
+
+	UFUNCTION(BlueprintPure, Category = "LearningAgents", meta = (AdvancedDisplay = 4))
+	static FLearningAgentsActionModifierElement MakeExclusiveUnionActionModifierFromArrays(ULearningAgentsActionModifier* Modifier, const TArray<FName>& ElementNames, const TArray<FLearningAgentsActionModifierElement>& Elements, const TArray<FName>& MaskedElements, const FName Tag = TEXT("ExclusiveUnionAction"));
+	static FLearningAgentsActionModifierElement MakeExclusiveUnionActionModifierFromArrayViews(ULearningAgentsActionModifier* Modifier, const TArrayView<const FName> ElementNames, const TArrayView<const FLearningAgentsActionModifierElement> Elements, const TArrayView<const FName> MaskedElements, const FName Tag = TEXT("ExclusiveUnionAction"));
+
+	UFUNCTION(BlueprintPure, Category = "LearningAgents", meta = (AdvancedDisplay = 3))
+	static FLearningAgentsActionModifierElement MakeInclusiveUnionActionModifier(ULearningAgentsActionModifier* Modifier, const TMap<FName, FLearningAgentsActionModifierElement>& Elements, const TArray<FName>& MaskedElements, const FName Tag = TEXT("InclusiveUnionAction"));
+
+	UFUNCTION(BlueprintPure, Category = "LearningAgents", meta = (AdvancedDisplay = 4))
+	static FLearningAgentsActionModifierElement MakeInclusiveUnionActionModifierFromArrays(ULearningAgentsActionModifier* Modifier, const TArray<FName>& ElementNames, const TArray<FLearningAgentsActionModifierElement>& Elements, const TArray<FName>& MaskedElements, const FName Tag = TEXT("InclusiveUnionAction"));
+	static FLearningAgentsActionModifierElement MakeInclusiveUnionActionModifierFromArrayViews(ULearningAgentsActionModifier* Modifier, const TArrayView<const FName> ElementNames, const TArrayView<const FLearningAgentsActionModifierElement> Elements, const TArrayView<const FName> MaskedElements, const FName Tag = TEXT("InclusiveUnionAction"));
+
+	UFUNCTION(BlueprintPure, Category = "LearningAgents", meta = (AdvancedDisplay = 2))
+	static FLearningAgentsActionModifierElement MakeStaticArrayActionModifier(ULearningAgentsActionModifier* Modifier, const TArray<FLearningAgentsActionModifierElement>& Elements, const FName Tag = TEXT("StaticArrayAction"));
+	static FLearningAgentsActionModifierElement MakeStaticArrayActionModifierFromArrayView(ULearningAgentsActionModifier* Modifier, const TArrayView<const FLearningAgentsActionModifierElement> Elements, const FName Tag = TEXT("StaticArrayAction"));
+
+	UFUNCTION(BlueprintPure, Category = "LearningAgents", meta = (AdvancedDisplay = 3))
+	static FLearningAgentsActionModifierElement MakePairActionModifier(ULearningAgentsActionModifier* Modifier, const FLearningAgentsActionModifierElement Key, const FLearningAgentsActionModifierElement Value, const FName Tag = TEXT("PairAction"));
+
+	UFUNCTION(BlueprintPure, Category = "LearningAgents", meta = (AdvancedDisplay = 3))
+	static FLearningAgentsActionModifierElement MakeEnumActionModifier(ULearningAgentsActionModifier* Modifier, const UEnum* Enum, const TArray<uint8>& EnumMaskedValues, const FName Tag = TEXT("EnumAction"));
+	static FLearningAgentsActionModifierElement MakeEnumActionModifierFromArrayView(ULearningAgentsActionModifier* Modifier, const UEnum* Enum, const TArrayView<const uint8> EnumMaskedValues, const FName Tag = TEXT("EnumAction"));
+
+	UFUNCTION(BlueprintPure, Category = "LearningAgents", meta = (AdvancedDisplay = 3))
+	static FLearningAgentsActionModifierElement MakeBitmaskActionModifier(ULearningAgentsActionModifier* Modifier, const UEnum* Enum, const int32 MaskedBitmask, const FName Tag = TEXT("BitmaskAction"));
+
+	UFUNCTION(BlueprintPure, Category = "LearningAgents", meta = (AdvancedDisplay = 4))
+	static FLearningAgentsActionModifierElement MakeOptionalActionModifier(ULearningAgentsActionModifier* Modifier, const FLearningAgentsActionModifierElement Element, const bool bAllowOnlyValid, const bool bAllowOnlyNull, const FName Tag = TEXT("OptionalAction"));
+
+	UFUNCTION(BlueprintPure, Category = "LearningAgents", meta = (AdvancedDisplay = 5))
+	static FLearningAgentsActionModifierElement MakeEitherActionModifier(ULearningAgentsActionModifier* Modifier, const FLearningAgentsActionModifierElement A, const FLearningAgentsActionModifierElement B, const bool bAllowOnlyA, const bool bAllowOnlyB, const FName Tag = TEXT("EitherAction"));
+
+	UFUNCTION(BlueprintPure, Category = "LearningAgents", meta = (AdvancedDisplay = 2))
+	static FLearningAgentsActionModifierElement MakeEncodingActionModifier(ULearningAgentsActionModifier* Modifier, const FLearningAgentsActionModifierElement Element, const FName Tag = TEXT("EncodingAction"));
+
+	UFUNCTION(BlueprintPure, Category = "LearningAgents", meta = (AdvancedDisplay = 2))
+	static FLearningAgentsActionModifierElement MakeBoolActionModifier(ULearningAgentsActionModifier* Modifier, const bool bValue, const FName Tag = TEXT("BoolAction"));
+
+	UFUNCTION(BlueprintPure, Category = "LearningAgents", meta = (AdvancedDisplay = 3))
+	static FLearningAgentsActionModifierElement MakeFloatActionModifier(ULearningAgentsActionModifier* Modifier, const float MaskedValue, const bool bMasked, const FName Tag = TEXT("FloatAction"));
+
+	UFUNCTION(BlueprintPure, Category = "LearningAgents", meta = (AdvancedDisplay = 6))
+	static FLearningAgentsActionModifierElement MakeLocationActionModifier(ULearningAgentsActionModifier* Modifier, const FVector MaskedLocation, const bool bMaskedX, const bool bMaskedY, const bool bMaskedZ, const FTransform RelativeTransform = FTransform(), const FName Tag = TEXT("LocationAction"));
+
+	UFUNCTION(BlueprintPure, Category = "LearningAgents", meta = (AdvancedDisplay = 6))
+	static FLearningAgentsActionModifierElement MakeScaleActionModifier(ULearningAgentsActionModifier* Modifier, const FVector MaskedScale, const bool bMaskedX, const bool bMaskedY, const bool bMaskedZ, const FVector RelativeScale = FVector(1, 1, 1), const FName Tag = TEXT("ScaleAction"));
+
+	UFUNCTION(BlueprintPure, Category = "LearningAgents", meta = (AdvancedDisplay = 4))
+	static FLearningAgentsActionModifierElement MakeAngleActionModifier(ULearningAgentsActionModifier* Modifier, const float MaskedAngle, const bool bMask, const float RelativeAngle = 0.0f, const FName Tag = TEXT("AngleAction"));
+
+	UFUNCTION(BlueprintPure, Category = "LearningAgents", meta = (AdvancedDisplay = 4))
+	static FLearningAgentsActionModifierElement MakeAngleActionModifierRadians(ULearningAgentsActionModifier* Modifier, const float MaskedAngle, const bool bMask, const float RelativeAngle = 0.0f, const FName Tag = TEXT("AngleAction"));
+
+	UFUNCTION(BlueprintPure, Category = "LearningAgents", meta = (AdvancedDisplay = 6))
+	static FLearningAgentsActionModifierElement MakeVelocityActionModifier(ULearningAgentsActionModifier* Modifier, const FVector MaskedVelocity, const bool bMaskedX, const bool bMaskedY, const bool bMaskedZ, const FTransform RelativeTransform = FTransform(), const FName Tag = TEXT("VelocityAction"));
+
+	UFUNCTION(BlueprintPure, Category = "LearningAgents", meta = (AdvancedDisplay = 6))
+	static FLearningAgentsActionModifierElement MakeDirectionActionModifier(ULearningAgentsActionModifier* Modifier, const FVector MaskedDirection, const bool bMaskedX, const bool bMaskedY, const bool bMaskedZ, const FTransform RelativeTransform = FTransform(), const FName Tag = TEXT("DirectionAction"));
+
 public:
 
 	/**
diff --git a/Engine/Plugins/Experimental/LearningAgents/Source/LearningAgents/Public/LearningAgentsInteractor.h b/Engine/Plugins/Experimental/LearningAgents/Source/LearningAgents/Public/LearningAgentsInteractor.h
--- a/Engine/Plugins/Experimental/LearningAgents/Source/LearningAgents/Public/LearningAgentsInteractor.h
+++ b/Engine/Plugins/Experimental/LearningAgents/Source/LearningAgents/Public/LearningAgentsInteractor.h
@@ -101,6 +101,7 @@ public:
 	UFUNCTION(BlueprintNativeEvent, Category = "LearningAgents", Meta = (ForceAsFunction))
 	void GatherAgentObservations(TArray<FLearningAgentsObservationObjectElement>& OutObservationObjectElements, ULearningAgentsObservationObject* InObservationObject, const TArray<int32>& AgentIds);
 
+
 // ----- Actions -----
 public:
 
@@ -136,6 +137,44 @@ public:
 	UFUNCTION(BlueprintNativeEvent, Category = "LearningAgents", Meta = (ForceAsFunction))
 	void PerformAgentActions(const ULearningAgentsActionObject* InActionObject, const TArray<FLearningAgentsActionObjectElement>& InActionObjectElements, const TArray<int32>& AgentIds);
 
+	/**
+	 * This callback can be optionally overridden by the Interactor to create an action modifier for a single agent. The structure of the Action 
+	 * Modifier Elements output by this function should match that of the actions defined by the Schema.
+	 *
+	 * @param OutActionModifierElement			Output Action Modifier Element.
+	 * @param InActionModifier					Action Modifier Object.
+	 * @param InObservationObject				Input Observation Object.
+	 * @param InObservationObjectElement		Input Observation Object Element.
+	 * @param AgentId							The Agent Id to make action modifiers for.
+	 */
+	UFUNCTION(BlueprintNativeEvent, Category = "LearningAgents", Meta = (ForceAsFunction))
+	void MakeAgentActionModifier(
+		FLearningAgentsActionModifierElement& OutActionModifierElement, 
+		ULearningAgentsActionModifier* InActionModifier, 
+		const ULearningAgentsObservationObject* InObservationObject,
+		const FLearningAgentsObservationObjectElement& InObservationObjectElement,
+		const int32 AgentId);
+
+	/**
+	 * This callback can be optionally overridden by the Interactor to create all the action modifier for the given agents. The structure of the 
+	 * Action Modifier Elements output by this function should match that of the actions defined by the Schema. The default implementation calls 
+	 * MakeAgentActionModifier on each agent.
+	 *
+	 * @param OutActionModifierElements			Output Action Modifier Elements. This should be the same size as AgentIds.
+	 * @param InActionModifier					Action Modifier Object.
+	 * @param InObservationObject				Input Observation Object.
+	 * @param InObservationObjectElements		Input Observation Object Elements.
+	 * @param AgentIds							Set of Agent Ids to make action modifiers for.
+	 */
+	UFUNCTION(BlueprintNativeEvent, Category = "LearningAgents", Meta = (ForceAsFunction))
+	void MakeAgentActionModifiers(
+		TArray<FLearningAgentsActionModifierElement>& OutActionModifierElements, 
+		ULearningAgentsActionModifier* InActionModifier, 
+		const ULearningAgentsObservationObject* InObservationObject,
+		const TArray<FLearningAgentsObservationObjectElement>& InObservationObjectElements,
+		const TArray<int32>& AgentIds);
+
+
 // ----- Blueprint public interface -----
 public:
 
@@ -143,6 +182,10 @@ public:
 	UFUNCTION(BlueprintCallable, Category = "LearningAgents")
 	void GatherObservations();
 
+	/** Makes all the action modifiers for all agents. This will call MakeAgentActionModifiers. Should be called even when Action Modifiers are not used. */
+	UFUNCTION(BlueprintCallable, Category = "LearningAgents")
+	void MakeActionModifiers();
+
 	/** Performs all the actions for all agents. This will call PerformAgentActions. */
 	UFUNCTION(BlueprintCallable, Category = "LearningAgents")
 	void PerformActions();
@@ -157,6 +200,16 @@ public:
 	UFUNCTION(BlueprintCallable, Category = "LearningAgents", Meta = (AgentId = "-1"))
 	void GetObservationVector(TArray<float>& OutObservationVector, int32& OutObservationCompatibilityHash, const int32 AgentId);
 
+	/**
+	 * Get the current buffered action modifier vector for the given agent.
+	 *
+	 * @param OutActionModifierVector			Output Action Modifier Vector
+	 * @param OutActionCompatibilityHash		Output Compatibility Hash used to identify which schema this vector is compatible with.
+	 * @param AgentId							Agent Id to look up the observation vector for.
+	 */
+	UFUNCTION(BlueprintCallable, Category = "LearningAgents", Meta = (AgentId = "-1"))
+	void GetActionModifierVector(TArray<float>& OutActionModifierVector, int32& OutActionCompatibilityHash, const int32 AgentId);
+
 	/**
 	 * Get the current buffered action vector for the given agent.
 	 *
@@ -178,6 +231,17 @@ public:
 	UFUNCTION(BlueprintCallable, Category = "LearningAgents", Meta = (AgentId = "-1"))
 	void SetObservationVector(const TArray<float>& ObservationVector, const int32 InObservationCompatibilityHash, const int32 AgentId, bool bIncrementIteration = true);
 
+	/**
+	 * Sets the current buffered action modifier vector for the given agent.
+	 *
+	 * @param ActionModifierVector				Action Modifier Vector
+	 * @param InActionCompatibilityHash			Compatibility Hash used to identify which schema this vector is compatible with.
+	 * @param AgentId							Agent Id to set the observation vector for.
+	 * @param bIncrementIteration				If to increment the iteration number used to keep track of associated actions and observations.
+	 */
+	UFUNCTION(BlueprintCallable, Category = "LearningAgents", Meta = (AgentId = "-1"))
+	void SetActionModifierVector(const TArray<float>& ActionModifierVector, const int32 InActionCompatibilityHash, const int32 AgentId, bool bIncrementIteration = true);
+
 	/**
 	 * Sets the current buffered action vector for the given agent.
 	 *
@@ -195,6 +259,12 @@ public:
 	UFUNCTION(BlueprintPure, Category = "LearningAgents", meta = (AgentId = "-1"))
 	bool HasObservationVector(const int32 AgentId) const;
 
+	/**
+	 * Returns true if MakeActionModifiers or SetActionModifierVector has been called and the action modifier vector already set for the given agent.
+	 */
+	UFUNCTION(BlueprintPure, Category = "LearningAgents", meta = (AgentId = "-1"))
+	bool HasActionModifierVector(const int32 AgentId) const;
+
 	/**
 	 * Returns true if DecodeAndSampleActions on the policy or SetActionVector has been called and the action vector already set for the given agent.
 	 */
@@ -217,6 +287,10 @@ public:
 	UFUNCTION(BlueprintPure, Category = "LearningAgents")
 	int32 GetActionDistributionVectorSize() const;
 
+	/** Gets the size of the action modifier vector used by this interactor. */
+	UFUNCTION(BlueprintPure, Category = "LearningAgents")
+	int32 GetActionModifierVectorSize() const;
+
 	/** Gets the size of the encoded action vector used by this interactor. */
 	UFUNCTION(BlueprintPure, Category = "LearningAgents")
 	int32 GetActionEncodedVectorSize() const;
@@ -227,6 +301,9 @@ public:
 	/** Encode Observations for a specific set of agents */
 	void GatherObservations(const UE::Learning::FIndexSet AgentSet, bool bIncrementIteration = true);
 
+	/** Make Action Modifiers for a specific set of agents */
+	void MakeActionModifiers(const UE::Learning::FIndexSet AgentSet, bool bIncrementIteration = true);
+
 	/** Perform Actions for a specific set of agents */
 	void PerformActions(const UE::Learning::FIndexSet AgentSet);
 
@@ -243,13 +320,19 @@ public:
 	const FLearningAgentsActionSchemaElement GetActionSchemaElement() const;
 
 	/** Gets the observation vectors as a const array view. */
-	TLearningArrayView<2, const float> GetObservationVectorArrayView() const;
+	TLearningArrayView<2, const float> GetObservationVectorsArrayView() const;
 
 	/** Gets the observation iteration value for the given agent id. */
 	uint64 GetObservationIteration(const int32 AgentId) const;
 
+	/** Gets the action modifier vectors as a const array view. */
+	TLearningArrayView<2, const float> GetActionModifierVectorsArrayView() const;
+
+	/** Gets the action modifier iteration value for the given agent id. */
+	uint64 GetActionModifierIteration(const int32 AgentId) const;
+
 	/** Gets the action vectors as a const array view. */
-	TLearningArrayView<2, const float> GetActionVectorArrayView() const;
+	TLearningArrayView<2, const float> GetActionVectorsArrayView() const;
 
 	/** Gets the action iteration value for the given agent id. */
 	uint64 GetActionIteration(const int32 AgentId) const;
@@ -260,6 +343,12 @@ public:
 	/** Gets the observation object elements. */
 	const TArray<FLearningAgentsObservationObjectElement>& GetObservationObjectElements() const;
 
+	/** Gets the action modifier. */
+	const ULearningAgentsActionModifier* GetActionModifier() const;
+
+	/** Gets the action modifier elements. */
+	const TArray<FLearningAgentsActionModifierElement>& GetActionModifierElements() const;
+
 	/** Gets the action object. */
 	ULearningAgentsActionObject* GetActionObject();
 
@@ -298,6 +387,14 @@ private:
 	UPROPERTY(VisibleAnywhere, Transient, Category = "LearningAgents")
 	TArray<FLearningAgentsObservationObjectElement> ObservationObjectElements;
 
+	/** Action Modifier used by this interactor */
+	UPROPERTY(VisibleAnywhere, Transient, Category = "LearningAgents")
+	TObjectPtr<ULearningAgentsActionModifier> ActionModifier;
+
+	/** Action Modifier Elements used by this interactor */
+	UPROPERTY(VisibleAnywhere, Transient, Category = "LearningAgents")
+	TArray<FLearningAgentsActionModifierElement> ActionModifierElements;
+
 	/** Action Object used by this interactor */
 	UPROPERTY(VisibleAnywhere, Transient, Category = "LearningAgents")
 	TObjectPtr<ULearningAgentsActionObject> ActionObject;
@@ -312,6 +409,9 @@ private:
 	/** Buffer of Observation Vectors for each agent */
 	TLearningArray<2, float> ObservationVectors;
 
+	/** Buffer of Action Modifier Vectors for each agent */
+	TLearningArray<2, float> ActionModifierVectors;
+
 	/** Buffer of Action Vectors for each agent */
 	TLearningArray<2, float> ActionVectors;
 
@@ -324,6 +424,9 @@ private:
 	/** Number of times observation vector has been set for all agents */
 	TLearningArray<1, uint64, TInlineAllocator<32>> ObservationVectorIteration;
 
+	/** Number of times action modifier vector has been set for all agents */
+	TLearningArray<1, uint64, TInlineAllocator<32>> ActionModifierVectorIteration;
+
 	/** Number of times action vector has been set for all agents */
 	TLearningArray<1, uint64, TInlineAllocator<32>> ActionVectorIteration;
 
diff --git a/Engine/Plugins/Experimental/LearningAgents/Source/LearningAgents/Public/LearningAgentsObservations.h b/Engine/Plugins/Experimental/LearningAgents/Source/LearningAgents/Public/LearningAgentsObservations.h
--- a/Engine/Plugins/Experimental/LearningAgents/Source/LearningAgents/Public/LearningAgentsObservations.h
+++ b/Engine/Plugins/Experimental/LearningAgents/Source/LearningAgents/Public/LearningAgentsObservations.h
@@ -39,12 +39,23 @@ struct LEARNINGAGENTS_API FLearningAgentsObservationObjectElement
 	UE::Learning::Observation::FObjectElement ObjectElement;
 };
 
-/** Comparison operator for Observation Object Elements */
+/** Comparison and Hashing operators for Observation Elements */
+
+LEARNINGAGENTS_API bool operator==(const FLearningAgentsObservationSchemaElement& Lhs, const FLearningAgentsObservationSchemaElement& Rhs);
 LEARNINGAGENTS_API bool operator==(const FLearningAgentsObservationObjectElement& Lhs, const FLearningAgentsObservationObjectElement& Rhs);
 
-/** Hashing operator for Observation Object Elements */
+LEARNINGAGENTS_API uint32 GetTypeHash(const FLearningAgentsObservationSchemaElement& Element);
 LEARNINGAGENTS_API uint32 GetTypeHash(const FLearningAgentsObservationObjectElement& Element);
 
+template<>
+struct TStructOpsTypeTraits<FLearningAgentsObservationSchemaElement> : public TStructOpsTypeTraitsBase2<FLearningAgentsObservationSchemaElement>
+{
+	enum
+	{
+		WithIdenticalViaEquality = true,
+	};
+};
+
 template<>
 struct TStructOpsTypeTraits<FLearningAgentsObservationObjectElement> : public TStructOpsTypeTraitsBase2<FLearningAgentsObservationObjectElement>
 {
@@ -54,6 +65,7 @@ struct TStructOpsTypeTraits<FLearningAgentsObservationObjectElement> : public TS
 	};
 };
 
+
 /**
  * Observation Schema
  *
diff --git a/Engine/Plugins/Experimental/LearningAgents/Source/LearningAgents/Public/LearningAgentsPolicy.h b/Engine/Plugins/Experimental/LearningAgents/Source/LearningAgents/Public/LearningAgentsPolicy.h
--- a/Engine/Plugins/Experimental/LearningAgents/Source/LearningAgents/Public/LearningAgentsPolicy.h
+++ b/Engine/Plugins/Experimental/LearningAgents/Source/LearningAgents/Public/LearningAgentsPolicy.h
@@ -160,7 +160,8 @@ public:
 	void EvaluatePolicy();
 
 	/**
-	 * Decodes and samples action vectors using the Decoder network. This should be called after EvaluatePolicy and before PerformActions.
+	 * Decodes and samples action vectors using the Decoder network. This should be called after both MakeActionModifiers and EvaluatePolicy, but 
+	 * before PerformActions.
 	 * 
 	 * @param ActionNoiseScale		Scale of the action noise to use during sampling. Set this to zero to always sample the mean (expected) action.
 	 */
diff --git a/Engine/Plugins/Experimental/LearningAgents/Source/LearningAgentsTraining/Private/LearningAgentsCommunicator.cpp b/Engine/Plugins/Experimental/LearningAgents/Source/LearningAgentsTraining/Private/LearningAgentsCommunicator.cpp
--- a/Engine/Plugins/Experimental/LearningAgents/Source/LearningAgentsTraining/Private/LearningAgentsCommunicator.cpp
+++ b/Engine/Plugins/Experimental/LearningAgents/Source/LearningAgentsTraining/Private/LearningAgentsCommunicator.cpp
@@ -7,11 +7,11 @@
 
 #include "Misc/Paths.h"
 
-FLearningAgentsTrainerProcess ULearningAgentsCommunicatorLibrary::SpawnSharedMemoryTrainingProcess(
+FLearningAgentsSharedMemoryTrainerProcess ULearningAgentsCommunicatorLibrary::SpawnSharedMemoryTrainingProcess(
 	const FLearningAgentsTrainerProcessSettings& TrainerProcessSettings,
 	const FLearningAgentsSharedMemoryCommunicatorSettings& SharedMemorySettings)
 {
-	FLearningAgentsTrainerProcess TrainerProcess;
+	FLearningAgentsSharedMemoryTrainerProcess TrainerProcess;
 
 	if (PLATFORM_MAC)
 	{
@@ -56,7 +56,7 @@ FLearningAgentsTrainerProcess ULearningAgentsCommunicatorLibrary::SpawnSharedMem
 }
 
 FLearningAgentsCommunicator ULearningAgentsCommunicatorLibrary::MakeSharedMemoryCommunicator(
-	const FLearningAgentsTrainerProcess& TrainerProcess,
+	const FLearningAgentsSharedMemoryTrainerProcess& TrainerProcess,
 	const FLearningAgentsSharedMemoryCommunicatorSettings& SharedMemorySettings)
 {
 	FLearningAgentsCommunicator Communicator;
@@ -82,11 +82,11 @@ FLearningAgentsCommunicator ULearningAgentsCommunicatorLibrary::MakeSharedMemory
 	return Communicator;
 }
 
-FLearningAgentsTrainerProcess ULearningAgentsCommunicatorLibrary::SpawnSocketTrainingProcess(
+FLearningAgentsSocketTrainerProcess ULearningAgentsCommunicatorLibrary::SpawnSocketTrainingProcess(
 	const FLearningAgentsTrainerProcessSettings& TrainerProcessSettings,
 	const FLearningAgentsSocketCommunicatorSettings& SocketSettings)
 {
-	FLearningAgentsTrainerProcess TrainerProcess;
+	FLearningAgentsSocketTrainerProcess TrainerProcess;
 
 	const FString PythonExecutablePath = UE::Learning::Trainer::GetPythonExecutablePath(TrainerProcessSettings.GetIntermediatePath());
 	if (!FPaths::FileExists(PythonExecutablePath))
@@ -125,7 +125,7 @@ FLearningAgentsTrainerProcess ULearningAgentsCommunicatorLibrary::SpawnSocketTra
 }
 
 FLearningAgentsCommunicator ULearningAgentsCommunicatorLibrary::MakeSocketCommunicator(
-	FLearningAgentsTrainerProcess TrainerProcess,
+	const FLearningAgentsSocketTrainerProcess& TrainerProcess,
 	const FLearningAgentsSocketCommunicatorSettings& SocketSettings)
 {
 	UE::Learning::ETrainerResponse Response = UE::Learning::ETrainerResponse::Success;
diff --git a/Engine/Plugins/Experimental/LearningAgents/Source/LearningAgentsTraining/Private/LearningAgentsPPOTrainer.cpp b/Engine/Plugins/Experimental/LearningAgents/Source/LearningAgentsTraining/Private/LearningAgentsPPOTrainer.cpp
--- a/Engine/Plugins/Experimental/LearningAgents/Source/LearningAgentsTraining/Private/LearningAgentsPPOTrainer.cpp
+++ b/Engine/Plugins/Experimental/LearningAgents/Source/LearningAgentsTraining/Private/LearningAgentsPPOTrainer.cpp
@@ -196,6 +196,10 @@ void ULearningAgentsPPOTrainer::SetupPPOTrainer(
 		TEXT("Actions"),
 		ActionSchemaId,
 		Interactor->GetActionVectorSize());
+	ActionModifierId = EpisodeBuffer->AddActionModifiers(
+		TEXT("ActionModifiers"),
+		ActionSchemaId,
+		Interactor->GetActionModifierVectorSize());
 	MemoryStateId = EpisodeBuffer->AddMemoryStates(
 		TEXT("MemoryStates"),
 		Policy->GetMemoryStateSize());
@@ -506,7 +510,7 @@ void ULearningAgentsPPOTrainer::ProcessExperience(const bool bResetAgentsOnUpdat
 		UE_LOG(LogLearning, Warning, TEXT("%s: No agents added to Manager."), *GetName());
 	}
 
-	// Check Observations, Actions, Rewards, and Completions have been completed and have matching iteration number
+	// Check Observations, ActionModifiers, Actions, Rewards, and Completions have been completed and have matching iteration number
 
 	TArray<int32> ValidAgentIds;
 	ValidAgentIds.Empty(Manager->GetMaxAgentNum());
@@ -515,20 +519,23 @@ void ULearningAgentsPPOTrainer::ProcessExperience(const bool bResetAgentsOnUpdat
 	{
 		if (Interactor->GetObservationIteration(AgentId) == 0 ||
 			Interactor->GetActionIteration(AgentId) == 0 ||
+			Interactor->GetActionModifierIteration(AgentId) == 0 ||
 			TrainingEnvironment->GetRewardIteration(AgentId) == 0 ||
 			TrainingEnvironment->GetCompletionIteration(AgentId) == 0)
 		{
-			UE_LOG(LogLearning, Display, TEXT("%s: Agent with id %i has not completed a full step of observations, actions, rewards, completions and so experience will not be processed for it."), *GetName(), AgentId);
+			UE_LOG(LogLearning, Display, TEXT("%s: Agent with id %i has not completed a full step of observations, action modifiers, actions, rewards, completions and so experience will not be processed for it."), *GetName(), AgentId);
 			continue;
 		}
 
 		if (Interactor->GetObservationIteration(AgentId) != Interactor->GetActionIteration(AgentId) ||
+			Interactor->GetObservationIteration(AgentId) != Interactor->GetActionModifierIteration(AgentId) ||
 			Interactor->GetObservationIteration(AgentId) != TrainingEnvironment->GetRewardIteration(AgentId) ||
 			Interactor->GetObservationIteration(AgentId) != TrainingEnvironment->GetCompletionIteration(AgentId))
 		{
-			UE_LOG(LogLearning, Warning, TEXT("%s: Agent with id %i has non-matching iteration numbers (observation: %i, action: %i, reward: %i, completion: %i). Experience will not be processed for it."), *GetName(), AgentId,
+			UE_LOG(LogLearning, Warning, TEXT("%s: Agent with id %i has non-matching iteration numbers (observation: %i, action modifiers: %i, action: %i, reward: %i, completion: %i). Experience will not be processed for it."), *GetName(), AgentId,
 				Interactor->GetObservationIteration(AgentId),
 				Interactor->GetActionIteration(AgentId),
+				Interactor->GetActionModifierIteration(AgentId),
 				TrainingEnvironment->GetRewardIteration(AgentId),
 				TrainingEnvironment->GetCompletionIteration(AgentId));
 			continue;
@@ -551,8 +558,9 @@ void ULearningAgentsPPOTrainer::ProcessExperience(const bool bResetAgentsOnUpdat
 	}
 
 	// Add Experience to Episode Buffer
-	EpisodeBuffer->PushObservations(ObservationId, Interactor->GetObservationVectorArrayView(), ValidAgentSet);
-	EpisodeBuffer->PushActions(ActionId, Interactor->GetActionVectorArrayView(), ValidAgentSet);
+	EpisodeBuffer->PushObservations(ObservationId, Interactor->GetObservationVectorsArrayView(), ValidAgentSet);
+	EpisodeBuffer->PushActions(ActionId, Interactor->GetActionVectorsArrayView(), ValidAgentSet);
+	EpisodeBuffer->PushActionModifiers(ActionId, Interactor->GetActionModifierVectorsArrayView(), ValidAgentSet);
 	EpisodeBuffer->PushMemoryStates(MemoryStateId, Policy->GetPreEvaluationMemoryState(), ValidAgentSet);
 	EpisodeBuffer->PushRewards(RewardId, TrainingEnvironment->GetRewardArrayView(), ValidAgentSet);
 	EpisodeBuffer->IncrementEpisodeStepNums(ValidAgentSet);
@@ -580,7 +588,7 @@ void ULearningAgentsPPOTrainer::ProcessExperience(const bool bResetAgentsOnUpdat
 	// And push those episodes to the Replay Buffer
 	const bool bReplayBufferFull = ReplayBuffer->AddEpisodes(
 		TrainingEnvironment->GetAllCompletions(),
-		{ Interactor->GetObservationVectorArrayView() },
+		{ Interactor->GetObservationVectorsArrayView() },
 		{ Policy->GetMemoryState() },
 		*EpisodeBuffer,
 		TrainingEnvironment->GetResetBuffer().GetResetInstances());
diff --git a/Engine/Plugins/Experimental/LearningAgents/Source/LearningAgentsTraining/Private/LearningAgentsRecorder.cpp b/Engine/Plugins/Experimental/LearningAgents/Source/LearningAgentsTraining/Private/LearningAgentsRecorder.cpp
--- a/Engine/Plugins/Experimental/LearningAgents/Source/LearningAgentsTraining/Private/LearningAgentsRecorder.cpp
+++ b/Engine/Plugins/Experimental/LearningAgents/Source/LearningAgentsTraining/Private/LearningAgentsRecorder.cpp
@@ -281,7 +281,7 @@ void ULearningAgentsRecorder::AddExperience()
 			continue;
 		}
 
-		RecordBuffers[AgentId].Push(Interactor->GetObservationVectorArrayView()[AgentId], Interactor->GetActionVectorArrayView()[AgentId]);
+		RecordBuffers[AgentId].Push(Interactor->GetObservationVectorsArrayView()[AgentId], Interactor->GetActionVectorsArrayView()[AgentId]);
 	}
 }
 
diff --git a/Engine/Plugins/Experimental/LearningAgents/Source/LearningAgentsTraining/Public/LearningAgentsCommunicator.h b/Engine/Plugins/Experimental/LearningAgents/Source/LearningAgentsTraining/Public/LearningAgentsCommunicator.h
--- a/Engine/Plugins/Experimental/LearningAgents/Source/LearningAgentsTraining/Public/LearningAgentsCommunicator.h
+++ b/Engine/Plugins/Experimental/LearningAgents/Source/LearningAgentsTraining/Public/LearningAgentsCommunicator.h
@@ -15,7 +15,8 @@
 namespace UE::Learning
 {
 	struct IExternalTrainer;
-	struct ITrainerProcess;
+	struct FSharedMemoryTrainerServerProcess;
+	struct FSocketTrainerServerProcess;
 }
 
 class ULearningAgentsInteractor;
@@ -58,13 +59,22 @@ public:
 	float Timeout = 10.0f;
 };
 
-/** Blueprint-compatible wrapper struct for ITrainerProcess. */
+/** Blueprint-compatible wrapper struct for Shared Memory Training Process. */
 USTRUCT(BlueprintType)
-struct LEARNINGAGENTSTRAINING_API FLearningAgentsTrainerProcess
+struct LEARNINGAGENTSTRAINING_API FLearningAgentsSharedMemoryTrainerProcess
 {
 	GENERATED_BODY()
 
-	TSharedPtr<UE::Learning::ITrainerProcess> TrainerProcess;
+	TSharedPtr<UE::Learning::FSharedMemoryTrainerServerProcess> TrainerProcess;
+};
+
+/** Blueprint-compatible wrapper struct for Socket Training Process. */
+USTRUCT(BlueprintType)
+struct LEARNINGAGENTSTRAINING_API FLearningAgentsSocketTrainerProcess
+{
+	GENERATED_BODY()
+
+	TSharedPtr<UE::Learning::FSocketTrainerServerProcess> TrainerProcess;
 };
 
 /** Blueprint-compatible wrapper struct for IExternalTrainer. */
@@ -94,7 +104,7 @@ public:
 	 * @param FLearningAgentsSharedMemoryCommunicatorSettings	Settings specific to shared memory communicators.
 	 */ 
 	UFUNCTION(BlueprintCallable, Category = "LearningAgents", meta = (AutoCreateRefTerm = "TrainerProcessSettings,SharedMemorySettings"))
-	static FLearningAgentsTrainerProcess SpawnSharedMemoryTrainingProcess(
+	static FLearningAgentsSharedMemoryTrainerProcess SpawnSharedMemoryTrainingProcess(
 		const FLearningAgentsTrainerProcessSettings& TrainerProcessSettings = FLearningAgentsTrainerProcessSettings(),
 		const FLearningAgentsSharedMemoryCommunicatorSettings& SharedMemorySettings = FLearningAgentsSharedMemoryCommunicatorSettings());
 
@@ -107,7 +117,7 @@ public:
 	 */
 	UFUNCTION(BlueprintCallable, Category = "LearningAgents", meta = (AutoCreateRefTerm = "SharedMemorySettings"))
 	static FLearningAgentsCommunicator MakeSharedMemoryCommunicator(
-		const FLearningAgentsTrainerProcess& TrainerProcess,
+		const FLearningAgentsSharedMemoryTrainerProcess& TrainerProcess,
 		const FLearningAgentsSharedMemoryCommunicatorSettings& SharedMemorySettings = FLearningAgentsSharedMemoryCommunicatorSettings());
 
 	/**
@@ -121,7 +131,7 @@ public:
 	 * @param FLearningAgentsSocketCommunicatorSettings	Settings specific to socket communicators.
 	 */
 	UFUNCTION(BlueprintCallable, Category = "LearningAgents", meta = (AutoCreateRefTerm = "TrainerProcessSettings,SocketSettings"))
-	static FLearningAgentsTrainerProcess SpawnSocketTrainingProcess(
+	static FLearningAgentsSocketTrainerProcess SpawnSocketTrainingProcess(
 		const FLearningAgentsTrainerProcessSettings& TrainerProcessSettings = FLearningAgentsTrainerProcessSettings(),
 		const FLearningAgentsSocketCommunicatorSettings& SocketSettings = FLearningAgentsSocketCommunicatorSettings());
 
@@ -134,6 +144,6 @@ public:
 	 */
 	UFUNCTION(BlueprintCallable, Category = "LearningAgents", meta = (AutoCreateRefTerm = "SocketSettings"))
 	static FLearningAgentsCommunicator MakeSocketCommunicator(
-		FLearningAgentsTrainerProcess TrainerProcess,
+		const FLearningAgentsSocketTrainerProcess& TrainerProcess,
 		const FLearningAgentsSocketCommunicatorSettings& SocketSettings = FLearningAgentsSocketCommunicatorSettings());
 };
diff --git a/Engine/Plugins/Experimental/LearningAgents/Source/LearningAgentsTraining/Public/LearningAgentsPPOTrainer.h b/Engine/Plugins/Experimental/LearningAgents/Source/LearningAgentsTraining/Public/LearningAgentsPPOTrainer.h
--- a/Engine/Plugins/Experimental/LearningAgents/Source/LearningAgentsTraining/Public/LearningAgentsPPOTrainer.h
+++ b/Engine/Plugins/Experimental/LearningAgents/Source/LearningAgentsTraining/Public/LearningAgentsPPOTrainer.h
@@ -430,6 +430,7 @@ private:
 
 	int32 ObservationId = INDEX_NONE;
 	int32 ActionId = INDEX_NONE;
+	int32 ActionModifierId = INDEX_NONE;
 	int32 MemoryStateId = INDEX_NONE;
 	int32 RewardId = INDEX_NONE;
 };
diff --git a/Engine/Plugins/Experimental/LearningAgents/Source/LearningTraining/Private/LearningExperience.cpp b/Engine/Plugins/Experimental/LearningAgents/Source/LearningTraining/Private/LearningExperience.cpp
--- a/Engine/Plugins/Experimental/LearningAgents/Source/LearningTraining/Private/LearningExperience.cpp
+++ b/Engine/Plugins/Experimental/LearningAgents/Source/LearningTraining/Private/LearningExperience.cpp
@@ -31,6 +31,13 @@ namespace UE::Learning
 			ActionArrays[NameIndex].SetNumUninitialized({ InMaxInstanceNum, InMaxStepNum, ActionSize });
 		}
 
+		// Action Modifiers
+		for (int32 NameIndex = 0; NameIndex < ActionModifierNames.Num(); NameIndex++)
+		{
+			int32 ActionModifierSize = ActionModifierSizes[NameIndex];
+			ActionModifierArrays[NameIndex].SetNumUninitialized({ InMaxInstanceNum, InMaxStepNum, ActionModifierSize });
+		}
+
 		// Memory States
 		for (int32 NameIndex = 0; NameIndex < MemoryStateNames.Num(); NameIndex++)
 		{
@@ -131,6 +138,43 @@ namespace UE::Learning
 		return ActionArrays[ActionId][InstanceIdx].Slice(0, EpisodeStepNums[InstanceIdx]);
 	}
 
+	int32 FEpisodeBuffer::AddActionModifiers(const FName& Name, const int32 SchemaId, const int32 Size)
+	{
+		UE_LEARNING_CHECKF(!ActionModifierNames.Contains(Name), TEXT("Action Modifier name collision!"));
+
+		int32 ActionModifierId = ActionModifierNames.Num();
+		ActionModifierNames.Add(Name);
+		ActionModifierSchemaIds.Add(SchemaId);
+		ActionModifierSizes.Add(Size);
+		TLearningArray<3, float>& NewActionModifier = ActionModifierArrays.AddDefaulted_GetRef();
+
+		if (bHasBeenSized)
+		{
+			NewActionModifier.SetNumUninitialized({ MaxInstanceNum, MaxStepNum, Size });
+		}
+
+		return ActionModifierId;
+	}
+
+	void FEpisodeBuffer::PushActionModifiers(const int32 ActionModifierId, const TLearningArrayView<2, const float> InActionModifiers, const FIndexSet Instances)
+	{
+		UE_LEARNING_TRACE_CPUPROFILER_EVENT_SCOPE(Learning::FEpisodeBuffer::PushActionModifiers);
+
+		UE_LEARNING_CHECKF(ActionModifierId >= 0 && ActionModifierId < ActionModifierArrays.Num(), TEXT("Action Modifier id invalid!"));
+
+		for (const int32 InstanceIdx : Instances)
+		{
+			Array::Copy(ActionModifierArrays[ActionModifierId][InstanceIdx][EpisodeStepNums[InstanceIdx]], InActionModifiers[InstanceIdx]);
+		}
+	}
+
+	const TLearningArrayView<2, const float> FEpisodeBuffer::GetActionModifiers(const int32 ActionModifierId, const int32 InstanceIdx) const
+	{
+		UE_LEARNING_CHECKF(ActionModifierId >= 0 && ActionModifierId < ActionModifierArrays.Num(), TEXT("Action Modifier id invalid!"));
+
+		return ActionModifierArrays[ActionModifierId][InstanceIdx].Slice(0, EpisodeStepNums[InstanceIdx]);
+	}
+
 	int32 FEpisodeBuffer::AddMemoryStates(const FName& Name, const int32 Size)
 	{
 		UE_LEARNING_CHECKF(!MemoryStateNames.Contains(Name), TEXT("Memory State name collision!"));
@@ -289,6 +333,18 @@ namespace UE::Learning
 			NewAction.SetNumUninitialized({ InMaxStepNum, ActionSize });
 		}
 
+		// Action Modifiers
+		ActionModifiers.Empty();
+		for (int32 NameIndex = 0; NameIndex < EpisodeBuffer.ActionModifierNames.Num(); NameIndex++)
+		{
+			ActionModifierNames.Add(EpisodeBuffer.ActionModifierNames[NameIndex]);
+			ActionModifierSchemaIds.Add(EpisodeBuffer.ActionModifierSchemaIds[NameIndex]);
+
+			int32 ActionModifierSize = EpisodeBuffer.ActionModifierSizes[NameIndex];
+			TLearningArray<2, float>& NewActionModifier = ActionModifiers.AddDefaulted_GetRef();
+			NewActionModifier.SetNumUninitialized({ InMaxStepNum, ActionModifierSize });
+		}
+
 		// Memory States
 		MemoryStates.Empty();
 		for (int32 NameIndex = 0; NameIndex < EpisodeBuffer.MemoryStateNames.Num(); NameIndex++)
@@ -335,6 +391,7 @@ namespace UE::Learning
 		UE_LEARNING_CHECKF(EpisodeBuffer.ObservationArrays.Num() == Observations.Num(), TEXT("Observation number mismatch!"));
 		UE_LEARNING_CHECKF(EpisodeBuffer.ObservationArrays.Num() == InEpisodeFinalObservations.Num(), TEXT("Final Observation number mismatch!"));
 		UE_LEARNING_CHECKF(EpisodeBuffer.ActionArrays.Num() == Actions.Num(), TEXT("Action number mismatch!"));
+		UE_LEARNING_CHECKF(EpisodeBuffer.ActionModifierArrays.Num() == ActionModifiers.Num(), TEXT("Action Modifier number mismatch!"));
 		UE_LEARNING_CHECKF(EpisodeBuffer.MemoryStateArrays.Num() == MemoryStates.Num(), TEXT("Memory State number mismatch!"));
 		UE_LEARNING_CHECKF(EpisodeBuffer.MemoryStateArrays.Num() == InEpisodeFinalMemoryStates.Num(), TEXT("Final Memory State number mismatch!"));
 		UE_LEARNING_CHECKF(EpisodeBuffer.RewardArrays.Num() == Rewards.Num(), TEXT("Reward number mismatch!"));
@@ -368,6 +425,13 @@ namespace UE::Learning
 						EpisodeBuffer.ActionArrays[Index][InstanceIdx].Slice(0, EpisodeBuffer.EpisodeStepNums[InstanceIdx]));
 				}
 
+				for (int32 Index = 0; Index < ActionModifiers.Num(); Index++)
+				{
+					Array::Copy(
+						ActionModifiers[Index].Slice(StepNum, EpisodeStepNum),
+						EpisodeBuffer.ActionModifierArrays[Index][InstanceIdx].Slice(0, EpisodeBuffer.EpisodeStepNums[InstanceIdx]));
+				}
+
 				for (int32 Index = 0; Index < MemoryStates.Num(); Index++)
 				{
 					Array::Copy(
@@ -423,6 +487,13 @@ namespace UE::Learning
 						EpisodeBuffer.ActionArrays[Index][InstanceIdx].Slice(0, PartialStepNum));
 				}
 
+				for (int32 Index = 0; Index < ActionModifiers.Num(); Index++)
+				{
+					Array::Copy(
+						ActionModifiers[Index].Slice(StepNum, PartialStepNum),
+						EpisodeBuffer.ActionModifierArrays[Index][InstanceIdx].Slice(0, PartialStepNum));
+				}
+
 				for (int32 Index = 0; Index < MemoryStates.Num(); Index++)
 				{
 					Array::Copy(
@@ -575,6 +646,16 @@ namespace UE::Learning
 		return Actions[Index].Slice(0, StepNum);
 	}
 
+	int32 FReplayBuffer::GetActionModifiersNum() const
+	{
+		return ActionModifiers.Num();
+	}
+
+	const TLearningArrayView<2, const float> FReplayBuffer::GetActionModifiers(const int32 Index) const
+	{
+		return ActionModifiers[Index].Slice(0, StepNum);
+	}
+
 	int32 FReplayBuffer::GetMemoryStatesNum() const
 	{
 		return MemoryStates.Num();
@@ -644,6 +725,21 @@ namespace UE::Learning
 		}
 		ConfigObject->SetArrayField(TEXT("Actions"), ActionObjects);
 
+		// Action Modifiers
+		TArray<TSharedPtr<FJsonValue>> ActionModifierObjects;
+		for (int32 Index = 0; Index < ActionModifiers.Num(); Index++)
+		{
+			TSharedPtr<FJsonObject> BufferObject = MakeShared<FJsonObject>();
+			BufferObject->SetNumberField(TEXT("Id"), Index);
+			BufferObject->SetStringField(TEXT("Name"), ActionModifierNames[Index].ToString());
+			BufferObject->SetNumberField(TEXT("SchemaId"), ActionModifierSchemaIds[Index]);
+			BufferObject->SetNumberField(TEXT("VectorDimensionNum"), ActionModifiers[Index].Num<1>());
+
+			TSharedRef<FJsonValueObject> JsonValue = MakeShared<FJsonValueObject>(BufferObject);
+			ActionModifierObjects.Add(JsonValue);
+		}
+		ConfigObject->SetArrayField(TEXT("ActionModifiers"), ActionModifierObjects);
+
 		// Memory States
 		TArray<TSharedPtr<FJsonValue>> MemoryStateObjects;
 		for (int32 Index = 0; Index < MemoryStates.Num(); Index++)
diff --git a/Engine/Plugins/Experimental/LearningAgents/Source/LearningTraining/Private/LearningExternalTrainer.cpp b/Engine/Plugins/Experimental/LearningAgents/Source/LearningTraining/Private/LearningExternalTrainer.cpp
--- a/Engine/Plugins/Experimental/LearningAgents/Source/LearningTraining/Private/LearningExternalTrainer.cpp
+++ b/Engine/Plugins/Experimental/LearningAgents/Source/LearningTraining/Private/LearningExternalTrainer.cpp
@@ -175,6 +175,11 @@ namespace UE::Learning
 				SharedMemory::Deallocate(SharedMemoryArrayView);
 			}
 
+			for (TSharedMemoryArrayView<3, float>& SharedMemoryArrayView : ActionModifiers)
+			{
+				SharedMemory::Deallocate(SharedMemoryArrayView);
+			}
+
 			for (TSharedMemoryArrayView<3, float>& SharedMemoryArrayView : MemoryStates)
 			{
 				SharedMemory::Deallocate(SharedMemoryArrayView);
@@ -190,11 +195,10 @@ namespace UE::Learning
 	FSharedMemoryTrainer::FSharedMemoryTrainer(
 		const FString& InTaskName,
 		const int32 InProcessNum,
-		const TSharedPtr<UE::Learning::ITrainerProcess>& ExternalTrainerProcess,
+		const TSharedPtr<FSharedMemoryTrainerServerProcess>& ExternalTrainerProcess,
 		const float InTimeout)
 	{
-		FSharedMemoryTrainerServerProcess* SharedMemoryTrainerProcess = (FSharedMemoryTrainerServerProcess*)ExternalTrainerProcess.Get();
-		if (!SharedMemoryTrainerProcess)
+		if (!ExternalTrainerProcess)
 		{
 			UE_LOG(LogLearning, Error, TEXT("FSharedMemoryTrainer ctor: Trainer process is nullptr. Is it not a shared memory process?"));
 			return;
@@ -203,11 +207,11 @@ namespace UE::Learning
 		check(InProcessNum > 0);
 
 		TaskName = InTaskName;
-		ConfigPath = SharedMemoryTrainerProcess->GetConfigPath();
-		IntermediatePath = SharedMemoryTrainerProcess->GetIntermediatePath();
-		TrainingProcess = SharedMemoryTrainerProcess->GetTrainingSubprocess();
+		ConfigPath = ExternalTrainerProcess->GetConfigPath();
+		IntermediatePath = ExternalTrainerProcess->GetIntermediatePath();
+		TrainingProcess = ExternalTrainerProcess->GetTrainingSubprocess();
 		ProcessNum = InProcessNum;
-		Controls = SharedMemoryTrainerProcess->GetControlsSharedMemoryArrayView();
+		Controls = ExternalTrainerProcess->GetControlsSharedMemoryArrayView();
 		Timeout = InTimeout;
 
 		ProcessIdx = 0;
@@ -311,6 +315,15 @@ namespace UE::Learning
 			}
 			JsonObject->SetArrayField(TEXT("ActionsGuids"), ActionsGuidsArray);
 
+			TArray<TSharedPtr<FJsonValue>> ActionModifiersGuidsArray;
+			for (const TSharedMemoryArrayView<3, float>& ActionModifiers : SharedMemoryExperienceContainer.ActionModifiers)
+			{
+				ActionModifiersGuidsArray.Add(MakeShared<FJsonValueString>(*ActionModifiers.Guid.ToString(EGuidFormats::DigitsWithHyphensInBraces)));
+
+			}
+			JsonObject->SetArrayField(TEXT("ActionModifiersGuids"), ActionModifiersGuidsArray);
+
+
 			TArray<TSharedPtr<FJsonValue>> MemoryStatesGuidsArray;
 			for (const TSharedMemoryArrayView<3, float>& MemoryStates : SharedMemoryExperienceContainer.MemoryStates)
 			{
@@ -442,6 +455,12 @@ namespace UE::Learning
 				ExperienceContainer.Actions.Add(SharedMemory::Allocate<3, float>({ ProcessNum, ReplayBuffer.GetMaxStepNum(), DimNum }));
 			}
 
+			for (int32 Index = 0; Index < ReplayBuffer.GetActionModifiersNum(); Index++)
+			{
+				const int32 DimNum = ReplayBuffer.GetActionModifiers(Index).Num<1>();
+				ExperienceContainer.ActionModifiers.Add(SharedMemory::Allocate<3, float>({ ProcessNum, ReplayBuffer.GetMaxStepNum(), DimNum }));
+			}
+
 			for (int32 Index = 0; Index < ReplayBuffer.GetMemoryStatesNum(); Index++)
 			{
 				const int32 DimNum = ReplayBuffer.GetMemoryStates(Index).Num<1>();
@@ -759,17 +778,16 @@ namespace UE::Learning
 
 	FSocketTrainer::FSocketTrainer(
 		ETrainerResponse& OutResponse,
-		const TSharedPtr<UE::Learning::ITrainerProcess>& ExternalTrainerProcess,
+		const TSharedPtr<FSocketTrainerServerProcess>& ExternalTrainerProcess,
 		const TCHAR* IpAddress,
 		const uint32 Port,
 		const float InTimeout)
 	{
 		Timeout = InTimeout;
 
-		FSocketTrainerServerProcess* SocketTrainerProcess = (FSocketTrainerServerProcess*)ExternalTrainerProcess.Get();
-		if (SocketTrainerProcess)
+		if (ExternalTrainerProcess)
 		{
-			TrainingProcess = SocketTrainerProcess->GetTrainingSubprocess();
+			TrainingProcess = ExternalTrainerProcess->GetTrainingSubprocess();
 		}
 
 		ISocketSubsystem* SocketSubsystem = ISocketSubsystem::Get(PLATFORM_SOCKETSUBSYSTEM);
diff --git a/Engine/Plugins/Experimental/LearningAgents/Source/LearningTraining/Private/LearningTrainer.cpp b/Engine/Plugins/Experimental/LearningAgents/Source/LearningTraining/Private/LearningTrainer.cpp
--- a/Engine/Plugins/Experimental/LearningAgents/Source/LearningTraining/Private/LearningTrainer.cpp
+++ b/Engine/Plugins/Experimental/LearningAgents/Source/LearningTraining/Private/LearningTrainer.cpp
@@ -231,6 +231,7 @@ namespace UE::Learning::Trainer
 		Object->SetNumberField(TEXT("VectorSize"), ActionSchema.GetActionVectorSize(ActionSchemaElement));
 		Object->SetNumberField(TEXT("DistributionSize"), ActionSchema.GetActionDistributionVectorSize(ActionSchemaElement));
 		Object->SetNumberField(TEXT("EncodedSize"), ActionSchema.GetEncodedVectorSize(ActionSchemaElement));
+		Object->SetNumberField(TEXT("ModifierSize"), ActionSchema.GetActionModifierVectorSize(ActionSchemaElement));
 
 		switch (ActionSchema.GetType(ActionSchemaElement))
 		{
diff --git a/Engine/Plugins/Experimental/LearningAgents/Source/LearningTraining/Public/LearningExperience.h b/Engine/Plugins/Experimental/LearningAgents/Source/LearningTraining/Public/LearningExperience.h
--- a/Engine/Plugins/Experimental/LearningAgents/Source/LearningTraining/Public/LearningExperience.h
+++ b/Engine/Plugins/Experimental/LearningAgents/Source/LearningTraining/Public/LearningExperience.h
@@ -56,6 +56,15 @@ namespace UE::Learning
 		const TLearningArrayView<2, const float> GetActions(const int32 ActionId, const int32 InstanceIdx) const;
 
 
+		// Action Modifiers
+
+		int32 AddActionModifiers(const FName& Name, const int32 SchemaId, const int32 Size);
+
+		void PushActionModifiers(const int32 ActionModifierId, const TLearningArrayView<2, const float> InActionModifiers, const FIndexSet Instances);
+
+		const TLearningArrayView<2, const float> GetActionModifiers(const int32 ActiActionModifierIdnId, const int32 InstanceIdx) const;
+
+
 		// Memory States
 
 		int32 AddMemoryStates(const FName& Name, const int32 Size);
@@ -104,6 +113,12 @@ namespace UE::Learning
 		TArray<int32, TInlineAllocator<1>> ActionSizes;
 		TArray<TLearningArray<3, float>, TInlineAllocator<1>> ActionArrays;
 
+		// Action Modifiers
+		TArray<FName, TInlineAllocator<1>> ActionModifierNames;
+		TArray<int32, TInlineAllocator<1>> ActionModifierSchemaIds;
+		TArray<int32, TInlineAllocator<1>> ActionModifierSizes;
+		TArray<TLearningArray<3, float>, TInlineAllocator<1>> ActionModifierArrays;
+
 		// Memory States
 		TArray<FName, TInlineAllocator<1>> MemoryStateNames;
 		TArray<int32, TInlineAllocator<1>> MemoryStateSizes;
@@ -202,6 +217,9 @@ namespace UE::Learning
 		int32 GetActionsNum() const;
 		const TLearningArrayView<2, const float> GetActions(const int32 Index) const;
 
+		int32 GetActionModifiersNum() const;
+		const TLearningArrayView<2, const float> GetActionModifiers(const int32 Index) const;
+
 		int32 GetMemoryStatesNum() const;
 		const TLearningArrayView<2, const float> GetMemoryStates(const int32 Index) const;
 		const TLearningArrayView<2, const float> GetEpisodeFinalMemoryStates(const int32 Index) const;
@@ -236,6 +254,10 @@ namespace UE::Learning
 		TArray<int32, TInlineAllocator<1>> ActionSchemaIds;
 		TArray<TLearningArray<2, float>, TInlineAllocator<1>> Actions;
 
+		TArray<FName, TInlineAllocator<1>> ActionModifierNames;
+		TArray<int32, TInlineAllocator<1>> ActionModifierSchemaIds;
+		TArray<TLearningArray<2, float>, TInlineAllocator<1>> ActionModifiers;
+
 		TArray<FName, TInlineAllocator<1>> MemoryStateNames;
 		TArray<TLearningArray<2, float>, TInlineAllocator<1>> MemoryStates;
 		TArray<TLearningArray<2, float>, TInlineAllocator<1>> EpisodeFinalMemoryStates;
diff --git a/Engine/Plugins/Experimental/LearningAgents/Source/LearningTraining/Public/LearningExternalTrainer.h b/Engine/Plugins/Experimental/LearningAgents/Source/LearningTraining/Public/LearningExternalTrainer.h
--- a/Engine/Plugins/Experimental/LearningAgents/Source/LearningTraining/Public/LearningExternalTrainer.h
+++ b/Engine/Plugins/Experimental/LearningAgents/Source/LearningTraining/Public/LearningExternalTrainer.h
@@ -119,37 +119,11 @@ namespace UE::Learning
 			const ELogSetting LogSettings = Trainer::DefaultLogSettings) = 0;
 	};
 
-	/** Interface for an external trainer process. */
-	struct ITrainerProcess
-	{
-		virtual ~ITrainerProcess() {}
-
-		/**
-		* Check if the server process is still running.
-		*/
-		virtual bool IsRunning() const = 0;
-
-		/**
-		* Wait for the server process to end
-		*
-		* @returns True if successful, otherwise false if it times out.
-		*/
-		virtual bool Wait() = 0;
-
-		/**
-		* Terminate the server process.
-		*/
-		virtual void Terminate() = 0;
-
-		/** Get the training subprocess. */
-		virtual FSubprocess* GetTrainingSubprocess() = 0;
-	};
-
 	/**
 	* This object allows you to launch the FSharedMemoryTrainer server as a subprocess,
 	* which is convenient when you want to train locally.
 	*/
-	struct LEARNINGTRAINING_API FSharedMemoryTrainerServerProcess : public ITrainerProcess
+	struct LEARNINGTRAINING_API FSharedMemoryTrainerServerProcess
 	{
 		/**
 		* Creates a training server as a subprocess using shared memory for communication. This will no-op if this UE
@@ -181,7 +155,7 @@ namespace UE::Learning
 			const ELogSetting LogSettings = ELogSetting::Normal);
 
 		/** Check if the server process is still running. */
-		virtual bool IsRunning() const override final;
+		bool IsRunning() const;
 
 		/**
 		* Wait for the server process to end
@@ -189,10 +163,10 @@ namespace UE::Learning
 		* @param Timeout		Timeout to wait in seconds
 		* @returns				true if successful, otherwise false if it times out
 		*/
-		virtual bool Wait() override final;
+		bool Wait();
 
 		/** Terminate the server process. */
-		virtual void Terminate() override final;
+		void Terminate();
 
 		/** Get the Controls shared memory array view. */
 		TSharedMemoryArrayView<2, volatile int32> GetControlsSharedMemoryArrayView() const;
@@ -204,7 +178,7 @@ namespace UE::Learning
 		const FString& GetConfigPath() const;
 
 		/** Get the training subprocess. */
-		virtual FSubprocess* GetTrainingSubprocess() override final;
+		FSubprocess* GetTrainingSubprocess();
 
 	private:
 
@@ -238,6 +212,7 @@ namespace UE::Learning
 			TArray<TSharedMemoryArrayView<3, float>, TInlineAllocator<1>> EpisodeFinalMemoryStates;
 			TArray<TSharedMemoryArrayView<3, float>, TInlineAllocator<1>> Observations;
 			TArray<TSharedMemoryArrayView<3, float>, TInlineAllocator<1>> Actions;
+			TArray<TSharedMemoryArrayView<3, float>, TInlineAllocator<1>> ActionModifiers;
 			TArray<TSharedMemoryArrayView<3, float>, TInlineAllocator<1>> MemoryStates;
 			TArray<TSharedMemoryArrayView<3, float>, TInlineAllocator<1>> Rewards;
 
@@ -257,7 +232,7 @@ namespace UE::Learning
 		FSharedMemoryTrainer(
 			const FString& InTaskName,
 			const int32 InProcessNum,
-			const TSharedPtr<UE::Learning::ITrainerProcess>& ExternalTrainerProcess,
+			const TSharedPtr<FSharedMemoryTrainerServerProcess>& ExternalTrainerProcess,
 			const float InTimeout = Trainer::DefaultTimeout);
 
 		~FSharedMemoryTrainer();
@@ -322,7 +297,7 @@ namespace UE::Learning
 	* This object allows you to launch the FSocketTrainer server as a subprocess,
 	* which is convenient when you want to train using it locally.
 	*/
-	struct LEARNINGTRAINING_API FSocketTrainerServerProcess : public ITrainerProcess
+	struct LEARNINGTRAINING_API FSocketTrainerServerProcess
 	{
 		/**
 		* Creates a training server as a subprocess
@@ -352,7 +327,7 @@ namespace UE::Learning
 		/**
 		* Check if the server process is still running
 		*/
-		virtual bool IsRunning() const override final;
+		bool IsRunning() const;
 
 		/**
 		* Wait for the server process to end
@@ -360,15 +335,15 @@ namespace UE::Learning
 		* @param Timeout		Timeout to wait in seconds
 		* @returns				true if successful, otherwise false if it times out
 		*/
-		virtual bool Wait() override final;
+		bool Wait();
 
 		/**
 		* Terminate the server process
 		*/
-		virtual void Terminate() override final;
+		void Terminate();
 
 		/** Get the training subprocess. */
-		virtual FSubprocess* GetTrainingSubprocess() override final;
+		FSubprocess* GetTrainingSubprocess();
 
 	private:
 
@@ -395,7 +370,7 @@ namespace UE::Learning
 		*/
 		FSocketTrainer(
 			ETrainerResponse& OutResponse,
-			const TSharedPtr<UE::Learning::ITrainerProcess>& ExternalTrainerProcess,
+			const TSharedPtr<FSocketTrainerServerProcess>& ExternalTrainerProcess,
 			const TCHAR* IpAddress = Trainer::DefaultIp,
 			const uint32 Port = Trainer::DefaultPort,
 			const float InTimeout = Trainer::DefaultTimeout);
diff --git a/Engine/Plugins/Experimental/LearningAgents/Source/LearningTraining/Public/LearningPPOTrainer.h b/Engine/Plugins/Experimental/LearningAgents/Source/LearningTraining/Public/LearningPPOTrainer.h
--- a/Engine/Plugins/Experimental/LearningAgents/Source/LearningTraining/Public/LearningPPOTrainer.h
+++ b/Engine/Plugins/Experimental/LearningAgents/Source/LearningTraining/Public/LearningPPOTrainer.h
@@ -27,7 +27,6 @@ class LEARNINGTRAINING_API ULearningSocketPPOTrainerServerCommandlet : public UC
 
 namespace UE::Learning
 {
-	struct ITrainerProcess;
 	struct IExternalTrainer;
 	struct FReplayBuffer;
 	struct FResetInstanceBuffer;
