LearningAgents: Added Named Discrete Exclusive and Inclusive Actions.

#virtualized

diff --git a/Engine/Plugins/Experimental/LearningAgents/Source/LearningAgents/Private/LearningAgentsActions.cpp b/Engine/Plugins/Experimental/LearningAgents/Source/LearningAgents/Private/LearningAgentsActions.cpp
--- a/Engine/Plugins/Experimental/LearningAgents/Source/LearningAgents/Private/LearningAgentsActions.cpp
+++ b/Engine/Plugins/Experimental/LearningAgents/Source/LearningAgents/Private/LearningAgentsActions.cpp
@@ -998,6 +998,64 @@ FLearningAgentsActionSchemaElement ULearningAgentsActions::SpecifyExclusiveDiscr
 	return { Schema->ActionSchema.CreateDiscreteExclusive({ Size, MakeArrayView(NormalizedPriorProbabilities) }, Tag)};
 }
 
+FLearningAgentsActionSchemaElement ULearningAgentsActions::SpecifyNamedExclusiveDiscreteAction(ULearningAgentsActionSchema* Schema, const TArray<FName>& Names, const TMap<FName, float>& PriorProbabilities, const FName Tag)
+{
+	if (Names.Num() == 0)
+	{
+		UE_LOG(LogLearning, Warning, TEXT("SpecifyNamedExclusiveDiscreteAction: Specifying zero-sized Named Exclusive Discrete Action."));
+	}
+
+	const int32 NameNum = Names.Num();
+
+	TArray<int32, TInlineAllocator<16>> NameIndices;
+	TArray<float, TInlineAllocator<16>> NamePriorProbabilities;
+	NameIndices.Empty(NameNum);
+	NamePriorProbabilities.Empty(NameNum);
+
+	int32 Index = 0;
+	for (const FName& Name : Names)
+	{
+		NameIndices.Add(Index);
+		const float* PriorProb = PriorProbabilities.Find(Name);
+		NamePriorProbabilities.Add(PriorProb ? *PriorProb : 1.0f / NameNum);
+		Index++;
+	}
+
+	// Sort Names According to FName
+
+	NameIndices.Sort([Names](const int32 Lhs, const int32 Rhs)
+		{
+			return Names[Lhs].ToString().ToLower() < Names[Rhs].ToString().ToLower();
+		});
+
+	TArray<FName, TInlineAllocator<16>> SortedNames;
+	TArray<float, TInlineAllocator<16>> SortedPriorProbabilities;
+	SortedNames.SetNumUninitialized(NameNum);
+	SortedPriorProbabilities.SetNumUninitialized(NameNum);
+	for (int32 Idx = 0; Idx < NameNum; Idx++)
+	{
+		SortedNames[Idx] = Names[NameIndices[Idx]];
+		SortedPriorProbabilities[Idx] = NamePriorProbabilities[NameIndices[Idx]];
+	}
+
+	return SpecifyNamedExclusiveDiscreteActionFromArrayViews(Schema, SortedNames, SortedPriorProbabilities, Tag);
+}
+
+FLearningAgentsActionSchemaElement ULearningAgentsActions::SpecifyNamedExclusiveDiscreteActionFromArrayViews(ULearningAgentsActionSchema* Schema, const TArrayView<const FName> Names, const TArrayView<const float> PriorProbabilities, const FName Tag)
+{
+	const int32 NameNum = Names.Num();
+
+	TArray<FLearningAgentsActionSchemaElement, TInlineAllocator<16>> Elements;
+	Elements.SetNumUninitialized(NameNum);
+
+	for (int32 NameIdx = 0; NameIdx < NameNum; NameIdx++)
+	{
+		Elements[NameIdx] = SpecifyNullAction(Schema);
+	}
+
+	return SpecifyExclusiveUnionActionFromArrayViews(Schema, Names, Elements, PriorProbabilities, Tag);
+}
+
 FLearningAgentsActionSchemaElement ULearningAgentsActions::SpecifyInclusiveDiscreteAction(ULearningAgentsActionSchema* Schema, const int32 Size, const TArray<float>& PriorProbabilities, const FName Tag)
 {
 	return SpecifyInclusiveDiscreteActionFromArrayView(Schema, Size, PriorProbabilities, Tag);
@@ -1033,6 +1091,64 @@ FLearningAgentsActionSchemaElement ULearningAgentsActions::SpecifyInclusiveDiscr
 	return { Schema->ActionSchema.CreateDiscreteInclusive({ Size, MakeArrayView(NormalizedPriorProbabilities) }, Tag) };
 }
 
+FLearningAgentsActionSchemaElement ULearningAgentsActions::SpecifyNamedInclusiveDiscreteAction(ULearningAgentsActionSchema* Schema, const TArray<FName> Names, const TMap<FName, float>& PriorProbabilities, const FName Tag)
+{
+	if (Names.Num() == 0)
+	{
+		UE_LOG(LogLearning, Warning, TEXT("SpecifyNamedInclusiveDiscreteAction: Specifying zero-sized Named Inclusive Discrete Action."));
+	}
+
+	const int32 NameNum = Names.Num();
+
+	TArray<int32, TInlineAllocator<16>> NameIndices;
+	TArray<float, TInlineAllocator<16>> NamePriorProbabilities;
+	NameIndices.Empty(NameNum);
+	NamePriorProbabilities.Empty(NameNum);
+
+	int32 Index = 0;
+	for (const FName& Name : Names)
+	{
+		NameIndices.Add(Index);
+		const float* PriorProb = PriorProbabilities.Find(Name);
+		NamePriorProbabilities.Add(PriorProb ? *PriorProb : 1.0f / NameNum);
+		Index++;
+	}
+
+	// Sort Names According to FName
+
+	NameIndices.Sort([Names](const int32 Lhs, const int32 Rhs)
+		{
+			return Names[Lhs].ToString().ToLower() < Names[Rhs].ToString().ToLower();
+		});
+
+	TArray<FName, TInlineAllocator<16>> SortedNames;
+	TArray<float, TInlineAllocator<16>> SortedPriorProbabilities;
+	SortedNames.SetNumUninitialized(NameNum);
+	SortedPriorProbabilities.SetNumUninitialized(NameNum);
+	for (int32 Idx = 0; Idx < NameNum; Idx++)
+	{
+		SortedNames[Idx] = Names[NameIndices[Idx]];
+		SortedPriorProbabilities[Idx] = NamePriorProbabilities[NameIndices[Idx]];
+	}
+
+	return SpecifyNamedInclusiveDiscreteActionFromArrayViews(Schema, SortedNames, SortedPriorProbabilities, Tag);
+}
+
+FLearningAgentsActionSchemaElement ULearningAgentsActions::SpecifyNamedInclusiveDiscreteActionFromArrayViews(ULearningAgentsActionSchema* Schema, const TArrayView<const FName> Names, const TArrayView<const float> PriorProbabilities, const FName Tag)
+{
+	const int32 NameNum = Names.Num();
+
+	TArray<FLearningAgentsActionSchemaElement, TInlineAllocator<16>> Elements;
+	Elements.SetNumUninitialized(NameNum);
+
+	for (int32 NameIdx = 0; NameIdx < NameNum; NameIdx++)
+	{
+		Elements[NameIdx] = SpecifyNullAction(Schema);
+	}
+
+	return SpecifyInclusiveUnionActionFromArrayViews(Schema, Names, Elements, PriorProbabilities, Tag);
+}
+
 FLearningAgentsActionSchemaElement ULearningAgentsActions::SpecifyStructAction(ULearningAgentsActionSchema* Schema, const TMap<FName, FLearningAgentsActionSchemaElement>& Elements, const FName Tag)
 {
 	if (Elements.Num() == 0)
@@ -1620,6 +1736,11 @@ FLearningAgentsActionObjectElement ULearningAgentsActions::MakeExclusiveDiscrete
 	return { Object->ActionObject.CreateDiscreteExclusive({ Index }, Tag) };
 }
 
+FLearningAgentsActionObjectElement ULearningAgentsActions::MakeNamedExclusiveDiscreteAction(ULearningAgentsActionObject* Object, const FName Name, const FName Tag)
+{
+	return MakeExclusiveUnionAction(Object, Name, MakeNullAction(Object), Tag);
+}
+
 FLearningAgentsActionObjectElement ULearningAgentsActions::MakeInclusiveDiscreteAction(ULearningAgentsActionObject* Object, const TArray<int32>& Indices, const FName Tag)
 {
 	return MakeInclusiveDiscreteActionFromArrayView(Object, Indices, Tag);
@@ -1653,6 +1774,26 @@ FLearningAgentsActionObjectElement ULearningAgentsActions::MakeInclusiveDiscrete
 	return { Object->ActionObject.CreateDiscreteInclusive({ Indices }, Tag) };
 }
 
+FLearningAgentsActionObjectElement ULearningAgentsActions::MakeNamedInclusiveDiscreteAction(ULearningAgentsActionObject* Object, const TArray<FName>& Names, const FName Tag)
+{
+	return MakeNamedInclusiveDiscreteActionFromArrayView(Object, Names, Tag);
+}
+
+FLearningAgentsActionObjectElement ULearningAgentsActions::MakeNamedInclusiveDiscreteActionFromArrayView(ULearningAgentsActionObject* Object, const TArrayView<const FName> Names, const FName Tag)
+{
+	const int32 NameNum = Names.Num();
+
+	TArray<FLearningAgentsActionObjectElement, TInlineAllocator<32>> Elements;
+	Elements.SetNumUninitialized(NameNum);
+
+	for (int32 NameIdx = 0; NameIdx < NameNum; NameIdx++)
+	{
+		Elements[NameIdx] = MakeNullAction(Object);
+	}
+
+	return MakeInclusiveUnionActionFromArrayViews(Object, Names, Elements, Tag);
+}
+
 FLearningAgentsActionObjectElement ULearningAgentsActions::MakeStructAction(ULearningAgentsActionObject* Object, const TMap<FName, FLearningAgentsActionObjectElement>& Elements, const FName Tag)
 {
 	if (Elements.Num() == 0)
@@ -2084,6 +2225,16 @@ FLearningAgentsActionModifierElement ULearningAgentsActions::MakeExclusiveDiscre
 	return { Modifier->ActionModifier.CreateDiscreteExclusive({ MaskedIndices }, Tag) };
 }
 
+FLearningAgentsActionModifierElement ULearningAgentsActions::MakeNamedExclusiveDiscreteActionModifier(ULearningAgentsActionModifier* Modifier, const TArray<FName>& MaskedNames, const FName Tag)
+{
+	return MakeNamedExclusiveDiscreteActionModifierFromArrayView(Modifier, MaskedNames, Tag);
+}
+
+FLearningAgentsActionModifierElement ULearningAgentsActions::MakeNamedExclusiveDiscreteActionModifierFromArrayView(ULearningAgentsActionModifier* Modifier, const TArrayView<const FName> MaskedNames, const FName Tag)
+{
+	return MakeExclusiveUnionActionModifierFromArrayViews(Modifier, {}, {}, MaskedNames, Tag);
+}
+
 FLearningAgentsActionModifierElement ULearningAgentsActions::MakeInclusiveDiscreteActionModifier(ULearningAgentsActionModifier* Modifier, const TArray<int32>& MaskedIndices, const FName Tag)
 {
 	return MakeInclusiveDiscreteActionModifierFromArrayView(Modifier, MaskedIndices, Tag);
@@ -2100,6 +2251,16 @@ FLearningAgentsActionModifierElement ULearningAgentsActions::MakeInclusiveDiscre
 	return { Modifier->ActionModifier.CreateDiscreteInclusive({ MaskedIndices }, Tag) };
 }
 
+FLearningAgentsActionModifierElement ULearningAgentsActions::MakeNamedInclusiveDiscreteActionModifier(ULearningAgentsActionModifier* Modifier, const TArray<FName>& MaskedNames, const FName Tag)
+{
+	return MakeNamedInclusiveDiscreteActionModifierFromArrayView(Modifier, MaskedNames, Tag);
+}
+
+FLearningAgentsActionModifierElement ULearningAgentsActions::MakeNamedInclusiveDiscreteActionModifierFromArrayView(ULearningAgentsActionModifier* Modifier, const TArrayView<const FName> MaskedNames, const FName Tag)
+{
+	return MakeInclusiveUnionActionModifierFromArrayViews(Modifier, {}, {}, MaskedNames, Tag);
+}
+
 FLearningAgentsActionModifierElement ULearningAgentsActions::MakeStructActionModifier(ULearningAgentsActionModifier* Modifier, const TMap<FName, FLearningAgentsActionModifierElement>& Elements, const FName Tag)
 {
 	if (Elements.Num() == 0)
@@ -2749,6 +2910,41 @@ bool ULearningAgentsActions::GetExclusiveDiscreteAction(
 	return true;
 }
 
+bool ULearningAgentsActions::GetNamedExclusiveDiscreteAction(
+	FName& OutName,
+	const ULearningAgentsActionObject* Object,
+	const FLearningAgentsActionObjectElement Element,
+	const FName Tag,
+	const bool bVisualLoggerEnabled,
+	ULearningAgentsManagerListener* VisualLoggerListener,
+	const int32 VisualLoggerAgentId,
+	const FVector VisualLoggerLocation,
+	const FLinearColor VisualLoggerColor)
+{
+	FLearningAgentsActionObjectElement OutElement;
+	if (!GetExclusiveUnionAction(OutName, OutElement, Object, Element, Tag))
+	{
+		return false;
+	}
+
+#if UE_LEARNING_AGENTS_ENABLE_VISUAL_LOG
+	if (bVisualLoggerEnabled && VisualLoggerListener)
+	{
+		const ULearningAgentsVisualLoggerObject* VisualLoggerObject = VisualLoggerListener->GetOrAddVisualLoggerObject(Tag);
+
+		UE_LEARNING_AGENTS_VLOG_STRING(VisualLoggerObject, LogLearning, Display, VisualLoggerLocation,
+			VisualLoggerColor.ToFColor(true),
+			TEXT("Listener: %s\nTag: %s\nAgent Id: % 3i\nName: [%s]"),
+			*VisualLoggerListener->GetName(),
+			*Tag.ToString(),
+			VisualLoggerAgentId,
+			*OutName.ToString());
+	}
+#endif
+
+	return true;
+}
+
 bool ULearningAgentsActions::GetInclusiveDiscreteActionNum(int32& OutNum, const ULearningAgentsActionObject* Object, const FLearningAgentsActionObjectElement Element, const FName Tag)
 {
 	if (!Object)
@@ -2884,6 +3080,88 @@ bool ULearningAgentsActions::GetInclusiveDiscreteActionToArrayView(
 	return true;
 }
 
+bool ULearningAgentsActions::GetNamedInclusiveDiscreteActionNum(int32& OutNum, const ULearningAgentsActionObject* Object, const FLearningAgentsActionObjectElement Element, const FName Tag)
+{
+	return GetInclusiveUnionActionNum(OutNum, Object, Element, Tag);
+}
+
+bool ULearningAgentsActions::GetNamedInclusiveDiscreteAction(
+	TArray<FName>& OutNames,
+	const ULearningAgentsActionObject* Object,
+	const FLearningAgentsActionObjectElement Element,
+	const FName Tag,
+	const bool bVisualLoggerEnabled,
+	ULearningAgentsManagerListener* VisualLoggerListener,
+	const int32 VisualLoggerAgentId,
+	const FVector VisualLoggerLocation,
+	const FLinearColor VisualLoggerColor)
+{
+	int32 OutNameNum = 0;
+	if (!GetNamedInclusiveDiscreteActionNum(OutNameNum, Object, Element, Tag))
+	{
+		OutNames.Empty();
+		return false;
+	}
+
+	OutNames.SetNumUninitialized(OutNameNum);
+
+	if (!GetNamedInclusiveDiscreteActionToArrayView(OutNames, Object, Element, Tag, bVisualLoggerEnabled, VisualLoggerListener, VisualLoggerAgentId, VisualLoggerLocation, VisualLoggerColor))
+	{
+		OutNames.Empty();
+		return false;
+	}
+
+	return true;
+}
+
+namespace UE::Learning::Array
+{
+	// We declare this here instead of in LearningArray.h because FName is part of the UObject system and right now the 
+	// Learning core module does not have a dependency on the UObject part of UE which is a nice thing to preserve.	
+	static inline FString FormatNames(const TLearningArrayView<1, const FName> Array, const int32 MaxItemNum = 16)
+	{
+		return Format<FName>(Array, [](const FName& Value) { return FString::Printf(TEXT("%s"), *Value.ToString()); }, MaxItemNum);
+	}
+}
+
+bool ULearningAgentsActions::GetNamedInclusiveDiscreteActionToArrayView(
+	TArrayView<FName> OutNames,
+	const ULearningAgentsActionObject* Object,
+	const FLearningAgentsActionObjectElement Element,
+	const FName Tag,
+	const bool bVisualLoggerEnabled,
+	ULearningAgentsManagerListener* VisualLoggerListener,
+	const int32 VisualLoggerAgentId,
+	const FVector VisualLoggerLocation,
+	const FLinearColor VisualLoggerColor)
+{
+	TArray<FLearningAgentsActionObjectElement, TInlineAllocator<16>> OutElements;
+	OutElements.SetNumUninitialized(OutNames.Num());
+
+	if (!GetInclusiveUnionActionToArrayViews(OutNames, OutElements, Object, Element, Tag))
+	{
+		return false;
+	}
+
+#if UE_LEARNING_AGENTS_ENABLE_VISUAL_LOG
+	if (bVisualLoggerEnabled && VisualLoggerListener)
+	{
+		const ULearningAgentsVisualLoggerObject* VisualLoggerObject = VisualLoggerListener->GetOrAddVisualLoggerObject(Tag);
+
+		UE_LEARNING_AGENTS_VLOG_STRING(VisualLoggerObject, LogLearning, Display, VisualLoggerLocation,
+			VisualLoggerColor.ToFColor(true),
+			TEXT("Listener: %s\nTag: %s\nAgent Id: % 3i\nNames: %s"),
+			*VisualLoggerListener->GetName(),
+			*Tag.ToString(),
+			VisualLoggerAgentId,
+			*UE::Learning::Array::FormatNames(OutNames, 256));
+	}
+#endif
+
+	return true;
+}
+
+
 bool ULearningAgentsActions::GetStructActionNum(int32& OutNum, const ULearningAgentsActionObject* Object, const FLearningAgentsActionObjectElement Element, const FName Tag)
 {
 	if (!Object)
diff --git a/Engine/Plugins/Experimental/LearningAgents/Source/LearningAgents/Public/LearningAgentsActions.h b/Engine/Plugins/Experimental/LearningAgents/Source/LearningAgents/Public/LearningAgentsActions.h
--- a/Engine/Plugins/Experimental/LearningAgents/Source/LearningAgents/Public/LearningAgentsActions.h
+++ b/Engine/Plugins/Experimental/LearningAgents/Source/LearningAgents/Public/LearningAgentsActions.h
@@ -242,6 +242,32 @@ public:
 	 */
 	static FLearningAgentsActionSchemaElement SpecifyExclusiveDiscreteActionFromArrayView(ULearningAgentsActionSchema* Schema, const int32 Size, const TArrayView<const float> PriorProbabilities = {}, const FName Tag = TEXT("DiscreteExclusiveAction"));
 
+	/**
+	 * Specifies a new named exclusive discrete action. This represents an action which is an exclusive choice from a number of discrete options, sampled
+	 * from a Categorical distribution.
+	 *
+	 * @param Schema The Action Schema
+	 * @param Names The names of the discrete options in the action.
+	 * @param PriorProbabilities The prior probabilities of each option. Can be left empty to use a uniform distribution over options. Should sum to one.
+	 * @param Tag The tag of this new action. Used during action object validation and debugging.
+	 * @return The newly created action schema element.
+	 */
+	UFUNCTION(BlueprintPure, Category = "LearningAgents", meta = (AdvancedDisplay = 2, AutoCreateRefTerm = "PriorProbabilities"))
+	static FLearningAgentsActionSchemaElement SpecifyNamedExclusiveDiscreteAction(ULearningAgentsActionSchema* Schema, const TArray<FName>& Names, const TMap<FName, float>& PriorProbabilities, const FName Tag = TEXT("NamedDiscreteExclusiveAction"));
+
+	/**
+	 * Specifies a new named exclusive discrete action. This represents an action which is an exclusive choice from a number of discrete options, sampled
+	 * from a Categorical distribution.
+	 *
+	 * @param Schema The Action Schema
+	 * @param Names The names of the discrete options in the action.
+	 * @param PriorProbabilityNames The names of the prior probabilities. Can be left empty to use a uniform distribution over options.
+	 * @param PriorProbabilityValues The prior probabilities of each option. Can be left empty to use a uniform distribution over options. Should sum to one.
+	 * @param Tag The tag of this new action. Used during action object validation and debugging.
+	 * @return The newly created action schema element.
+	 */
+	static FLearningAgentsActionSchemaElement SpecifyNamedExclusiveDiscreteActionFromArrayViews(ULearningAgentsActionSchema* Schema, const TArrayView<const FName> Names, const TArrayView<const float> PriorProbabilities, const FName Tag = TEXT("NamedDiscreteExclusiveAction"));
+
 	/**
 	 * Specifies a new inclusive discrete action. This represents an action which is an inclusive choice from a number of discrete options, sampled
 	 * from a Bernoulli distribution.
@@ -267,6 +293,31 @@ public:
 	 */
 	static FLearningAgentsActionSchemaElement SpecifyInclusiveDiscreteActionFromArrayView(ULearningAgentsActionSchema* Schema, const int32 Size, const TArrayView<const float> PriorProbabilities = {}, const FName Tag = TEXT("DiscreteInclusiveAction"));
 
+	/**
+	 * Specifies a new named inclusive discrete action. This represents an action which is an inclusive choice from a number of discrete options, sampled
+	 * from a Bernoulli distribution.
+	 *
+	 * @param Schema The Action Schema
+	 * @param Names The names of the discrete options in the action.
+	 * @param PriorProbabilities The prior probabilities of each option. Can be left empty to use a probability of 0.5 for each option.
+	 * @param Tag The tag of this new action. Used during action object validation and debugging.
+	 * @return The newly created action schema element.
+	 */
+	UFUNCTION(BlueprintPure, Category = "LearningAgents", meta = (AdvancedDisplay = 2, AutoCreateRefTerm = "PriorProbabilities"))
+	static FLearningAgentsActionSchemaElement SpecifyNamedInclusiveDiscreteAction(ULearningAgentsActionSchema* Schema, const TArray<FName> Names, const TMap<FName, float>& PriorProbabilities, const FName Tag = TEXT("NamedDiscreteInclusiveAction"));
+
+	/**
+	 * Specifies a new named inclusive discrete action. This represents an action which is an inclusive choice from a number of discrete options, sampled
+	 * from a Bernoulli distribution.
+	 *
+	 * @param Schema The Action Schema
+	 * @param Names The names of the discrete options in the action.
+	 * @param PriorProbabilities The prior probabilities of each option. Can be left empty to use a probability of 0.5 for each option.
+	 * @param Tag The tag of this new action. Used during action object validation and debugging.
+	 * @return The newly created action schema element.
+	 */
+	static FLearningAgentsActionSchemaElement SpecifyNamedInclusiveDiscreteActionFromArrayViews(ULearningAgentsActionSchema* Schema, const TArrayView<const FName> Names, const TArrayView<const float> PriorProbabilities = {}, const FName Tag = TEXT("NamedDiscreteInclusiveAction"));
+
 	/**
 	 * Specifies a new struct action. This represents an action which is made up of a number of named sub-actions.
 	 *
@@ -635,10 +686,17 @@ public:
 	UFUNCTION(BlueprintPure, Category = "LearningAgents", meta = (AdvancedDisplay = 2))
 	static FLearningAgentsActionObjectElement MakeExclusiveDiscreteAction(ULearningAgentsActionObject* Object, const int32 Index, const FName Tag = TEXT("DiscreteExclusiveAction"));
 
+	UFUNCTION(BlueprintPure, Category = "LearningAgents", meta = (AdvancedDisplay = 2))
+	static FLearningAgentsActionObjectElement MakeNamedExclusiveDiscreteAction(ULearningAgentsActionObject* Object, const FName Name, const FName Tag = TEXT("NamedDiscreteExclusiveAction"));
+
 	UFUNCTION(BlueprintPure, Category = "LearningAgents", meta = (AdvancedDisplay = 2))
 	static FLearningAgentsActionObjectElement MakeInclusiveDiscreteAction(ULearningAgentsActionObject* Object, const TArray<int32>& Indices, const FName Tag = TEXT("DiscreteInclusiveAction"));
 	static FLearningAgentsActionObjectElement MakeInclusiveDiscreteActionFromArrayView(ULearningAgentsActionObject* Object, const TArrayView<const int32> Indices, const FName Tag = TEXT("DiscreteInclusiveAction"));
 
+	UFUNCTION(BlueprintPure, Category = "LearningAgents", meta = (AdvancedDisplay = 2))
+	static FLearningAgentsActionObjectElement MakeNamedInclusiveDiscreteAction(ULearningAgentsActionObject* Object, const TArray<FName>& Names, const FName Tag = TEXT("NamedDiscreteInclusiveAction"));
+	static FLearningAgentsActionObjectElement MakeNamedInclusiveDiscreteActionFromArrayView(ULearningAgentsActionObject* Object, const TArrayView<const FName> Names, const FName Tag = TEXT("NamedDiscreteInclusiveAction"));
+
 	UFUNCTION(BlueprintPure, Category = "LearningAgents", meta = (AdvancedDisplay = 2))
 	static FLearningAgentsActionObjectElement MakeStructAction(ULearningAgentsActionObject* Object, const TMap<FName, FLearningAgentsActionObjectElement>& Elements, const FName Tag = TEXT("StructAction"));
 
@@ -736,10 +794,18 @@ public:
 	static FLearningAgentsActionModifierElement MakeExclusiveDiscreteActionModifier(ULearningAgentsActionModifier* Modifier, const TArray<int32>& MaskedIndices, const FName Tag = TEXT("DiscreteExclusiveAction"));
 	static FLearningAgentsActionModifierElement MakeExclusiveDiscreteActionModifierFromArrayView(ULearningAgentsActionModifier* Modifier, const TArrayView<const int32> MaskedIndices, const FName Tag = TEXT("DiscreteExclusiveAction"));
 
+	UFUNCTION(BlueprintPure, Category = "LearningAgents", meta = (AdvancedDisplay = 2))
+	static FLearningAgentsActionModifierElement MakeNamedExclusiveDiscreteActionModifier(ULearningAgentsActionModifier* Modifier, const TArray<FName>& MaskedNames, const FName Tag = TEXT("NamedDiscreteExclusiveAction"));
+	static FLearningAgentsActionModifierElement MakeNamedExclusiveDiscreteActionModifierFromArrayView(ULearningAgentsActionModifier* Modifier, const TArrayView<const FName> MaskedNames, const FName Tag = TEXT("NamedDiscreteExclusiveAction"));
+
 	UFUNCTION(BlueprintPure, Category = "LearningAgents", meta = (AdvancedDisplay = 2))
 	static FLearningAgentsActionModifierElement MakeInclusiveDiscreteActionModifier(ULearningAgentsActionModifier* Modifier, const TArray<int32>& MaskedIndices, const FName Tag = TEXT("DiscreteInclusiveAction"));
 	static FLearningAgentsActionModifierElement MakeInclusiveDiscreteActionModifierFromArrayView(ULearningAgentsActionModifier* Modifier, const TArrayView<const int32> MaskedIndices, const FName Tag = TEXT("DiscreteInclusiveAction"));
 
+	UFUNCTION(BlueprintPure, Category = "LearningAgents", meta = (AdvancedDisplay = 2))
+	static FLearningAgentsActionModifierElement MakeNamedInclusiveDiscreteActionModifier(ULearningAgentsActionModifier* Modifier, const TArray<FName>& MaskedNames, const FName Tag = TEXT("NamedDiscreteInclusiveAction"));
+	static FLearningAgentsActionModifierElement MakeNamedInclusiveDiscreteActionModifierFromArrayView(ULearningAgentsActionModifier* Modifier, const TArrayView<const FName> MaskedNames, const FName Tag = TEXT("NamedDiscreteInclusiveAction"));
+
 	UFUNCTION(BlueprintPure, Category = "LearningAgents", meta = (AdvancedDisplay = 2))
 	static FLearningAgentsActionModifierElement MakeStructActionModifier(ULearningAgentsActionModifier* Modifier, const TMap<FName, FLearningAgentsActionModifierElement>& Elements, const FName Tag = TEXT("StructAction"));
 
@@ -910,6 +976,32 @@ public:
 		const FVector VisualLoggerLocation = FVector::ZeroVector,
 		const FLinearColor VisualLoggerColor = FLinearColor::Blue);
 
+	/**
+	 * Get the name for a named exclusive discrete action.
+	 *
+	 * @param OutName The output name.
+	 * @param Object The Action Object
+	 * @param Element The Action Object Element
+	 * @param Tag The tag of the corresponding action. Must match the tag given during Specify.
+	 * @param bVisualLoggerEnabled When true, debug data will be sent to the visual logger.
+	 * @param VisualLoggerListener The listener object which is making this action. This must be set to use logging.
+	 * @param VisualLoggerAgentId The agent id associated with this action.
+	 * @param VisualLoggerLocation A location for the visual logger information in the world.
+	 * @param VisualLoggerColor The color for the visual logger display.
+	 * @return true if the provided Element is the correct type, otherwise false.
+	 */
+	UFUNCTION(BlueprintPure = false, Category = "LearningAgents", meta = (AdvancedDisplay = 3, ReturnDisplayName = "Success", DefaultToSelf = "VisualLoggerListener"))
+	static bool GetNamedExclusiveDiscreteAction(
+		FName& OutName,
+		const ULearningAgentsActionObject* Object,
+		const FLearningAgentsActionObjectElement Element,
+		const FName Tag = TEXT("NamedDiscreteExclusiveAction"),
+		const bool bVisualLoggerEnabled = false,
+		ULearningAgentsManagerListener* VisualLoggerListener = nullptr,
+		const int32 VisualLoggerAgentId = -1,
+		const FVector VisualLoggerLocation = FVector::ZeroVector,
+		const FLinearColor VisualLoggerColor = FLinearColor::Blue);
+
 	/**
 	 * Get the number of indices for an inclusive discrete action.
 	 *
@@ -973,6 +1065,69 @@ public:
 		const FVector VisualLoggerLocation = FVector::ZeroVector,
 		const FLinearColor VisualLoggerColor = FLinearColor::Blue);
 
+	/**
+	 * Get the number of names for a named inclusive discrete action.
+	 *
+	 * @param OutNum The output number of names.
+	 * @param Object The Action Object
+	 * @param Element The Action Object Element
+	 * @param Tag The tag of the corresponding action. Must match the tag given during Specify.
+	 * @return true if the provided Element is the correct type, otherwise false.
+	 */
+	UFUNCTION(BlueprintPure = false, Category = "LearningAgents", meta = (AdvancedDisplay = 3, ReturnDisplayName = "Success"))
+	static bool GetNamedInclusiveDiscreteActionNum(int32& OutNum, const ULearningAgentsActionObject* Object, const FLearningAgentsActionObjectElement Element, const FName Tag = TEXT("NamedDiscreteInclusiveAction"));
+
+	/**
+	 * Get the names for a named inclusive discrete action.
+	 *
+	 * @param OutNames The output names.
+	 * @param Object The Action Object
+	 * @param Element The Action Object Element
+	 * @param Tag The tag of the corresponding action. Must match the tag given during Specify.
+	 * @param bVisualLoggerEnabled When true, debug data will be sent to the visual logger.
+	 * @param VisualLoggerListener The listener object which is making this action. This must be set to use logging.
+	 * @param VisualLoggerAgentId The agent id associated with this action.
+	 * @param VisualLoggerLocation A location for the visual logger information in the world.
+	 * @param VisualLoggerColor The color for the visual logger display.
+	 * @return true if the provided Element is the correct type, otherwise false.
+	 */
+	UFUNCTION(BlueprintPure = false, Category = "LearningAgents", meta = (AdvancedDisplay = 3, ReturnDisplayName = "Success", DefaultToSelf = "VisualLoggerListener"))
+	static bool GetNamedInclusiveDiscreteAction(
+		TArray<FName>& OutNames,
+		const ULearningAgentsActionObject* Object,
+		const FLearningAgentsActionObjectElement Element,
+		const FName Tag = TEXT("NamedDiscreteInclusiveAction"),
+		const bool bVisualLoggerEnabled = false,
+		ULearningAgentsManagerListener* VisualLoggerListener = nullptr,
+		const int32 VisualLoggerAgentId = -1,
+		const FVector VisualLoggerLocation = FVector::ZeroVector,
+		const FLinearColor VisualLoggerColor = FLinearColor::Blue);
+
+	/**
+	 * Get the names for a named inclusive discrete action. The OutNames ArrayView must be the correct size.
+	 *
+	 * @param OutNames The output names.
+	 * @param Object The Action Object
+	 * @param Element The Action Object Element
+	 * @param Tag The tag of the corresponding action. Must match the tag given during Specify.
+	 * @param bVisualLoggerEnabled When true, debug data will be sent to the visual logger.
+	 * @param VisualLoggerListener The listener object which is making this action. This must be set to use logging.
+	 * @param VisualLoggerAgentId The agent id associated with this action.
+	 * @param VisualLoggerLocation A location for the visual logger information in the world.
+	 * @param VisualLoggerColor The color for the visual logger display.
+	 * @return true if the provided Element is the correct type, otherwise false.
+	 */
+	static bool GetNamedInclusiveDiscreteActionToArrayView(
+		TArrayView<FName> OutNames,
+		const ULearningAgentsActionObject* Object,
+		const FLearningAgentsActionObjectElement Element,
+		const FName Tag = TEXT("NamedDiscreteInclusiveAction"),
+		const bool bVisualLoggerEnabled = false,
+		ULearningAgentsManagerListener* VisualLoggerListener = nullptr,
+		const int32 VisualLoggerAgentId = -1,
+		const FVector VisualLoggerLocation = FVector::ZeroVector,
+		const FLinearColor VisualLoggerColor = FLinearColor::Blue);
+
 	/**
 	 * Get the number of sub-actions for a struct action.
 	 *
