LearningAgents: Fix issue with the training process being mandatory in the communicator

#rnx

#virtualized

diff --git a/Engine/Plugins/Experimental/LearningAgents/Source/LearningAgentsTraining/Private/LearningAgentsCommunicator.cpp b/Engine/Plugins/Experimental/LearningAgents/Source/LearningAgentsTraining/Private/LearningAgentsCommunicator.cpp
--- a/Engine/Plugins/Experimental/LearningAgents/Source/LearningAgentsTraining/Private/LearningAgentsCommunicator.cpp
+++ b/Engine/Plugins/Experimental/LearningAgents/Source/LearningAgentsTraining/Private/LearningAgentsCommunicator.cpp
@@ -56,7 +56,7 @@ FLearningAgentsSharedMemoryTrainerProcess ULearningAgentsCommunicatorLibrary::Sp
 }
 
 FLearningAgentsCommunicator ULearningAgentsCommunicatorLibrary::MakeSharedMemoryCommunicator(
-	const FLearningAgentsSharedMemoryTrainerProcess& TrainerProcess,
+	FLearningAgentsSharedMemoryTrainerProcess TrainerProcess,
 	const FLearningAgentsSharedMemoryCommunicatorSettings& SharedMemorySettings)
 {
 	FLearningAgentsCommunicator Communicator;
@@ -125,7 +125,7 @@ FLearningAgentsSocketTrainerProcess ULearningAgentsCommunicatorLibrary::SpawnSoc
 }
 
 FLearningAgentsCommunicator ULearningAgentsCommunicatorLibrary::MakeSocketCommunicator(
-	const FLearningAgentsSocketTrainerProcess& TrainerProcess,
+	FLearningAgentsSocketTrainerProcess TrainerProcess,
 	const FLearningAgentsSocketCommunicatorSettings& SocketSettings)
 {
 	UE::Learning::ETrainerResponse Response = UE::Learning::ETrainerResponse::Success;
diff --git a/Engine/Plugins/Experimental/LearningAgents/Source/LearningAgentsTraining/Public/LearningAgentsCommunicator.h b/Engine/Plugins/Experimental/LearningAgents/Source/LearningAgentsTraining/Public/LearningAgentsCommunicator.h
--- a/Engine/Plugins/Experimental/LearningAgents/Source/LearningAgentsTraining/Public/LearningAgentsCommunicator.h
+++ b/Engine/Plugins/Experimental/LearningAgents/Source/LearningAgentsTraining/Public/LearningAgentsCommunicator.h
@@ -117,7 +117,7 @@ public:
 	 */
 	UFUNCTION(BlueprintCallable, Category = "LearningAgents", meta = (AutoCreateRefTerm = "SharedMemorySettings"))
 	static FLearningAgentsCommunicator MakeSharedMemoryCommunicator(
-		const FLearningAgentsSharedMemoryTrainerProcess& TrainerProcess,
+		FLearningAgentsSharedMemoryTrainerProcess TrainerProcess,
 		const FLearningAgentsSharedMemoryCommunicatorSettings& SharedMemorySettings = FLearningAgentsSharedMemoryCommunicatorSettings());
 
 	/**
@@ -144,6 +144,6 @@ public:
 	 */
 	UFUNCTION(BlueprintCallable, Category = "LearningAgents", meta = (AutoCreateRefTerm = "SocketSettings"))
 	static FLearningAgentsCommunicator MakeSocketCommunicator(
-		const FLearningAgentsSocketTrainerProcess& TrainerProcess,
+		FLearningAgentsSocketTrainerProcess TrainerProcess,
 		const FLearningAgentsSocketCommunicatorSettings& SocketSettings = FLearningAgentsSocketCommunicatorSettings());
 };
