LearningAgents: ActionModifier plumbing

#rnx

#virtualized

diff --git a/Engine/Plugins/Experimental/LearningAgents/Source/LearningTraining/Private/LearningExternalTrainer.cpp b/Engine/Plugins/Experimental/LearningAgents/Source/LearningTraining/Private/LearningExternalTrainer.cpp
--- a/Engine/Plugins/Experimental/LearningAgents/Source/LearningTraining/Private/LearningExternalTrainer.cpp
+++ b/Engine/Plugins/Experimental/LearningAgents/Source/LearningTraining/Private/LearningExternalTrainer.cpp
@@ -654,6 +654,12 @@ namespace UE::Learning
 			Actions.Add(Acts.View[ProcessIdx]);
 		}
 
+		TArray<TLearningArrayView<2, float>> ActionModifiers;
+		for (TSharedMemoryArrayView<3, float>& Mods : SharedMemoryExperienceContainers[ReplayBufferId].ActionModifiers)
+		{
+			ActionModifiers.Add(Mods.View[ProcessIdx]);
+		}
+
 		TArray<TLearningArrayView<2, float>> MemoryStates;
 		for (TSharedMemoryArrayView<3, float>& Mems : SharedMemoryExperienceContainers[ReplayBufferId].MemoryStates)
 		{
@@ -675,6 +681,7 @@ namespace UE::Learning
 			EpisodeFinalMemoryStates,
 			Observations,
 			Actions,
+			ActionModifiers,
 			MemoryStates,
 			Rewards,
 			Controls.View[ProcessIdx],
diff --git a/Engine/Plugins/Experimental/LearningAgents/Source/LearningTraining/Private/LearningSharedMemoryTraining.cpp b/Engine/Plugins/Experimental/LearningAgents/Source/LearningTraining/Private/LearningSharedMemoryTraining.cpp
--- a/Engine/Plugins/Experimental/LearningAgents/Source/LearningTraining/Private/LearningSharedMemoryTraining.cpp
+++ b/Engine/Plugins/Experimental/LearningAgents/Source/LearningTraining/Private/LearningSharedMemoryTraining.cpp
@@ -203,6 +203,7 @@ namespace UE::Learning::SharedMemoryTraining
 		TArrayView<TLearningArrayView<2, float>> EpisodeFinalMemoryStates,
 		TArrayView<TLearningArrayView<2, float>> Observations,
 		TArrayView<TLearningArrayView<2, float>> Actions,
+		TArrayView<TLearningArrayView<2, float>> ActionModifiers,
 		TArrayView<TLearningArrayView<2, float>> MemoryStates,
 		TArrayView<TLearningArrayView<2, float>> Rewards,
 		TLearningArrayView<1, volatile int32> Controls,
@@ -285,6 +286,11 @@ namespace UE::Learning::SharedMemoryTraining
 			Array::Copy(Actions[Index].Slice(0, StepNum), ReplayBuffer.GetActions(Index));
 		}
 
+		for (int32 Index = 0; Index < ReplayBuffer.GetActionModifiersNum(); Index++)
+		{
+			Array::Copy(ActionModifiers[Index].Slice(0, StepNum), ReplayBuffer.GetActionModifiers(Index));
+		}
+
 		for (int32 Index = 0; Index < ReplayBuffer.GetMemoryStatesNum(); Index++)
 		{
 			Array::Copy(MemoryStates[Index].Slice(0, StepNum), ReplayBuffer.GetMemoryStates(Index));
diff --git a/Engine/Plugins/Experimental/LearningAgents/Source/LearningTraining/Private/LearningSharedMemoryTraining.h b/Engine/Plugins/Experimental/LearningAgents/Source/LearningTraining/Private/LearningSharedMemoryTraining.h
--- a/Engine/Plugins/Experimental/LearningAgents/Source/LearningTraining/Private/LearningSharedMemoryTraining.h
+++ b/Engine/Plugins/Experimental/LearningAgents/Source/LearningTraining/Private/LearningSharedMemoryTraining.h
@@ -68,6 +68,7 @@ namespace UE::Learning
 			TArrayView<TLearningArrayView<2, float>> EpisodeFinalMemoryStates,
 			TArrayView<TLearningArrayView<2, float>> Observations,
 			TArrayView<TLearningArrayView<2, float>> Actions,
+			TArrayView<TLearningArrayView<2, float>> ActionModifiers,
 			TArrayView<TLearningArrayView<2, float>> MemoryStates,
 			TArrayView<TLearningArrayView<2, float>> Rewards,
 			TLearningArrayView<1, volatile int32> Controls,
diff --git a/Engine/Plugins/Experimental/LearningAgents/Source/LearningTraining/Private/LearningSocketTraining.cpp b/Engine/Plugins/Experimental/LearningAgents/Source/LearningTraining/Private/LearningSocketTraining.cpp
--- a/Engine/Plugins/Experimental/LearningAgents/Source/LearningTraining/Private/LearningSocketTraining.cpp
+++ b/Engine/Plugins/Experimental/LearningAgents/Source/LearningTraining/Private/LearningSocketTraining.cpp
@@ -354,6 +354,12 @@ namespace UE::Learning::SocketTraining
 			Response = SendWithTimeout(Socket, Process, (const uint8*)ReplayBuffer.GetActions(Index).GetData(), ReplayBuffer.GetActions(Index).Num() * sizeof(float), Timeout);
 			if (Response != ETrainerResponse::Success) { return Response; }
 		}
+
+		for (int32 Index = 0; Index < ReplayBuffer.GetActionModifiersNum(); Index++)
+		{
+			Response = SendWithTimeout(Socket, Process, (const uint8*)ReplayBuffer.GetActionModifiers(Index).GetData(), ReplayBuffer.GetActionModifiers(Index).Num() * sizeof(float), Timeout);
+			if (Response != ETrainerResponse::Success) { return Response; }
+		}
 		
 		for (int32 Index = 0; Index < ReplayBuffer.GetMemoryStatesNum(); Index++)
 		{
