LearningAgents: Fix issues with action modifiers

#rnx

diff --git a/Engine/Plugins/Experimental/LearningAgents/Source/LearningAgentsTraining/Private/LearningAgentsPPOTrainer.cpp b/Engine/Plugins/Experimental/LearningAgents/Source/LearningAgentsTraining/Private/LearningAgentsPPOTrainer.cpp
--- a/Engine/Plugins/Experimental/LearningAgents/Source/LearningAgentsTraining/Private/LearningAgentsPPOTrainer.cpp
+++ b/Engine/Plugins/Experimental/LearningAgents/Source/LearningAgentsTraining/Private/LearningAgentsPPOTrainer.cpp
@@ -510,7 +510,7 @@ void ULearningAgentsPPOTrainer::ProcessExperience(const bool bResetAgentsOnUpdat
 		UE_LOG(LogLearning, Warning, TEXT("%s: No agents added to Manager."), *GetName());
 	}
 
-	// Check Observations, ActionModifiers, Actions, Rewards, and Completions have been completed and have matching iteration number
+	// Check Observations, Actions, ActionModifiers, Rewards, and Completions have been completed and have matching iteration number
 
 	TArray<int32> ValidAgentIds;
 	ValidAgentIds.Empty(Manager->GetMaxAgentNum());
@@ -532,7 +532,7 @@ void ULearningAgentsPPOTrainer::ProcessExperience(const bool bResetAgentsOnUpdat
 			Interactor->GetObservationIteration(AgentId) != TrainingEnvironment->GetRewardIteration(AgentId) ||
 			Interactor->GetObservationIteration(AgentId) != TrainingEnvironment->GetCompletionIteration(AgentId))
 		{
-			UE_LOG(LogLearning, Warning, TEXT("%s: Agent with id %i has non-matching iteration numbers (observation: %i, action modifiers: %i, action: %i, reward: %i, completion: %i). Experience will not be processed for it."), *GetName(), AgentId,
+			UE_LOG(LogLearning, Warning, TEXT("%s: Agent with id %i has non-matching iteration numbers (observation: %i, action: %i, action modifiers: %i, reward: %i, completion: %i). Experience will not be processed for it."), *GetName(), AgentId,
 				Interactor->GetObservationIteration(AgentId),
 				Interactor->GetActionIteration(AgentId),
 				Interactor->GetActionModifierIteration(AgentId),
