Learning Agents: Shared Memory Debugging & Scale out

#virtualized

diff --git a/Engine/Plugins/Experimental/LearningAgents/Source/LearningAgentsTraining/Private/LearningAgentsCommunicator.cpp b/Engine/Plugins/Experimental/LearningAgents/Source/LearningAgentsTraining/Private/LearningAgentsCommunicator.cpp
--- a/Engine/Plugins/Experimental/LearningAgents/Source/LearningAgentsTraining/Private/LearningAgentsCommunicator.cpp
+++ b/Engine/Plugins/Experimental/LearningAgents/Source/LearningAgentsTraining/Private/LearningAgentsCommunicator.cpp
@@ -5,14 +5,30 @@
 #include "LearningAgentsInteractor.h"
 #include "LearningExternalTrainer.h"
 
+#include "Misc/CommandLine.h"
+#include "Misc/Guid.h"
 #include "Misc/Paths.h"
 
+FLearningAgentsCommunicator ULearningAgentsCommunicatorLibrary::MakeSharedMemoryTrainingProcess(
+	const FLearningAgentsTrainerProcessSettings& TrainerProcessSettings,
+	const FLearningAgentsSharedMemoryCommunicatorSettings& SharedMemorySettings)
+{
+	FLearningAgentsSharedMemoryTrainerProcess TrainerProcess = SpawnSharedMemoryTrainingProcess(TrainerProcessSettings, SharedMemorySettings);
+	return MakeSharedMemoryCommunicator(TrainerProcess, TrainerProcessSettings, SharedMemorySettings);
+}
+
 FLearningAgentsSharedMemoryTrainerProcess ULearningAgentsCommunicatorLibrary::SpawnSharedMemoryTrainingProcess(
 	const FLearningAgentsTrainerProcessSettings& TrainerProcessSettings,
 	const FLearningAgentsSharedMemoryCommunicatorSettings& SharedMemorySettings)
 {
 	FLearningAgentsSharedMemoryTrainerProcess TrainerProcess;
 
+	if (FParse::Param(FCommandLine::Get(), TEXT("LearningAgentsUseExternalTrainingProcess")) || SharedMemorySettings.bUseExternalTrainingProcess)
+	{
+		UE_LOG(LogLearning, Display, TEXT("SpawnSharedMemoryTrainingProcess: Skipping spawning because bUseExternalTrainingProcess is True."));
+		return TrainerProcess;
+	}
+
 	if (PLATFORM_MAC)
 	{
 		UE_LOG(LogLearning, Error, TEXT("SpawnSharedMemoryTrainingProcess: Shared Memory not supported on Mac. Switch to Socket Communicator instead."));
@@ -49,7 +65,6 @@ FLearningAgentsSharedMemoryTrainerProcess ULearningAgentsCommunicatorLibrary::Sp
 		PythonExecutablePath,
 		PythonContentPath,
 		IntermediatePath,
-		1, // ProcessNum hard-coded to 1 for now
 		SharedMemorySettings.Timeout);
 
 	return TrainerProcess;
@@ -57,6 +72,7 @@ FLearningAgentsSharedMemoryTrainerProcess ULearningAgentsCommunicatorLibrary::Sp
 
 FLearningAgentsCommunicator ULearningAgentsCommunicatorLibrary::MakeSharedMemoryCommunicator(
 	FLearningAgentsSharedMemoryTrainerProcess TrainerProcess,
+	const FLearningAgentsTrainerProcessSettings& TrainerProcessSettings,
 	const FLearningAgentsSharedMemoryCommunicatorSettings& SharedMemorySettings)
 {
 	FLearningAgentsCommunicator Communicator;
@@ -67,27 +83,47 @@ FLearningAgentsCommunicator ULearningAgentsCommunicatorLibrary::MakeSharedMemory
 		return Communicator;
 	}
 
-	if (!TrainerProcess.TrainerProcess)
+	const FString IntermediatePath = UE::Learning::Trainer::GetIntermediatePath(TrainerProcessSettings.GetIntermediatePath());
+
+	FString CmdLineControlsGuidString;
+	FGuid CmdLineControlsGuid;
+	if (FParse::Value(FCommandLine::Get(), TEXT("LearningAgentsControlsGuid="), CmdLineControlsGuidString) && FGuid::Parse(CmdLineControlsGuidString, CmdLineControlsGuid))
 	{
-		UE_LOG(LogLearning, Error, TEXT("MakeSharedMemoryCommunicator: TrainerProcess is nullptr"));
-		return Communicator;
+		UE_LOG(LogLearning, Display, TEXT("Mapping Controls Guid override from the cmdline: %s"), *CmdLineControlsGuid.ToString(EGuidFormats::DigitsWithHyphensInBraces));
 	}
 
 	Communicator.Trainer = MakeShared<UE::Learning::FSharedMemoryTrainer>(
 		SharedMemorySettings.TaskName,
-		1, // ProcessNum hard-coded to 1 for now
 		TrainerProcess.TrainerProcess,
+		FParse::Param(FCommandLine::Get(), TEXT("LearningAgentsUseExternalTrainingProcess")) || SharedMemorySettings.bUseExternalTrainingProcess,
+		CmdLineControlsGuid.IsValid() ? CmdLineControlsGuid : SharedMemorySettings.ControlsGuid,
+		SharedMemorySettings.ConfigPath.Path,
+		IntermediatePath,
 		SharedMemorySettings.Timeout);
 
 	return Communicator;
 }
 
+FLearningAgentsCommunicator ULearningAgentsCommunicatorLibrary::MakeSocketTrainingProcess(
+	const FLearningAgentsTrainerProcessSettings& TrainerProcessSettings,
+	const FLearningAgentsSocketCommunicatorSettings& SocketSettings)
+{
+	FLearningAgentsSocketTrainerProcess TrainerProcess = SpawnSocketTrainingProcess(TrainerProcessSettings, SocketSettings);
+	return MakeSocketCommunicator(TrainerProcess, SocketSettings);
+}
+
 FLearningAgentsSocketTrainerProcess ULearningAgentsCommunicatorLibrary::SpawnSocketTrainingProcess(
 	const FLearningAgentsTrainerProcessSettings& TrainerProcessSettings,
 	const FLearningAgentsSocketCommunicatorSettings& SocketSettings)
 {
 	FLearningAgentsSocketTrainerProcess TrainerProcess;
 
+	if (FParse::Param(FCommandLine::Get(), TEXT("LearningAgentsUseExternalTrainingProcess")) || SocketSettings.bUseExternalTrainingProcess)
+	{
+		UE_LOG(LogLearning, Display, TEXT("SpawnSocketTrainingProcess: Skipping spawning because bUseExternalTrainingProcess is True."));
+		return TrainerProcess;
+	}
+
 	const FString PythonExecutablePath = UE::Learning::Trainer::GetPythonExecutablePath(TrainerProcessSettings.GetIntermediatePath());
 	if (!FPaths::FileExists(PythonExecutablePath))
 	{
@@ -130,12 +166,26 @@ FLearningAgentsCommunicator ULearningAgentsCommunicatorLibrary::MakeSocketCommun
 {
 	UE::Learning::ETrainerResponse Response = UE::Learning::ETrainerResponse::Success;
 
+	FString CmdLineIpAddress;
+	bool bIpAddressOverridden = FParse::Value(FCommandLine::Get(), TEXT("LearningAgentsIpAddress="), CmdLineIpAddress);
+	if (bIpAddressOverridden)
+	{
+		UE_LOG(LogLearning, Display, TEXT("Using Ip Address override from the cmdline: %s"), *CmdLineIpAddress);
+	}
+
+	uint32 CmdLinePort;
+	bool bPortOverridden = FParse::Value(FCommandLine::Get(), TEXT("LearningAgentsPort="), CmdLinePort);
+	if (bPortOverridden)
+	{
+		UE_LOG(LogLearning, Display, TEXT("Using Port override from the cmdline: %u"), CmdLinePort);
+	}
+
 	FLearningAgentsCommunicator Communicator;
 	Communicator.Trainer = MakeShared<UE::Learning::FSocketTrainer>(
 		Response,
 		TrainerProcess.TrainerProcess,
-		*SocketSettings.IpAddress,
-		SocketSettings.Port,
+		bIpAddressOverridden ? *CmdLineIpAddress : *SocketSettings.IpAddress,
+		bPortOverridden ? CmdLinePort : SocketSettings.Port,
 		SocketSettings.Timeout);
 
 	if (Response != UE::Learning::ETrainerResponse::Success)
diff --git a/Engine/Plugins/Experimental/LearningAgents/Source/LearningAgentsTraining/Public/LearningAgentsCommunicator.h b/Engine/Plugins/Experimental/LearningAgents/Source/LearningAgentsTraining/Public/LearningAgentsCommunicator.h
--- a/Engine/Plugins/Experimental/LearningAgents/Source/LearningAgentsTraining/Public/LearningAgentsCommunicator.h
+++ b/Engine/Plugins/Experimental/LearningAgents/Source/LearningAgentsTraining/Public/LearningAgentsCommunicator.h
@@ -36,6 +36,18 @@ public:
 	/** Time in seconds to wait for the training process before timing out. */
 	UPROPERTY(EditAnywhere, Category = "LearningAgents", meta = (ClampMin = "0.0", UIMin = "0.0"))
 	float Timeout = 10.0f;
+
+	/** If true, then we will attach to existing training process and skip spawning one. */
+	UPROPERTY(EditAnywhere, Category = "LearningAgents")
+	bool bUseExternalTrainingProcess = false;
+
+	/** The existing controls memory to attach to. */
+	UPROPERTY(EditAnywhere, Category = "LearningAgents", meta = (EditCondition = "bUseExternalTrainingProcess"))
+	FGuid ControlsGuid;
+
+	/** The absolute path to the config directory. */
+	UPROPERTY(EditAnywhere, Category = "LearningAgents", meta = (EditCondition = "bUseExternalTrainingProcess"))
+	FDirectoryPath ConfigPath;
 };
 
 /** Settings specific to socket communicators. */
@@ -57,6 +69,10 @@ public:
 	/** Time in seconds to wait for the training process before timing out. */
 	UPROPERTY(EditAnywhere, Category = "LearningAgents", meta = (ClampMin = "0.0", UIMin = "0.0"))
 	float Timeout = 10.0f;
+
+	/** If true, then we will attach to existing training process and skip spawning one. */
+	UPROPERTY(EditAnywhere, Category = "LearningAgents")
+	bool bUseExternalTrainingProcess = false;
 };
 
 /** Blueprint-compatible wrapper struct for Shared Memory Training Process. */
@@ -94,6 +110,29 @@ class LEARNINGAGENTSTRAINING_API ULearningAgentsCommunicatorLibrary : public UBl
 
 public:
 
+	/**
+	 * Start and connect to a training sub-process which will communicate via shared memory. Shared memory has the least
+	 * communication overhead so prefer this for local development.
+	 * 
+	 * This must be called on game thread!
+	 * 
+	 * You can use this in place of calling "SpawnSharedMemoryTrainingProcess" followed by "MakeSharedMemoryCommunicator"
+	 * for most use-cases. 
+	 *
+	 * If you wish to connect to an externally launched trainer (for debugging or scale-out reasons, etc.), you can use
+	 * the following command-line options to override some settings:
+	 *		-LearningAgentsUseExternalTrainingProcess	If present, spawning a sub-process will be skipped.
+	 *		-LearningAgentsControlsGuid={guid}			This optional guid will override the SharedMemorySettings, so you
+	 *													can attach to shared memory control created by the trainer.
+	 *
+	 * @param TrainerProcessSettings							Settings universal to all trainer processes.
+	 * @param FLearningAgentsSharedMemoryCommunicatorSettings	Settings specific to shared memory communicators.
+	 */
+	UFUNCTION(BlueprintCallable, Category = "LearningAgents", meta = (AutoCreateRefTerm = "TrainerProcessSettings,SharedMemorySettings"))
+	static FLearningAgentsCommunicator MakeSharedMemoryTrainingProcess(
+		const FLearningAgentsTrainerProcessSettings& TrainerProcessSettings = FLearningAgentsTrainerProcessSettings(),
+		const FLearningAgentsSharedMemoryCommunicatorSettings& SharedMemorySettings = FLearningAgentsSharedMemoryCommunicatorSettings());
+
 	/**
 	 * Start a local python training sub-process which will communicate via shared memory. Shared memory has the least
 	 * communication overhead so prefer this for local development.
@@ -118,8 +157,30 @@ public:
 	UFUNCTION(BlueprintCallable, Category = "LearningAgents", meta = (AutoCreateRefTerm = "SharedMemorySettings"))
 	static FLearningAgentsCommunicator MakeSharedMemoryCommunicator(
 		FLearningAgentsSharedMemoryTrainerProcess TrainerProcess,
+		const FLearningAgentsTrainerProcessSettings& TrainerProcessSettings = FLearningAgentsTrainerProcessSettings(),
 		const FLearningAgentsSharedMemoryCommunicatorSettings& SharedMemorySettings = FLearningAgentsSharedMemoryCommunicatorSettings());
 
+	/**
+	 * Start and connect to a training sub-process which will communicate via sockets. Sockets have some overhead
+	 * compared to shared memory but can work over networked connects. This provides no encryption so do not use on
+	 * public internet if privacy is a concern.
+	 *
+	 * This must be called on game thread!
+	 * 
+	 * If you wish to connect to an externally launched trainer (for debugging or scale-out reasons, etc.), you can use
+	 * the following command-line options to override some settings:
+	 *		-LearningAgentsUseExternalTrainingProcess	If present, spawning a sub-process will be skipped.
+	 *		-LearningAgentsIpAddress=					This optional IP address will override the SocketSettings
+	 *		-LearningAgentsPort=						This optional Port will override the SocketSettings
+	 *
+	 * @param TrainerProcessSettings					Settings universal to all trainer processes.
+	 * @param FLearningAgentsSocketCommunicatorSettings	Settings specific to socket communicators.
+	 */
+	UFUNCTION(BlueprintCallable, Category = "LearningAgents", meta = (AutoCreateRefTerm = "TrainerProcessSettings,SocketSettings"))
+	static FLearningAgentsCommunicator MakeSocketTrainingProcess(
+		const FLearningAgentsTrainerProcessSettings& TrainerProcessSettings = FLearningAgentsTrainerProcessSettings(),
+		const FLearningAgentsSocketCommunicatorSettings& SocketSettings = FLearningAgentsSocketCommunicatorSettings());
+
 	/**
 	 * Start a local python training sub-process which will communicate via sockets. Sockets have some overhead
 	 * compared to shared memory but can work over networked connects. This provides no encryption so do not use on
diff --git a/Engine/Plugins/Experimental/LearningAgents/Source/LearningAgentsTraining/Public/LearningAgentsPPOTrainer.h b/Engine/Plugins/Experimental/LearningAgents/Source/LearningAgentsTraining/Public/LearningAgentsPPOTrainer.h
--- a/Engine/Plugins/Experimental/LearningAgents/Source/LearningAgentsTraining/Public/LearningAgentsPPOTrainer.h
+++ b/Engine/Plugins/Experimental/LearningAgents/Source/LearningAgentsTraining/Public/LearningAgentsPPOTrainer.h
@@ -324,7 +324,7 @@ public:
 	 * @param TrainingGameSettings		The settings that will affect the game's simulation.
 	 * @param bResetAgentsOnBegin		If true, reset all agents at the beginning of training.
 	 */
-	UFUNCTION(BlueprintCallable, Category = "LearningAgents", meta = (AutoCreateRefTerm = "TrainerTrainingSettings,TrainingGameSettings,TrainerPathSettings"))
+	UFUNCTION(BlueprintCallable, Category = "LearningAgents", meta = (AutoCreateRefTerm = "TrainerTrainingSettings,TrainingGameSettings"))
 	void BeginTraining(
 		const FLearningAgentsPPOTrainingSettings& TrainerTrainingSettings = FLearningAgentsPPOTrainingSettings(),
 		const FLearningAgentsTrainingGameSettings& TrainingGameSettings = FLearningAgentsTrainingGameSettings(),
@@ -355,7 +355,7 @@ public:
 	 * @param bResetAgentsOnBegin				If true, reset all agents at the beginning of training.
 	 * @param bResetAgentsOnUpdate				If true, reset all agents whenever an updated policy is received.
 	 */
-	UFUNCTION(BlueprintCallable, Category = "LearningAgents", meta = (AutoCreateRefTerm = "TrainerTrainingSettings,TrainingGameSettings,TrainerPathSettings"))
+	UFUNCTION(BlueprintCallable, Category = "LearningAgents", meta = (AutoCreateRefTerm = "TrainerTrainingSettings,TrainingGameSettings"))
 	void RunTraining(
 		const FLearningAgentsPPOTrainingSettings& TrainerTrainingSettings = FLearningAgentsPPOTrainingSettings(),
 		const FLearningAgentsTrainingGameSettings& TrainingGameSettings = FLearningAgentsTrainingGameSettings(),
diff --git a/Engine/Plugins/Experimental/LearningAgents/Source/LearningTraining/Private/LearningExternalTrainer.cpp b/Engine/Plugins/Experimental/LearningAgents/Source/LearningTraining/Private/LearningExternalTrainer.cpp
--- a/Engine/Plugins/Experimental/LearningAgents/Source/LearningTraining/Private/LearningExternalTrainer.cpp
+++ b/Engine/Plugins/Experimental/LearningAgents/Source/LearningTraining/Private/LearningExternalTrainer.cpp
@@ -29,65 +29,45 @@ namespace UE::Learning
 		const FString& PythonExecutablePath,
 		const FString& PythonContentPath,
 		const FString& InIntermediatePath,
-		const int32 ProcessNum,
 		const float InTimeout,
 		const ESubprocessFlags TrainingProcessFlags,
 		const ELogSetting LogSettings)
 	{
-		check(ProcessNum > 0);
+		UE_LEARNING_CHECK(FPaths::FileExists(PythonExecutablePath));
+		UE_LEARNING_CHECK(FPaths::DirectoryExists(PythonContentPath));
 
-		int32 ProcessIdx = 0;
-		FParse::Value(FCommandLine::Get(), TEXT("LearningProcessIdx"), ProcessIdx);
-		
 		Timeout = InTimeout;
 		IntermediatePath = InIntermediatePath;
 
-		if (ProcessIdx == 0)
-		{
-			// Allocate the control memory if we are the parent UE process
-			Controls = SharedMemory::Allocate<2, volatile int32>({ ProcessNum, SharedMemoryTraining::GetControlNum()});
-		}
-		else
-		{
-			FGuid ControlsGuid; ensure(FParse::Value(FCommandLine::Get(), TEXT("LearningControlsGuid"), ControlsGuid));
-			Controls = SharedMemory::Map<2, volatile int32>(ControlsGuid, { ProcessNum, SharedMemoryTraining::GetControlNum()});
-
-			// We do not want to launch another training process if we are a child process
-			return;
-		}
-
-		UE_LEARNING_CHECK(FPaths::FileExists(PythonExecutablePath));
-		UE_LEARNING_CHECK(FPaths::DirectoryExists(PythonContentPath));
+		Controls = SharedMemory::Allocate<1, volatile int32>({ SharedMemoryTraining::GetControlNum() });
 
 		// We need to zero the control memory before we start the training sub-process since it may contain
 		// uninitialized values or those left over from previous runs.
 		Array::Zero(Controls.View);
 
 		// Set the ID columns to -1
-		for (int32 Index = 0; Index < Controls.View.Num<0>(); Index++)
-		{
-			Controls.View[Index][(uint8)SharedMemoryTraining::EControls::NetworkId] = -1;
-			Controls.View[Index][(uint8)SharedMemoryTraining::EControls::ReplayBufferId] = -1;
-		}
+		Controls.View[(uint8)SharedMemoryTraining::EControls::NetworkId] = -1;
+		Controls.View[(uint8)SharedMemoryTraining::EControls::ReplayBufferId] = -1;
 
 		const FString TimeStamp = FDateTime::Now().ToFormattedString(TEXT("%Y-%m-%d_%H-%M-%S"));
 		const FString TrainerType = TEXT("SharedMemory");
-		ConfigPath = InIntermediatePath / TEXT("Configs") / FString::Printf(TEXT("%s_%s_%s_%s.json"), *TaskName, *TrainerFileName, *TrainerType, *TimeStamp);
+		ConfigPath = InIntermediatePath / TEXT("Configs") / FString::Printf(TEXT("%s.json"), *Controls.Guid.ToString(EGuidFormats::DigitsWithHyphensInBraces));
 
 		IFileManager& FileManager = IFileManager::Get();
-		const FString CommandLineArguments = FString::Printf(TEXT("\"%s\" -p \"%s\" -m \"%s\" \"%s\" SharedMemory \"%s\" %i \"%s\""),
+		const FString CommandLineArguments = FString::Printf(TEXT("\"%s\" -p \"%s\" -m \"%s\" \"%s\" SharedMemory \"%s\" -g \"%s\""),
 			*FileManager.ConvertToAbsolutePathForExternalAppForRead(*(PythonContentPath / TEXT("train.py"))),
 			*FileManager.ConvertToAbsolutePathForExternalAppForRead(*CustomTrainerPath),
 			*TrainerFileName,
 			LogSettings == ELogSetting::Normal ? TEXT("-l") : TEXT(""),
-			*Controls.Guid.ToString(EGuidFormats::DigitsWithHyphensInBraces),
-			ProcessNum,
-			*FileManager.ConvertToAbsolutePathForExternalAppForRead(*ConfigPath));
+			*FileManager.ConvertToAbsolutePathForExternalAppForRead(*ConfigPath),
+			*Controls.Guid.ToString(EGuidFormats::DigitsWithHyphensInBraces));
 
-		TrainingProcess.Launch(
+		bool bLaunched = TrainingProcess.Launch(
 			FileManager.ConvertToAbsolutePathForExternalAppForRead(*PythonExecutablePath),
 			CommandLineArguments, 
 			TrainingProcessFlags);
+
+		ensure(bLaunched);
 	}
 
 	bool FSharedMemoryTrainerServerProcess::IsRunning() const
@@ -119,7 +99,7 @@ namespace UE::Learning
 		TrainingProcess.Terminate();
 	}
 
-	TSharedMemoryArrayView<2, volatile int32> FSharedMemoryTrainerServerProcess::GetControlsSharedMemoryArrayView() const
+	TSharedMemoryArrayView<1, volatile int32> FSharedMemoryTrainerServerProcess::GetControlsSharedMemoryArrayView() const
 	{
 		return Controls;
 	}
@@ -155,37 +135,37 @@ namespace UE::Learning
 			SharedMemory::Deallocate(EpisodeLengths);
 			SharedMemory::Deallocate(EpisodeCompletionModes);
 
-			for(TSharedMemoryArrayView<3, float>& SharedMemoryArrayView : EpisodeFinalObservations)
+			for(TSharedMemoryArrayView<2, float>& SharedMemoryArrayView : EpisodeFinalObservations)
 			{
 				SharedMemory::Deallocate(SharedMemoryArrayView);
 			}
 
-			for (TSharedMemoryArrayView<3, float>& SharedMemoryArrayView : EpisodeFinalMemoryStates)
+			for (TSharedMemoryArrayView<2, float>& SharedMemoryArrayView : EpisodeFinalMemoryStates)
 			{
 				SharedMemory::Deallocate(SharedMemoryArrayView);
 			}
 
-			for (TSharedMemoryArrayView<3, float>& SharedMemoryArrayView : Observations)
+			for (TSharedMemoryArrayView<2, float>& SharedMemoryArrayView : Observations)
 			{
 				SharedMemory::Deallocate(SharedMemoryArrayView);
 			}
 
-			for (TSharedMemoryArrayView<3, float>& SharedMemoryArrayView : Actions)
+			for (TSharedMemoryArrayView<2, float>& SharedMemoryArrayView : Actions)
 			{
 				SharedMemory::Deallocate(SharedMemoryArrayView);
 			}
 
-			for (TSharedMemoryArrayView<3, float>& SharedMemoryArrayView : ActionModifiers)
+			for (TSharedMemoryArrayView<2, float>& SharedMemoryArrayView : ActionModifiers)
 			{
 				SharedMemory::Deallocate(SharedMemoryArrayView);
 			}
 
-			for (TSharedMemoryArrayView<3, float>& SharedMemoryArrayView : MemoryStates)
+			for (TSharedMemoryArrayView<2, float>& SharedMemoryArrayView : MemoryStates)
 			{
 				SharedMemory::Deallocate(SharedMemoryArrayView);
 			}
 
-			for (TSharedMemoryArrayView<3, float>& SharedMemoryArrayView : Rewards)
+			for (TSharedMemoryArrayView<2, float>& SharedMemoryArrayView : Rewards)
 			{
 				SharedMemory::Deallocate(SharedMemoryArrayView);
 			}
@@ -194,28 +174,43 @@ namespace UE::Learning
 
 	FSharedMemoryTrainer::FSharedMemoryTrainer(
 		const FString& InTaskName,
-		const int32 InProcessNum,
 		const TSharedPtr<FSharedMemoryTrainerServerProcess>& ExternalTrainerProcess,
+		const bool bUseExternalTrainingProcess,
+		const FGuid ControlsGuid,
+		const FString& InConfigPath,
+		const FString& InIntermediatePath,
 		const float InTimeout)
 	{
-		if (!ExternalTrainerProcess)
+		TaskName = InTaskName;
+		Timeout = InTimeout;
+
+		if (ExternalTrainerProcess)
+		{
+			TrainingProcess = ExternalTrainerProcess;
+			ConfigPath = ExternalTrainerProcess->GetConfigPath();
+			IntermediatePath = ExternalTrainerProcess->GetIntermediatePath();
+			Controls = ExternalTrainerProcess->GetControlsSharedMemoryArrayView();	
+		}
+		else if (!bUseExternalTrainingProcess)
 		{
-			UE_LOG(LogLearning, Error, TEXT("FSharedMemoryTrainer ctor: Trainer process is nullptr. Is it not a shared memory process?"));
+			UE_LOG(LogLearning, Error, TEXT("ExternalTrainerProcess was null, but attach to existing trainer process was false. Either pass in external trainer process or set bUseExternalTrainingProcess to true in the settings (along with associated settings)."))
 			return;
 		}
+		else
+		{
+			Controls = SharedMemory::Map<1, volatile int32>(ControlsGuid, { SharedMemoryTraining::GetControlNum() });
 
-		check(InProcessNum > 0);
+			// We need to zero the control memory before we start the training sub-process since it may contain
+			// uninitialized values or those left over from previous runs.
+			Array::Zero(Controls.View);
 
-		TaskName = InTaskName;
-		ConfigPath = ExternalTrainerProcess->GetConfigPath();
-		IntermediatePath = ExternalTrainerProcess->GetIntermediatePath();
-		TrainingProcess = ExternalTrainerProcess->GetTrainingSubprocess();
-		ProcessNum = InProcessNum;
-		Controls = ExternalTrainerProcess->GetControlsSharedMemoryArrayView();
-		Timeout = InTimeout;
+			// Set the ID columns to -1
+			Controls.View[(uint8)SharedMemoryTraining::EControls::NetworkId] = -1;
+			Controls.View[(uint8)SharedMemoryTraining::EControls::ReplayBufferId] = -1;
 
-		ProcessIdx = 0;
-		FParse::Value(FCommandLine::Get(), TEXT("LearningProcessIdx"), ProcessIdx);
+			ConfigPath = InConfigPath / FString::Printf(TEXT("%s.json"), *ControlsGuid.ToString(EGuidFormats::DigitsWithHyphensInBraces));
+			IntermediatePath = InIntermediatePath;
+		}
 	}
 
 	FSharedMemoryTrainer::~FSharedMemoryTrainer()
@@ -230,7 +225,7 @@ namespace UE::Learning
 
 	bool FSharedMemoryTrainer::HasNetworkOrCompleted()
 	{
-		return SharedMemoryTraining::HasNetworkOrCompleted(Controls.View[ProcessIdx]);
+		return SharedMemoryTraining::HasNetworkOrCompleted(Controls.View);
 	}
 
 	void FSharedMemoryTrainer::Terminate()
@@ -240,29 +235,17 @@ namespace UE::Learning
 
 	ETrainerResponse FSharedMemoryTrainer::SendStop()
 	{
-		check(ProcessIdx != INDEX_NONE);
 		checkf(Controls.Region, TEXT("SendStop: Controls Shared Memory Region is nullptr"));
-
-		return SharedMemoryTraining::SendStop(Controls.View[ProcessIdx]);
+		return SharedMemoryTraining::SendStop(Controls.View);
 	}
 
 	ETrainerResponse FSharedMemoryTrainer::SendConfig(const TSharedRef<FJsonObject>& ConfigObject, const ELogSetting LogSettings)
 	{
-		check(ProcessNum > 0);
-
-		if (ProcessIdx != 0)
-		{
-			// Only the parent process will send the config
-			return ETrainerResponse::Success;
-		}
-
 		IFileManager& FileManager = IFileManager::Get();
 		ConfigObject->SetStringField(TEXT("IntermediatePath"), *FileManager.ConvertToAbsolutePathForExternalAppForRead(*IntermediatePath));
 
 		TSharedPtr<FJsonObject> SharedMemoryObject = MakeShared<FJsonObject>();
 
-		SharedMemoryObject->SetNumberField(TEXT("ProcessNum"), ProcessNum);
-
 		TArray<TSharedPtr<FJsonValue>> NetworkGuidsArray;
 		for(int32 Index = 0; Index < NeuralNetworkSharedMemoryArrayViews.Num(); Index++)
 		{
@@ -284,7 +267,7 @@ namespace UE::Learning
 			JsonObject->SetStringField(TEXT("EpisodeCompletionModesGuid"), *SharedMemoryExperienceContainer.EpisodeCompletionModes.Guid.ToString(EGuidFormats::DigitsWithHyphensInBraces));
 
 			TArray<TSharedPtr<FJsonValue>> EpisodeFinalObservationsGuidsArray;
-			for (const TSharedMemoryArrayView<3, float>& EpisodeFinalObservations : SharedMemoryExperienceContainer.EpisodeFinalObservations)
+			for (const TSharedMemoryArrayView<2, float>& EpisodeFinalObservations : SharedMemoryExperienceContainer.EpisodeFinalObservations)
 			{
 				EpisodeFinalObservationsGuidsArray.Add(MakeShared<FJsonValueString>(*EpisodeFinalObservations.Guid.ToString(EGuidFormats::DigitsWithHyphensInBraces)));
 				
@@ -292,7 +275,7 @@ namespace UE::Learning
 			JsonObject->SetArrayField(TEXT("EpisodeFinalObservationsGuids"), EpisodeFinalObservationsGuidsArray);
 
 			TArray<TSharedPtr<FJsonValue>> EpisodeFinalMemoryStatesGuidsArray;
-			for (const TSharedMemoryArrayView<3, float>& EpisodeFinalMemoryStates : SharedMemoryExperienceContainer.EpisodeFinalMemoryStates)
+			for (const TSharedMemoryArrayView<2, float>& EpisodeFinalMemoryStates : SharedMemoryExperienceContainer.EpisodeFinalMemoryStates)
 			{
 				EpisodeFinalMemoryStatesGuidsArray.Add(MakeShared<FJsonValueString>(*EpisodeFinalMemoryStates.Guid.ToString(EGuidFormats::DigitsWithHyphensInBraces)));
 
@@ -300,7 +283,7 @@ namespace UE::Learning
 			JsonObject->SetArrayField(TEXT("EpisodeFinalMemoryStatesGuids"), EpisodeFinalMemoryStatesGuidsArray);
 
 			TArray<TSharedPtr<FJsonValue>> ObservationsGuidsArray;
-			for (const TSharedMemoryArrayView<3, float>& Observations : SharedMemoryExperienceContainer.Observations)
+			for (const TSharedMemoryArrayView<2, float>& Observations : SharedMemoryExperienceContainer.Observations)
 			{
 				ObservationsGuidsArray.Add(MakeShared<FJsonValueString>(*Observations.Guid.ToString(EGuidFormats::DigitsWithHyphensInBraces)));
 
@@ -308,7 +291,7 @@ namespace UE::Learning
 			JsonObject->SetArrayField(TEXT("ObservationsGuids"), ObservationsGuidsArray);
 
 			TArray<TSharedPtr<FJsonValue>> ActionsGuidsArray;
-			for (const TSharedMemoryArrayView<3, float>& Actions : SharedMemoryExperienceContainer.Actions)
+			for (const TSharedMemoryArrayView<2, float>& Actions : SharedMemoryExperienceContainer.Actions)
 			{
 				ActionsGuidsArray.Add(MakeShared<FJsonValueString>(*Actions.Guid.ToString(EGuidFormats::DigitsWithHyphensInBraces)));
 
@@ -316,7 +299,7 @@ namespace UE::Learning
 			JsonObject->SetArrayField(TEXT("ActionsGuids"), ActionsGuidsArray);
 
 			TArray<TSharedPtr<FJsonValue>> ActionModifiersGuidsArray;
-			for (const TSharedMemoryArrayView<3, float>& ActionModifiers : SharedMemoryExperienceContainer.ActionModifiers)
+			for (const TSharedMemoryArrayView<2, float>& ActionModifiers : SharedMemoryExperienceContainer.ActionModifiers)
 			{
 				ActionModifiersGuidsArray.Add(MakeShared<FJsonValueString>(*ActionModifiers.Guid.ToString(EGuidFormats::DigitsWithHyphensInBraces)));
 
@@ -325,7 +308,7 @@ namespace UE::Learning
 
 
 			TArray<TSharedPtr<FJsonValue>> MemoryStatesGuidsArray;
-			for (const TSharedMemoryArrayView<3, float>& MemoryStates : SharedMemoryExperienceContainer.MemoryStates)
+			for (const TSharedMemoryArrayView<2, float>& MemoryStates : SharedMemoryExperienceContainer.MemoryStates)
 			{
 				MemoryStatesGuidsArray.Add(MakeShared<FJsonValueString>(*MemoryStates.Guid.ToString(EGuidFormats::DigitsWithHyphensInBraces)));
 
@@ -333,7 +316,7 @@ namespace UE::Learning
 			JsonObject->SetArrayField(TEXT("MemoryStatesGuids"), MemoryStatesGuidsArray);
 
 			TArray<TSharedPtr<FJsonValue>> RewardsGuidsArray;
-			for (const TSharedMemoryArrayView<3, float>& Rewards : SharedMemoryExperienceContainer.Rewards)
+			for (const TSharedMemoryArrayView<2, float>& Rewards : SharedMemoryExperienceContainer.Rewards)
 			{
 				RewardsGuidsArray.Add(MakeShared<FJsonValueString>(*Rewards.Guid.ToString(EGuidFormats::DigitsWithHyphensInBraces)));
 
@@ -352,7 +335,9 @@ namespace UE::Learning
 		FJsonSerializer::Serialize(ConfigObject, JsonWriter, true);
 		FFileHelper::SaveStringToFile(ConfigString, *ConfigPath);
 
-		return SharedMemoryTraining::SendConfigSignal(Controls.View[ProcessIdx], LogSettings);
+		UE_LOG(LogLearning, Display, TEXT("Wrote Config File to %s. Sending Config Signal..."), *ConfigPath);
+
+		return SharedMemoryTraining::SendConfigSignal(Controls.View, LogSettings);
 	}
 
 	int32 FSharedMemoryTrainer::AddNetwork(const ULearningNeuralNetworkData& Network)
@@ -368,7 +353,6 @@ namespace UE::Learning
 		FRWLock* NetworkLock,
 		const ELogSetting LogSettings)
 	{
-		check(ProcessIdx != INDEX_NONE);
 		checkf(Controls.Region, TEXT("ReceiveNetwork: Controls Shared Memory Region is nullptr"));
 		if (!ensureMsgf(NeuralNetworkSharedMemoryArrayViews.Num() >= NetworkId, TEXT("Network %d has not been added. Call AddNetwork prior to ReceiveNetwork."), NetworkId))
 		{
@@ -376,10 +360,10 @@ namespace UE::Learning
 		}
 
 		return SharedMemoryTraining::RecvNetwork(
-			Controls.View[ProcessIdx],
+			Controls.View,
 			NetworkId,
 			OutNetwork,
-			*TrainingProcess,
+			TrainingProcess.IsValid() ? TrainingProcess->GetTrainingSubprocess() : nullptr,
 			NeuralNetworkSharedMemoryArrayViews[NetworkId].View,
 			Timeout,
 			NetworkLock,
@@ -392,7 +376,6 @@ namespace UE::Learning
 		FRWLock* NetworkLock,
 		const ELogSetting LogSettings)
 	{
-		check(ProcessIdx != INDEX_NONE);
 		checkf(Controls.Region, TEXT("SendNetwork: Controls Shared Memory Region is nullptr"));
 		if (!ensureMsgf(NeuralNetworkSharedMemoryArrayViews.Num() >= NetworkId, TEXT("Network %d has not been added. Call AddNetwork prior to SendNetwork."), NetworkId))
 		{
@@ -400,10 +383,10 @@ namespace UE::Learning
 		}
 
 		return SharedMemoryTraining::SendNetwork(
-			Controls.View[ProcessIdx],
+			Controls.View,
 			NetworkId,
 			NeuralNetworkSharedMemoryArrayViews[NetworkId].View,
-			*TrainingProcess,
+			TrainingProcess.IsValid() ? TrainingProcess->GetTrainingSubprocess() : nullptr,
 			Network,
 			Timeout,
 			NetworkLock,
@@ -412,208 +395,62 @@ namespace UE::Learning
 
 	int32 FSharedMemoryTrainer::AddReplayBuffer(const FReplayBuffer& ReplayBuffer)
 	{
-		check(ProcessNum > 0);
-
 		FSharedMemoryExperienceContainer ExperienceContainer;
-		if (ProcessIdx == 0)
-		{
-			ExperienceContainer.EpisodeStarts = SharedMemory::Allocate<2, int32>({ ProcessNum, ReplayBuffer.GetMaxEpisodeNum() });
-			ExperienceContainer.EpisodeLengths = SharedMemory::Allocate<2, int32>({ ProcessNum, ReplayBuffer.GetMaxEpisodeNum() });
-
-			if (ReplayBuffer.HasCompletions())
-			{
-				ExperienceContainer.EpisodeCompletionModes = SharedMemory::Allocate<2, ECompletionMode>({ ProcessNum, ReplayBuffer.GetMaxEpisodeNum() });
-			}
 
-			if (ReplayBuffer.HasFinalObservations())
-			{
-				for (int32 Index = 0; Index < ReplayBuffer.GetObservationsNum(); Index++)
-				{
-					const int32 DimNum = ReplayBuffer.GetEpisodeFinalObservations(Index).Num<1>();
-					ExperienceContainer.EpisodeFinalObservations.Add(SharedMemory::Allocate<3, float>({ ProcessNum, ReplayBuffer.GetMaxEpisodeNum(), DimNum }));
-				}
-			}
+		ExperienceContainer.EpisodeStarts = SharedMemory::Allocate<1, int32>({ ReplayBuffer.GetMaxEpisodeNum() });
+		ExperienceContainer.EpisodeLengths = SharedMemory::Allocate<1, int32>({ ReplayBuffer.GetMaxEpisodeNum() });
 
-			if (ReplayBuffer.HasFinalMemoryStates())
-			{
-				for (int32 Index = 0; Index < ReplayBuffer.GetMemoryStatesNum(); Index++)
-				{
-					const int32 DimNum = ReplayBuffer.GetEpisodeFinalMemoryStates(Index).Num<1>();
-					ExperienceContainer.EpisodeFinalMemoryStates.Add(SharedMemory::Allocate<3, float>({ ProcessNum, ReplayBuffer.GetMaxEpisodeNum(), DimNum }));
-				}
-			}
+		if (ReplayBuffer.HasCompletions())
+		{
+			ExperienceContainer.EpisodeCompletionModes = SharedMemory::Allocate<1, ECompletionMode>({ ReplayBuffer.GetMaxEpisodeNum() });
+		}
 
+		if (ReplayBuffer.HasFinalObservations())
+		{
 			for (int32 Index = 0; Index < ReplayBuffer.GetObservationsNum(); Index++)
 			{
-				const int32 DimNum = ReplayBuffer.GetObservations(Index).Num<1>();
-				ExperienceContainer.Observations.Add(SharedMemory::Allocate<3, float>({ ProcessNum, ReplayBuffer.GetMaxStepNum(), DimNum }));
-			}
-
-			for (int32 Index = 0; Index < ReplayBuffer.GetActionsNum(); Index++)
-			{
-				const int32 DimNum = ReplayBuffer.GetActions(Index).Num<1>();
-				ExperienceContainer.Actions.Add(SharedMemory::Allocate<3, float>({ ProcessNum, ReplayBuffer.GetMaxStepNum(), DimNum }));
-			}
-
-			for (int32 Index = 0; Index < ReplayBuffer.GetActionModifiersNum(); Index++)
-			{
-				const int32 DimNum = ReplayBuffer.GetActionModifiers(Index).Num<1>();
-				ExperienceContainer.ActionModifiers.Add(SharedMemory::Allocate<3, float>({ ProcessNum, ReplayBuffer.GetMaxStepNum(), DimNum }));
+				const int32 DimNum = ReplayBuffer.GetEpisodeFinalObservations(Index).Num<1>();
+				ExperienceContainer.EpisodeFinalObservations.Add(SharedMemory::Allocate<2, float>({ ReplayBuffer.GetMaxEpisodeNum(), DimNum }));
 			}
+		}
 
+		if (ReplayBuffer.HasFinalMemoryStates())
+		{
 			for (int32 Index = 0; Index < ReplayBuffer.GetMemoryStatesNum(); Index++)
 			{
-				const int32 DimNum = ReplayBuffer.GetMemoryStates(Index).Num<1>();
-				ExperienceContainer.MemoryStates.Add(SharedMemory::Allocate<3, float>({ ProcessNum, ReplayBuffer.GetMaxStepNum(), DimNum }));
-			}
-
-			for (int32 Index = 0; Index < ReplayBuffer.GetRewardsNum(); Index++)
-			{
-				const int32 DimNum = ReplayBuffer.GetRewards(Index).Num<1>();
-				ExperienceContainer.Rewards.Add(SharedMemory::Allocate<3, float>({ ProcessNum, ReplayBuffer.GetMaxStepNum(), DimNum }));
+				const int32 DimNum = ReplayBuffer.GetEpisodeFinalMemoryStates(Index).Num<1>();
+				ExperienceContainer.EpisodeFinalMemoryStates.Add(SharedMemory::Allocate<2, float>({ ReplayBuffer.GetMaxEpisodeNum(), DimNum }));
 			}
 		}
-		else
-		{
-			FGuid EpisodeStartsGuid; ensure(FParse::Value(FCommandLine::Get(), TEXT("LearningEpisodeStartsGuid"), EpisodeStartsGuid));
-			ExperienceContainer.EpisodeStarts = SharedMemory::Map<2, int32>(EpisodeStartsGuid, { ProcessNum, ReplayBuffer.GetMaxEpisodeNum() });
-
-			FGuid EpisodeLengthsGuid; ensure(FParse::Value(FCommandLine::Get(), TEXT("LearningEpisodeLengthsGuid"), EpisodeLengthsGuid));
-			ExperienceContainer.EpisodeLengths = SharedMemory::Map<2, int32>(EpisodeLengthsGuid, { ProcessNum, ReplayBuffer.GetMaxEpisodeNum() });
-
-			if (ReplayBuffer.HasCompletions())
-			{
-				FGuid EpisodeCompletionModesGuid; ensure(FParse::Value(FCommandLine::Get(), TEXT("LearningEpisodeCompletionModesGuid"), EpisodeCompletionModesGuid));
-				ExperienceContainer.EpisodeCompletionModes = SharedMemory::Map<2, ECompletionMode>(EpisodeCompletionModesGuid, { ProcessNum, ReplayBuffer.GetMaxEpisodeNum() });
-			}
-
-			// Final Observations
-			if (ReplayBuffer.HasFinalObservations())
-			{
-				FString StringOfGuids; ensure(FParse::Value(FCommandLine::Get(), TEXT("LearningEpisodeFinalObservationsGuids"), StringOfGuids));
-				TArray<FString> StringGuids;
-				StringOfGuids.ParseIntoArray(StringGuids, TEXT(","));
-				TArray<FGuid> Guids;
-				for (const FString& StringGuid : StringGuids)
-				{
-					FGuid Guid;
-					FGuid::Parse(StringGuid, Guid);
-					Guids.Add(Guid);
-				}
-				check(Guids.Num() == ReplayBuffer.GetObservationsNum());
-
-				for (int32 Index = 0; Index < ReplayBuffer.GetObservationsNum(); Index++)
-				{
-					const int32 DimNum = ReplayBuffer.GetEpisodeFinalObservations(Index).Num<1>();
-					ExperienceContainer.EpisodeFinalObservations.Add(SharedMemory::Map<3, float>(Guids[Index], { ProcessNum, ReplayBuffer.GetMaxEpisodeNum(), DimNum }));
-				}
-			}
 
-			// Final Memory States
-			if (ReplayBuffer.HasFinalMemoryStates())
-			{
-				FString StringOfGuids; ensure(FParse::Value(FCommandLine::Get(), TEXT("LearningEpisodeFinalMemoryStatesGuid"), StringOfGuids));
-				TArray<FString> StringGuids;
-				StringOfGuids.ParseIntoArray(StringGuids, TEXT(","));
-				TArray<FGuid> Guids;
-				for (const FString& StringGuid : StringGuids)
-				{
-					FGuid Guid;
-					FGuid::Parse(StringGuid, Guid);
-					Guids.Add(Guid);
-				}
-				check(Guids.Num() == ReplayBuffer.GetMemoryStatesNum());
-
-				for (int32 Index = 0; Index < ReplayBuffer.GetMemoryStatesNum(); Index++)
-				{
-					const int32 DimNum = ReplayBuffer.GetEpisodeFinalMemoryStates(Index).Num<1>();
-					ExperienceContainer.EpisodeFinalMemoryStates.Add(SharedMemory::Map<3, float>(Guids[Index], { ProcessNum, ReplayBuffer.GetMaxEpisodeNum(), DimNum }));
-				}
-			}
+		for (int32 Index = 0; Index < ReplayBuffer.GetObservationsNum(); Index++)
+		{
+			const int32 DimNum = ReplayBuffer.GetObservations(Index).Num<1>();
+			ExperienceContainer.Observations.Add(SharedMemory::Allocate<2, float>({ ReplayBuffer.GetMaxStepNum(), DimNum }));
+		}
 
-			// Observations
-			{
-				FString StringOfGuids; ensure(FParse::Value(FCommandLine::Get(), TEXT("LearningObservationsGuid"), StringOfGuids));
-				TArray<FString> StringGuids;
-				StringOfGuids.ParseIntoArray(StringGuids, TEXT(","));
-				TArray<FGuid> Guids;
-				for (const FString& StringGuid : StringGuids)
-				{
-					FGuid Guid;
-					FGuid::Parse(StringGuid, Guid);
-					Guids.Add(Guid);
-				}
-				check(Guids.Num() == ReplayBuffer.GetObservationsNum());
-
-				for (int32 Index = 0; Index < ReplayBuffer.GetObservationsNum(); Index++)
-				{
-					const int32 DimNum = ReplayBuffer.GetObservations(Index).Num<1>();
-					ExperienceContainer.Observations.Add(SharedMemory::Map<3, float>(Guids[Index], { ProcessNum, ReplayBuffer.GetMaxStepNum(), DimNum }));
-				}
-			}
+		for (int32 Index = 0; Index < ReplayBuffer.GetActionsNum(); Index++)
+		{
+			const int32 DimNum = ReplayBuffer.GetActions(Index).Num<1>();
+			ExperienceContainer.Actions.Add(SharedMemory::Allocate<2, float>({ ReplayBuffer.GetMaxStepNum(), DimNum }));
+		}
 
-			// Actions
-			{
-				FString StringOfGuids; ensure(FParse::Value(FCommandLine::Get(), TEXT("LearningActionsGuid"), StringOfGuids));
-				TArray<FString> StringGuids;
-				StringOfGuids.ParseIntoArray(StringGuids, TEXT(","));
-				TArray<FGuid> Guids;
-				for (const FString& StringGuid : StringGuids)
-				{
-					FGuid Guid;
-					FGuid::Parse(StringGuid, Guid);
-					Guids.Add(Guid);
-				}
-				check(Guids.Num() == ReplayBuffer.GetActionsNum());
-
-				for (int32 Index = 0; Index < ReplayBuffer.GetActionsNum(); Index++)
-				{
-					const int32 DimNum = ReplayBuffer.GetActions(Index).Num<1>();
-					ExperienceContainer.Actions.Add(SharedMemory::Map<3, float>(Guids[Index], { ProcessNum, ReplayBuffer.GetMaxStepNum(), DimNum }));
-				}
-			}
+		for (int32 Index = 0; Index < ReplayBuffer.GetActionModifiersNum(); Index++)
+		{
+			const int32 DimNum = ReplayBuffer.GetActionModifiers(Index).Num<1>();
+			ExperienceContainer.ActionModifiers.Add(SharedMemory::Allocate<2, float>({ ReplayBuffer.GetMaxStepNum(), DimNum }));
+		}
 
-			// Memory States
-			{
-				FString StringOfGuids; ensure(FParse::Value(FCommandLine::Get(), TEXT("LearningMemoryStatesGuid"), StringOfGuids));
-				TArray<FString> StringGuids;
-				StringOfGuids.ParseIntoArray(StringGuids, TEXT(","));
-				TArray<FGuid> Guids;
-				for (const FString& StringGuid : StringGuids)
-				{
-					FGuid Guid;
-					FGuid::Parse(StringGuid, Guid);
-					Guids.Add(Guid);
-				}
-				check(Guids.Num() == ReplayBuffer.GetMemoryStatesNum());
-
-				for (int32 Index = 0; Index < ReplayBuffer.GetMemoryStatesNum(); Index++)
-				{
-					const int32 DimNum = ReplayBuffer.GetMemoryStates(Index).Num<1>();
-					ExperienceContainer.MemoryStates.Add(SharedMemory::Map<3, float>(Guids[Index], { ProcessNum, ReplayBuffer.GetMaxStepNum(), DimNum }));
-				}
-			}
+		for (int32 Index = 0; Index < ReplayBuffer.GetMemoryStatesNum(); Index++)
+		{
+			const int32 DimNum = ReplayBuffer.GetMemoryStates(Index).Num<1>();
+			ExperienceContainer.MemoryStates.Add(SharedMemory::Allocate<2, float>({ ReplayBuffer.GetMaxStepNum(), DimNum }));
+		}
 
-			// Rewards
-			{
-				FString StringOfGuids; ensure(FParse::Value(FCommandLine::Get(), TEXT("LearningRewardsGuid"), StringOfGuids));
-				TArray<FString> StringGuids;
-				StringOfGuids.ParseIntoArray(StringGuids, TEXT(","));
-				TArray<FGuid> Guids;
-				for (const FString& StringGuid : StringGuids)
-				{
-					FGuid Guid;
-					FGuid::Parse(StringGuid, Guid);
-					Guids.Add(Guid);
-				}
-				check(Guids.Num() == ReplayBuffer.GetRewardsNum());
-
-				for (int32 Index = 0; Index < ReplayBuffer.GetRewardsNum(); Index++)
-				{
-					const int32 DimNum = ReplayBuffer.GetRewards(Index).Num<1>();
-					ExperienceContainer.Rewards.Add(SharedMemory::Map<3, float>(Guids[Index], { ProcessNum, ReplayBuffer.GetMaxStepNum(), DimNum }));
-				}
-			}
+		for (int32 Index = 0; Index < ReplayBuffer.GetRewardsNum(); Index++)
+		{
+			const int32 DimNum = ReplayBuffer.GetRewards(Index).Num<1>();
+			ExperienceContainer.Rewards.Add(SharedMemory::Allocate<2, float>({ ReplayBuffer.GetMaxStepNum(), DimNum }));
 		}
 		
 		const int32 ReplayBufferId = SharedMemoryExperienceContainers.Num();
@@ -623,7 +460,6 @@ namespace UE::Learning
 
 	ETrainerResponse FSharedMemoryTrainer::SendReplayBuffer(const int32 ReplayBufferId, const FReplayBuffer& ReplayBuffer, const ELogSetting LogSettings)
 	{
-		check(ProcessIdx != INDEX_NONE);
 		checkf(Controls.Region, TEXT("SendReplayBuffer: Controls Shared Memory Region is nullptr"));
 		if (!ensureMsgf(SharedMemoryExperienceContainers.Num() >= ReplayBufferId, TEXT("ReplayBuffer %d has not been added. Call AddReplayBuffer prior to SendReplayBuffer."), ReplayBufferId))
 		{
@@ -631,52 +467,52 @@ namespace UE::Learning
 		}
 
 		TArray<TLearningArrayView<2, float>> EpisodeFinalObservations;
-		for (TSharedMemoryArrayView<3, float>& EpisodeFinalObs : SharedMemoryExperienceContainers[ReplayBufferId].EpisodeFinalObservations)
+		for (TSharedMemoryArrayView<2, float>& EpisodeFinalObs : SharedMemoryExperienceContainers[ReplayBufferId].EpisodeFinalObservations)
 		{
-			EpisodeFinalObservations.Add(EpisodeFinalObs.View[ProcessIdx]);
+			EpisodeFinalObservations.Add(EpisodeFinalObs.View);
 		}
 
 		TArray<TLearningArrayView<2, float>> EpisodeFinalMemoryStates;
-		for (TSharedMemoryArrayView<3, float>& EpisodeFinalMems : SharedMemoryExperienceContainers[ReplayBufferId].EpisodeFinalMemoryStates)
+		for (TSharedMemoryArrayView<2, float>& EpisodeFinalMems : SharedMemoryExperienceContainers[ReplayBufferId].EpisodeFinalMemoryStates)
 		{
-			EpisodeFinalMemoryStates.Add(EpisodeFinalMems.View[ProcessIdx]);
+			EpisodeFinalMemoryStates.Add(EpisodeFinalMems.View);
 		}
 
 		TArray<TLearningArrayView<2, float>> Observations;
-		for (TSharedMemoryArrayView<3, float>& Obs : SharedMemoryExperienceContainers[ReplayBufferId].Observations)
+		for (TSharedMemoryArrayView<2, float>& Obs : SharedMemoryExperienceContainers[ReplayBufferId].Observations)
 		{
-			Observations.Add(Obs.View[ProcessIdx]);
+			Observations.Add(Obs.View);
 		}
 
 		TArray<TLearningArrayView<2, float>> Actions;
-		for (TSharedMemoryArrayView<3, float>& Acts : SharedMemoryExperienceContainers[ReplayBufferId].Actions)
+		for (TSharedMemoryArrayView<2, float>& Acts : SharedMemoryExperienceContainers[ReplayBufferId].Actions)
 		{
-			Actions.Add(Acts.View[ProcessIdx]);
+			Actions.Add(Acts.View);
 		}
 
 		TArray<TLearningArrayView<2, float>> ActionModifiers;
-		for (TSharedMemoryArrayView<3, float>& Mods : SharedMemoryExperienceContainers[ReplayBufferId].ActionModifiers)
+		for (TSharedMemoryArrayView<2, float>& Mods : SharedMemoryExperienceContainers[ReplayBufferId].ActionModifiers)
 		{
-			ActionModifiers.Add(Mods.View[ProcessIdx]);
+			ActionModifiers.Add(Mods.View);
 		}
 
 		TArray<TLearningArrayView<2, float>> MemoryStates;
-		for (TSharedMemoryArrayView<3, float>& Mems : SharedMemoryExperienceContainers[ReplayBufferId].MemoryStates)
+		for (TSharedMemoryArrayView<2, float>& Mems : SharedMemoryExperienceContainers[ReplayBufferId].MemoryStates)
 		{
-			MemoryStates.Add(Mems.View[ProcessIdx]);
+			MemoryStates.Add(Mems.View);
 		}
 
 		TArray<TLearningArrayView<2, float>> Rewards;
-		for (TSharedMemoryArrayView<3, float>& Rews : SharedMemoryExperienceContainers[ReplayBufferId].Rewards)
+		for (TSharedMemoryArrayView<2, float>& Rews : SharedMemoryExperienceContainers[ReplayBufferId].Rewards)
 		{
-			Rewards.Add(Rews.View[ProcessIdx]);
+			Rewards.Add(Rews.View);
 		}
 
 		TLearningArrayView<1, ECompletionMode> EmptyCompletionsArray;
 		return SharedMemoryTraining::SendExperience(
-			SharedMemoryExperienceContainers[ReplayBufferId].EpisodeStarts.View[ProcessIdx],
-			SharedMemoryExperienceContainers[ReplayBufferId].EpisodeLengths.View[ProcessIdx],
-			ReplayBuffer.HasCompletions() ? SharedMemoryExperienceContainers[ReplayBufferId].EpisodeCompletionModes.View[ProcessIdx] : EmptyCompletionsArray,
+			SharedMemoryExperienceContainers[ReplayBufferId].EpisodeStarts.View,
+			SharedMemoryExperienceContainers[ReplayBufferId].EpisodeLengths.View,
+			ReplayBuffer.HasCompletions() ? SharedMemoryExperienceContainers[ReplayBufferId].EpisodeCompletionModes.View : EmptyCompletionsArray,
 			EpisodeFinalObservations,
 			EpisodeFinalMemoryStates,
 			Observations,
@@ -684,8 +520,8 @@ namespace UE::Learning
 			ActionModifiers,
 			MemoryStates,
 			Rewards,
-			Controls.View[ProcessIdx],
-			*TrainingProcess,
+			Controls.View,
+			TrainingProcess.IsValid() ? TrainingProcess->GetTrainingSubprocess() : nullptr,
 			ReplayBufferId,
 			ReplayBuffer,
 			Timeout,
@@ -737,11 +573,13 @@ namespace UE::Learning
 			Port,
 			*FileManager.ConvertToAbsolutePathForExternalAppForRead(*IntermediatePath));
 
-		TrainingProcess.Launch(
+		bool bLaunched = TrainingProcess.Launch(
 			FileManager.ConvertToAbsolutePathForExternalAppForRead(*PythonExecutablePath),
 			CommandLineArguments, 
 			TrainingProcessFlags);
 
+		ensure(bLaunched);
+
 		if (PLATFORM_MAC)
 		{
 			// TODO we seem to have to sleep on Mac so the trainer can start listening before we try to connect
@@ -794,7 +632,7 @@ namespace UE::Learning
 
 		if (ExternalTrainerProcess)
 		{
-			TrainingProcess = ExternalTrainerProcess->GetTrainingSubprocess();
+			TrainingProcess = ExternalTrainerProcess;
 		}
 
 		ISocketSubsystem* SocketSubsystem = ISocketSubsystem::Get(PLATFORM_SOCKETSUBSYSTEM);
@@ -814,7 +652,11 @@ namespace UE::Learning
 
 		Socket = FTcpSocketBuilder(TEXT("LearningTrainerSocket")).AsBlocking().Build();
 
-		OutResponse = SocketTraining::WaitForConnection(*Socket, TrainingProcess, *Address, Timeout);
+		OutResponse = SocketTraining::WaitForConnection(
+			*Socket,
+			TrainingProcess.IsValid() ? TrainingProcess->GetTrainingSubprocess() : nullptr,
+			*Address,
+			Timeout);
 	}
 
 	FSocketTrainer::~FSocketTrainer()
@@ -829,7 +671,7 @@ namespace UE::Learning
 
 	bool FSocketTrainer::HasNetworkOrCompleted()
 	{
-		return SocketTraining::HasNetworkOrCompleted(*Socket, TrainingProcess);
+		return SocketTraining::HasNetworkOrCompleted(*Socket, TrainingProcess.IsValid() ? TrainingProcess->GetTrainingSubprocess() : nullptr);
 	}
 
 	void FSocketTrainer::Terminate()
@@ -849,7 +691,7 @@ namespace UE::Learning
 			return ETrainerResponse::Unexpected;
 		}
 
-		return SocketTraining::SendStop(*Socket, TrainingProcess, Timeout);
+		return SocketTraining::SendStop(*Socket, TrainingProcess.IsValid() ? TrainingProcess->GetTrainingSubprocess() : nullptr, Timeout);
 	}
 
 	ETrainerResponse FSocketTrainer::SendConfig(
@@ -866,7 +708,12 @@ namespace UE::Learning
 		TSharedRef<TJsonWriter<>> JsonWriter = TJsonWriterFactory<>::Create(&ConfigString, 0);
 		FJsonSerializer::Serialize(ConfigObject, JsonWriter, true);
 
-		return SocketTraining::SendConfig(*Socket, ConfigString, TrainingProcess, Timeout, LogSettings);
+		return SocketTraining::SendConfig(
+			*Socket,
+			ConfigString,
+			TrainingProcess.IsValid() ? TrainingProcess->GetTrainingSubprocess() : nullptr,
+			Timeout,
+			LogSettings);
 	}
 
 	int32 FSocketTrainer::AddNetwork(const ULearningNeuralNetworkData& Network)
@@ -894,7 +741,15 @@ namespace UE::Learning
 			return ETrainerResponse::Unexpected;
 		}
 
-		return SocketTraining::RecvNetwork(*Socket, NetworkId, OutNetwork, TrainingProcess, NetworkBuffers[NetworkId], Timeout, NetworkLock, LogSettings);
+		return SocketTraining::RecvNetwork(
+			*Socket,
+			NetworkId,
+			OutNetwork,
+			TrainingProcess.IsValid() ? TrainingProcess->GetTrainingSubprocess() : nullptr,
+			NetworkBuffers[NetworkId],
+			Timeout,
+			NetworkLock,
+			LogSettings);
 	}
 
 	ETrainerResponse FSocketTrainer::SendNetwork(
@@ -914,7 +769,15 @@ namespace UE::Learning
 			return ETrainerResponse::Unexpected;
 		}
 
-		return SocketTraining::SendNetwork(*Socket, NetworkBuffers[NetworkId], TrainingProcess, NetworkId, Network, Timeout, NetworkLock, LogSettings);
+		return SocketTraining::SendNetwork(
+			*Socket,
+			NetworkBuffers[NetworkId],
+			TrainingProcess.IsValid() ? TrainingProcess->GetTrainingSubprocess() : nullptr,
+			NetworkId,
+			Network,
+			Timeout,
+			NetworkLock,
+			LogSettings);
 	}
 
 	int32 FSocketTrainer::AddReplayBuffer(const FReplayBuffer& ReplayBuffer)
@@ -936,6 +799,12 @@ namespace UE::Learning
 			return ETrainerResponse::Unexpected;
 		}
 
-		return SocketTraining::SendExperience(*Socket, ReplayBufferId, ReplayBuffer, TrainingProcess, Timeout, LogSettings);
+		return SocketTraining::SendExperience(
+			*Socket,
+			ReplayBufferId,
+			ReplayBuffer,
+			TrainingProcess.IsValid() ? TrainingProcess->GetTrainingSubprocess() : nullptr,
+			Timeout,
+			LogSettings);
 	}
 }
diff --git a/Engine/Plugins/Experimental/LearningAgents/Source/LearningTraining/Private/LearningSharedMemoryTraining.cpp b/Engine/Plugins/Experimental/LearningAgents/Source/LearningTraining/Private/LearningSharedMemoryTraining.cpp
--- a/Engine/Plugins/Experimental/LearningAgents/Source/LearningTraining/Private/LearningSharedMemoryTraining.cpp
+++ b/Engine/Plugins/Experimental/LearningAgents/Source/LearningTraining/Private/LearningSharedMemoryTraining.cpp
@@ -45,7 +45,7 @@ namespace UE::Learning::SharedMemoryTraining
 		TLearningArrayView<1, volatile int32> Controls,
 		const int32 NetworkId,
 		ULearningNeuralNetworkData& OutNetwork,
-		FSubprocess& Process,
+		FSubprocess* Process,
 		const TLearningArrayView<1, const uint8> NetworkData,
 		const float Timeout,
 		FRWLock* NetworkLock,
@@ -66,7 +66,7 @@ namespace UE::Learning::SharedMemoryTraining
 			}
 
 			// If we're monitoring a process, then has it has exited?
-			if (!Process.Update())
+			if (Process && !Process->Update())
 			{
 				return ETrainerResponse::Unexpected;
 			}
@@ -131,7 +131,7 @@ namespace UE::Learning::SharedMemoryTraining
 		TLearningArrayView<1, volatile int32> Controls,
 		const int32 NetworkId,
 		TLearningArrayView<1, uint8> NetworkData,
-		FSubprocess& Process,
+		FSubprocess* Process,
 		const ULearningNeuralNetworkData& Network,
 		const float Timeout,
 		FRWLock* NetworkLock,
@@ -144,7 +144,7 @@ namespace UE::Learning::SharedMemoryTraining
 		while (!Controls[(uint8)EControls::NetworkSignal])
 		{
 			// If we're monitoring a process, then has it has exited?
-			if (!Process.Update())
+			if (Process && !Process->Update())
 			{
 				return ETrainerResponse::Unexpected;
 			}
@@ -207,7 +207,7 @@ namespace UE::Learning::SharedMemoryTraining
 		TArrayView<TLearningArrayView<2, float>> MemoryStates,
 		TArrayView<TLearningArrayView<2, float>> Rewards,
 		TLearningArrayView<1, volatile int32> Controls,
-		FSubprocess& Process,
+		FSubprocess* Process,
 		const int32 ReplayBufferId,
 		const FReplayBuffer& ReplayBuffer,
 		const float Timeout,
@@ -220,7 +220,7 @@ namespace UE::Learning::SharedMemoryTraining
 		while (Controls[(uint8)EControls::ExperienceSignal])
 		{
 			// If we're monitoring a process, then has it has exited?
-			if (!Process.Update())
+			if (Process && !Process->Update())
 			{
 				return ETrainerResponse::Unexpected;
 			}
diff --git a/Engine/Plugins/Experimental/LearningAgents/Source/LearningTraining/Private/LearningSharedMemoryTraining.h b/Engine/Plugins/Experimental/LearningAgents/Source/LearningTraining/Private/LearningSharedMemoryTraining.h
--- a/Engine/Plugins/Experimental/LearningAgents/Source/LearningTraining/Private/LearningSharedMemoryTraining.h
+++ b/Engine/Plugins/Experimental/LearningAgents/Source/LearningTraining/Private/LearningSharedMemoryTraining.h
@@ -39,7 +39,7 @@ namespace UE::Learning
 			TLearningArrayView<1, volatile int32> Controls,
 			const int32 NetworkId,
 			ULearningNeuralNetworkData& OutNetwork,
-			FSubprocess& Process,
+			FSubprocess* Process,
 			const TLearningArrayView<1, const uint8> NetworkData,
 			const float Timeout = Trainer::DefaultTimeout,
 			FRWLock* NetworkLock = nullptr,
@@ -54,7 +54,7 @@ namespace UE::Learning
 			TLearningArrayView<1, volatile int32> Controls,
 			const int32 NetworkId,
 			TLearningArrayView<1, uint8> NetworkData,
-			FSubprocess& Process,
+			FSubprocess* Process,
 			const ULearningNeuralNetworkData& Network,
 			const float Timeout = Trainer::DefaultTimeout,
 			FRWLock* NetworkLock = nullptr,
@@ -72,7 +72,7 @@ namespace UE::Learning
 			TArrayView<TLearningArrayView<2, float>> MemoryStates,
 			TArrayView<TLearningArrayView<2, float>> Rewards,
 			TLearningArrayView<1, volatile int32> Controls,
-			FSubprocess& Process,
+			FSubprocess* Process,
 			const int32 ReplayBufferId,
 			const FReplayBuffer& ReplayBuffer,
 			const float Timeout = Trainer::DefaultTimeout,
diff --git a/Engine/Plugins/Experimental/LearningAgents/Source/LearningTraining/Public/LearningExternalTrainer.h b/Engine/Plugins/Experimental/LearningAgents/Source/LearningTraining/Public/LearningExternalTrainer.h
--- a/Engine/Plugins/Experimental/LearningAgents/Source/LearningTraining/Public/LearningExternalTrainer.h
+++ b/Engine/Plugins/Experimental/LearningAgents/Source/LearningTraining/Public/LearningExternalTrainer.h
@@ -135,10 +135,8 @@ namespace UE::Learning
 		* @param PythonExecutablePath		Path to the python executable used for training. In general should be the
 		*									python shipped with Unreal Editor.
 		* @param PythonContentPath			Path to the Python Content folder provided by the Learning plugin
-		* @param IntermediatePath			Path to the intermediate folder to write temporary files, logs, and
+		* @param InIntermediatePath			Path to the intermediate folder to write temporary files, logs, and
 		*									snapshots to
-		* @param ProcessNum					Number of processes to use for multi-processed experience gathering (used
-		*									to allocate enough shared memory)
 		* @param TrainingProcessFlags		Training server subprocess flags
 		* @param LogSettings				Logging settings to use
 		*/
@@ -148,8 +146,7 @@ namespace UE::Learning
 			const FString& TrainerFileName,
 			const FString& PythonExecutablePath,
 			const FString& PythonContentPath,
-			const FString& IntermediatePath,
-			const int32 ProcessNum,
+			const FString& InIntermediatePath,
 			const float InTimeout = Trainer::DefaultTimeout,
 			const ESubprocessFlags TrainingProcessFlags = ESubprocessFlags::None,
 			const ELogSetting LogSettings = ELogSetting::Normal);
@@ -169,7 +166,7 @@ namespace UE::Learning
 		void Terminate();
 
 		/** Get the Controls shared memory array view. */
-		TSharedMemoryArrayView<2, volatile int32> GetControlsSharedMemoryArrayView() const;
+		TSharedMemoryArrayView<1, volatile int32> GetControlsSharedMemoryArrayView() const;
 
 		/** Get the intermediate path. */
 		const FString& GetIntermediatePath() const;
@@ -188,7 +185,7 @@ namespace UE::Learning
 		FString IntermediatePath;
 		FString ConfigPath;
 
-		TSharedMemoryArrayView<2, volatile int32> Controls; // Mark as volatile to avoid compiler optimizing away reads without writes etc.
+		TSharedMemoryArrayView<1, volatile int32> Controls; // Mark as volatile to avoid compiler optimizing away reads without writes etc.
 
 		FSubprocess TrainingProcess;
 		float Timeout = Trainer::DefaultTimeout;
@@ -204,17 +201,17 @@ namespace UE::Learning
 	{
 		struct FSharedMemoryExperienceContainer
 		{
-			TSharedMemoryArrayView<2, int32> EpisodeStarts;
-			TSharedMemoryArrayView<2, int32> EpisodeLengths;
-			TSharedMemoryArrayView<2, ECompletionMode> EpisodeCompletionModes;
-
-			TArray<TSharedMemoryArrayView<3, float>, TInlineAllocator<1>> EpisodeFinalObservations;
-			TArray<TSharedMemoryArrayView<3, float>, TInlineAllocator<1>> EpisodeFinalMemoryStates;
-			TArray<TSharedMemoryArrayView<3, float>, TInlineAllocator<1>> Observations;
-			TArray<TSharedMemoryArrayView<3, float>, TInlineAllocator<1>> Actions;
-			TArray<TSharedMemoryArrayView<3, float>, TInlineAllocator<1>> ActionModifiers;
-			TArray<TSharedMemoryArrayView<3, float>, TInlineAllocator<1>> MemoryStates;
-			TArray<TSharedMemoryArrayView<3, float>, TInlineAllocator<1>> Rewards;
+			TSharedMemoryArrayView<1, int32> EpisodeStarts;
+			TSharedMemoryArrayView<1, int32> EpisodeLengths;
+			TSharedMemoryArrayView<1, ECompletionMode> EpisodeCompletionModes;
+
+			TArray<TSharedMemoryArrayView<2, float>, TInlineAllocator<1>> EpisodeFinalObservations;
+			TArray<TSharedMemoryArrayView<2, float>, TInlineAllocator<1>> EpisodeFinalMemoryStates;
+			TArray<TSharedMemoryArrayView<2, float>, TInlineAllocator<1>> Observations;
+			TArray<TSharedMemoryArrayView<2, float>, TInlineAllocator<1>> Actions;
+			TArray<TSharedMemoryArrayView<2, float>, TInlineAllocator<1>> ActionModifiers;
+			TArray<TSharedMemoryArrayView<2, float>, TInlineAllocator<1>> MemoryStates;
+			TArray<TSharedMemoryArrayView<2, float>, TInlineAllocator<1>> Rewards;
 
 			/** Free and deallocate all shared memory. */
 			void Deallocate();
@@ -223,16 +220,21 @@ namespace UE::Learning
 		/**
 		* Creates a new SharedMemory trainer
 		*
-		* @param InTaskName				Unique name for this training task - used to avoid config filename conflicts
-		* @param ProcessNum				Number of processes to use for multi-processed experience gathering (used
-		*								to allocate enough shared memory)
-		* @param ExternalTrainerProcess	Shared memory used for communicating status to the trainer server process
-		* @param InTimeout				Timeout to wait in seconds for connection and initial data transfer
+		* @param InTaskName						Unique name for this training task - used to avoid config filename conflicts
+		* @param ExternalTrainerProcess			Shared memory used for communicating status to the trainer server process
+		* @param bUseExternalTrainingProcess	If true, attach to existing external training process.
+		* @param ControlsGuid					If bUseExternalTrainingProcess is true, attach to this controls memory.
+		* @param InConfigPath					The path to write the config to.
+		* @param InIntermediatePath				The Intermediate Path - for writing temporary files.
+		* @param InTimeout						Timeout to wait in seconds for connection and initial data transfer
 		*/
 		FSharedMemoryTrainer(
 			const FString& InTaskName,
-			const int32 InProcessNum,
 			const TSharedPtr<FSharedMemoryTrainerServerProcess>& ExternalTrainerProcess,
+			const bool bUseExternalTrainingProcess = false,
+			const FGuid ControlsGuid = FGuid(),
+			const FString& InConfigPath = FString(),
+			const FString& InIntermediatePath = FString(),
 			const float InTimeout = Trainer::DefaultTimeout);
 
 		~FSharedMemoryTrainer();
@@ -276,19 +278,14 @@ namespace UE::Learning
 		void Deallocate();
 
 		FString TaskName;
-
 		FString IntermediatePath;
-
 		FString ConfigPath;
 
-		int32 ProcessNum;
-
-		FSubprocess* TrainingProcess = nullptr;
+		TSharedPtr<FSharedMemoryTrainerServerProcess> TrainingProcess = nullptr;
 
 		float Timeout = Trainer::DefaultTimeout;
 
-		int32 ProcessIdx = INDEX_NONE;
-		TSharedMemoryArrayView<2, volatile int32> Controls;
+		TSharedMemoryArrayView<1, volatile int32> Controls;
 		TArray<TSharedMemoryArrayView<1, uint8>> NeuralNetworkSharedMemoryArrayViews;
 		TArray<FSharedMemoryExperienceContainer> SharedMemoryExperienceContainers;
 	};
@@ -417,7 +414,7 @@ namespace UE::Learning
 
 		float Timeout = Trainer::DefaultTimeout;
 
-		FSubprocess* TrainingProcess = nullptr;
+		TSharedPtr<FSocketTrainerServerProcess> TrainingProcess = nullptr;
 
 		FSocket* Socket = nullptr;
 	};
