LearningAgents: Update comments about TensorBoard plugin

diff --git a/Engine/Plugins/Experimental/LearningAgents/Source/LearningAgentsTraining/Public/LearningAgentsImitationTrainer.h b/Engine/Plugins/Experimental/LearningAgents/Source/LearningAgentsTraining/Public/LearningAgentsImitationTrainer.h
--- a/Engine/Plugins/Experimental/LearningAgents/Source/LearningAgentsTraining/Public/LearningAgentsImitationTrainer.h
+++ b/Engine/Plugins/Experimental/LearningAgents/Source/LearningAgentsTraining/Public/LearningAgentsImitationTrainer.h
@@ -96,7 +96,12 @@ public:
 	UPROPERTY(EditAnywhere, Category = "LearningAgents")
 	ELearningAgentsTrainingDevice Device = ELearningAgentsTrainingDevice::CPU;
 
-	/** If true, TensorBoard logs will be emitted to the intermediate directory. */
+	/**
+	 * If true, TensorBoard logs will be emitted to the intermediate directory.
+	 *
+	 * TensorBoard will only work if it is installed in Unreal Engine's python environment. This can be done by
+	 * enabling the "Tensorboard" plugin in your project.
+	 */
 	UPROPERTY(EditAnywhere, Category = "LearningAgents")
 	bool bUseTensorboard = false;
 
diff --git a/Engine/Plugins/Experimental/LearningAgents/Source/LearningAgentsTraining/Public/LearningAgentsPPOTrainer.h b/Engine/Plugins/Experimental/LearningAgents/Source/LearningAgentsTraining/Public/LearningAgentsPPOTrainer.h
--- a/Engine/Plugins/Experimental/LearningAgents/Source/LearningAgentsTraining/Public/LearningAgentsPPOTrainer.h
+++ b/Engine/Plugins/Experimental/LearningAgents/Source/LearningAgentsTraining/Public/LearningAgentsPPOTrainer.h
@@ -229,7 +229,12 @@ public:
 	UPROPERTY(EditAnywhere, Category = "LearningAgents")
 	ELearningAgentsTrainingDevice Device = ELearningAgentsTrainingDevice::GPU;
 
-	/** If true, TensorBoard logs will be emitted to the intermediate directory. */
+	/**
+	 * If true, TensorBoard logs will be emitted to the intermediate directory.
+	 * 
+	 * TensorBoard will only work if it is installed in Unreal Engine's python environment. This can be done by
+	 * enabling the "Tensorboard" plugin in your project.
+	 */
 	UPROPERTY(EditAnywhere, Category = "LearningAgents")
 	bool bUseTensorboard = false;
 
diff --git a/Engine/Plugins/Experimental/LearningAgents/Source/LearningTraining/Public/LearningPPOTrainer.h b/Engine/Plugins/Experimental/LearningAgents/Source/LearningTraining/Public/LearningPPOTrainer.h
--- a/Engine/Plugins/Experimental/LearningAgents/Source/LearningTraining/Public/LearningPPOTrainer.h
+++ b/Engine/Plugins/Experimental/LearningAgents/Source/LearningTraining/Public/LearningPPOTrainer.h
@@ -139,10 +139,8 @@ namespace UE::Learning
 
 		// If to use TensorBoard for logging and tracking the training progress.
 		// 
-		// Even when enabled, TensorBoard will only work if it is installed in your Unreal Editor
-		// bundled version of Python, which is not the case by default. TensorBoard can be installed 
-		// for this version of Python by going to your Unreal Editor Python Binaries directory 
-		// (e.g. "\Engine\Binaries\ThirdParty\Python3\Win64") and running `python -m pip install tensorboard`. 
+		// TensorBoard will only work if it is installed in Unreal Engine's python environment. This can be done by
+		// enabling the "Tensorboard" plugin in your project.
 		bool bUseTensorboard = false;
 
 		// If to save snapshots of the trained networks every 1000 iterations
