Learning: Added NumBytes method.

diff --git a/Engine/Plugins/Experimental/LearningAgents/Source/Learning/Public/MultiArray.h b/Engine/Plugins/Experimental/LearningAgents/Source/Learning/Public/MultiArray.h
--- a/Engine/Plugins/Experimental/LearningAgents/Source/Learning/Public/MultiArray.h
+++ b/Engine/Plugins/Experimental/LearningAgents/Source/Learning/Public/MultiArray.h
@@ -541,6 +541,16 @@ public:
 		return ArrayShape.Total();
 	}
 
+	/**
+	 * Returns the total number of bytes used by the array
+	 *
+	 * @returns Total number of bytes used by the array.
+	 */
+	FORCEINLINE int64 NumBytes() const
+	{
+		return ArrayShape.Total() * sizeof(ElementType);
+	}
+
 	/**
 	 * Returns number of elements in a dimension.
 	 *
@@ -1195,6 +1205,16 @@ public:
 		return ArrayShape.Total();
 	}
 
+	/**
+	 * Returns the total number of bytes used by the array
+	 *
+	 * @returns Total number of bytes used by the array.
+	 */
+	FORCEINLINE int64 NumBytes() const
+	{
+		return ArrayShape.Total() * sizeof(ElementType);
+	}
+
 	/**
 	 * Returns number of elements in a dimension.
 	 *
diff --git a/Engine/Plugins/Experimental/LearningAgents/Source/Learning/Public/MultiArrayView.h b/Engine/Plugins/Experimental/LearningAgents/Source/Learning/Public/MultiArrayView.h
--- a/Engine/Plugins/Experimental/LearningAgents/Source/Learning/Public/MultiArrayView.h
+++ b/Engine/Plugins/Experimental/LearningAgents/Source/Learning/Public/MultiArrayView.h
@@ -412,6 +412,16 @@ public:
 		return ArrayShape.Total();
 	}
 
+	/**
+	 * Returns the total number of bytes used by the array
+	 *
+	 * @returns Total number of bytes used by the array.
+	 */
+	FORCEINLINE int64 NumBytes() const
+	{
+		return ArrayShape.Total() * sizeof(ElementType);
+	}
+
 	/**
 	 * Returns the number of elements in a dimension.
 	 *
@@ -777,6 +787,16 @@ public:
 		return ArrayShape.Total();
 	}
 
+	/**
+	 * Returns the total number of bytes used by the array
+	 *
+	 * @returns Total number of bytes used by the array.
+	 */
+	FORCEINLINE int64 NumBytes() const
+	{
+		return ArrayShape.Total() * sizeof(ElementType);
+	}
+
 	/**
 	 * Returns the number of elements in a dimension.
 	 *
