LearningAgents: Custom Trainer Non Editor Path

diff --git a/Engine/Plugins/Experimental/LearningAgents/Source/LearningAgentsTraining/Private/LearningAgentsTrainer.cpp b/Engine/Plugins/Experimental/LearningAgents/Source/LearningAgentsTraining/Private/LearningAgentsTrainer.cpp
--- a/Engine/Plugins/Experimental/LearningAgents/Source/LearningAgentsTraining/Private/LearningAgentsTrainer.cpp
+++ b/Engine/Plugins/Experimental/LearningAgents/Source/LearningAgentsTraining/Private/LearningAgentsTrainer.cpp
@@ -200,5 +200,14 @@ FString FLearningAgentsTrainerProcessSettings::GetIntermediatePath() const
 
 FString FLearningAgentsTrainerProcessSettings::GetCustomTrainerModulePath() const
 {
+#if WITH_EDITOR
 	return CustomTrainerModulePath.Path;
+#else
+	if (NonEditorCustomTrainerModulePath.IsEmpty())
+	{
+		UE_LOG(LogLearning, Warning, TEXT("GetCustomTrainerModulePath: NonEditorCustomTrainerModulePath not set"));
+	}
+
+	return NonEditorCustomTrainerModulePath;
+#endif
 }
diff --git a/Engine/Plugins/Experimental/LearningAgents/Source/LearningAgentsTraining/Public/LearningAgentsTrainer.h b/Engine/Plugins/Experimental/LearningAgents/Source/LearningAgentsTraining/Public/LearningAgentsTrainer.h
--- a/Engine/Plugins/Experimental/LearningAgents/Source/LearningAgentsTraining/Public/LearningAgentsTrainer.h
+++ b/Engine/Plugins/Experimental/LearningAgents/Source/LearningAgentsTraining/Public/LearningAgentsTrainer.h
@@ -144,6 +144,10 @@ public:
 	UPROPERTY(EditAnywhere, Category = "LearningAgents")
 	FDirectoryPath CustomTrainerModulePath;
 
+	/** The complete path to a custom trainer module when training with a custom trainer for non-editor builds. */
+	UPROPERTY(EditAnywhere, Category = "LearningAgents")
+	FString NonEditorCustomTrainerModulePath;
+
 	/** Trainer file name. The name of the python file to use for training. Do NOT include the '.py' file extension. */
 	UPROPERTY(EditAnywhere, Category = "LearningAgents")
 	FString TrainerFileName = TEXT("train_ppo");
