Learning: Removed FLT_MAX and -FLT_MAX from checks. Fixed natvis file. Added Find function to IndexSet.

diff --git a/Engine/Plugins/Experimental/LearningAgents/Source/Learning/Public/LearningArray.h b/Engine/Plugins/Experimental/LearningAgents/Source/Learning/Public/LearningArray.h
--- a/Engine/Plugins/Experimental/LearningAgents/Source/Learning/Public/LearningArray.h
+++ b/Engine/Plugins/Experimental/LearningAgents/Source/Learning/Public/LearningArray.h
@@ -29,7 +29,7 @@
 //#define UE_LEARNING_DEVELOPMENT 0
 #endif
 
-// There are various checks in Learning to see if arrays contain inf/nan/MAX_flt/-MAX_flt
+// There are various checks in Learning to see if arrays contain nan/inf/-inf
 // which can be disabled or enabled with this flag. These will be disabled in 
 // Shipping mode but in certain cases you may want to disable them in editor 
 // mode too to improve performance.
@@ -216,6 +216,30 @@ namespace UE::Learning
 			return IsSlice() ? (Index >= SliceStart && Index < SliceStart + SliceNum) : Indices.Contains(Index);
 		}
 
+		inline int32 Find(const int32 Index) const
+		{
+			if (IsSlice())
+			{
+				if (Index >= SliceStart && Index < SliceStart + SliceNum)
+				{
+					return Index - SliceStart;
+				}
+				else
+				{
+					return INDEX_NONE;
+				}
+			}
+			else
+			{
+				const int32 IndexNum = Indices.Num();
+				for (int32 Idx = 0; Idx < IndexNum; Idx++)
+				{
+					if (Indices[Idx] == Index) { return Idx; }
+				}
+				return INDEX_NONE;
+			}
+		}
+
 		inline FIndexSetIterator begin() const
 		{
 			return IsSlice() ?
@@ -348,7 +372,7 @@ namespace UE::Learning
 		}
 
 		/**
-		* Check that an array view does not contain any NaN, Inf, -Inf, MAX_flt, or -MAX_flt
+		* Check that an array view does not contain any NaN, Inf, -Inf
 		*/
 		template<uint8 InDimNum, typename InElementType>
 		inline void Check(const TLearningArrayView<InDimNum, InElementType> View)
@@ -357,14 +381,14 @@ namespace UE::Learning
 			for (int64 Idx = 0; Idx < View.Num(); Idx++)
 			{
 				UE_LEARNING_ARRAY_VALUE_CHECKF(
-					FMath::IsFinite(View.GetData()[Idx]) && View.GetData()[Idx] != MAX_flt && View.GetData()[Idx] != -MAX_flt,
+					FMath::IsFinite(View.GetData()[Idx]),
 					TEXT("Invalid value %f found at flat array index %lld"), View.GetData()[Idx], Idx);
 			}
 #endif
 		}
 
 		/**
-		* Check that an array does not contain any NaN, Inf, -Inf, MAX_flt, or -MAX_flt
+		* Check that an array does not contain any NaN, Inf, -Inf
 		*/
 		template<uint8 InDimNum, typename InElementType, typename Allocator>
 		inline void Check(const TLearningArray<InDimNum, InElementType, Allocator>& View)
@@ -373,7 +397,7 @@ namespace UE::Learning
 			for (int64 Idx = 0; Idx < View.Num(); Idx++)
 			{
 				UE_LEARNING_ARRAY_VALUE_CHECKF(
-					FMath::IsFinite(View.GetData()[Idx]) && View.GetData()[Idx] != MAX_flt && View.GetData()[Idx] != -MAX_flt,
+					FMath::IsFinite(View.GetData()[Idx]),
 					TEXT("Invalid value %f found at flat array index %lld"), View.GetData()[Idx], Idx);
 			}
 #endif
@@ -381,7 +405,7 @@ namespace UE::Learning
 
 
 		/**
-		* Check that an array view does not contain any NaN, Inf, -Inf, MAX_flt, or -MAX_flt
+		* Check that an array view does not contain any NaN, Inf, -Inf
 		*/
 		template<uint8 InDimNum, typename InElementType>
 		inline void Check(const TLearningArrayView<InDimNum, const InElementType> View)
@@ -390,14 +414,14 @@ namespace UE::Learning
 			for (int64 Idx = 0; Idx < View.Num(); Idx++)
 			{
 				UE_LEARNING_ARRAY_VALUE_CHECKF(
-					FMath::IsFinite(View.GetData()[Idx]) && View.GetData()[Idx] != MAX_flt && View.GetData()[Idx] != -MAX_flt,
+					FMath::IsFinite(View.GetData()[Idx]),
 					TEXT("Invalid value %f found at flat array index %lld"), View.GetData()[Idx], Idx);
 			}
 #endif
 		}
 
 		/**
-		* Check that an array view does not contain any NaN, Inf, -Inf, MAX_flt, or -MAX_flt
+		* Check that an array view does not contain any NaN, Inf, -Inf
 		*/
 		template<uint8 InDimNum, typename InElementType>
 		inline void Check(const TLearningArrayView<InDimNum, const InElementType> View, const FIndexSet Indices)
@@ -415,7 +439,7 @@ namespace UE::Learning
 				for (const int32 Idx : Indices)
 				{
 					UE_LEARNING_ARRAY_VALUE_CHECKF(
-						FMath::IsFinite(View.GetData()[Idx]) && View.GetData()[Idx] != MAX_flt && View.GetData()[Idx] != -MAX_flt,
+						FMath::IsFinite(View.GetData()[Idx]) && View.GetData()[Idx],
 						TEXT("Invalid value %f found at flat array index %lld"), View.GetData()[Idx], Idx);
 				}
 			}
@@ -423,7 +447,7 @@ namespace UE::Learning
 		}
 
 		/**
-		* Check that an array view does not contain any NaN, Inf, -Inf, MAX_flt, or -MAX_flt
+		* Check that an array view does not contain any NaN, Inf, -Inf
 		*/
 		template<uint8 InDimNum, typename InElementType>
 		inline void Check(const TLearningArrayView<InDimNum, InElementType> View, const FIndexSet Indices)
@@ -441,7 +465,7 @@ namespace UE::Learning
 				for (const int32 Idx : Indices)
 				{
 					UE_LEARNING_ARRAY_VALUE_CHECKF(
-						FMath::IsFinite(View.GetData()[Idx]) && View.GetData()[Idx] != MAX_flt && View.GetData()[Idx] != -MAX_flt,
+						FMath::IsFinite(View.GetData()[Idx]) && View.GetData()[Idx],
 						TEXT("Invalid value %f found at flat array index %lld"), View.GetData()[Idx], Idx);
 				}
 			}
@@ -449,7 +473,7 @@ namespace UE::Learning
 		}
 
 		/**
-		* Check that an array view does not contain any NaN, Inf, -Inf, MAX_flt, or -MAX_flt
+		* Check that an array view does not contain any NaN, Inf, -Inf
 		*/
 		template<uint8 InDimNum, typename InElementType>
 		inline void Check(const TLearningArray<InDimNum, InElementType>& View, const FIndexSet Indices)
@@ -467,7 +491,7 @@ namespace UE::Learning
 				for (const int32 Idx : Indices)
 				{
 					UE_LEARNING_ARRAY_VALUE_CHECKF(
-						FMath::IsFinite(View.GetData()[Idx]) && View.GetData()[Idx] != MAX_flt && View.GetData()[Idx] != -MAX_flt,
+						FMath::IsFinite(View.GetData()[Idx]),
 						TEXT("Invalid value %f found at flat array index %lld"), View.GetData()[Idx], Idx);
 				}
 			}
