Learning: Small fixes to MultiArrayView natvis. Added additional constructors to FIndexSet.

diff --git a/Engine/Plugins/Experimental/LearningAgents/Source/Learning/Public/LearningArray.h b/Engine/Plugins/Experimental/LearningAgents/Source/Learning/Public/LearningArray.h
--- a/Engine/Plugins/Experimental/LearningAgents/Source/Learning/Public/LearningArray.h
+++ b/Engine/Plugins/Experimental/LearningAgents/Source/Learning/Public/LearningArray.h
@@ -144,6 +144,8 @@ namespace UE::Learning
 		FIndexSet(const int32 InSliceStart, const int32 InSliceNum) : SliceStart(InSliceStart), SliceNum(InSliceNum) {}
 		FIndexSet(const TArrayView<const int32> InIndices) : Indices(InIndices) {}
 		FIndexSet(const TLearningArrayView<1, const int32> InIndices) : Indices(InIndices) {}
+		FIndexSet(const TArrayView<int32> InIndices) : Indices(InIndices) {}
+		FIndexSet(const TLearningArrayView<1, int32> InIndices) : Indices(InIndices) {}
 		FIndexSet(std::initializer_list<int32> InIndices) : Indices(InIndices) {}
 		
 		template<typename InAllocatorType>
