LearningAgents: Minor fixes from working on FruitTaxi

diff --git a/Engine/Plugins/Experimental/LearningAgents/Source/LearningAgentsTraining/Public/LearningAgentsCommunicator.h b/Engine/Plugins/Experimental/LearningAgents/Source/LearningAgentsTraining/Public/LearningAgentsCommunicator.h
--- a/Engine/Plugins/Experimental/LearningAgents/Source/LearningAgentsTraining/Public/LearningAgentsCommunicator.h
+++ b/Engine/Plugins/Experimental/LearningAgents/Source/LearningAgentsTraining/Public/LearningAgentsCommunicator.h
@@ -154,7 +154,7 @@ public:
 	 * @param TrainerProcessSettings							Settings universal to all trainer processes.
 	 * @param FLearningAgentsSharedMemoryCommunicatorSettings	Settings specific to shared memory communicators.
 	 */
-	UFUNCTION(BlueprintCallable, Category = "LearningAgents", meta = (AutoCreateRefTerm = "SharedMemorySettings"))
+	UFUNCTION(BlueprintCallable, Category = "LearningAgents", meta = (AutoCreateRefTerm = "TrainerProcessSettings,SharedMemorySettings"))
 	static FLearningAgentsCommunicator MakeSharedMemoryCommunicator(
 		FLearningAgentsSharedMemoryTrainerProcess TrainerProcess,
 		const FLearningAgentsTrainerProcessSettings& TrainerProcessSettings = FLearningAgentsTrainerProcessSettings(),
