[LearningAgents][LearningCompanions] Migration

* EnttiesManagerComponent, EntityInterface, Gym and GymsManager are moved to LearningAgentsTraining

#rnx

diff --git a/Engine/Plugins/Experimental/LearningAgents/Source/LearningAgentsTraining/LearningAgentsTraining.Build.cs b/Engine/Plugins/Experimental/LearningAgents/Source/LearningAgentsTraining/LearningAgentsTraining.Build.cs
--- a/Engine/Plugins/Experimental/LearningAgents/Source/LearningAgentsTraining/LearningAgentsTraining.Build.cs
+++ b/Engine/Plugins/Experimental/LearningAgents/Source/LearningAgentsTraining/LearningAgentsTraining.Build.cs
@@ -12,19 +12,21 @@ public class LearningAgentsTraining : ModuleRules
 			new string[] {
 				// ... add public include paths required here ...
 			});
-		
+
 		PrivateIncludePaths.AddRange(
 			new string[] {
 				// ... add other private include paths required here ...
 			});
-		
+
 		PublicDependencyModuleNames.AddRange(
 			new string[]
 			{
 				"Core",
+				"NavigationSystem",
+				"AIModule",
 				// ... add other public dependencies that you statically link with here ...
 			});
-		
+
 		PrivateDependencyModuleNames.AddRange(
 			new string[]
 			{
diff --git a/Engine/Plugins/Experimental/LearningAgents/Source/LearningAgentsTraining/Private/LearningAgentsEntitiesManagerComponent.cpp b/Engine/Plugins/Experimental/LearningAgents/Source/LearningAgentsTraining/Private/LearningAgentsEntitiesManagerComponent.cpp
new file mode 100644
--- /dev/null
+++ b/Engine/Plugins/Experimental/LearningAgents/Source/LearningAgentsTraining/Private/LearningAgentsEntitiesManagerComponent.cpp
@@ -0,0 +1,231 @@
+// Copyright Epic Games, Inc. All Rights Reserved.
+
+#include "LearningAgentsEntitiesManagerComponent.h"
+#include "Engine/World.h"
+#include "LearningLog.h"
+#include "LearningAgentsGym.h"
+#include "LearningAgentsEntityInterface.h"
+
+void ULearningAgentsEntitiesManagerComponent::InitializeEntities()
+{
+	TObjectPtr<ALearningAgentsGymBase> Gym = Cast<ALearningAgentsGymBase>(GetOwner());
+	if (Gym)
+	{
+		if (!CheckEntityClasses())
+		{
+			return;
+		}
+
+		if (Gym->GetLocalRole() == ROLE_Authority)
+		{
+			for (FEntityInfo& EntityInfo : Entities)
+			{
+				SpawnEntities(EntityInfo, EntityInfo.EpisodeStartEntitySpawnCount);
+			}
+		}
+
+		// This is in a separate loop because initialization may be dependent on when all entities are spawned 
+		for (TPair<FName, FSpawnedEntitiesInfo>& UniqueTypeEntities : EntitiesPool)
+		{
+			for (TScriptInterface<ILearningAgentsEntityInterface>& Entity : UniqueTypeEntities.Value.SpawnedEntities)
+			{
+				Entity.GetInterface()->InitializeEntity(Gym);
+			}
+		}
+
+	}
+}
+
+void ULearningAgentsEntitiesManagerComponent::ResetEntities()
+{
+	TObjectPtr<ALearningAgentsGymBase> Gym = Cast<ALearningAgentsGymBase>(GetOwner());
+	if (Gym)
+	{
+		if (!CheckEntityClasses())
+		{
+			return;
+		}
+		for (TPair<FName, FSpawnedEntitiesInfo>& UniqueTypeEntities : EntitiesPool)
+		{
+			for (int32 i = 0; i < UniqueTypeEntities.Value.SpawnedEntities.Num(); ++i)
+			{
+				if (i < UniqueTypeEntities.Value.EntityInfo.EpisodeStartEntitySpawnCount)
+				{
+					// Make sure transform is applied before reset to wipe all data (ex: distance traveled)
+					if (TObjectPtr<AActor> ActorEntity = Cast<AActor>(UniqueTypeEntities.Value.SpawnedEntities[i].GetObject()))
+					{
+						FTransform RandomTransform;
+						RandomizeTransform(RandomTransform, UniqueTypeEntities.Value.EntityInfo.EntitySpawnZOffset);
+						ActorEntity->SetActorTransform(RandomTransform);
+					}
+					UniqueTypeEntities.Value.SpawnedEntities[i].GetInterface()->ResetEntity(Gym);
+					UniqueTypeEntities.Value.SpawnedEntities[i].GetInterface()->Enable();
+				}
+				else
+				{
+					UniqueTypeEntities.Value.SpawnedEntities[i].GetInterface()->Disable();
+				}
+			}
+		}
+	}
+}
+
+TScriptInterface<ILearningAgentsEntityInterface> ULearningAgentsEntitiesManagerComponent::SpawnEntity(const FEntityInfo& EntityInfo, const FTransform& InTransform)
+{
+	TArray<TScriptInterface<ILearningAgentsEntityInterface>> SpawnedEntities = SpawnEntities(EntityInfo, 1, false, InTransform);
+	if (!SpawnedEntities.IsEmpty())
+	{
+		return SpawnedEntities.Top();
+	}
+	else
+	{
+		UE_LOG(LogLearning, Warning, TEXT("Could not spawn entity from class %s!"), *EntityInfo.EntityClass->GetName());
+		return TScriptInterface<ILearningAgentsEntityInterface>();
+	}
+}
+
+TArray<TScriptInterface<ILearningAgentsEntityInterface>> ULearningAgentsEntitiesManagerComponent::SpawnEntities(const FEntityInfo& EntityInfo, int32 SpawnCount, bool bRandomizeTransformInGym, const FTransform& InTransform)
+{
+	TArray<TScriptInterface<ILearningAgentsEntityInterface>> SpawnedEntities;
+	SpawnedEntities.SetNumUninitialized(SpawnCount);
+
+	FActorSpawnParameters SpawnParameters;
+	SpawnParameters.Owner = GetOwner();
+
+	FTransform Transform = InTransform;
+
+	int32 SuccessfulSpawnCount = 0;
+	if (FSpawnedEntitiesInfo* SpawnEntitiesInfo = EntitiesPool.Find(EntityInfo.EntityClass->GetFName()))
+	{
+		for (TScriptInterface<ILearningAgentsEntityInterface>& Entity : SpawnEntitiesInfo->SpawnedEntities)
+		{
+			if (SuccessfulSpawnCount == SpawnCount)
+			{
+				return SpawnedEntities;
+			}
+			if (!Entity.GetInterface()->IsEnabled())
+			{
+				if (TObjectPtr<AActor> ActorEntity = Cast<AActor>(Entity.GetObject()))
+				{
+					if (bRandomizeTransformInGym)
+					{
+						RandomizeTransform(Transform, EntityInfo.EntitySpawnZOffset);
+					}
+					else
+					{
+						ProjectTransform(Transform);
+					}
+					ActorEntity->SetActorTransform(Transform);
+				}
+
+				Entity.GetInterface()->Enable();
+				SpawnedEntities.Add(Entity);
+
+				SuccessfulSpawnCount++;
+			}
+		}
+		// Allocate new object and add to pool because there's not enough in pool to pull from.
+		for (int32 Index = SuccessfulSpawnCount; SuccessfulSpawnCount < SpawnCount; ++SuccessfulSpawnCount)
+		{
+			if (bRandomizeTransformInGym)
+			{
+				RandomizeTransform(Transform, EntityInfo.EntitySpawnZOffset);
+			}
+			else
+			{
+				ProjectTransform(Transform);
+			}
+			TObjectPtr<AActor> Entity = GetWorld()->SpawnActor<AActor>(EntityInfo.EntityClass, Transform.GetLocation(), Transform.GetRotation().Rotator(), SpawnParameters);
+			SpawnEntitiesInfo->SpawnedEntities.Add(Entity);
+			SpawnedEntities.Add(TScriptInterface<ILearningAgentsEntityInterface>(Entity));
+		}
+	}
+	// New entity spawned during an episode
+	else
+	{
+		FSpawnedEntitiesInfo NewSpawnEntitiesInfo;
+		NewSpawnEntitiesInfo.EntityInfo = EntityInfo;
+
+		if (bRandomizeTransformInGym)
+		{
+			RandomizeTransform(Transform, EntityInfo.EntitySpawnZOffset);
+		}
+		else
+		{
+			ProjectTransform(Transform);
+		}
+		for (int32 Index = 0; Index < SpawnCount; ++Index)
+		{
+			TObjectPtr<AActor> Entity = GetWorld()->SpawnActor<AActor>(EntityInfo.EntityClass, Transform.GetLocation(), Transform.GetRotation().Rotator(), SpawnParameters);
+			NewSpawnEntitiesInfo.SpawnedEntities.Add(Entity);
+			SpawnedEntities.Add(TScriptInterface<ILearningAgentsEntityInterface>(Entity));
+		}
+		EntitiesPool.Add(EntityInfo.EntityClass->GetFName(), NewSpawnEntitiesInfo);
+	}
+	return SpawnedEntities;
+}
+
+void ULearningAgentsEntitiesManagerComponent::BeginPlay()
+{
+	Super::BeginPlay();
+	CheckEntityClasses();
+
+	for (FEntityInfo& EntityInfo : Entities)
+	{
+		FSpawnedEntitiesInfo SpawnedEntitiesInfo;
+		SpawnedEntitiesInfo.EntityInfo = EntityInfo;
+		EntitiesPool.Add(EntityInfo.EntityClass->GetFName(), SpawnedEntitiesInfo);
+	}
+}
+
+void ULearningAgentsEntitiesManagerComponent::ProjectTransform(FTransform& Transform) const
+{
+	if (TObjectPtr<ALearningAgentsGymBase> Gym = Cast<ALearningAgentsGymBase>(GetOwner()))
+	{
+		Transform.SetLocation(Gym->ProjectPointToGym(Transform.GetLocation()));
+	}
+	else
+	{
+		UE_LOG(LogLearning, Warning, TEXT("Entities Manager Component attached to an owner that is not a ALearningAgentsGymBase!"));
+	}
+}
+
+void ULearningAgentsEntitiesManagerComponent::RandomizeTransform(FTransform& OutTransform, float LocationZOffset) const
+{
+	if (TObjectPtr<ALearningAgentsGymBase> Gym = Cast<ALearningAgentsGymBase>(GetOwner()))
+	{
+		FVector Location = Gym->GenerateRandomLocationInGym();
+		Location.Z += LocationZOffset;
+		OutTransform.SetLocation(Location);
+		OutTransform.SetRotation(Gym->GenerateRandomRotationInGym().Quaternion());
+	}
+	else
+	{
+		UE_LOG(LogLearning, Warning, TEXT("Entities Manager Component attached to an owner that is not a ALearningAgentsGymBase!"));
+	}
+}
+
+bool ULearningAgentsEntitiesManagerComponent::CheckEntityClasses() const
+{
+	if (Entities.IsEmpty())
+	{
+		UE_LOG(LogLearning, Error, TEXT("No entity classes are setup for %s!"), *GetName());
+		return false;
+	}
+
+	for (const FEntityInfo& EntityInfo : Entities)
+	{
+		if (!EntityInfo.EntityClass)
+		{
+			UE_LOG(LogLearning, Error, TEXT("An entity spawn entry is set to None for %s!"), *GetName());
+			return false;
+		}
+
+		if (!EntityInfo.EntityClass->ImplementsInterface(ULearningAgentsEntityInterface::StaticClass()))
+		{
+			UE_LOG(LogLearning, Error, TEXT("Invalid entity class! %s does not implement ILearningAgentsEntityInterface!"), *EntityInfo.EntityClass->GetName());
+			return false;
+		}
+	}
+	return true;
+}
diff --git a/Engine/Plugins/Experimental/LearningAgents/Source/LearningAgentsTraining/Private/LearningAgentsGym.cpp b/Engine/Plugins/Experimental/LearningAgents/Source/LearningAgentsTraining/Private/LearningAgentsGym.cpp
new file mode 100644
--- /dev/null
+++ b/Engine/Plugins/Experimental/LearningAgents/Source/LearningAgentsTraining/Private/LearningAgentsGym.cpp
@@ -0,0 +1,55 @@
+// Copyright Epic Games, Inc. All Rights Reserved.
+
+
+#include "LearningAgentsGym.h"
+#include "LearningLog.h"
+#include "LearningAgentsEntityInterface.h"
+#include "LearningAgentsEntitiesManagerComponent.h"
+#include "Kismet/KismetMathLibrary.h"
+
+ALearningAgentsGymBase::ALearningAgentsGymBase()
+{
+	EntitiesManagerComponent = CreateDefaultSubobject<ULearningAgentsEntitiesManagerComponent>(TEXT("EntitiesManagerComponent"));
+
+	PrimaryActorTick.bCanEverTick = false;
+	PrimaryActorTick.bStartWithTickEnabled = false;
+}
+
+void ALearningAgentsGymBase::Initialize()
+{
+	if (!RandomStream.IsValid())
+	{
+		RandomStream = MakeShareable(new FRandomStream);
+		RandomStream->Initialize(RandomSeed);
+	}
+	EntitiesManagerComponent->InitializeEntities();
+}
+
+void ALearningAgentsGymBase::Reset()
+{
+	EntitiesManagerComponent->ResetEntities();
+}
+
+void ALearningAgentsGymBase::GetRandomStream(FRandomStream& OutRandomStream) const
+{
+	OutRandomStream = *RandomStream;
+}
+
+TSharedPtr<FRandomStream> ALearningAgentsGymBase::GetRandomStream() const
+{
+	return RandomStream;
+}
+
+void ALearningAgentsGymBase::SetRandomStream(const TSharedPtr<FRandomStream>& InRandomStream)
+{
+	RandomStream = InRandomStream;
+}
+
+bool ALearningAgentsGymBase::IsMemberOfGym(TObjectPtr<AActor> Actor) const
+{
+	if (ILearningAgentsEntityTrainingInterface* EntityInterface = Cast<ILearningAgentsEntityTrainingInterface>(Actor))
+	{
+		return EntityInterface->GetGym() == this;
+	}
+	return false;
+}
diff --git a/Engine/Plugins/Experimental/LearningAgents/Source/LearningAgentsTraining/Private/LearningAgentsGymSimple.cpp b/Engine/Plugins/Experimental/LearningAgents/Source/LearningAgentsTraining/Private/LearningAgentsGymSimple.cpp
new file mode 100644
--- /dev/null
+++ b/Engine/Plugins/Experimental/LearningAgents/Source/LearningAgentsTraining/Private/LearningAgentsGymSimple.cpp
@@ -0,0 +1,66 @@
+// Copyright Epic Games, Inc. All Rights Reserved.
+
+
+#include "LearningAgentsGymSimple.h"
+#include "LearningLog.h"
+#include "Engine/World.h"
+#include "NavigationSystem.h"
+#include "AI/NavigationSystemBase.h"
+#include "NavMesh/NavMeshBoundsVolume.h"
+#include "Components/StaticMeshComponent.h"
+#include "Kismet/KismetMathLibrary.h"
+
+ALearningAgentsGymSimple::ALearningAgentsGymSimple()
+{
+	RootComponent = CreateDefaultSubobject<USceneComponent>(TEXT("RootComponent"));
+	SimpleGymFloor = CreateDefaultSubobject<UStaticMeshComponent>(TEXT("SimpleGymFloor"));
+	SimpleGymFloor->SetupAttachment(RootComponent);
+}
+
+void ALearningAgentsGymSimple::GetGymExtents(FVector& OutMinBounds, FVector& OutMaxBounds) const
+{
+	SimpleGymFloor->GetLocalBounds(OutMinBounds, OutMaxBounds);
+	OutMinBounds *= SimpleGymFloor->GetComponentScale();
+	OutMaxBounds *= SimpleGymFloor->GetComponentScale();
+}
+
+FRotator ALearningAgentsGymSimple::GenerateRandomRotationInGym() const
+{
+	check(RandomStream);
+
+	return FRotator(0.0f, RandomStream->FRandRange(0.0f, 360.0f), 0.0f);
+}
+
+FVector ALearningAgentsGymSimple::GenerateRandomLocationInGym() const
+{
+	FVector MinBounds, MaxBounds;
+	GetGymExtents(MinBounds, MaxBounds);
+
+	// Find random point in radius taken from min axis of max bounds
+	FVector RandomPoint = GetActorLocation() + RandomStream->VRand() * RandomStream->FRandRange(0.0f, FMath::Min(MaxBounds.X, MaxBounds.Y));
+	RandomPoint.Z = GetActorLocation().Z;
+	return ProjectPointToGym(RandomPoint);
+}
+
+FVector ALearningAgentsGymSimple::ProjectPointToGym(const FVector& InLocation) const
+{
+	FVector MinBounds, MaxBounds;
+	GetGymExtents(MinBounds, MaxBounds);
+	MaxBounds.Z += 100.0f;
+
+	TObjectPtr<UWorld> World = GetWorld();
+	
+	check(World);
+	TObjectPtr<UNavigationSystemV1> NavigationSystem = FNavigationSystem::GetCurrent<UNavigationSystemV1>(World);
+	
+	if (NavigationSystem)
+	{
+		FNavLocation NavLocation;
+		NavigationSystem->ProjectPointToNavigation(InLocation, NavLocation, MaxBounds);
+		return NavLocation.Location;
+	}
+	else
+	{
+		return InLocation;
+	}
+}
diff --git a/Engine/Plugins/Experimental/LearningAgents/Source/LearningAgentsTraining/Private/LearningAgentsGymsManager.cpp b/Engine/Plugins/Experimental/LearningAgents/Source/LearningAgentsTraining/Private/LearningAgentsGymsManager.cpp
new file mode 100644
--- /dev/null
+++ b/Engine/Plugins/Experimental/LearningAgents/Source/LearningAgentsTraining/Private/LearningAgentsGymsManager.cpp
@@ -0,0 +1,115 @@
+// Copyright Epic Games, Inc. All Rights Reserved.
+
+
+#include "LearningAgentsGymsManager.h"
+#include "LearningAgentsGym.h"
+#include "LearningLog.h"
+#include "Engine/World.h"
+#include "NavigationSystem.h"
+#include "AI/NavigationSystemBase.h"
+#include "Misc/CommandLine.h"
+
+ALearningAgentsGymsManager::ALearningAgentsGymsManager()
+{
+	PrimaryActorTick.bCanEverTick = false;
+	PrimaryActorTick.bStartWithTickEnabled = false;
+}
+
+void ALearningAgentsGymsManager::SpawnGyms()
+{
+	InitializeRandomStream();
+
+	FVector TemplateSpawnLocation = GetActorLocation();
+
+	for (FSpawnGymInfo& GymTemplate : GymTemplates)
+	{
+		FActorSpawnParameters SpawnParameters;
+		SpawnParameters.Owner = this;
+
+		int32 CurrentSpawnCount = 0;
+		int32 GridSize = FMath::CeilToInt32(FMath::Sqrt(static_cast<float>(GymTemplate.SpawnCount)));
+		bool bFirstGym = true;
+
+		FVector MinBounds = FVector::ZeroVector;
+		FVector MaxBounds = FVector::ZeroVector;
+		FName GymTag;
+
+		for (int32 Row = 0; Row < GridSize; ++Row)
+		{
+			for (int32 Col = 0; Col < GridSize; ++Col)
+			{
+				if (CurrentSpawnCount == GymTemplate.SpawnCount)
+				{
+					break;
+				}
+
+				TObjectPtr<ALearningAgentsGymBase> Gym;
+
+				if (bFirstGym)
+				{
+					bFirstGym = false;
+					Gym = GetWorld()->SpawnActor<ALearningAgentsGymBase>(GymTemplate.GymClass, TemplateSpawnLocation, FRotator::ZeroRotator, SpawnParameters);
+					Gym->GetGymExtents(MinBounds, MaxBounds);
+				}
+				else
+				{
+					FVector SpawnLocation = FVector(Row * (GymsSpacing + (MaxBounds.X - MinBounds.X)) + TemplateSpawnLocation.X, Col * (GymsSpacing + (MaxBounds.Y - MinBounds.Y) + TemplateSpawnLocation.Y), TemplateSpawnLocation.Z);
+					Gym = GetWorld()->SpawnActor<ALearningAgentsGymBase>(GymTemplate.GymClass, SpawnLocation, FRotator::ZeroRotator, SpawnParameters);
+				}
+
+				SpawnedGyms.Add(Gym);
+				
+				check(RandomStream.IsValid());
+				Gym->SetRandomStream(RandomStream);
+
+				CurrentSpawnCount++;
+			}
+		}
+		TemplateSpawnLocation.X += (GymsSpacing + (MaxBounds.X - MinBounds.X)) * (GridSize);
+	}
+	
+	if (TObjectPtr<UNavigationSystemV1> NavigationSystem = FNavigationSystem::GetCurrent<UNavigationSystemV1>(GetWorld()))
+	{
+		NavigationSystem->Build();
+	}
+	else
+	{
+		UE_LOG(LogLearning, Log, TEXT("Unable to build navigation system due to invalid navigation system reference!"));
+	}
+}
+
+int32 ALearningAgentsGymsManager::GetGymsCount() const
+{
+	int32 Count = 0;
+	for (const FSpawnGymInfo& GymTemplate : GymTemplates)
+	{
+		Count += GymTemplate.SpawnCount;
+	}
+	return Count;
+}
+
+void ALearningAgentsGymsManager::InitializeRandomStream()
+{
+	FString RandomSeedInput;
+	if (FParse::Value(FCommandLine::Get(), TEXT("GymsManagerRandomSeed"), RandomSeedInput))
+	{
+		if (RandomSeedInput.StartsWith(TEXT("=")))
+		{
+			RandomSeedInput = RandomSeedInput.RightChop(1);
+			UE_LOG(LogLearning, Log, TEXT("Removed leading '=' from RandomSeedInput: %s"), *RandomSeedInput);
+		}
+		RandomSeed = FCString::Atoi(*RandomSeedInput);
+	}
+	RandomStream = MakeShareable(new FRandomStream);
+	RandomStream->Initialize(RandomSeed);
+	UE_LOG(LogLearning, Log, TEXT("GymsManager initialized with seed %d!"), RandomSeed);
+}
+
+void ALearningAgentsGymsManager::Start()
+{
+	SpawnGyms();
+	for (TObjectPtr<ALearningAgentsGymBase>& SpawnedGym : SpawnedGyms)
+	{
+		SpawnedGym->Initialize();
+	}
+}
diff --git a/Engine/Plugins/Experimental/LearningAgents/Source/LearningAgentsTraining/Public/LearningAgentsEntitiesManagerComponent.h b/Engine/Plugins/Experimental/LearningAgents/Source/LearningAgentsTraining/Public/LearningAgentsEntitiesManagerComponent.h
new file mode 100644
--- /dev/null
+++ b/Engine/Plugins/Experimental/LearningAgents/Source/LearningAgentsTraining/Public/LearningAgentsEntitiesManagerComponent.h
@@ -0,0 +1,85 @@
+// Copyright Epic Games, Inc. All Rights Reserved.
+
+#pragma once
+
+#include "CoreMinimal.h"
+#include "Components/ActorComponent.h"
+#include "LearningAgentsEntitiesManagerComponent.generated.h"
+
+class ILearningAgentsEntityInterface;
+
+/**
+ * FEntityInfo holds spawn information of a single entity type.
+ */
+USTRUCT(BlueprintType)
+struct FEntityInfo
+{
+	GENERATED_BODY()
+
+	/** Specify the entity class to spawn. */
+	UPROPERTY(EditDefaultsOnly, BlueprintReadWrite, Category = "LearningAgents")
+	TSubclassOf<AActor> EntityClass;
+
+	/** Specify the Z offset to apply to spawn locations. */
+	UPROPERTY(EditDefaultsOnly, BlueprintReadWrite, Category = "LearningAgents")
+	float EntitySpawnZOffset = 0.0f;
+
+	/** The number of entities to spawn at the start of an episode. */
+	UPROPERTY(EditDefaultsOnly, BlueprintReadWrite, Category = "LearningAgents")
+	int32 EpisodeStartEntitySpawnCount = 0;
+};
+
+/**
+ * FSpawnedEntitiesInfo holds references to spawned entities.
+ */
+USTRUCT(BlueprintType)
+struct FSpawnedEntitiesInfo
+{
+	GENERATED_BODY()
+
+	UPROPERTY(VisibleAnywhere, Category = "LearningAgents")
+	FEntityInfo EntityInfo;
+
+	UPROPERTY(VisibleAnywhere, Category = "LearningAgents")
+	TArray<TScriptInterface<ILearningAgentsEntityInterface>> SpawnedEntities;
+};
+
+/**
+ * ULearningAgentsEntitiesManagerComponent handles the spawn and reset of multiple entity types during training.
+ */
+UCLASS(ClassGroup=(LearningAgents), meta=(BlueprintSpawnableComponent))
+class LEARNINGAGENTSTRAINING_API ULearningAgentsEntitiesManagerComponent : public UActorComponent 
+{
+	GENERATED_BODY()
+
+public:
+	/** Specify the entity types to spawn. */
+	UPROPERTY(EditDefaultsOnly, Category = "LearningAgents")
+	TArray<FEntityInfo> Entities;
+
+	/** Initializes the entities that the component manages. */
+	void InitializeEntities();
+
+	/** Resets entities that the component manages. */
+	void ResetEntities();
+
+	/** Spawns a single pooled entity. */
+	UFUNCTION(BlueprintCallable, Category = "LearningAgents")
+	TScriptInterface<ILearningAgentsEntityInterface> SpawnEntity(const FEntityInfo& EntityInfo, const FTransform& Transform);
+
+	/** Spawns multiple pooled entities. */
+	UFUNCTION(BlueprintCallable, Category = "LearningAgents", meta = (ClampMin = "0", UIMin = "0"))
+	TArray<TScriptInterface<ILearningAgentsEntityInterface>> SpawnEntities(const FEntityInfo& EntityInfo, int32 SpawnCount, bool bRandomizeTransformInGym = true, const FTransform& Transform = FTransform());
+
+protected:
+	virtual void BeginPlay() override;
+
+private:
+	TMap<FName, FSpawnedEntitiesInfo> EntitiesPool;
+
+	void ProjectTransform(FTransform& Transform) const;
+
+	void RandomizeTransform(FTransform& OutTransform, float LocationZOffset) const;
+	
+	bool CheckEntityClasses() const;
+};
diff --git a/Engine/Plugins/Experimental/LearningAgents/Source/LearningAgentsTraining/Public/LearningAgentsEntityInterface.h b/Engine/Plugins/Experimental/LearningAgents/Source/LearningAgentsTraining/Public/LearningAgentsEntityInterface.h
new file mode 100644
--- /dev/null
+++ b/Engine/Plugins/Experimental/LearningAgents/Source/LearningAgentsTraining/Public/LearningAgentsEntityInterface.h
@@ -0,0 +1,66 @@
+// Copyright Epic Games, Inc. All Rights Reserved.
+
+#pragma once
+
+#include "CoreMinimal.h"
+#include "GameFramework/Actor.h"
+#include "LearningAgentsGym.h"
+#include "LearningAgentsEntityInterface.generated.h"
+
+/**
+* Classes that implement the Entity Training Interface need to provide definitions for how to initialize an object for training and resetting the 
+* object after a training episode ends.
+*/
+UINTERFACE(MinimalAPI, meta = (CannotImplementInterfaceInBlueprint), BlueprintType)
+class ULearningAgentsEntityTrainingInterface : public UInterface
+{
+	GENERATED_BODY()
+};
+
+/** Completion modes for episodes. */
+class LEARNINGAGENTSTRAINING_API ILearningAgentsEntityTrainingInterface
+{
+	GENERATED_BODY()
+
+public:
+	/** Initialize an entity at the start of training. */
+	UFUNCTION(BlueprintCallable, Category = "LearningAgents")
+	virtual	void InitializeEntity(ALearningAgentsGymBase* InGym) = 0;
+
+	/** Resets an entity for a new training episode. */
+	UFUNCTION(BlueprintCallable, Category = "LearningAgents")
+	virtual void ResetEntity(ALearningAgentsGymBase* InGym) = 0;
+
+	/** Resets an entity for a new training episode. */
+	UFUNCTION(BlueprintCallable, Category = "LearningAgents")
+	virtual ALearningAgentsGymBase* GetGym() const = 0;
+};
+
+/**
+* Classes that implement the Entity Interface need to provide definitions for enabling and disabling an entity during training.
+*
+* If using the LearningAgentsGym, this interface should be implemented by any actor that is spawned during training.
+*/
+UINTERFACE(MinimalAPI, meta = (CannotImplementInterfaceInBlueprint), BlueprintType)
+class ULearningAgentsEntityInterface : public ULearningAgentsEntityTrainingInterface
+{
+	GENERATED_BODY()
+};
+
+class LEARNINGAGENTSTRAINING_API ILearningAgentsEntityInterface: public ILearningAgentsEntityTrainingInterface
+{
+	GENERATED_BODY()
+
+public:
+	/** Enables an entity to be used for training. */
+	UFUNCTION(BlueprintCallable, Category = "LearningAgents")
+	virtual void Enable() = 0;
+
+	/** Disables an entity used for training. */
+	UFUNCTION(BlueprintCallable, Category = "LearningAgents")
+	virtual void Disable() = 0;
+
+	/** Checks whether an entity is eligible to be used for training. */
+	UFUNCTION(BlueprintCallable, Category = "LearningAgents")
+	virtual bool IsEnabled() const = 0;
+};
diff --git a/Engine/Plugins/Experimental/LearningAgents/Source/LearningAgentsTraining/Public/LearningAgentsGym.h b/Engine/Plugins/Experimental/LearningAgents/Source/LearningAgentsTraining/Public/LearningAgentsGym.h
new file mode 100644
--- /dev/null
+++ b/Engine/Plugins/Experimental/LearningAgents/Source/LearningAgentsTraining/Public/LearningAgentsGym.h
@@ -0,0 +1,67 @@
+// Copyright Epic Games, Inc. All Rights Reserved.
+
+#pragma once
+
+#include "CoreMinimal.h"
+#include "GameFramework/Actor.h"
+#include "LearningAgentsEntitiesManagerComponent.h"
+#include "LearningAgentsGym.generated.h"
+
+#define INVALID_GYM_VECTOR (FVector::ZeroVector)
+#define INVALID_GYM_ROTATOR (FRotator::ZeroRotator)
+
+/**
+ * The Gym Base abstract class handles the start and reset of entities training in a single gym.
+ */
+UCLASS(Abstract, NotBlueprintable)
+class LEARNINGAGENTSTRAINING_API ALearningAgentsGymBase : public AActor
+{
+	GENERATED_BODY()
+	
+public:
+	ALearningAgentsGymBase();
+
+	/** Initializes the gym at the start of training. */
+	UFUNCTION(BlueprintCallable, Category = "LearningAgents")
+	void Initialize();
+
+	/** Resets the gym for a new training episode. */
+	UFUNCTION(BlueprintCallable, Category = "LearningAgents")
+	void Reset();
+
+public:
+	UPROPERTY(EditDefaultsOnly, Category = "LearningAgents")
+	TObjectPtr<ULearningAgentsEntitiesManagerComponent> EntitiesManagerComponent;
+
+	/** Retrieves the current random stream used by the Gym.*/
+	UFUNCTION(BlueprintPure, Category = "LearningAgents")
+	void GetRandomStream(FRandomStream& OutRandomStream) const;
+	
+	/** Retrieves the current random stream used by the Gym.*/
+	TSharedPtr<FRandomStream> GetRandomStream() const;
+
+	/** Sets the random stream used by the Gym.*/
+	void SetRandomStream(const TSharedPtr<FRandomStream>& InRandomStream);
+
+	/** Checks if an actor is training in this gym.*/
+	bool IsMemberOfGym(TObjectPtr<AActor> Actor) const;
+
+	/** Generates a random rotator using the gym's random stream.*/
+	virtual FRotator GenerateRandomRotationInGym() const PURE_VIRTUAL(ALearningAgentsGymBase::GenerateRandomRotationInGym, return INVALID_GYM_ROTATOR;);
+
+	/** Generates a valid random point in the gym using the gym's random stream. Must be overridden in a derived class. */
+	virtual FVector GenerateRandomLocationInGym() const PURE_VIRTUAL(ALearningAgentsGymBase::GenerateRandomLocationInGym, return INVALID_GYM_VECTOR;);
+
+	/** Projects a point onto a valid location in the gym. Must be overridden in a derived class. */
+	virtual FVector ProjectPointToGym(const FVector& InPoint) const PURE_VIRTUAL(ALearningAgentsGymBase::ProjectPointToGym, return INVALID_GYM_VECTOR;);
+
+	/** Gets the gym max and min bounds. Must be overridden in a derived class. */
+	virtual void GetGymExtents(FVector& OutMinExtents, FVector& OutMaxExtents) const PURE_VIRTUAL(ALearningAgentsGymBase::GetGymExtents, );
+
+protected:
+	/** The random seed used for spawn locations. Note: This is only used if a random stream is not setup by the GymsManager.*/
+	UPROPERTY(EditDefaultsOnly, Category = "LearningAgents")
+	int32 RandomSeed = 1234;
+
+	TSharedPtr<FRandomStream> RandomStream; 
+};
diff --git a/Engine/Plugins/Experimental/LearningAgents/Source/LearningAgentsTraining/Public/LearningAgentsGymSimple.h b/Engine/Plugins/Experimental/LearningAgents/Source/LearningAgentsTraining/Public/LearningAgentsGymSimple.h
new file mode 100644
--- /dev/null
+++ b/Engine/Plugins/Experimental/LearningAgents/Source/LearningAgentsTraining/Public/LearningAgentsGymSimple.h
@@ -0,0 +1,32 @@
+// Copyright Epic Games, Inc. All Rights Reserved.
+
+#pragma once
+
+#include "CoreMinimal.h"
+#include "LearningAgentsGym.h"
+#include "LearningAgentsGymSimple.generated.h"
+
+class UStaticMeshComponent;
+
+/**
+ * A simple gym template class that uses a static mesh as its training floor.
+ */
+UCLASS(ClassGroup = (LearningAgents), BlueprintType, Blueprintable)
+class LEARNINGAGENTSTRAINING_API ALearningAgentsGymSimple : public ALearningAgentsGymBase
+{
+	GENERATED_BODY()
+	
+public:
+	ALearningAgentsGymSimple();
+
+	UPROPERTY(EditDefaultsOnly, Category = "LearningAgents")
+	TObjectPtr<UStaticMeshComponent> SimpleGymFloor;
+
+	virtual void GetGymExtents(FVector& OutMinExtents, FVector& OutMaxExtents) const override;
+
+	virtual FRotator GenerateRandomRotationInGym() const override;
+
+	virtual FVector GenerateRandomLocationInGym() const override;
+
+	virtual FVector ProjectPointToGym(const FVector& InPoint) const override;
+};
diff --git a/Engine/Plugins/Experimental/LearningAgents/Source/LearningAgentsTraining/Public/LearningAgentsGymsManager.h b/Engine/Plugins/Experimental/LearningAgents/Source/LearningAgentsTraining/Public/LearningAgentsGymsManager.h
new file mode 100644
--- /dev/null
+++ b/Engine/Plugins/Experimental/LearningAgents/Source/LearningAgentsTraining/Public/LearningAgentsGymsManager.h
@@ -0,0 +1,68 @@
+// Copyright Epic Games, Inc. All Rights Reserved.
+
+#pragma once
+
+#include "CoreMinimal.h"
+#include "GameFramework/Actor.h"
+#include "LearningAgentsGymsManager.generated.h"
+
+class ALearningAgentsGymBase;
+class UNavigationSystemV1;
+
+/**
+ * This struct holds information (num to spawn etc) on one gym template to spawn.
+ */
+USTRUCT(BlueprintType)
+struct FSpawnGymInfo
+{
+	GENERATED_BODY()
+
+	/** The gym class used to construct gym instances. */
+	UPROPERTY(EditAnywhere, BlueprintReadWrite, Category = "LearningAgents")
+	TSubclassOf<ALearningAgentsGymBase> GymClass;
+
+	/** The number of gym instances to spawn. */
+	UPROPERTY(EditAnywhere, BlueprintReadWrite, Category = "LearningAgents")
+	int32 SpawnCount = 0;
+};
+
+
+/**
+ * The Learning Agents GymsManager centralizes the start and reset of training for multiple gym templates.
+ */
+UCLASS(ClassGroup = (LearningAgents), Blueprintable)
+class LEARNINGAGENTSTRAINING_API ALearningAgentsGymsManager : public AActor
+{
+	GENERATED_BODY()
+	
+	ALearningAgentsGymsManager();
+
+public:		
+	/** The random seed to initialize the random stream used by the GymsManager. */
+	UPROPERTY(EditAnywhere, Category = "LearningAgents")
+	int32 RandomSeed = 1234;
+
+	/** The distance between each parallel gym when spawned. */
+	UPROPERTY(EditAnywhere, Category = "LearningAgents")
+	float GymsSpacing = 300.0f;
+
+	/** Spawns and initializes gym instances. */
+	UFUNCTION(BlueprintCallable, Category = "LearningAgents")
+	void Start();
+
+	/** Get the number of gyms managed by the GymsManager. */
+	UFUNCTION(BlueprintCallable, Category = "LearningAgents")
+	int32 GetGymsCount() const;
+
+private:
+	UPROPERTY(EditAnywhere, Category = "LearningAgents")
+	TArray<FSpawnGymInfo> GymTemplates;
+
+	TArray<TObjectPtr<ALearningAgentsGymBase>> SpawnedGyms;
+
+	TSharedPtr<FRandomStream> RandomStream;
+
+	void InitializeRandomStream();
+
+	void SpawnGyms();
+};
