LearningAgents: Fixed bug in GetSetObservation

diff --git a/Engine/Plugins/Experimental/LearningAgents/Source/LearningAgents/Private/LearningAgentsObservations.cpp b/Engine/Plugins/Experimental/LearningAgents/Source/LearningAgents/Private/LearningAgentsObservations.cpp
--- a/Engine/Plugins/Experimental/LearningAgents/Source/LearningAgents/Private/LearningAgentsObservations.cpp
+++ b/Engine/Plugins/Experimental/LearningAgents/Source/LearningAgents/Private/LearningAgentsObservations.cpp
@@ -3293,15 +3293,6 @@ bool ULearningAgentsObservations::GetSetObservation(TSet<FLearningAgentsObservat
 
 	const UE::Learning::Observation::FObjectSetParameters Parameters = Object->ObservationObject.GetSet(Element.ObjectElement);
 
-	if (Parameters.Elements.Num() != OutElements.Num())
-	{
-		UE_LOG(LogLearning, Error, TEXT("GetSetObservation: Observation '%s' size does not match. Observation is '%i' elements but asked for '%i'."),
-			*Object->ObservationObject.GetTag(Element.ObjectElement).ToString(),
-			Parameters.Elements.Num(), OutElements.Num());
-		OutElements.Empty();
-		return false;
-	}
-
 	OutElements.Empty(Parameters.Elements.Num());
 	for (int32 ElementIdx = 0; ElementIdx < Parameters.Elements.Num(); ElementIdx++)
 	{
