[LearningAgents] [LearningCompanions] Updates to EntitiesManager, added LearningAgentsLearningComponentInterface and other changes

diff --git a/Engine/Plugins/Experimental/LearningAgents/Source/LearningAgentsTraining/Private/LearningAgentsEntitiesManagerComponent.cpp b/Engine/Plugins/Experimental/LearningAgents/Source/LearningAgentsTraining/Private/LearningAgentsEntitiesManagerComponent.cpp
--- a/Engine/Plugins/Experimental/LearningAgents/Source/LearningAgentsTraining/Private/LearningAgentsEntitiesManagerComponent.cpp
+++ b/Engine/Plugins/Experimental/LearningAgents/Source/LearningAgentsTraining/Private/LearningAgentsEntitiesManagerComponent.cpp
@@ -6,7 +6,7 @@
 #include "LearningAgentsGym.h"
 #include "LearningAgentsEntityInterface.h"
 
-void ULearningAgentsEntitiesManagerComponent::InitializeEntities()
+void ULearningAgentsEntitiesManagerComponent::InitializeLearningComponent()
 {
 	TObjectPtr<ALearningAgentsGymBase> Gym = Cast<ALearningAgentsGymBase>(GetOwner());
 	if (Gym)
@@ -29,14 +29,14 @@ void ULearningAgentsEntitiesManagerComponent::InitializeEntities()
 		{
 			for (TScriptInterface<ILearningAgentsEntityInterface>& Entity : UniqueTypeEntities.Value.SpawnedEntities)
 			{
-				Entity.GetInterface()->InitializeEntity(Gym);
+				ILearningAgentsEntityInterface::Execute_InitializeEntity(Entity.GetObject(), Gym);
 			}
 		}
 
 	}
 }
 
-void ULearningAgentsEntitiesManagerComponent::ResetEntities()
+void ULearningAgentsEntitiesManagerComponent::ResetLearningComponent()
 {
 	TObjectPtr<ALearningAgentsGymBase> Gym = Cast<ALearningAgentsGymBase>(GetOwner());
 	if (Gym)
@@ -58,12 +58,13 @@ void ULearningAgentsEntitiesManagerComponent::ResetEntities()
 						RandomizeTransform(RandomTransform, UniqueTypeEntities.Value.EntityInfo.EntitySpawnZOffset);
 						ActorEntity->SetActorTransform(RandomTransform);
 					}
-					UniqueTypeEntities.Value.SpawnedEntities[i].GetInterface()->ResetEntity(Gym);
-					UniqueTypeEntities.Value.SpawnedEntities[i].GetInterface()->Enable();
+
+					ILearningAgentsEntityInterface::Execute_ResetEntity(UniqueTypeEntities.Value.SpawnedEntities[i].GetObject(), Gym);
+					ILearningAgentsEntityInterface::Execute_EnableEntity(UniqueTypeEntities.Value.SpawnedEntities[i].GetObject());
 				}
 				else
 				{
-					UniqueTypeEntities.Value.SpawnedEntities[i].GetInterface()->Disable();
+					ILearningAgentsEntityInterface::Execute_DisableEntity(UniqueTypeEntities.Value.SpawnedEntities[i].GetObject());
 				}
 			}
 		}
@@ -103,7 +104,7 @@ TArray<TScriptInterface<ILearningAgentsEntityInterface>> ULearningAgentsEntities
 			{
 				return SpawnedEntities;
 			}
-			if (!Entity.GetInterface()->IsEnabled())
+			if (!ILearningAgentsEntityInterface::Execute_IsEntityEnabled(Entity.GetObject()))
 			{
 				if (TObjectPtr<AActor> ActorEntity = Cast<AActor>(Entity.GetObject()))
 				{
@@ -118,7 +119,7 @@ TArray<TScriptInterface<ILearningAgentsEntityInterface>> ULearningAgentsEntities
 					ActorEntity->SetActorTransform(Transform);
 				}
 
-				Entity.GetInterface()->Enable();
+				ILearningAgentsEntityInterface::Execute_EnableEntity(Entity.GetObject());
 				SpawnedEntities.Add(Entity);
 
 				SuccessfulSpawnCount++;
diff --git a/Engine/Plugins/Experimental/LearningAgents/Source/LearningAgentsTraining/Private/LearningAgentsGym.cpp b/Engine/Plugins/Experimental/LearningAgents/Source/LearningAgentsTraining/Private/LearningAgentsGym.cpp
--- a/Engine/Plugins/Experimental/LearningAgents/Source/LearningAgentsTraining/Private/LearningAgentsGym.cpp
+++ b/Engine/Plugins/Experimental/LearningAgents/Source/LearningAgentsTraining/Private/LearningAgentsGym.cpp
@@ -4,13 +4,11 @@
 #include "LearningAgentsGym.h"
 #include "LearningLog.h"
 #include "LearningAgentsEntityInterface.h"
-#include "LearningAgentsEntitiesManagerComponent.h"
+#include "LearningAgentsLearningComponentInterface.h"
 #include "Kismet/KismetMathLibrary.h"
 
 ALearningAgentsGymBase::ALearningAgentsGymBase()
 {
-	EntitiesManagerComponent = CreateDefaultSubobject<ULearningAgentsEntitiesManagerComponent>(TEXT("EntitiesManagerComponent"));
-
 	PrimaryActorTick.bCanEverTick = false;
 	PrimaryActorTick.bStartWithTickEnabled = false;
 }
@@ -22,12 +20,19 @@ void ALearningAgentsGymBase::Initialize()
 		RandomStream = MakeShareable(new FRandomStream);
 		RandomStream->Initialize(RandomSeed);
 	}
-	EntitiesManagerComponent->InitializeEntities();
+	PopulateLearningComponents();
+	for (TScriptInterface<ILearningAgentsLearningComponentInterface>& LearningComponent : LearningComponents)
+	{
+		LearningComponent->InitializeLearningComponent();
+	}
 }
 
 void ALearningAgentsGymBase::Reset()
 {
-	EntitiesManagerComponent->ResetEntities();
+	for (TScriptInterface<ILearningAgentsLearningComponentInterface>& LearningComponent : LearningComponents)
+	{
+		LearningComponent->ResetLearningComponent();
+	}
 }
 
 void ALearningAgentsGymBase::GetRandomStream(FRandomStream& OutRandomStream) const
@@ -47,9 +52,24 @@ void ALearningAgentsGymBase::SetRandomStream(const TSharedPtr<FRandomStream>& In
 
 bool ALearningAgentsGymBase::IsMemberOfGym(TObjectPtr<AActor> Actor) const
 {
-	if (ILearningAgentsEntityTrainingInterface* EntityInterface = Cast<ILearningAgentsEntityTrainingInterface>(Actor))
+	if (Actor and Actor->Implements<ULearningAgentsEntityInterface>())
 	{
-		return EntityInterface->GetGym() == this;
+		return ILearningAgentsEntityInterface::Execute_GetGym(Actor) == this;
 	}
 	return false;
 }
+
+void ALearningAgentsGymBase::PopulateLearningComponents()
+{
+	TArray<TObjectPtr<UActorComponent>> AllComponents;
+	this->GetComponents(AllComponents);
+
+	for (TObjectPtr<UActorComponent> Component : AllComponents)
+	{
+		if (Component && Component->Implements<ULearningAgentsLearningComponentInterface>())
+		{
+			TScriptInterface<ILearningAgentsLearningComponentInterface> LearningComponentInterface = TScriptInterface<ILearningAgentsLearningComponentInterface>(Component);
+			LearningComponents.Add(LearningComponentInterface);
+		}
+	}
+}
diff --git a/Engine/Plugins/Experimental/LearningAgents/Source/LearningAgentsTraining/Public/LearningAgentsEntitiesManagerComponent.h b/Engine/Plugins/Experimental/LearningAgents/Source/LearningAgentsTraining/Public/LearningAgentsEntitiesManagerComponent.h
--- a/Engine/Plugins/Experimental/LearningAgents/Source/LearningAgentsTraining/Public/LearningAgentsEntitiesManagerComponent.h
+++ b/Engine/Plugins/Experimental/LearningAgents/Source/LearningAgentsTraining/Public/LearningAgentsEntitiesManagerComponent.h
@@ -4,6 +4,7 @@
 
 #include "CoreMinimal.h"
 #include "Components/ActorComponent.h"
+#include "LearningAgentsLearningComponentInterface.h"
 #include "LearningAgentsEntitiesManagerComponent.generated.h"
 
 class ILearningAgentsEntityInterface;
@@ -48,7 +49,7 @@ struct FSpawnedEntitiesInfo
  * ULearningAgentsEntitiesManagerComponent handles the spawn and reset of multiple entity types during training.
  */
 UCLASS(ClassGroup=(LearningAgents), meta=(BlueprintSpawnableComponent))
-class LEARNINGAGENTSTRAINING_API ULearningAgentsEntitiesManagerComponent : public UActorComponent 
+class LEARNINGAGENTSTRAINING_API ULearningAgentsEntitiesManagerComponent : public UActorComponent, public ILearningAgentsLearningComponentInterface
 {
 	GENERATED_BODY()
 
@@ -58,10 +59,10 @@ public:
 	TArray<FEntityInfo> Entities;
 
 	/** Initializes the entities that the component manages. */
-	void InitializeEntities();
+	void InitializeLearningComponent();
 
 	/** Resets entities that the component manages. */
-	void ResetEntities();
+	void ResetLearningComponent();
 
 	/** Spawns a single pooled entity. */
 	UFUNCTION(BlueprintCallable, Category = "LearningAgents")
@@ -75,6 +76,7 @@ protected:
 	virtual void BeginPlay() override;
 
 private:
+	UPROPERTY(VisibleAnywhere, Category = "LearningAgents")
 	TMap<FName, FSpawnedEntitiesInfo> EntitiesPool;
 
 	void ProjectTransform(FTransform& Transform) const;
diff --git a/Engine/Plugins/Experimental/LearningAgents/Source/LearningAgentsTraining/Public/LearningAgentsEntityInterface.h b/Engine/Plugins/Experimental/LearningAgents/Source/LearningAgentsTraining/Public/LearningAgentsEntityInterface.h
--- a/Engine/Plugins/Experimental/LearningAgents/Source/LearningAgentsTraining/Public/LearningAgentsEntityInterface.h
+++ b/Engine/Plugins/Experimental/LearningAgents/Source/LearningAgentsTraining/Public/LearningAgentsEntityInterface.h
@@ -11,7 +11,7 @@
 * Classes that implement the Entity Training Interface need to provide definitions for how to initialize an object for training and resetting the 
 * object after a training episode ends.
 */
-UINTERFACE(MinimalAPI, meta = (CannotImplementInterfaceInBlueprint), BlueprintType)
+UINTERFACE(MinimalAPI, Blueprintable, BlueprintType)
 class ULearningAgentsEntityTrainingInterface : public UInterface
 {
 	GENERATED_BODY()
@@ -24,16 +24,16 @@ class LEARNINGAGENTSTRAINING_API ILearningAgentsEntityTrainingInterface
 
 public:
 	/** Initialize an entity at the start of training. */
-	UFUNCTION(BlueprintCallable, Category = "LearningAgents")
-	virtual	void InitializeEntity(ALearningAgentsGymBase* InGym) = 0;
+	UFUNCTION(BlueprintCallable, BlueprintNativeEvent, Category = "LearningAgents")
+	void InitializeEntity(ALearningAgentsGymBase* InGym);
 
 	/** Resets an entity for a new training episode. */
-	UFUNCTION(BlueprintCallable, Category = "LearningAgents")
-	virtual void ResetEntity(ALearningAgentsGymBase* InGym) = 0;
+	UFUNCTION(BlueprintCallable, BlueprintNativeEvent, Category = "LearningAgents")
+	void ResetEntity(ALearningAgentsGymBase* InGym);
 
 	/** Resets an entity for a new training episode. */
-	UFUNCTION(BlueprintCallable, Category = "LearningAgents")
-	virtual ALearningAgentsGymBase* GetGym() const = 0;
+	UFUNCTION(BlueprintCallable, BlueprintNativeEvent, Category = "LearningAgents")
+	ALearningAgentsGymBase* GetGym() const;
 };
 
 /**
@@ -41,7 +41,7 @@ public:
 *
 * If using the LearningAgentsGym, this interface should be implemented by any actor that is spawned during training.
 */
-UINTERFACE(MinimalAPI, meta = (CannotImplementInterfaceInBlueprint), BlueprintType)
+UINTERFACE(MinimalAPI, Blueprintable, BlueprintType)
 class ULearningAgentsEntityInterface : public ULearningAgentsEntityTrainingInterface
 {
 	GENERATED_BODY()
@@ -53,14 +53,14 @@ class LEARNINGAGENTSTRAINING_API ILearningAgentsEntityInterface: public ILearnin
 
 public:
 	/** Enables an entity to be used for training. */
-	UFUNCTION(BlueprintCallable, Category = "LearningAgents")
-	virtual void Enable() = 0;
+	UFUNCTION(BlueprintCallable, BlueprintNativeEvent, Category = "LearningAgents")
+	void EnableEntity();
 
 	/** Disables an entity used for training. */
-	UFUNCTION(BlueprintCallable, Category = "LearningAgents")
-	virtual void Disable() = 0;
+	UFUNCTION(BlueprintCallable, BlueprintNativeEvent, Category = "LearningAgents")
+	void DisableEntity();
 
 	/** Checks whether an entity is eligible to be used for training. */
-	UFUNCTION(BlueprintCallable, Category = "LearningAgents")
-	virtual bool IsEnabled() const = 0;
+	UFUNCTION(BlueprintCallable, BlueprintNativeEvent, Category = "LearningAgents")
+	bool IsEntityEnabled() const;
 };
diff --git a/Engine/Plugins/Experimental/LearningAgents/Source/LearningAgentsTraining/Public/LearningAgentsGym.h b/Engine/Plugins/Experimental/LearningAgents/Source/LearningAgentsTraining/Public/LearningAgentsGym.h
--- a/Engine/Plugins/Experimental/LearningAgents/Source/LearningAgentsTraining/Public/LearningAgentsGym.h
+++ b/Engine/Plugins/Experimental/LearningAgents/Source/LearningAgentsTraining/Public/LearningAgentsGym.h
@@ -10,6 +10,8 @@
 #define INVALID_GYM_VECTOR (FVector::ZeroVector)
 #define INVALID_GYM_ROTATOR (FRotator::ZeroRotator)
 
+class ILearningAgentsLearningComponentInterface;
+
 /**
  * The Gym Base abstract class handles the start and reset of entities training in a single gym.
  */
@@ -30,9 +32,6 @@ public:
 	void Reset();
 
 public:
-	UPROPERTY(EditDefaultsOnly, Category = "LearningAgents")
-	TObjectPtr<ULearningAgentsEntitiesManagerComponent> EntitiesManagerComponent;
-
 	/** Retrieves the current random stream used by the Gym.*/
 	UFUNCTION(BlueprintPure, Category = "LearningAgents")
 	void GetRandomStream(FRandomStream& OutRandomStream) const;
@@ -64,4 +63,9 @@ protected:
 	int32 RandomSeed = 1234;
 
 	TSharedPtr<FRandomStream> RandomStream; 
+
+	UPROPERTY()
+	TArray<TScriptInterface<ILearningAgentsLearningComponentInterface>> LearningComponents;
+
+	void PopulateLearningComponents();
 };
diff --git a/Engine/Plugins/Experimental/LearningAgents/Source/LearningAgentsTraining/Public/LearningAgentsLearningComponentInterface.h b/Engine/Plugins/Experimental/LearningAgents/Source/LearningAgentsTraining/Public/LearningAgentsLearningComponentInterface.h
new file mode 100644
--- /dev/null
+++ b/Engine/Plugins/Experimental/LearningAgents/Source/LearningAgentsTraining/Public/LearningAgentsLearningComponentInterface.h
@@ -0,0 +1,29 @@
+// Copyright Epic Games, Inc. All Rights Reserved.
+
+#pragma once
+
+#include "CoreMinimal.h"
+#include "UObject/Interface.h"
+#include "LearningAgentsLearningComponentInterface.generated.h"
+
+/**
+ */
+UINTERFACE(MinimalAPI, meta = (CannotImplementInterfaceInBlueprint), BlueprintType)
+class ULearningAgentsLearningComponentInterface : public UInterface
+{
+	GENERATED_BODY()
+};
+
+class LEARNINGAGENTSTRAINING_API ILearningAgentsLearningComponentInterface
+{
+	GENERATED_BODY()
+
+public:
+	/** Initializes the component at the start of training. */
+	UFUNCTION(BlueprintCallable, Category = "LearningAgents")
+	virtual void InitializeLearningComponent() = 0;
+
+	/** Resets the component for a new training episode. */
+	UFUNCTION(BlueprintCallable, Category = "LearningAgents")
+	virtual void ResetLearningComponent() = 0;
+};
