LearningAgents: Added ability to choose weight initialization method to encoder and decoder networks.

diff --git a/Engine/Plugins/Experimental/LearningAgents/Source/Learning/Private/LearningAction.cpp b/Engine/Plugins/Experimental/LearningAgents/Source/Learning/Private/LearningAction.cpp
--- a/Engine/Plugins/Experimental/LearningAgents/Source/Learning/Private/LearningAction.cpp
+++ b/Engine/Plugins/Experimental/LearningAgents/Source/Learning/Private/LearningAction.cpp
@@ -1637,6 +1637,15 @@ namespace UE::Learning::Action
 				NNE::RuntimeBasic::FModelBuilder::ELinearLayerType::Compressed :
 				NNE::RuntimeBasic::FModelBuilder::ELinearLayerType::Normal;
 
+			switch (NetworkSettings.WeightInitialization)
+			{
+			case EWeightInitialization::KaimingGaussian: LinearLayerSettings.WeightInitializationSettings.Type = 
+				NNE::RuntimeBasic::FModelBuilder::EWeightInitializationType::KaimingGaussian; break;
+			case EWeightInitialization::KaimingUniform: LinearLayerSettings.WeightInitializationSettings.Type =
+				NNE::RuntimeBasic::FModelBuilder::EWeightInitializationType::KaimingUniform; break;
+			default: UE_LEARNING_NOT_IMPLEMENTED();
+			}
+
 			OutElement = Builder.MakeSequence({
 				Builder.MakeActivation(Parameters.EncodingSize, Private::GetNNEActivationFunction(Parameters.ActivationFunction)),
 				Builder.MakeMLP(
diff --git a/Engine/Plugins/Experimental/LearningAgents/Source/Learning/Private/LearningObservation.cpp b/Engine/Plugins/Experimental/LearningAgents/Source/Learning/Private/LearningObservation.cpp
--- a/Engine/Plugins/Experimental/LearningAgents/Source/Learning/Private/LearningObservation.cpp
+++ b/Engine/Plugins/Experimental/LearningAgents/Source/Learning/Private/LearningObservation.cpp
@@ -971,6 +971,15 @@ namespace UE::Learning::Observation
 					NNE::RuntimeBasic::FModelBuilder::ELinearLayerType::Compressed :
 					NNE::RuntimeBasic::FModelBuilder::ELinearLayerType::Normal;
 
+				switch (NetworkSettings.WeightInitialization)
+				{
+				case EWeightInitialization::KaimingGaussian: LinearLayerSettings.WeightInitializationSettings.Type = 
+					NNE::RuntimeBasic::FModelBuilder::EWeightInitializationType::KaimingGaussian; break;
+				case EWeightInitialization::KaimingUniform: LinearLayerSettings.WeightInitializationSettings.Type =
+					NNE::RuntimeBasic::FModelBuilder::EWeightInitializationType::KaimingUniform; break;
+				default: UE_LEARNING_NOT_IMPLEMENTED();
+				}
+
 				BuilderEncoders.Emplace(Builder.MakeLinearLayer(SubElementEncodedSize, Parameters.EncodingSize, LinearLayerSettings));
 			}
 
@@ -1001,6 +1010,15 @@ namespace UE::Learning::Observation
 					NNE::RuntimeBasic::FModelBuilder::ELinearLayerType::Compressed :
 					NNE::RuntimeBasic::FModelBuilder::ELinearLayerType::Normal;
 
+				switch (NetworkSettings.WeightInitialization)
+				{
+				case EWeightInitialization::KaimingGaussian: LinearLayerSettings.WeightInitializationSettings.Type = 
+					NNE::RuntimeBasic::FModelBuilder::EWeightInitializationType::KaimingGaussian; break;
+				case EWeightInitialization::KaimingUniform: LinearLayerSettings.WeightInitializationSettings.Type =
+					NNE::RuntimeBasic::FModelBuilder::EWeightInitializationType::KaimingUniform; break;
+				default: UE_LEARNING_NOT_IMPLEMENTED();
+				}
+
 				BuilderQueryLayers.Emplace(Builder.MakeLinearLayer(SubElementEncodedSize, Parameters.AttentionHeadNum * Parameters.AttentionEncodingSize, LinearLayerSettings));
 				BuilderKeyLayers.Emplace(Builder.MakeLinearLayer(SubElementEncodedSize, Parameters.AttentionHeadNum * Parameters.AttentionEncodingSize, LinearLayerSettings));
 				BuilderValueLayers.Emplace(Builder.MakeLinearLayer(SubElementEncodedSize, Parameters.AttentionHeadNum * Parameters.ValueEncodingSize, LinearLayerSettings));
@@ -1042,6 +1060,15 @@ namespace UE::Learning::Observation
 				NNE::RuntimeBasic::FModelBuilder::ELinearLayerType::Compressed :
 				NNE::RuntimeBasic::FModelBuilder::ELinearLayerType::Normal;
 
+			switch (NetworkSettings.WeightInitialization)
+			{
+			case EWeightInitialization::KaimingGaussian: LinearLayerSettings.WeightInitializationSettings.Type = 
+				NNE::RuntimeBasic::FModelBuilder::EWeightInitializationType::KaimingGaussian; break;
+			case EWeightInitialization::KaimingUniform: LinearLayerSettings.WeightInitializationSettings.Type =
+				NNE::RuntimeBasic::FModelBuilder::EWeightInitializationType::KaimingUniform; break;
+			default: UE_LEARNING_NOT_IMPLEMENTED();
+			}
+
 			OutElement = Builder.MakeAggregateSet(
 				Parameters.MaxNum,
 				Parameters.ValueEncodingSize,
@@ -1069,6 +1096,15 @@ namespace UE::Learning::Observation
 				NNE::RuntimeBasic::FModelBuilder::ELinearLayerType::Compressed :
 				NNE::RuntimeBasic::FModelBuilder::ELinearLayerType::Normal;
 
+			switch (NetworkSettings.WeightInitialization)
+			{
+			case EWeightInitialization::KaimingGaussian: LinearLayerSettings.WeightInitializationSettings.Type = 
+				NNE::RuntimeBasic::FModelBuilder::EWeightInitializationType::KaimingGaussian; break;
+			case EWeightInitialization::KaimingUniform: LinearLayerSettings.WeightInitializationSettings.Type =
+				NNE::RuntimeBasic::FModelBuilder::EWeightInitializationType::KaimingUniform; break;
+			default: UE_LEARNING_NOT_IMPLEMENTED();
+			}
+
 			OutElement = Builder.MakeSequence({
 				BuilderSubElement,
 				Builder.MakeMLP(
diff --git a/Engine/Plugins/Experimental/LearningAgents/Source/Learning/Public/LearningAction.h b/Engine/Plugins/Experimental/LearningAgents/Source/Learning/Public/LearningAction.h
--- a/Engine/Plugins/Experimental/LearningAgents/Source/Learning/Public/LearningAction.h
+++ b/Engine/Plugins/Experimental/LearningAgents/Source/Learning/Public/LearningAction.h
@@ -603,6 +603,13 @@ namespace UE::Learning::Action
 		const FSchema& SchemaB,
 		const FSchemaElement SchemaElementB);
 
+	/** Network weight initialization type */
+	enum class EWeightInitialization : uint8
+	{
+		KaimingGaussian = 0,
+		KaimingUniform = 1,
+	};
+
 	/**
 	 * Settings for building a network from a schema.
 	 */
@@ -610,6 +617,9 @@ namespace UE::Learning::Action
 	{
 		// If to use compressed linear layers. This reduces the memory usage by half at the cost of some small impact on quality and evaluation speed.
 		bool bUseCompressedLinearLayers = false;
+
+		// Which weight initialization to use.
+		EWeightInitialization WeightInitialization = EWeightInitialization::KaimingGaussian;
 	};
 
 	/**
diff --git a/Engine/Plugins/Experimental/LearningAgents/Source/Learning/Public/LearningObservation.h b/Engine/Plugins/Experimental/LearningAgents/Source/Learning/Public/LearningObservation.h
--- a/Engine/Plugins/Experimental/LearningAgents/Source/Learning/Public/LearningObservation.h
+++ b/Engine/Plugins/Experimental/LearningAgents/Source/Learning/Public/LearningObservation.h
@@ -432,6 +432,13 @@ namespace UE::Learning::Observation
 		const FSchema& SchemaB,
 		const FSchemaElement SchemaElementB);
 
+	/** Network weight initialization type */
+	enum class EWeightInitialization : uint8
+	{
+		KaimingGaussian = 0,
+		KaimingUniform = 1,
+	};
+
 	/**
 	 * Settings for building a network from a schema.
 	 */
@@ -439,6 +446,9 @@ namespace UE::Learning::Observation
 	{
 		// If to use compressed linear layers. This reduces the memory usage by half at the cost of some small impact on quality and evaluation speed.
 		bool bUseCompressedLinearLayers = false;
+
+		// Which weight initialization to use.
+		EWeightInitialization WeightInitialization = EWeightInitialization::KaimingGaussian;
 	};
 
 	/**
