LearningAgents: Added support for GELU activation function

diff --git a/Engine/Plugins/Experimental/LearningAgents/Source/Learning/Private/LearningAction.cpp b/Engine/Plugins/Experimental/LearningAgents/Source/Learning/Private/LearningAction.cpp
--- a/Engine/Plugins/Experimental/LearningAgents/Source/Learning/Private/LearningAction.cpp
+++ b/Engine/Plugins/Experimental/LearningAgents/Source/Learning/Private/LearningAction.cpp
@@ -1262,6 +1262,7 @@ namespace UE::Learning::Action
 			case EEncodingActivationFunction::ReLU: return NNE::RuntimeBasic::FModelBuilder::EActivationFunction::ReLU;
 			case EEncodingActivationFunction::ELU: return NNE::RuntimeBasic::FModelBuilder::EActivationFunction::ELU;
 			case EEncodingActivationFunction::TanH: return NNE::RuntimeBasic::FModelBuilder::EActivationFunction::TanH;
+			case EEncodingActivationFunction::GELU: return NNE::RuntimeBasic::FModelBuilder::EActivationFunction::GELU;
 			default: UE_LEARNING_NOT_IMPLEMENTED(); return NNE::RuntimeBasic::FModelBuilder::EActivationFunction::ReLU;
 			}
 		}
diff --git a/Engine/Plugins/Experimental/LearningAgents/Source/Learning/Private/LearningObservation.cpp b/Engine/Plugins/Experimental/LearningAgents/Source/Learning/Private/LearningObservation.cpp
--- a/Engine/Plugins/Experimental/LearningAgents/Source/Learning/Private/LearningObservation.cpp
+++ b/Engine/Plugins/Experimental/LearningAgents/Source/Learning/Private/LearningObservation.cpp
@@ -691,6 +691,7 @@ namespace UE::Learning::Observation
 			case EEncodingActivationFunction::ReLU: return NNE::RuntimeBasic::FModelBuilder::EActivationFunction::ReLU;
 			case EEncodingActivationFunction::ELU: return NNE::RuntimeBasic::FModelBuilder::EActivationFunction::ELU;
 			case EEncodingActivationFunction::TanH: return NNE::RuntimeBasic::FModelBuilder::EActivationFunction::TanH;
+			case EEncodingActivationFunction::GELU: return NNE::RuntimeBasic::FModelBuilder::EActivationFunction::GELU;
 			default: UE_LEARNING_NOT_IMPLEMENTED(); return NNE::RuntimeBasic::FModelBuilder::EActivationFunction::ReLU;
 			}
 		}
diff --git a/Engine/Plugins/Experimental/LearningAgents/Source/Learning/Public/LearningAction.h b/Engine/Plugins/Experimental/LearningAgents/Source/Learning/Public/LearningAction.h
--- a/Engine/Plugins/Experimental/LearningAgents/Source/Learning/Public/LearningAction.h
+++ b/Engine/Plugins/Experimental/LearningAgents/Source/Learning/Public/LearningAction.h
@@ -132,6 +132,7 @@ namespace UE::Learning::Action
 		ELU = 0,
 		ReLU = 1,
 		TanH = 2,
+		GELU = 3,
 	};
 
 	struct FSchemaEncodingParameters
diff --git a/Engine/Plugins/Experimental/LearningAgents/Source/Learning/Public/LearningObservation.h b/Engine/Plugins/Experimental/LearningAgents/Source/Learning/Public/LearningObservation.h
--- a/Engine/Plugins/Experimental/LearningAgents/Source/Learning/Public/LearningObservation.h
+++ b/Engine/Plugins/Experimental/LearningAgents/Source/Learning/Public/LearningObservation.h
@@ -135,6 +135,7 @@ namespace UE::Learning::Observation
 		ELU = 0,
 		ReLU = 1,
 		TanH = 2,
+		GELU = 3,
 	};
 
 	struct FSchemaEncodingParameters
