[LearningAgents] Trivial change in LearningAgentsGymBase to make it a Blueprint Type

diff --git a/Engine/Plugins/Experimental/LearningAgents/Source/LearningAgentsTraining/Public/LearningAgentsEntitiesManagerComponent.h b/Engine/Plugins/Experimental/LearningAgents/Source/LearningAgentsTraining/Public/LearningAgentsEntitiesManagerComponent.h
--- a/Engine/Plugins/Experimental/LearningAgents/Source/LearningAgentsTraining/Public/LearningAgentsEntitiesManagerComponent.h
+++ b/Engine/Plugins/Experimental/LearningAgents/Source/LearningAgentsTraining/Public/LearningAgentsEntitiesManagerComponent.h
@@ -59,10 +59,10 @@ public:
 	TArray<FEntityInfo> Entities;
 
 	/** Initializes the entities that the component manages. */
-	void InitializeLearningComponent();
+	virtual void InitializeLearningComponent();
 
 	/** Resets entities that the component manages. */
-	void ResetLearningComponent();
+	virtual void ResetLearningComponent();
 
 	/** Spawns a single pooled entity. */
 	UFUNCTION(BlueprintCallable, Category = "LearningAgents")
diff --git a/Engine/Plugins/Experimental/LearningAgents/Source/LearningAgentsTraining/Public/LearningAgentsGym.h b/Engine/Plugins/Experimental/LearningAgents/Source/LearningAgentsTraining/Public/LearningAgentsGym.h
--- a/Engine/Plugins/Experimental/LearningAgents/Source/LearningAgentsTraining/Public/LearningAgentsGym.h
+++ b/Engine/Plugins/Experimental/LearningAgents/Source/LearningAgentsTraining/Public/LearningAgentsGym.h
@@ -15,7 +15,7 @@ class ILearningAgentsLearningComponentInterface;
 /**
  * The Gym Base abstract class handles the start and reset of entities training in a single gym.
  */
-UCLASS(Abstract, NotBlueprintable)
+UCLASS(Abstract, BlueprintType)
 class LEARNINGAGENTSTRAINING_API ALearningAgentsGymBase : public AActor
 {
 	GENERATED_BODY()
