[LearningAgents] SetAgentMaxNum

A new method to allow blueprints to set max agent num in LearningAgentsManager.

diff --git a/Engine/Plugins/Experimental/LearningAgents/Source/LearningAgents/Private/LearningAgentsManager.cpp b/Engine/Plugins/Experimental/LearningAgents/Source/LearningAgents/Private/LearningAgentsManager.cpp
--- a/Engine/Plugins/Experimental/LearningAgents/Source/LearningAgents/Private/LearningAgentsManager.cpp
+++ b/Engine/Plugins/Experimental/LearningAgents/Source/LearningAgents/Private/LearningAgentsManager.cpp
@@ -50,6 +50,37 @@ int32 ULearningAgentsManager::GetMaxAgentNum() const
 	return MaxAgentNum;
 }
 
+void ULearningAgentsManager::SetMaxAgentNum(const int32 InMaxAgentNum)
+{
+	if (InMaxAgentNum <= 0)
+	{
+		UE_LOG(LogLearning, Error, TEXT("%s: Attempted to SetMaxAgentNum but the input is less than or equal to 0."), *GetName());
+		return;
+	}
+	
+	if (GetAgentNum() > 0)
+	{
+		UE_LOG(LogLearning, Error, TEXT("%s: Cannot SetMaxAgentNum because agents have already been added."), *GetName());
+		return;
+	}
+
+	int32 PreviousMaxAgentNum = MaxAgentNum;
+	MaxAgentNum = InMaxAgentNum;
+
+	OnEventAgentIds.Reserve(MaxAgentNum);
+	OccupiedAgentIds.Reserve(MaxAgentNum);
+	VacantAgentIds.Reserve(MaxAgentNum);
+	Agents.Reserve(MaxAgentNum);
+
+	for (int32 AgentId = MaxAgentNum - 1; AgentId >= PreviousMaxAgentNum; AgentId--)
+	{
+		VacantAgentIds.Push(AgentId);
+		Agents.Add(nullptr);
+	}
+
+	UpdateAgentSets();
+}
+
 int32 ULearningAgentsManager::AddAgent(UObject* Agent)
 {
 	if (Agent == nullptr)
diff --git a/Engine/Plugins/Experimental/LearningAgents/Source/LearningAgents/Public/LearningAgentsManager.h b/Engine/Plugins/Experimental/LearningAgents/Source/LearningAgents/Public/LearningAgentsManager.h
--- a/Engine/Plugins/Experimental/LearningAgents/Source/LearningAgents/Public/LearningAgentsManager.h
+++ b/Engine/Plugins/Experimental/LearningAgents/Source/LearningAgents/Public/LearningAgentsManager.h
@@ -46,6 +46,10 @@ public:
 	/** Returns the maximum number of agents that this manager is configured to handle. */
 	UFUNCTION(BlueprintPure, Category = "LearningAgents")
 	int32 GetMaxAgentNum() const;
+	
+	/** Sets the maximum number of agents that this manager is configured to handle. */
+	UFUNCTION(BlueprintCallable, Category = "LearningAgents", meta = (ClampMin = "1", UIMin = "1"))
+	void SetMaxAgentNum(const int32 InMaxAgentNum);
 
 public:
 
diff --git a/Engine/Plugins/Experimental/LearningAgents/Source/LearningAgentsTraining/Public/LearningAgentsGymsManager.h b/Engine/Plugins/Experimental/LearningAgents/Source/LearningAgentsTraining/Public/LearningAgentsGymsManager.h
--- a/Engine/Plugins/Experimental/LearningAgents/Source/LearningAgentsTraining/Public/LearningAgentsGymsManager.h
+++ b/Engine/Plugins/Experimental/LearningAgents/Source/LearningAgentsTraining/Public/LearningAgentsGymsManager.h
@@ -51,7 +51,7 @@ public:
 	void Start();
 
 	/** Get the number of gyms managed by the GymsManager. */
-	UFUNCTION(BlueprintCallable, Category = "LearningAgents")
+	UFUNCTION(BlueprintPure, Category = "LearningAgents")
 	int32 GetGymsCount() const;
 
 private:
diff --git a/Engine/Plugins/Experimental/LearningAgents/Source/LearningAgentsTraining/Public/LearningAgentsPPOTrainer.h b/Engine/Plugins/Experimental/LearningAgents/Source/LearningAgentsTraining/Public/LearningAgentsPPOTrainer.h
--- a/Engine/Plugins/Experimental/LearningAgents/Source/LearningAgentsTraining/Public/LearningAgentsPPOTrainer.h
+++ b/Engine/Plugins/Experimental/LearningAgents/Source/LearningAgentsTraining/Public/LearningAgentsPPOTrainer.h
@@ -38,7 +38,7 @@ public:
 	 * Maximum number of steps recorded in an episode before it is added to the replay buffer. This can generally be left at the default value and
 	 * does not have a large impact on training.
 	 */
-	UPROPERTY(EditAnywhere, Category = "LearningAgents", meta = (ClampMin = "1", UIMin = "1"))
+	UPROPERTY(EditAnywhere, BlueprintReadWrite, Category = "LearningAgents", meta = (ClampMin = "1", UIMin = "1"))
 	int32 MaxEpisodeStepNum = 512;
 
 	/**
@@ -46,7 +46,7 @@ public:
 	 * MaximumRecordedEpisodesPerIteration is reached. Typical values for this should be around 1000. Setting this too small means there is not
 	 * enough data each iteration for the system to train. Setting it too large means training will be very slow.
 	 */
-	UPROPERTY(EditAnywhere, Category = "LearningAgents", meta = (ClampMin = "1", UIMin = "1"))
+	UPROPERTY(EditAnywhere, BlueprintReadWrite, Category = "LearningAgents", meta = (ClampMin = "1", UIMin = "1"))
 	int32 MaximumRecordedEpisodesPerIteration = 1000;
 
 	/**
@@ -54,7 +54,7 @@ public:
 	 * MaximumRecordedEpisodesPerIteration is reached. Typical values for this should be around 10000. Setting this too small means there is not
 	 * enough data each iteration for the system to train. Setting it too large means training will be very slow.
 	 */
-	UPROPERTY(EditAnywhere, Category = "LearningAgents", meta = (ClampMin = "1", UIMin = "1"))
+	UPROPERTY(EditAnywhere, BlueprintReadWrite, Category = "LearningAgents", meta = (ClampMin = "1", UIMin = "1"))
 	int32 MaximumRecordedStepsPerIteration = 10000;
 };
 
