LearningAgents: Prevent situations where invalid training process can cause a crash

diff --git a/Engine/Plugins/Experimental/LearningAgents/Source/LearningAgentsTraining/Private/LearningAgentsPPOTrainer.cpp b/Engine/Plugins/Experimental/LearningAgents/Source/LearningAgentsTraining/Private/LearningAgentsPPOTrainer.cpp
--- a/Engine/Plugins/Experimental/LearningAgents/Source/LearningAgentsTraining/Private/LearningAgentsPPOTrainer.cpp
+++ b/Engine/Plugins/Experimental/LearningAgents/Source/LearningAgentsTraining/Private/LearningAgentsPPOTrainer.cpp
@@ -169,7 +169,7 @@ void ULearningAgentsPPOTrainer::SetupPPOTrainer(
 		return;
 	}
 
-	if (!Communicator.Trainer)
+	if (!Communicator.Trainer || !Communicator.Trainer->IsValid())
 	{
 		UE_LOG(LogLearning, Error, TEXT("%s: Communicator's Trainer is nullptr."), *GetName());
 		return;
diff --git a/Engine/Plugins/Experimental/LearningAgents/Source/LearningTraining/Private/LearningExternalTrainer.cpp b/Engine/Plugins/Experimental/LearningAgents/Source/LearningTraining/Private/LearningExternalTrainer.cpp
--- a/Engine/Plugins/Experimental/LearningAgents/Source/LearningTraining/Private/LearningExternalTrainer.cpp
+++ b/Engine/Plugins/Experimental/LearningAgents/Source/LearningTraining/Private/LearningExternalTrainer.cpp
@@ -218,6 +218,11 @@ namespace UE::Learning
 		Terminate();
 	}
 
+	bool FSharedMemoryTrainer::IsValid()
+	{
+		return TrainingProcess != nullptr;
+	}
+
 	ETrainerResponse FSharedMemoryTrainer::Wait()
 	{
 		return ETrainerResponse::Success;
@@ -664,6 +669,11 @@ namespace UE::Learning
 		Terminate();
 	}
 
+	bool FSocketTrainer::IsValid()
+	{
+		return TrainingProcess != nullptr;
+	}
+
 	ETrainerResponse FSocketTrainer::Wait()
 	{
 		return ETrainerResponse::Success;
diff --git a/Engine/Plugins/Experimental/LearningAgents/Source/LearningTraining/Public/LearningExternalTrainer.h b/Engine/Plugins/Experimental/LearningAgents/Source/LearningTraining/Public/LearningExternalTrainer.h
--- a/Engine/Plugins/Experimental/LearningAgents/Source/LearningTraining/Public/LearningExternalTrainer.h
+++ b/Engine/Plugins/Experimental/LearningAgents/Source/LearningTraining/Public/LearningExternalTrainer.h
@@ -25,6 +25,9 @@ namespace UE::Learning
 	{
 		virtual ~IExternalTrainer() {}
 
+		/** Returns true if this trainer is valid. Otherwise, false. */
+		virtual bool IsValid() = 0;
+
 		/** Terminate the trainer immediately. */
 		virtual void Terminate() = 0;
 
@@ -239,6 +242,8 @@ namespace UE::Learning
 
 		~FSharedMemoryTrainer();
 
+		virtual bool IsValid() override final;
+
 		virtual void Terminate() override final;
 
 		virtual ETrainerResponse SendStop() override final;
@@ -374,6 +379,8 @@ namespace UE::Learning
 
 		~FSocketTrainer();
 
+		virtual bool IsValid() override final;
+
 		virtual void Terminate() override final;
 
 		virtual ETrainerResponse SendStop() override final;
