Learning Agents: Expand cmdline options, add task directories

#virtualized

diff --git a/Engine/Plugins/Experimental/LearningAgents/Source/LearningAgentsTraining/Private/LearningAgentsCommunicator.cpp b/Engine/Plugins/Experimental/LearningAgents/Source/LearningAgentsTraining/Private/LearningAgentsCommunicator.cpp
--- a/Engine/Plugins/Experimental/LearningAgents/Source/LearningAgentsTraining/Private/LearningAgentsCommunicator.cpp
+++ b/Engine/Plugins/Experimental/LearningAgents/Source/LearningAgentsTraining/Private/LearningAgentsCommunicator.cpp
@@ -83,23 +83,47 @@ FLearningAgentsCommunicator ULearningAgentsCommunicatorLibrary::MakeSharedMemory
 		return Communicator;
 	}
 
-	const FString IntermediatePath = UE::Learning::Trainer::GetIntermediatePath(TrainerProcessSettings.GetIntermediatePath());
+	// Task Name
+	FString CmdLineTaskName;
+	bool bTaskNameOverridden = FParse::Value(FCommandLine::Get(), TEXT("LearningAgentsTaskName="), CmdLineTaskName);
+	if (bTaskNameOverridden)
+	{
+		UE_LOG(LogLearning, Display, TEXT("Overriding Task Name from the cmdline: %s"), *CmdLineTaskName);
+	}
+
+	// Timeout
+	float CmdLineTimeout;
+	bool bTimeoutOverridden = FParse::Value(FCommandLine::Get(), TEXT("LearningAgentsTimeout="), CmdLineTimeout);
+	if (bTimeoutOverridden)
+	{
+		UE_LOG(LogLearning, Display, TEXT("Overriding Timeout with value from the cmdline: %f"), CmdLineTimeout);
+	}
 
+	// Controls Guid
+	// We have to use FGuid::Parse to accept more guid formats
 	FString CmdLineControlsGuidString;
 	FGuid CmdLineControlsGuid;
 	if (FParse::Value(FCommandLine::Get(), TEXT("LearningAgentsControlsGuid="), CmdLineControlsGuidString) && FGuid::Parse(CmdLineControlsGuidString, CmdLineControlsGuid))
 	{
-		UE_LOG(LogLearning, Display, TEXT("Mapping Controls Guid override from the cmdline: %s"), *CmdLineControlsGuid.ToString(EGuidFormats::DigitsWithHyphensInBraces));
+		UE_LOG(LogLearning, Display, TEXT("Overriding Controls Guid with value from the cmdline: %s"), *CmdLineControlsGuid.ToString(EGuidFormats::DigitsWithHyphensInBraces));
+	}
+
+	// Config Path
+	FString CmdLineConfigPath;
+	bool bConfigPathOverridden = FParse::Value(FCommandLine::Get(), TEXT("LearningAgentsConfigPath="), CmdLineConfigPath);
+	if (bConfigPathOverridden)
+	{
+		UE_LOG(LogLearning, Display, TEXT("Overriding Config Path with value from the cmdline: %s"), *CmdLineConfigPath);
 	}
 
 	Communicator.Trainer = MakeShared<UE::Learning::FSharedMemoryTrainer>(
-		SharedMemorySettings.TaskName,
+		bTaskNameOverridden ? CmdLineTaskName : SharedMemorySettings.TaskName,
 		TrainerProcess.TrainerProcess,
 		FParse::Param(FCommandLine::Get(), TEXT("LearningAgentsUseExternalTrainingProcess")) || SharedMemorySettings.bUseExternalTrainingProcess,
 		CmdLineControlsGuid.IsValid() ? CmdLineControlsGuid : SharedMemorySettings.ControlsGuid,
-		SharedMemorySettings.ConfigPath.Path,
-		IntermediatePath,
-		SharedMemorySettings.Timeout);
+		bConfigPathOverridden ? CmdLineConfigPath : SharedMemorySettings.ConfigPath.Path,
+		UE::Learning::Trainer::GetIntermediatePath(TrainerProcessSettings.GetIntermediatePath()),
+		bTimeoutOverridden ? CmdLineTimeout : SharedMemorySettings.Timeout);
 
 	return Communicator;
 }
@@ -166,6 +190,7 @@ FLearningAgentsCommunicator ULearningAgentsCommunicatorLibrary::MakeSocketCommun
 {
 	UE::Learning::ETrainerResponse Response = UE::Learning::ETrainerResponse::Success;
 
+	// Ip Address
 	FString CmdLineIpAddress;
 	bool bIpAddressOverridden = FParse::Value(FCommandLine::Get(), TEXT("LearningAgentsIpAddress="), CmdLineIpAddress);
 	if (bIpAddressOverridden)
@@ -173,6 +198,7 @@ FLearningAgentsCommunicator ULearningAgentsCommunicatorLibrary::MakeSocketCommun
 		UE_LOG(LogLearning, Display, TEXT("Using Ip Address override from the cmdline: %s"), *CmdLineIpAddress);
 	}
 
+	// Port
 	uint32 CmdLinePort;
 	bool bPortOverridden = FParse::Value(FCommandLine::Get(), TEXT("LearningAgentsPort="), CmdLinePort);
 	if (bPortOverridden)
@@ -180,13 +206,22 @@ FLearningAgentsCommunicator ULearningAgentsCommunicatorLibrary::MakeSocketCommun
 		UE_LOG(LogLearning, Display, TEXT("Using Port override from the cmdline: %u"), CmdLinePort);
 	}
 
+	// Timeout
+	float CmdLineTimeout;
+	bool bTimeoutOverridden = FParse::Value(FCommandLine::Get(), TEXT("LearningAgentsTimeout="), CmdLineTimeout);
+	if (bTimeoutOverridden)
+	{
+		UE_LOG(LogLearning, Display, TEXT("Overriding Timeout with value from the cmdline: %f"), CmdLineTimeout);
+	}
+
 	FLearningAgentsCommunicator Communicator;
 	Communicator.Trainer = MakeShared<UE::Learning::FSocketTrainer>(
 		Response,
 		TrainerProcess.TrainerProcess,
+		FParse::Param(FCommandLine::Get(), TEXT("LearningAgentsUseExternalTrainingProcess")) || SocketSettings.bUseExternalTrainingProcess,
 		bIpAddressOverridden ? *CmdLineIpAddress : *SocketSettings.IpAddress,
 		bPortOverridden ? CmdLinePort : SocketSettings.Port,
-		SocketSettings.Timeout);
+		bTimeoutOverridden ? CmdLineTimeout : SocketSettings.Timeout);
 
 	if (Response != UE::Learning::ETrainerResponse::Success)
 	{
diff --git a/Engine/Plugins/Experimental/LearningAgents/Source/LearningAgentsTraining/Private/LearningAgentsImitationTrainer.cpp b/Engine/Plugins/Experimental/LearningAgents/Source/LearningAgentsTraining/Private/LearningAgentsImitationTrainer.cpp
--- a/Engine/Plugins/Experimental/LearningAgents/Source/LearningAgentsTraining/Private/LearningAgentsImitationTrainer.cpp
+++ b/Engine/Plugins/Experimental/LearningAgents/Source/LearningAgentsTraining/Private/LearningAgentsImitationTrainer.cpp
@@ -275,10 +275,11 @@ void ULearningAgentsImitationTrainer::BeginTraining(
 	DecoderNetworkId = Trainer->AddNetwork(*Policy->GetDecoderNetworkAsset()->NeuralNetworkData);
 	ReplayBufferId = Trainer->AddReplayBuffer(*ReplayBuffer);
 
-	TSharedRef<FJsonObject> ConfigObject = CreateConfig(ImitationTrainerTrainingSettings);
+	TSharedRef<FJsonObject> DataConfigObject = CreateDataConfig();
+	TSharedRef<FJsonObject> TrainerConfigObject = CreateTrainerConfig(ImitationTrainerTrainingSettings);
 
-	UE_LOG(LogLearning, Display, TEXT("%s: Sending config..."), *GetName());
-	SendConfig(ConfigObject);
+	UE_LOG(LogLearning, Display, TEXT("%s: Sending configs..."), *GetName());
+	SendConfigs(DataConfigObject, TrainerConfigObject);
 
 	UE_LOG(LogLearning, Display, TEXT("%s: Imitation Training Started"), *GetName());
 	UE::Learning::ETrainerResponse Response = UE::Learning::ETrainerResponse::Success;
@@ -324,15 +325,10 @@ void ULearningAgentsImitationTrainer::BeginTraining(
 	bIsTraining = true;
 }
 
-TSharedRef<FJsonObject> ULearningAgentsImitationTrainer::CreateConfig(const FLearningAgentsImitationTrainerTrainingSettings& TrainingSettings) const
+TSharedRef<FJsonObject> ULearningAgentsImitationTrainer::CreateDataConfig() const
 {
 	TSharedRef<FJsonObject> ConfigObject = MakeShared<FJsonObject>();
 	
-	// Add Training Config Entries
-	ConfigObject->SetStringField(TEXT("TaskName"), TEXT("Training"));
-	ConfigObject->SetStringField(TEXT("TrainerMethod"), TEXT("BehaviorCloning"));
-	ConfigObject->SetStringField(TEXT("TimeStamp"), *FDateTime::Now().ToFormattedString(TEXT("%Y-%m-%d_%H-%M-%S")));
-
 	const int32 ObservationSchemaId = 0;
 	const int32 ActionSchemaId = 0;
 
@@ -419,6 +415,18 @@ TSharedRef<FJsonObject> ULearningAgentsImitationTrainer::CreateConfig(const FLea
 
 	ConfigObject->SetObjectField(TEXT("Schemas"), SchemasObject);
 
+	return ConfigObject;
+}
+
+TSharedRef<FJsonObject> ULearningAgentsImitationTrainer::CreateTrainerConfig(const FLearningAgentsImitationTrainerTrainingSettings& TrainingSettings) const
+{
+	TSharedRef<FJsonObject> ConfigObject = MakeShared<FJsonObject>();
+
+	// Add Training Config Entries
+	ConfigObject->SetStringField(TEXT("TaskName"), TEXT("Training"));
+	ConfigObject->SetStringField(TEXT("TrainerMethod"), TEXT("BehaviorCloning"));
+	ConfigObject->SetStringField(TEXT("TimeStamp"), *FDateTime::Now().ToFormattedString(TEXT("%Y-%m-%d_%H-%M-%S")));
+
 	// Add Imitation Specific Config Entries
 	ConfigObject->SetObjectField(TEXT("BehaviorCloningSettings"), TrainingSettings.AsJsonConfig());
 	ConfigObject->SetNumberField(TEXT("MemoryStateNum"), Policy->GetMemoryStateSize());
@@ -426,10 +434,10 @@ TSharedRef<FJsonObject> ULearningAgentsImitationTrainer::CreateConfig(const FLea
 	return ConfigObject;
 }
 
-void ULearningAgentsImitationTrainer::SendConfig(const TSharedRef<FJsonObject>& ConfigObject)
+void ULearningAgentsImitationTrainer::SendConfigs(const TSharedRef<FJsonObject>& DataConfigObject, const TSharedRef<FJsonObject>& TrainerConfigObject)
 {
 	UE::Learning::ETrainerResponse Response = UE::Learning::ETrainerResponse::Success;
-	Response = Trainer->SendConfig(ConfigObject, UE::Learning::ELogSetting::Normal); // todo pipe the log settings up to the blueprint graph
+	Response = Trainer->SendConfigs(DataConfigObject, TrainerConfigObject);
 
 	if (Response != UE::Learning::ETrainerResponse::Success)
 	{
diff --git a/Engine/Plugins/Experimental/LearningAgents/Source/LearningAgentsTraining/Private/LearningAgentsPPOTrainer.cpp b/Engine/Plugins/Experimental/LearningAgents/Source/LearningAgentsTraining/Private/LearningAgentsPPOTrainer.cpp
--- a/Engine/Plugins/Experimental/LearningAgents/Source/LearningAgentsTraining/Private/LearningAgentsPPOTrainer.cpp
+++ b/Engine/Plugins/Experimental/LearningAgents/Source/LearningAgentsTraining/Private/LearningAgentsPPOTrainer.cpp
@@ -283,10 +283,11 @@ void ULearningAgentsPPOTrainer::BeginTraining(
 	DecoderNetworkId = Trainer->AddNetwork(*Policy->GetDecoderNetworkAsset()->NeuralNetworkData);
 	ReplayBufferId = Trainer->AddReplayBuffer(*ReplayBuffer);
 	
-	TSharedRef<FJsonObject> ConfigObject = CreateConfig(TrainingSettings);
+	TSharedRef<FJsonObject> DataConfigObject = CreateDataConfig();
+	TSharedRef<FJsonObject> TrainerConfigObject = CreateTrainerConfig(TrainingSettings);
 
-	UE_LOG(LogLearning, Display, TEXT("%s: Sending config..."), *GetName());
-	SendConfig(ConfigObject);
+	UE_LOG(LogLearning, Display, TEXT("%s: Sending configs..."), *GetName());
+	SendConfigs(DataConfigObject, TrainerConfigObject);
 
 	UE_LOG(LogLearning, Display, TEXT("%s: Sending initial policy..."), *GetName());
 	UE::Learning::ETrainerResponse Response = UE::Learning::ETrainerResponse::Success;
@@ -337,15 +338,10 @@ void ULearningAgentsPPOTrainer::BeginTraining(
 	bIsTraining = true;
 }
 
-TSharedRef<FJsonObject> ULearningAgentsPPOTrainer::CreateConfig(const FLearningAgentsPPOTrainingSettings& TrainingSettings) const
+TSharedRef<FJsonObject> ULearningAgentsPPOTrainer::CreateDataConfig() const
 {
 	TSharedRef<FJsonObject> ConfigObject = MakeShared<FJsonObject>();
 
-	// Add Training Config Entries
-	ConfigObject->SetStringField(TEXT("TaskName"), TEXT("Training"));
-	ConfigObject->SetStringField(TEXT("TrainerMethod"), TEXT("PPO"));
-	ConfigObject->SetStringField(TEXT("TimeStamp"), *FDateTime::Now().ToFormattedString(TEXT("%Y-%m-%d_%H-%M-%S")));
-
 	const int32 ObservationSchemaId = 0;
 	const int32 ActionSchemaId = 0;
 
@@ -444,16 +440,28 @@ TSharedRef<FJsonObject> ULearningAgentsPPOTrainer::CreateConfig(const FLearningA
 
 	ConfigObject->SetObjectField(TEXT("Schemas"), SchemasObject);
 
-	// Add PPO Specific Config Entries
+	return ConfigObject;
+}
+
+TSharedRef<FJsonObject> ULearningAgentsPPOTrainer::CreateTrainerConfig(const FLearningAgentsPPOTrainingSettings& TrainingSettings) const
+{
+	TSharedRef<FJsonObject> ConfigObject = MakeShared<FJsonObject>();
+
+	// Add Training Task-specific Config Entries
+	ConfigObject->SetStringField(TEXT("TaskName"), TEXT("Training"));
+	ConfigObject->SetStringField(TEXT("TrainerMethod"), TEXT("PPO"));
+	ConfigObject->SetStringField(TEXT("TimeStamp"), *FDateTime::Now().ToFormattedString(TEXT("%Y-%m-%d_%H-%M-%S")));
+
+	// Add PPO-specific Config Entries
 	ConfigObject->SetObjectField(TEXT("PPOSettings"), TrainingSettings.AsJsonConfig());
 
 	return ConfigObject;
 }
 
-void ULearningAgentsPPOTrainer::SendConfig(const TSharedRef<FJsonObject>& ConfigObject)
+void ULearningAgentsPPOTrainer::SendConfigs(const TSharedRef<FJsonObject>& DataConfigObject, const TSharedRef<FJsonObject>& TrainerConfigObject)
 {
 	UE::Learning::ETrainerResponse Response = UE::Learning::ETrainerResponse::Success;
-	Response = Trainer->SendConfig(ConfigObject, UE::Learning::ELogSetting::Normal);
+	Response = Trainer->SendConfigs(DataConfigObject, TrainerConfigObject);
 	
 	if (Response != UE::Learning::ETrainerResponse::Success)
 	{
diff --git a/Engine/Plugins/Experimental/LearningAgents/Source/LearningAgentsTraining/Private/LearningAgentsTrainer.cpp b/Engine/Plugins/Experimental/LearningAgents/Source/LearningAgentsTraining/Private/LearningAgentsTrainer.cpp
--- a/Engine/Plugins/Experimental/LearningAgents/Source/LearningAgentsTraining/Private/LearningAgentsTrainer.cpp
+++ b/Engine/Plugins/Experimental/LearningAgents/Source/LearningAgentsTraining/Private/LearningAgentsTrainer.cpp
@@ -172,6 +172,14 @@ FLearningAgentsTrainerProcessSettings::FLearningAgentsTrainerProcessSettings()
 
 FString FLearningAgentsTrainerProcessSettings::GetEditorEnginePath() const
 {
+	FString CmdLineEnginePath;
+	bool bEnginePathOverridden = FParse::Value(FCommandLine::Get(), TEXT("LearningAgentsEnginePath="), CmdLineEnginePath);
+	if (bEnginePathOverridden)
+	{
+		UE_LOG(LogLearning, Display, TEXT("Overriding Engine Path with value from the cmdline: %s"), *CmdLineEnginePath);
+		return CmdLineEnginePath;
+	}
+
 #if WITH_EDITOR
 	return EditorEngineRelativePath.Path;
 #else
@@ -186,6 +194,14 @@ FString FLearningAgentsTrainerProcessSettings::GetEditorEnginePath() const
 
 FString FLearningAgentsTrainerProcessSettings::GetIntermediatePath() const
 {
+	FString CmdLineIntermediatePath;
+	bool bIntermediatePathOverridden = FParse::Value(FCommandLine::Get(), TEXT("LearningAgentsIntermediatePath="), CmdLineIntermediatePath);
+	if (bIntermediatePathOverridden)
+	{
+		UE_LOG(LogLearning, Display, TEXT("Overriding Intermediate Path with value from the cmdline: %s"), *CmdLineIntermediatePath);
+		return CmdLineIntermediatePath;
+	}
+
 #if WITH_EDITOR
 	return EditorIntermediateRelativePath.Path;
 #else
@@ -200,6 +216,14 @@ FString FLearningAgentsTrainerProcessSettings::GetIntermediatePath() const
 
 FString FLearningAgentsTrainerProcessSettings::GetCustomTrainerModulePath() const
 {
+	FString CmdLineCustomTrainerPath;
+	bool bCustomTrainerPathOverridden = FParse::Value(FCommandLine::Get(), TEXT("LearningAgentsCustomTrainerPath="), CmdLineCustomTrainerPath);
+	if (bCustomTrainerPathOverridden)
+	{
+		UE_LOG(LogLearning, Display, TEXT("Overriding Custom Trainer Module Path with value from the cmdline: %s"), *CmdLineCustomTrainerPath);
+		return CmdLineCustomTrainerPath;
+	}
+
 #if WITH_EDITOR
 	return CustomTrainerModulePath.Path;
 #else
diff --git a/Engine/Plugins/Experimental/LearningAgents/Source/LearningAgentsTraining/Public/LearningAgentsImitationTrainer.h b/Engine/Plugins/Experimental/LearningAgents/Source/LearningAgentsTraining/Public/LearningAgentsImitationTrainer.h
--- a/Engine/Plugins/Experimental/LearningAgents/Source/LearningAgentsTraining/Public/LearningAgentsImitationTrainer.h
+++ b/Engine/Plugins/Experimental/LearningAgents/Source/LearningAgentsTraining/Public/LearningAgentsImitationTrainer.h
@@ -244,8 +244,9 @@ private:
 	UPROPERTY(VisibleAnywhere, Transient, Category = "LearningAgents")
 	bool bHasTrainingFailed = false;
 
-	TSharedRef<FJsonObject> CreateConfig(const FLearningAgentsImitationTrainerTrainingSettings& TrainingSettings) const;
-	void SendConfig(const TSharedRef<FJsonObject>& ConfigObject);
+	TSharedRef<FJsonObject> CreateDataConfig() const;
+	TSharedRef<FJsonObject> CreateTrainerConfig(const FLearningAgentsImitationTrainerTrainingSettings& TrainingSettings) const;
+	void SendConfigs(const TSharedRef<FJsonObject>& DataConfigObject, const TSharedRef<FJsonObject>& TrainerConfigObject);
 
 	void DoneTraining();
 
diff --git a/Engine/Plugins/Experimental/LearningAgents/Source/LearningAgentsTraining/Public/LearningAgentsPPOTrainer.h b/Engine/Plugins/Experimental/LearningAgents/Source/LearningAgentsTraining/Public/LearningAgentsPPOTrainer.h
--- a/Engine/Plugins/Experimental/LearningAgents/Source/LearningAgentsTraining/Public/LearningAgentsPPOTrainer.h
+++ b/Engine/Plugins/Experimental/LearningAgents/Source/LearningAgentsTraining/Public/LearningAgentsPPOTrainer.h
@@ -419,8 +419,19 @@ private:
 	TUniquePtr<UE::Learning::FReplayBuffer> ReplayBuffer;
 	TSharedPtr<UE::Learning::IExternalTrainer> Trainer;
 
-	TSharedRef<FJsonObject> CreateConfig(const FLearningAgentsPPOTrainingSettings& TrainingSettings) const;
-	void SendConfig(const TSharedRef<FJsonObject>& ConfigObject);
+	/**
+	 * The data config contains the info needed to create the neural network models and the supporting data buffers.
+	 * These need to stay synchronized between UE and the trainer process, otherwise we will run into memory errors.
+	 */
+	TSharedRef<FJsonObject> CreateDataConfig() const;
+
+	/**
+	 * The trainer config contains the info needed run our specific training algorithm.
+	 * In theory, most of these values can be easily overridden on the trainer process side without causing any errors.
+	 */
+	TSharedRef<FJsonObject> CreateTrainerConfig(const FLearningAgentsPPOTrainingSettings& TrainingSettings) const;
+
+	void SendConfigs(const TSharedRef<FJsonObject>& DataConfigObject, const TSharedRef<FJsonObject>& TrainerConfigObject);
 
 	void DoneTraining();
 
diff --git a/Engine/Plugins/Experimental/LearningAgents/Source/LearningTraining/Private/LearningExternalTrainer.cpp b/Engine/Plugins/Experimental/LearningAgents/Source/LearningTraining/Private/LearningExternalTrainer.cpp
--- a/Engine/Plugins/Experimental/LearningAgents/Source/LearningTraining/Private/LearningExternalTrainer.cpp
+++ b/Engine/Plugins/Experimental/LearningAgents/Source/LearningTraining/Private/LearningExternalTrainer.cpp
@@ -7,6 +7,9 @@
 #include "LearningSharedMemoryTraining.h"
 #include "LearningSocketTraining.h"
 
+#include "GenericPlatform/GenericPlatformFile.h"
+#include "HAL/PlatformFileManager.h"
+
 #include "Misc/CommandLine.h"
 #include "Misc/FileHelper.h"
 #include "Misc/Guid.h"
@@ -51,15 +54,31 @@ namespace UE::Learning
 
 		const FString TimeStamp = FDateTime::Now().ToFormattedString(TEXT("%Y-%m-%d_%H-%M-%S"));
 		const FString TrainerType = TEXT("SharedMemory");
-		ConfigPath = InIntermediatePath / TEXT("Configs") / FString::Printf(TEXT("%s.json"), *Controls.Guid.ToString(EGuidFormats::DigitsWithHyphensInBraces));
+		
+		// Make a config directory for the task
+		int32 TaskId = 0;
+		while (true)
+		{
+			FString CandidatePath = InIntermediatePath / TaskName + FString::FromInt(TaskId) / TEXT("Configs");
+			IPlatformFile& PlatformFile = FPlatformFileManager::Get().GetPlatformFile();
+			if (!PlatformFile.DirectoryExists(*CandidatePath))
+			{
+				PlatformFile.CreateDirectory(*CandidatePath);
+				ConfigPath = CandidatePath;
+				break;
+			}
+
+			TaskId++;
+		}
 
 		IFileManager& FileManager = IFileManager::Get();
-		const FString CommandLineArguments = FString::Printf(TEXT("\"%s\" -p \"%s\" -m \"%s\" \"%s\" SharedMemory \"%s\" -g \"%s\""),
+		const FString CommandLineArguments = FString::Printf(TEXT("\"%s\" \"%s\" -p \"%s\" -m \"%s\" \"%s\" SharedMemory \"%s\" -g \"%s\""),
 			*FileManager.ConvertToAbsolutePathForExternalAppForRead(*(PythonContentPath / TEXT("train.py"))),
+			*TaskName,
 			*FileManager.ConvertToAbsolutePathForExternalAppForRead(*CustomTrainerPath),
 			*TrainerFileName,
 			LogSettings == ELogSetting::Normal ? TEXT("-l") : TEXT(""),
-			*FileManager.ConvertToAbsolutePathForExternalAppForRead(*ConfigPath),
+			*FileManager.ConvertToAbsolutePathForExternalAppForRead(*(InIntermediatePath / TaskName + FString::FromInt(TaskId))),
 			*Controls.Guid.ToString(EGuidFormats::DigitsWithHyphensInBraces));
 
 		bool bLaunched = TrainingProcess.Launch(
@@ -183,6 +202,7 @@ namespace UE::Learning
 	{
 		TaskName = InTaskName;
 		Timeout = InTimeout;
+		bUseExternalTrainer = bUseExternalTrainingProcess;
 
 		if (ExternalTrainerProcess)
 		{
@@ -208,7 +228,7 @@ namespace UE::Learning
 			Controls.View[(uint8)SharedMemoryTraining::EControls::NetworkId] = -1;
 			Controls.View[(uint8)SharedMemoryTraining::EControls::ReplayBufferId] = -1;
 
-			ConfigPath = InConfigPath / FString::Printf(TEXT("%s.json"), *ControlsGuid.ToString(EGuidFormats::DigitsWithHyphensInBraces));
+			ConfigPath = InConfigPath;
 			IntermediatePath = InIntermediatePath;
 		}
 	}
@@ -220,7 +240,7 @@ namespace UE::Learning
 
 	bool FSharedMemoryTrainer::IsValid()
 	{
-		return TrainingProcess != nullptr;
+		return bUseExternalTrainer || TrainingProcess != nullptr;
 	}
 
 	ETrainerResponse FSharedMemoryTrainer::Wait()
@@ -244,12 +264,17 @@ namespace UE::Learning
 		return SharedMemoryTraining::SendStop(Controls.View);
 	}
 
-	ETrainerResponse FSharedMemoryTrainer::SendConfig(const TSharedRef<FJsonObject>& ConfigObject, const ELogSetting LogSettings)
+	ETrainerResponse FSharedMemoryTrainer::SendConfigs(
+		const TSharedRef<FJsonObject>& DataConfigObject,
+		const TSharedRef<FJsonObject>& TrainerConfigObject,
+		const ELogSetting LogSettings)
 	{
 		IFileManager& FileManager = IFileManager::Get();
-		ConfigObject->SetStringField(TEXT("IntermediatePath"), *FileManager.ConvertToAbsolutePathForExternalAppForRead(*IntermediatePath));
 
-		TSharedPtr<FJsonObject> SharedMemoryObject = MakeShared<FJsonObject>();
+		// Add intermediate path as a temp directory for the tensorboard logging on python side
+		TrainerConfigObject->SetStringField(TEXT("TempDirectory"), *FileManager.ConvertToAbsolutePathForExternalAppForRead(*IntermediatePath));
+
+		TSharedRef<FJsonObject> SharedMemoryConfigObject = MakeShared<FJsonObject>();
 
 		TArray<TSharedPtr<FJsonValue>> NetworkGuidsArray;
 		for(int32 Index = 0; Index < NeuralNetworkSharedMemoryArrayViews.Num(); Index++)
@@ -261,7 +286,7 @@ namespace UE::Learning
 			TSharedRef<FJsonValueObject> JsonValue = MakeShared<FJsonValueObject>(JsonObject);
 			NetworkGuidsArray.Add(JsonValue);
 		}
-		SharedMemoryObject->SetArrayField(TEXT("NetworkGuids"), NetworkGuidsArray);
+		SharedMemoryConfigObject->SetArrayField(TEXT("NetworkGuids"), NetworkGuidsArray);
 
 		TArray<TSharedPtr<FJsonValue>> ExperienceContainerObjectsArray;
 		for (const FSharedMemoryExperienceContainer& SharedMemoryExperienceContainer : SharedMemoryExperienceContainers)
@@ -331,16 +356,25 @@ namespace UE::Learning
 			TSharedRef<FJsonValueObject> JsonValue = MakeShared<FJsonValueObject>(JsonObject);
 			ExperienceContainerObjectsArray.Add(JsonValue);
 		}
-		SharedMemoryObject->SetArrayField(TEXT("ReplayBuffers"), ExperienceContainerObjectsArray);
-
-		ConfigObject->SetObjectField(TEXT("SharedMemory"), SharedMemoryObject);
+		SharedMemoryConfigObject->SetArrayField(TEXT("ReplayBuffers"), ExperienceContainerObjectsArray);
 		
-		FString ConfigString;
-		TSharedRef<TJsonWriter<>> JsonWriter = TJsonWriterFactory<>::Create(&ConfigString, 0);
-		FJsonSerializer::Serialize(ConfigObject, JsonWriter, true);
-		FFileHelper::SaveStringToFile(ConfigString, *ConfigPath);
+		// Write all the configs
+		FString SharedMemoryConfigString;
+		TSharedRef<TJsonWriter<>> SharedMemoryJsonWriter = TJsonWriterFactory<>::Create(&SharedMemoryConfigString, 0);
+		FJsonSerializer::Serialize(SharedMemoryConfigObject, SharedMemoryJsonWriter, true);
+		FFileHelper::SaveStringToFile(SharedMemoryConfigString, *(ConfigPath / FString::Printf(TEXT("shared-memory-%s.json"), *Controls.Guid.ToString(EGuidFormats::DigitsWithHyphensInBraces))));
 
-		UE_LOG(LogLearning, Display, TEXT("Wrote Config File to %s. Sending Config Signal..."), *ConfigPath);
+		FString DataConfigString;
+		TSharedRef<TJsonWriter<>> DataJsonWriter = TJsonWriterFactory<>::Create(&DataConfigString, 0);
+		FJsonSerializer::Serialize(DataConfigObject, DataJsonWriter, true);
+		FFileHelper::SaveStringToFile(DataConfigString, *(ConfigPath / "data-config.json"));
+
+		FString TrainerConfigString;
+		TSharedRef<TJsonWriter<>> TrainerJsonWriter = TJsonWriterFactory<>::Create(&TrainerConfigString, 0);
+		FJsonSerializer::Serialize(TrainerConfigObject, TrainerJsonWriter, true);
+		FFileHelper::SaveStringToFile(TrainerConfigString, *(ConfigPath / "trainer-config.json"));
+
+		UE_LOG(LogLearning, Display, TEXT("Wrote Config Files to %s. Sending Config Signal..."), *ConfigPath);
 
 		return SharedMemoryTraining::SendConfigSignal(Controls.View, LogSettings);
 	}
@@ -569,7 +603,7 @@ namespace UE::Learning
 		UE_LEARNING_CHECK(FPaths::DirectoryExists(PythonContentPath));
 
 		IFileManager& FileManager = IFileManager::Get();
-		const FString CommandLineArguments = FString::Printf(TEXT("\"%s\" -p \"%s\" -m \"%s\" \"%s\" Socket \"%s:%i\" \"%s\""),
+		const FString CommandLineArguments = FString::Printf(TEXT("\"%s\" Training -p \"%s\" -m \"%s\" \"%s\" Socket \"%s:%i\" \"%s\""),
 			*FileManager.ConvertToAbsolutePathForExternalAppForRead(*(PythonContentPath / TEXT("train.py"))),
 			*FileManager.ConvertToAbsolutePathForExternalAppForRead(*CustomTrainerPath),
 			*TrainerFileName,
@@ -629,11 +663,13 @@ namespace UE::Learning
 	FSocketTrainer::FSocketTrainer(
 		ETrainerResponse& OutResponse,
 		const TSharedPtr<FSocketTrainerServerProcess>& ExternalTrainerProcess,
+		const bool bUseExternalTrainerProcess,
 		const TCHAR* IpAddress,
 		const uint32 Port,
 		const float InTimeout)
 	{
 		Timeout = InTimeout;
+		bUseExternalTrainer = bUseExternalTrainerProcess;
 
 		if (ExternalTrainerProcess)
 		{
@@ -671,7 +707,7 @@ namespace UE::Learning
 
 	bool FSocketTrainer::IsValid()
 	{
-		return TrainingProcess != nullptr;
+		return bUseExternalTrainer || TrainingProcess != nullptr;
 	}
 
 	ETrainerResponse FSocketTrainer::Wait()
@@ -704,19 +740,23 @@ namespace UE::Learning
 		return SocketTraining::SendStop(*Socket, TrainingProcess.IsValid() ? TrainingProcess->GetTrainingSubprocess() : nullptr, Timeout);
 	}
 
-	ETrainerResponse FSocketTrainer::SendConfig(
-		const TSharedRef<FJsonObject>& ConfigObject,
+	ETrainerResponse FSocketTrainer::SendConfigs(
+		const TSharedRef<FJsonObject>& DataConfigObject,
+		const TSharedRef<FJsonObject>& TrainerConfigObject,
 		const ELogSetting LogSettings)
 	{
+
 		if (!Socket)
 		{
 			UE_LOG(LogLearning, Error, TEXT("Training socket is nullptr"));
 			return ETrainerResponse::Unexpected;
 		}
 
+		DataConfigObject->SetObjectField(TEXT("TrainerSettings"), TrainerConfigObject);
+
 		FString ConfigString;
 		TSharedRef<TJsonWriter<>> JsonWriter = TJsonWriterFactory<>::Create(&ConfigString, 0);
-		FJsonSerializer::Serialize(ConfigObject, JsonWriter, true);
+		FJsonSerializer::Serialize(DataConfigObject, JsonWriter, true);
 
 		return SocketTraining::SendConfig(
 			*Socket,
diff --git a/Engine/Plugins/Experimental/LearningAgents/Source/LearningTraining/Private/LearningPPOTrainer.cpp b/Engine/Plugins/Experimental/LearningAgents/Source/LearningTraining/Private/LearningPPOTrainer.cpp
--- a/Engine/Plugins/Experimental/LearningAgents/Source/LearningTraining/Private/LearningPPOTrainer.cpp
+++ b/Engine/Plugins/Experimental/LearningAgents/Source/LearningTraining/Private/LearningPPOTrainer.cpp
@@ -152,8 +152,6 @@ namespace UE::Learning::PPOTrainer
 
 		ETrainerResponse Response = ETrainerResponse::Success;
 
-		// Send initial Policy
-
 		if (LogSettings != ELogSetting::Silent)
 		{
 			UE_LOG(LogLearning, Display, TEXT("Sending initial Policy..."));
@@ -169,11 +167,8 @@ namespace UE::Learning::PPOTrainer
 		const int32 ActionVectorDimensionNum = ReplayBuffer.GetActions(0).Num<1>();
 		const int32 MemoryStateVectorDimensionNum = ReplayBuffer.GetMemoryStates(0).Num<1>();
 
-		// Write PPO Config
-		TSharedRef<FJsonObject> ConfigObject = MakeShared<FJsonObject>();
-		ConfigObject->SetStringField(TEXT("TaskName"), TEXT("Training"));
-		ConfigObject->SetStringField(TEXT("TrainerMethod"), TEXT("PPO"));
-		ConfigObject->SetStringField(TEXT("TimeStamp"), *FDateTime::Now().ToFormattedString(TEXT("%Y-%m-%d_%H-%M-%S")));
+		// Write Data Config
+		TSharedRef<FJsonObject> DataConfigObject = MakeShared<FJsonObject>();
 
 		const int32 ObservationSchemaId = 0;
 		const int32 ActionSchemaId = 0;
@@ -228,13 +223,13 @@ namespace UE::Learning::PPOTrainer
 			NetworkObjects.Add(JsonValue);
 		}
 
-		ConfigObject->SetArrayField(TEXT("Networks"), NetworkObjects);
+		DataConfigObject->SetArrayField(TEXT("Networks"), NetworkObjects);
 
 		// Add Replay Buffers Config Entries
 		TArray<TSharedPtr<FJsonValue>> ReplayBufferObjects;
 		TSharedRef<FJsonValueObject> ReplayBufferJsonValue = MakeShared<FJsonValueObject>(ReplayBuffer.AsJsonConfig(ReplayBufferId));
 		ReplayBufferObjects.Add(ReplayBufferJsonValue);
-		ConfigObject->SetArrayField(TEXT("ReplayBuffers"), ReplayBufferObjects);
+		DataConfigObject->SetArrayField(TEXT("ReplayBuffers"), ReplayBufferObjects);
 
 		// Schemas
 		TSharedPtr<FJsonObject> SchemasObject = MakeShared<FJsonObject>();
@@ -269,43 +264,45 @@ namespace UE::Learning::PPOTrainer
 		}
 		SchemasObject->SetArrayField(TEXT("Actions"), ActionSchemaObjects);
 
-		ConfigObject->SetObjectField(TEXT("Schemas"), SchemasObject);
+		DataConfigObject->SetObjectField(TEXT("Schemas"), SchemasObject);
 
 		// Add PPO Specific Config Entries
-		TSharedRef<FJsonObject> PPOConfigObject = MakeShared<FJsonObject>();
-
-		PPOConfigObject->SetNumberField(TEXT("IterationNum"), TrainerSettings.IterationNum);
-		PPOConfigObject->SetNumberField(TEXT("LearningRatePolicy"), TrainerSettings.LearningRatePolicy);
-		PPOConfigObject->SetNumberField(TEXT("LearningRateCritic"), TrainerSettings.LearningRateCritic);
-		PPOConfigObject->SetNumberField(TEXT("LearningRateDecay"), TrainerSettings.LearningRateDecay);
-		PPOConfigObject->SetNumberField(TEXT("WeightDecay"), TrainerSettings.WeightDecay);
-		PPOConfigObject->SetNumberField(TEXT("PolicyBatchSize"), TrainerSettings.PolicyBatchSize);
-		PPOConfigObject->SetNumberField(TEXT("CriticBatchSize"), TrainerSettings.CriticBatchSize);
-		PPOConfigObject->SetNumberField(TEXT("PolicyWindow"), TrainerSettings.PolicyWindow);
-		PPOConfigObject->SetNumberField(TEXT("IterationsPerGather"), TrainerSettings.IterationsPerGather);
-		PPOConfigObject->SetNumberField(TEXT("CriticWarmupIterations"), TrainerSettings.CriticWarmupIterations);
-		PPOConfigObject->SetNumberField(TEXT("EpsilonClip"), TrainerSettings.EpsilonClip);
-		PPOConfigObject->SetNumberField(TEXT("ActionSurrogateWeight"), TrainerSettings.ActionSurrogateWeight);
-		PPOConfigObject->SetNumberField(TEXT("ActionRegularizationWeight"), TrainerSettings.ActionRegularizationWeight);
-		PPOConfigObject->SetNumberField(TEXT("ActionEntropyWeight"), TrainerSettings.ActionEntropyWeight);
-		PPOConfigObject->SetNumberField(TEXT("ReturnRegularizationWeight"), TrainerSettings.ReturnRegularizationWeight);
-		PPOConfigObject->SetNumberField(TEXT("GaeLambda"), TrainerSettings.GaeLambda);
-		PPOConfigObject->SetBoolField(TEXT("AdvantageNormalization"), TrainerSettings.bAdvantageNormalization);
-		PPOConfigObject->SetNumberField(TEXT("AdvantageMin"), TrainerSettings.AdvantageMin);
-		PPOConfigObject->SetNumberField(TEXT("AdvantageMax"), TrainerSettings.AdvantageMax);
-		PPOConfigObject->SetBoolField(TEXT("UseGradNormMaxClipping"), TrainerSettings.bUseGradNormMaxClipping);
-		PPOConfigObject->SetNumberField(TEXT("GradNormMax"), TrainerSettings.GradNormMax);
-		PPOConfigObject->SetNumberField(TEXT("TrimEpisodeStartStepNum"), TrainerSettings.TrimEpisodeStartStepNum);
-		PPOConfigObject->SetNumberField(TEXT("TrimEpisodeEndStepNum"), TrainerSettings.TrimEpisodeEndStepNum);
-		PPOConfigObject->SetNumberField(TEXT("Seed"), TrainerSettings.Seed);
-		PPOConfigObject->SetNumberField(TEXT("DiscountFactor"), TrainerSettings.DiscountFactor);
-		PPOConfigObject->SetStringField(TEXT("Device"), UE::Learning::Trainer::GetDeviceString(TrainerSettings.Device));
-		PPOConfigObject->SetBoolField(TEXT("UseTensorBoard"), TrainerSettings.bUseTensorboard);
-		PPOConfigObject->SetBoolField(TEXT("SaveSnapshots"), TrainerSettings.bSaveSnapshots);
-
-		ConfigObject->SetObjectField(TEXT("PPOSettings"), PPOConfigObject);
-
-		ExternalTrainer->SendConfig(ConfigObject, LogSettings);
+		TSharedRef<FJsonObject> TrainingConfigObject = MakeShared<FJsonObject>();
+
+		TrainingConfigObject->SetStringField(TEXT("TaskName"), TEXT("Training"));
+		TrainingConfigObject->SetStringField(TEXT("TrainerMethod"), TEXT("PPO"));
+		TrainingConfigObject->SetStringField(TEXT("TimeStamp"), *FDateTime::Now().ToFormattedString(TEXT("%Y-%m-%d_%H-%M-%S")));
+
+		TrainingConfigObject->SetNumberField(TEXT("IterationNum"), TrainerSettings.IterationNum);
+		TrainingConfigObject->SetNumberField(TEXT("LearningRatePolicy"), TrainerSettings.LearningRatePolicy);
+		TrainingConfigObject->SetNumberField(TEXT("LearningRateCritic"), TrainerSettings.LearningRateCritic);
+		TrainingConfigObject->SetNumberField(TEXT("LearningRateDecay"), TrainerSettings.LearningRateDecay);
+		TrainingConfigObject->SetNumberField(TEXT("WeightDecay"), TrainerSettings.WeightDecay);
+		TrainingConfigObject->SetNumberField(TEXT("PolicyBatchSize"), TrainerSettings.PolicyBatchSize);
+		TrainingConfigObject->SetNumberField(TEXT("CriticBatchSize"), TrainerSettings.CriticBatchSize);
+		TrainingConfigObject->SetNumberField(TEXT("PolicyWindow"), TrainerSettings.PolicyWindow);
+		TrainingConfigObject->SetNumberField(TEXT("IterationsPerGather"), TrainerSettings.IterationsPerGather);
+		TrainingConfigObject->SetNumberField(TEXT("CriticWarmupIterations"), TrainerSettings.CriticWarmupIterations);
+		TrainingConfigObject->SetNumberField(TEXT("EpsilonClip"), TrainerSettings.EpsilonClip);
+		TrainingConfigObject->SetNumberField(TEXT("ActionSurrogateWeight"), TrainerSettings.ActionSurrogateWeight);
+		TrainingConfigObject->SetNumberField(TEXT("ActionRegularizationWeight"), TrainerSettings.ActionRegularizationWeight);
+		TrainingConfigObject->SetNumberField(TEXT("ActionEntropyWeight"), TrainerSettings.ActionEntropyWeight);
+		TrainingConfigObject->SetNumberField(TEXT("ReturnRegularizationWeight"), TrainerSettings.ReturnRegularizationWeight);
+		TrainingConfigObject->SetNumberField(TEXT("GaeLambda"), TrainerSettings.GaeLambda);
+		TrainingConfigObject->SetBoolField(TEXT("AdvantageNormalization"), TrainerSettings.bAdvantageNormalization);
+		TrainingConfigObject->SetNumberField(TEXT("AdvantageMin"), TrainerSettings.AdvantageMin);
+		TrainingConfigObject->SetNumberField(TEXT("AdvantageMax"), TrainerSettings.AdvantageMax);
+		TrainingConfigObject->SetBoolField(TEXT("UseGradNormMaxClipping"), TrainerSettings.bUseGradNormMaxClipping);
+		TrainingConfigObject->SetNumberField(TEXT("GradNormMax"), TrainerSettings.GradNormMax);
+		TrainingConfigObject->SetNumberField(TEXT("TrimEpisodeStartStepNum"), TrainerSettings.TrimEpisodeStartStepNum);
+		TrainingConfigObject->SetNumberField(TEXT("TrimEpisodeEndStepNum"), TrainerSettings.TrimEpisodeEndStepNum);
+		TrainingConfigObject->SetNumberField(TEXT("Seed"), TrainerSettings.Seed);
+		TrainingConfigObject->SetNumberField(TEXT("DiscountFactor"), TrainerSettings.DiscountFactor);
+		TrainingConfigObject->SetStringField(TEXT("Device"), UE::Learning::Trainer::GetDeviceString(TrainerSettings.Device));
+		TrainingConfigObject->SetBoolField(TEXT("UseTensorBoard"), TrainerSettings.bUseTensorboard);
+		TrainingConfigObject->SetBoolField(TEXT("SaveSnapshots"), TrainerSettings.bSaveSnapshots);
+
+		ExternalTrainer->SendConfigs(DataConfigObject, TrainingConfigObject, LogSettings);
 
 		Response = ExternalTrainer->SendNetwork(PolicyNetworkId, PolicyNetwork, PolicyNetworkLock);
 
diff --git a/Engine/Plugins/Experimental/LearningAgents/Source/LearningTraining/Public/LearningExternalTrainer.h b/Engine/Plugins/Experimental/LearningAgents/Source/LearningTraining/Public/LearningExternalTrainer.h
--- a/Engine/Plugins/Experimental/LearningAgents/Source/LearningTraining/Public/LearningExternalTrainer.h
+++ b/Engine/Plugins/Experimental/LearningAgents/Source/LearningTraining/Public/LearningExternalTrainer.h
@@ -46,14 +46,16 @@ namespace UE::Learning
 		virtual bool HasNetworkOrCompleted() = 0;
 
 		/**
-		* Sends the given json config to the trainer process.
+		* Sends the given json configs to the trainer process.
 		*
-		* @param ConfigObject	The config to send
-		* @param LogSettings	The log verbosity level
-		* @returns				Trainer response
+		* @param DataConfigObject		The config to send with meta-data
+		* @param TrainerConfigObject	The config to send with trainer-specific settings
+		* @param LogSettings			The log verbosity level
+		* @returns						Trainer response
 		*/
-		virtual ETrainerResponse SendConfig(
-			const TSharedRef<FJsonObject>& ConfigObject,
+		virtual ETrainerResponse SendConfigs(
+			const TSharedRef<FJsonObject>& DataConfigObject,
+			const TSharedRef<FJsonObject>& TrainerConfigObject,
 			const ELogSetting LogSettings = Trainer::DefaultLogSettings) = 0;
 
 		/**
@@ -252,8 +254,9 @@ namespace UE::Learning
 
 		virtual bool HasNetworkOrCompleted() override final;
 
-		virtual ETrainerResponse SendConfig(
-			const TSharedRef<FJsonObject>& ConfigObject,
+		virtual ETrainerResponse SendConfigs(
+			const TSharedRef<FJsonObject>& DataConfigObject,
+			const TSharedRef<FJsonObject>& TrainerConfigObject,
 			const ELogSetting LogSettings = Trainer::DefaultLogSettings) override final;
 
 		virtual int32 AddNetwork(const ULearningNeuralNetworkData& Network) override final;
@@ -290,6 +293,8 @@ namespace UE::Learning
 
 		float Timeout = Trainer::DefaultTimeout;
 
+		bool bUseExternalTrainer = false;
+
 		TSharedMemoryArrayView<1, volatile int32> Controls;
 		TArray<TSharedMemoryArrayView<1, uint8>> NeuralNetworkSharedMemoryArrayViews;
 		TArray<FSharedMemoryExperienceContainer> SharedMemoryExperienceContainers;
@@ -373,6 +378,7 @@ namespace UE::Learning
 		FSocketTrainer(
 			ETrainerResponse& OutResponse,
 			const TSharedPtr<FSocketTrainerServerProcess>& ExternalTrainerProcess,
+			const bool bUseExternalTrainerProcess,
 			const TCHAR* IpAddress = Trainer::DefaultIp,
 			const uint32 Port = Trainer::DefaultPort,
 			const float InTimeout = Trainer::DefaultTimeout);
@@ -389,8 +395,9 @@ namespace UE::Learning
 
 		virtual bool HasNetworkOrCompleted() override final;
 
-		virtual ETrainerResponse SendConfig(
-			const TSharedRef<FJsonObject>& ConfigObject,
+		virtual ETrainerResponse SendConfigs(
+			const TSharedRef<FJsonObject>& DataConfigObject,
+			const TSharedRef<FJsonObject>& TrainerConfigObject,
 			const ELogSetting LogSettings = Trainer::DefaultLogSettings) override final;
 
 		virtual int32 AddNetwork(const ULearningNeuralNetworkData& Network) override final;
@@ -421,6 +428,8 @@ namespace UE::Learning
 
 		float Timeout = Trainer::DefaultTimeout;
 
+		bool bUseExternalTrainer = false;
+
 		TSharedPtr<FSocketTrainerServerProcess> TrainingProcess = nullptr;
 
 		FSocket* Socket = nullptr;
diff --git a/Engine/Plugins/Experimental/LearningAgents/Source/LearningTraining/Public/LearningTrainer.h b/Engine/Plugins/Experimental/LearningAgents/Source/LearningTraining/Public/LearningTrainer.h
--- a/Engine/Plugins/Experimental/LearningAgents/Source/LearningTraining/Public/LearningTrainer.h
+++ b/Engine/Plugins/Experimental/LearningAgents/Source/LearningTraining/Public/LearningTrainer.h
@@ -130,7 +130,7 @@ namespace UE::Learning
 		/**
 		* Default Log Settings to use during communication.
 		*/
-		static constexpr ELogSetting DefaultLogSettings = ELogSetting::Silent;
+		static constexpr ELogSetting DefaultLogSettings = ELogSetting::Normal;
 
 		/**
 		* Default IP to use for networked training
