Learning Agents: Fix Socket timeouts being too short

diff --git a/Engine/Plugins/Experimental/LearningAgents/Source/LearningTraining/Private/LearningSocketTraining.cpp b/Engine/Plugins/Experimental/LearningAgents/Source/LearningTraining/Private/LearningSocketTraining.cpp
--- a/Engine/Plugins/Experimental/LearningAgents/Source/LearningTraining/Private/LearningSocketTraining.cpp
+++ b/Engine/Plugins/Experimental/LearningAgents/Source/LearningTraining/Private/LearningSocketTraining.cpp
@@ -46,7 +46,7 @@ namespace UE::Learning::SocketTraining
 
 		while (!Process || Process->Update())
 		{
-			if (Socket.Wait(ESocketWaitConditions::WaitForRead, FTimespan(1000))) // 1 millisecond
+			if (Socket.Wait(ESocketWaitConditions::WaitForRead, FTimespan(10000))) // 1 millisecond, as FTimespan are measured in 0.1 microseconds
 			{
 				if (Socket.Recv(Bytes + TotalBytesRead, ByteNum - TotalBytesRead, BytesRead))
 				{
@@ -156,7 +156,7 @@ namespace UE::Learning::SocketTraining
 
 		while (!Process || Process->Update())
 		{
-			if (Socket.Wait(ESocketWaitConditions::WaitForWrite, FTimespan(1000))) // 1 millisecond
+			if (Socket.Wait(ESocketWaitConditions::WaitForWrite, FTimespan(10000))) // 1 millisecond, as FTimespan are measured in 0.1 microseconds
 			{
 				if (Socket.Send(Bytes + TotalBytesSent, ByteNum - TotalBytesSent, BytesSent))
 				{
