[LearningAgents] Added delegate events to LearningAgentsGym to mark gym restart & init

diff --git a/Engine/Plugins/Experimental/LearningAgents/Source/LearningAgentsTraining/Private/LearningAgentsGym.cpp b/Engine/Plugins/Experimental/LearningAgents/Source/LearningAgentsTraining/Private/LearningAgentsGym.cpp
--- a/Engine/Plugins/Experimental/LearningAgents/Source/LearningAgentsTraining/Private/LearningAgentsGym.cpp
+++ b/Engine/Plugins/Experimental/LearningAgents/Source/LearningAgentsTraining/Private/LearningAgentsGym.cpp
@@ -25,14 +25,20 @@ void ALearningAgentsGymBase::Initialize()
 	{
 		LearningComponent->InitializeLearningComponent();
 	}
+	
+	OnGymInitialized.Broadcast();
 }
 
 void ALearningAgentsGymBase::Reset()
 {
+	OnBeginGymReset.Broadcast();
+
 	for (TScriptInterface<ILearningAgentsLearningComponentInterface>& LearningComponent : LearningComponents)
 	{
 		LearningComponent->ResetLearningComponent();
 	}
+
+	OnPostGymReset.Broadcast();
 }
 
 void ALearningAgentsGymBase::GetRandomStream(FRandomStream& OutRandomStream) const
diff --git a/Engine/Plugins/Experimental/LearningAgents/Source/LearningAgentsTraining/Public/LearningAgentsGym.h b/Engine/Plugins/Experimental/LearningAgents/Source/LearningAgentsTraining/Public/LearningAgentsGym.h
--- a/Engine/Plugins/Experimental/LearningAgents/Source/LearningAgentsTraining/Public/LearningAgentsGym.h
+++ b/Engine/Plugins/Experimental/LearningAgents/Source/LearningAgentsTraining/Public/LearningAgentsGym.h
@@ -10,6 +10,10 @@
 #define INVALID_GYM_VECTOR (FVector::ZeroVector)
 #define INVALID_GYM_ROTATOR (FRotator::ZeroRotator)
 
+DECLARE_DYNAMIC_MULTICAST_DELEGATE(FOnGymInitializedSignature);
+DECLARE_DYNAMIC_MULTICAST_DELEGATE(FOnBeginGymResetSignature);
+DECLARE_DYNAMIC_MULTICAST_DELEGATE(FOnPostGymResetSignature);
+
 class ILearningAgentsLearningComponentInterface;
 
 /**
@@ -31,6 +35,19 @@ public:
 	UFUNCTION(BlueprintCallable, Category = "LearningAgents")
 	void Reset();
 
+public:
+	/** Event called at a gym's initialization. */
+	UPROPERTY(BlueprintAssignable, Category = "LearningAgents")
+	FOnGymInitializedSignature OnGymInitialized;
+
+	/** Event called at the start of a gym's reset. */
+	UPROPERTY(BlueprintAssignable, Category = "LearningAgents")
+	FOnBeginGymResetSignature OnBeginGymReset;
+
+	/** Event called at the end of a gym's reset. */
+	UPROPERTY(BlueprintAssignable, Category = "LearningAgents")
+	FOnPostGymResetSignature OnPostGymReset;
+
 public:
 	/** Retrieves the current random stream used by the Gym.*/
 	UFUNCTION(BlueprintPure, Category = "LearningAgents")
