[LearningAgents] EntitiesManager Blueprint Spawner Calls

#rnx

diff --git a/Engine/Plugins/Experimental/LearningAgents/Source/LearningAgentsTraining/Private/LearningAgentsEntitiesManagerComponent.cpp b/Engine/Plugins/Experimental/LearningAgents/Source/LearningAgentsTraining/Private/LearningAgentsEntitiesManagerComponent.cpp
--- a/Engine/Plugins/Experimental/LearningAgents/Source/LearningAgentsTraining/Private/LearningAgentsEntitiesManagerComponent.cpp
+++ b/Engine/Plugins/Experimental/LearningAgents/Source/LearningAgentsTraining/Private/LearningAgentsEntitiesManagerComponent.cpp
@@ -20,7 +20,7 @@ void ULearningAgentsEntitiesManagerComponent::InitializeLearningComponent()
 		{
 			for (FEntityInfo& EntityInfo : Entities)
 			{
-				SpawnEntities(EntityInfo, EntityInfo.EpisodeStartEntitySpawnCount);
+				SpawnEntitiesAtRandomLocations(EntityInfo.EntityClass, EntityInfo.EntitySpawnZOffset, EntityInfo.EpisodeStartEntitySpawnCount);
 			}
 		}
 
@@ -71,32 +71,64 @@ void ULearningAgentsEntitiesManagerComponent::ResetLearningComponent()
 	}
 }
 
-TScriptInterface<ILearningAgentsEntityInterface> ULearningAgentsEntitiesManagerComponent::SpawnEntity(const FEntityInfo& EntityInfo, const FTransform& InTransform)
+TScriptInterface<ILearningAgentsEntityInterface> ULearningAgentsEntitiesManagerComponent::SpawnEntitiesAtRandomLocations(TSubclassOf<AActor> EntityClass, float EntitySpawnZOffset, int32 SpawnCount)
 {
-	TArray<TScriptInterface<ILearningAgentsEntityInterface>> SpawnedEntities = SpawnEntities(EntityInfo, 1, false, InTransform);
+	if (SpawnCount < 0)
+	{
+		UE_LOG(LogLearning, Warning, TEXT("Spawn Count must be greater than 0 in %s!"), *EntityClass->GetName());
+		return TScriptInterface<ILearningAgentsEntityInterface>();
+	}
+
+	FTransform Transform = FTransform();
+	RandomizeTransform(Transform, EntitySpawnZOffset);
+
+	TArray<TScriptInterface<ILearningAgentsEntityInterface>> SpawnedEntities = SpawnEntities(EntityClass, EntitySpawnZOffset, SpawnCount, Transform);
 	if (!SpawnedEntities.IsEmpty())
 	{
 		return SpawnedEntities.Top();
 	}
 	else
 	{
-		UE_LOG(LogLearning, Warning, TEXT("Could not spawn entity from class %s!"), *EntityInfo.EntityClass->GetName());
+		UE_LOG(LogLearning, Warning, TEXT("Could not spawn entity from class %s!"), *EntityClass->GetName());
 		return TScriptInterface<ILearningAgentsEntityInterface>();
 	}
 }
 
-TArray<TScriptInterface<ILearningAgentsEntityInterface>> ULearningAgentsEntitiesManagerComponent::SpawnEntities(const FEntityInfo& EntityInfo, int32 SpawnCount, bool bRandomizeTransformInGym, const FTransform& InTransform)
+TScriptInterface<ILearningAgentsEntityInterface> ULearningAgentsEntitiesManagerComponent::SpawnEntityAtProjectedLocation(TSubclassOf<AActor> EntityClass, float EntitySpawnZOffset, const FTransform& InTransform)
+{
+	FTransform Transform = InTransform;
+	ProjectTransform(Transform);
+
+	TArray<TScriptInterface<ILearningAgentsEntityInterface>> SpawnedEntities = SpawnEntities(EntityClass, EntitySpawnZOffset, 1, Transform);
+	if (!SpawnedEntities.IsEmpty())
+	{
+		return SpawnedEntities.Top();
+	}
+	else
+	{
+		UE_LOG(LogLearning, Warning, TEXT("Could not spawn entity from class %s!"), *EntityClass->GetName());
+		return TScriptInterface<ILearningAgentsEntityInterface>();
+	}
+}
+
+TArray<TScriptInterface<ILearningAgentsEntityInterface>> ULearningAgentsEntitiesManagerComponent::SpawnEntities(TSubclassOf<AActor> EntityClass, float EntitySpawnZOffset, int32 SpawnCount, const FTransform& InTransform)
 {
 	TArray<TScriptInterface<ILearningAgentsEntityInterface>> SpawnedEntities;
 	SpawnedEntities.SetNumUninitialized(SpawnCount);
 
+	if (SpawnCount < 0)
+	{
+		UE_LOG(LogLearning, Warning, TEXT("Spawn Count must be greater than 0 in %s!"), *EntityClass->GetName());
+		return SpawnedEntities;
+	}
+
 	FActorSpawnParameters SpawnParameters;
 	SpawnParameters.Owner = GetOwner();
 
 	FTransform Transform = InTransform;
 
 	int32 SuccessfulSpawnCount = 0;
-	if (FSpawnedEntitiesInfo* SpawnEntitiesInfo = EntitiesPool.Find(EntityInfo.EntityClass->GetFName()))
+	if (FSpawnedEntitiesInfo* SpawnEntitiesInfo = EntitiesPool.Find(EntityClass->GetFName()))
 	{
 		for (TScriptInterface<ILearningAgentsEntityInterface>& Entity : SpawnEntitiesInfo->SpawnedEntities)
 		{
@@ -108,14 +140,6 @@ TArray<TScriptInterface<ILearningAgentsEntityInterface>> ULearningAgentsEntities
 			{
 				if (TObjectPtr<AActor> ActorEntity = Cast<AActor>(Entity.GetObject()))
 				{
-					if (bRandomizeTransformInGym)
-					{
-						RandomizeTransform(Transform, EntityInfo.EntitySpawnZOffset);
-					}
-					else
-					{
-						ProjectTransform(Transform);
-					}
 					ActorEntity->SetActorTransform(Transform);
 				}
 
@@ -128,15 +152,11 @@ TArray<TScriptInterface<ILearningAgentsEntityInterface>> ULearningAgentsEntities
 		// Allocate new object and add to pool because there's not enough in pool to pull from.
 		for (int32 Index = SuccessfulSpawnCount; SuccessfulSpawnCount < SpawnCount; ++SuccessfulSpawnCount)
 		{
-			if (bRandomizeTransformInGym)
+			TObjectPtr<AActor> Entity = GetWorld()->SpawnActor<AActor>(EntityClass, Transform.GetLocation(), Transform.GetRotation().Rotator(), SpawnParameters);
+			if (TObjectPtr<ALearningAgentsGymBase> Gym = Cast<ALearningAgentsGymBase>(GetOwner()))
 			{
-				RandomizeTransform(Transform, EntityInfo.EntitySpawnZOffset);
+				ILearningAgentsEntityInterface::Execute_InitializeEntity(Entity, Gym);
 			}
-			else
-			{
-				ProjectTransform(Transform);
-			}
-			TObjectPtr<AActor> Entity = GetWorld()->SpawnActor<AActor>(EntityInfo.EntityClass, Transform.GetLocation(), Transform.GetRotation().Rotator(), SpawnParameters);
 			SpawnEntitiesInfo->SpawnedEntities.Add(Entity);
 			SpawnedEntities.Add(TScriptInterface<ILearningAgentsEntityInterface>(Entity));
 		}
@@ -145,23 +165,19 @@ TArray<TScriptInterface<ILearningAgentsEntityInterface>> ULearningAgentsEntities
 	else
 	{
 		FSpawnedEntitiesInfo NewSpawnEntitiesInfo;
-		NewSpawnEntitiesInfo.EntityInfo = EntityInfo;
+		NewSpawnEntitiesInfo.EntityInfo = FEntityInfo(EntityClass, EntitySpawnZOffset, 0);
 
-		if (bRandomizeTransformInGym)
-		{
-			RandomizeTransform(Transform, EntityInfo.EntitySpawnZOffset);
-		}
-		else
-		{
-			ProjectTransform(Transform);
-		}
 		for (int32 Index = 0; Index < SpawnCount; ++Index)
 		{
-			TObjectPtr<AActor> Entity = GetWorld()->SpawnActor<AActor>(EntityInfo.EntityClass, Transform.GetLocation(), Transform.GetRotation().Rotator(), SpawnParameters);
+			TObjectPtr<AActor> Entity = GetWorld()->SpawnActor<AActor>(EntityClass, Transform.GetLocation(), Transform.GetRotation().Rotator(), SpawnParameters);
+			if (TObjectPtr<ALearningAgentsGymBase> Gym = Cast<ALearningAgentsGymBase>(GetOwner()))
+			{
+				ILearningAgentsEntityInterface::Execute_InitializeEntity(Entity, Gym);
+			}
 			NewSpawnEntitiesInfo.SpawnedEntities.Add(Entity);
 			SpawnedEntities.Add(TScriptInterface<ILearningAgentsEntityInterface>(Entity));
 		}
-		EntitiesPool.Add(EntityInfo.EntityClass->GetFName(), NewSpawnEntitiesInfo);
+		EntitiesPool.Add(EntityClass->GetFName(), NewSpawnEntitiesInfo);
 	}
 	return SpawnedEntities;
 }
diff --git a/Engine/Plugins/Experimental/LearningAgents/Source/LearningAgentsTraining/Public/LearningAgentsEntitiesManagerComponent.h b/Engine/Plugins/Experimental/LearningAgents/Source/LearningAgentsTraining/Public/LearningAgentsEntitiesManagerComponent.h
--- a/Engine/Plugins/Experimental/LearningAgents/Source/LearningAgentsTraining/Public/LearningAgentsEntitiesManagerComponent.h
+++ b/Engine/Plugins/Experimental/LearningAgents/Source/LearningAgentsTraining/Public/LearningAgentsEntitiesManagerComponent.h
@@ -64,13 +64,17 @@ public:
 	/** Resets entities that the component manages. */
 	virtual void ResetLearningComponent();
 
-	/** Spawns a single pooled entity. */
+	/** Spawns pooled entities at random locations. */
 	UFUNCTION(BlueprintCallable, Category = "LearningAgents")
-	TScriptInterface<ILearningAgentsEntityInterface> SpawnEntity(const FEntityInfo& EntityInfo, const FTransform& Transform);
+	TScriptInterface<ILearningAgentsEntityInterface> SpawnEntitiesAtRandomLocations(TSubclassOf<AActor> EntityClass, float EntitySpawnZOffset, int32 SpawnCount);
+
+	/** Spawns a single pooled entity with a specified transform projected in the gym. */
+	UFUNCTION(BlueprintCallable, Category = "LearningAgents")
+	TScriptInterface<ILearningAgentsEntityInterface> SpawnEntityAtProjectedLocation(TSubclassOf<AActor> EntityClass, float EntitySpawnZOffset, const FTransform& Transform);
 
 	/** Spawns multiple pooled entities. */
-	UFUNCTION(BlueprintCallable, Category = "LearningAgents", meta = (ClampMin = "0", UIMin = "0"))
-	TArray<TScriptInterface<ILearningAgentsEntityInterface>> SpawnEntities(const FEntityInfo& EntityInfo, int32 SpawnCount, bool bRandomizeTransformInGym = true, const FTransform& Transform = FTransform());
+	UFUNCTION(BlueprintCallable, Category = "LearningAgents")
+	TArray<TScriptInterface<ILearningAgentsEntityInterface>> SpawnEntities(TSubclassOf<AActor> EntityClass, float EntitySpawnZOffset, int32 SpawnCount, const FTransform& Transform);
 
 protected:
 	virtual void BeginPlay() override;
