[LearningAgents] Fixed missing constructor

diff --git a/Engine/Plugins/Experimental/LearningAgents/Source/LearningAgentsTraining/Private/LearningAgentsEntitiesManagerComponent.cpp b/Engine/Plugins/Experimental/LearningAgents/Source/LearningAgentsTraining/Private/LearningAgentsEntitiesManagerComponent.cpp
--- a/Engine/Plugins/Experimental/LearningAgents/Source/LearningAgentsTraining/Private/LearningAgentsEntitiesManagerComponent.cpp
+++ b/Engine/Plugins/Experimental/LearningAgents/Source/LearningAgentsTraining/Private/LearningAgentsEntitiesManagerComponent.cpp
@@ -165,7 +165,11 @@ TArray<TScriptInterface<ILearningAgentsEntityInterface>> ULearningAgentsEntities
 	else
 	{
 		FSpawnedEntitiesInfo NewSpawnEntitiesInfo;
-		NewSpawnEntitiesInfo.EntityInfo = FEntityInfo(EntityClass, EntitySpawnZOffset, 0);
+		FEntityInfo NewEntityInfo;
+		NewEntityInfo.EntityClass = EntityClass;
+		NewEntityInfo.EntitySpawnZOffset = EntitySpawnZOffset;
+		NewEntityInfo.EpisodeStartEntitySpawnCount = 0;
+		NewSpawnEntitiesInfo.EntityInfo = NewEntityInfo;
 
 		for (int32 Index = 0; Index < SpawnCount; ++Index)
 		{
