Learning Agents: Provide Hostname override on the cmdline

diff --git a/Engine/Plugins/Experimental/LearningAgents/Source/LearningAgentsTraining/Private/LearningAgentsCommunicator.cpp b/Engine/Plugins/Experimental/LearningAgents/Source/LearningAgentsTraining/Private/LearningAgentsCommunicator.cpp
--- a/Engine/Plugins/Experimental/LearningAgents/Source/LearningAgentsTraining/Private/LearningAgentsCommunicator.cpp
+++ b/Engine/Plugins/Experimental/LearningAgents/Source/LearningAgentsTraining/Private/LearningAgentsCommunicator.cpp
@@ -193,7 +193,15 @@ FLearningAgentsCommunicator ULearningAgentsCommunicatorLibrary::MakeSocketCommun
 	// Ip Address
 	FString CmdLineIpAddress;
 	bool bIpAddressOverridden = FParse::Value(FCommandLine::Get(), TEXT("LearningAgentsIpAddress="), CmdLineIpAddress);
-	if (bIpAddressOverridden)
+
+	// Hostname
+	FString CmdLineHostname;
+	bool bHostnameOverridden = FParse::Value(FCommandLine::Get(), TEXT("LearningAgentsHostname="), CmdLineHostname);
+	if (bHostnameOverridden)
+	{
+		UE_LOG(LogLearning, Display, TEXT("Using Hostname override from the cmdline: %s"), *CmdLineHostname);
+	}
+	else if (bIpAddressOverridden)
 	{
 		UE_LOG(LogLearning, Display, TEXT("Using Ip Address override from the cmdline: %s"), *CmdLineIpAddress);
 	}
@@ -219,9 +227,12 @@ FLearningAgentsCommunicator ULearningAgentsCommunicatorLibrary::MakeSocketCommun
 		Response,
 		TrainerProcess.TrainerProcess,
 		FParse::Param(FCommandLine::Get(), TEXT("LearningAgentsUseExternalTrainingProcess")) || SocketSettings.bUseExternalTrainingProcess,
-		bIpAddressOverridden ? *CmdLineIpAddress : *SocketSettings.IpAddress,
+		bHostnameOverridden
+			? *CmdLineHostname
+			: bIpAddressOverridden ? *CmdLineIpAddress: *SocketSettings.IpAddress,
 		bPortOverridden ? CmdLinePort : SocketSettings.Port,
-		bTimeoutOverridden ? CmdLineTimeout : SocketSettings.Timeout);
+		bTimeoutOverridden ? CmdLineTimeout : SocketSettings.Timeout,
+		bHostnameOverridden);
 
 	if (Response != UE::Learning::ETrainerResponse::Success)
 	{
diff --git a/Engine/Plugins/Experimental/LearningAgents/Source/LearningTraining/Private/LearningExternalTrainer.cpp b/Engine/Plugins/Experimental/LearningAgents/Source/LearningTraining/Private/LearningExternalTrainer.cpp
--- a/Engine/Plugins/Experimental/LearningAgents/Source/LearningTraining/Private/LearningExternalTrainer.cpp
+++ b/Engine/Plugins/Experimental/LearningAgents/Source/LearningTraining/Private/LearningExternalTrainer.cpp
@@ -21,6 +21,7 @@
 
 #include "Sockets.h"
 #include "Common/TcpSocketBuilder.h"
+#include "IPAddressAsyncResolve.h"
 #include "SocketSubsystem.h"
 
 namespace UE::Learning
@@ -664,9 +665,10 @@ namespace UE::Learning
 		ETrainerResponse& OutResponse,
 		const TSharedPtr<FSocketTrainerServerProcess>& ExternalTrainerProcess,
 		const bool bUseExternalTrainerProcess,
-		const TCHAR* IpAddress,
+		const TCHAR* IpAddressOrHostname,
 		const uint32 Port,
-		const float InTimeout)
+		const float InTimeout,
+		const bool IsHostname)
 	{
 		Timeout = InTimeout;
 		bUseExternalTrainer = bUseExternalTrainerProcess;
@@ -681,18 +683,53 @@ namespace UE::Learning
 
 		bool bIsValid = false;
 		TSharedRef<FInternetAddr> Address = SocketSubsystem->CreateInternetAddr();
-		Address->SetIp(IpAddress, bIsValid);
 		Address->SetPort(Port);
 
+		if (IsHostname)
+		{
+			FResolveInfo* ResolveInfo = SocketSubsystem->GetHostByName(TCHAR_TO_ANSI(IpAddressOrHostname));
+			
+			const float SleepTime = 0.001f;
+			float WaitTime = 0.0f;
+			while (!ResolveInfo->IsComplete())
+			{
+				FPlatformProcess::Sleep(SleepTime);
+				WaitTime += SleepTime;
+
+				if (WaitTime > Timeout)
+				{
+					UE_LOG(LogLearning, Warning, TEXT("Timed-out waiting for DNS..."));
+					break;
+				}
+			}
+
+			if (ResolveInfo->IsComplete() && ResolveInfo->GetErrorCode() == 0)
+			{
+				uint32 Ip;
+				const FInternetAddr& ResolvedAddress = ResolveInfo->GetResolvedAddress();
+				ResolvedAddress.GetIp(Ip);
+				Address->SetIp(Ip);
+			}
+			else
+			{
+				UE_LOG(LogLearning, Error, TEXT("Unable to resolve hostname %s. Error code %d..."), IpAddressOrHostname, ResolveInfo->GetErrorCode());
+				OutResponse = ETrainerResponse::Unexpected;
+				return;
+			}
+		}
+		else
+		{
+			Address->SetIp(IpAddressOrHostname, bIsValid);
+		}
+
 		if (!bIsValid)
 		{
-			UE_LOG(LogLearning, Error, TEXT("Invalid Ip Address \"%s\"..."), IpAddress);
+			UE_LOG(LogLearning, Error, TEXT("Invalid Ip Address \"%s\"..."), IpAddressOrHostname);
 			OutResponse = ETrainerResponse::Unexpected;
 			return;
 		}
 
 		Socket = FTcpSocketBuilder(TEXT("LearningTrainerSocket")).AsBlocking().Build();
-
 		OutResponse = SocketTraining::WaitForConnection(
 			*Socket,
 			TrainingProcess.IsValid() ? TrainingProcess->GetTrainingSubprocess() : nullptr,
diff --git a/Engine/Plugins/Experimental/LearningAgents/Source/LearningTraining/Public/LearningExternalTrainer.h b/Engine/Plugins/Experimental/LearningAgents/Source/LearningTraining/Public/LearningExternalTrainer.h
--- a/Engine/Plugins/Experimental/LearningAgents/Source/LearningTraining/Public/LearningExternalTrainer.h
+++ b/Engine/Plugins/Experimental/LearningAgents/Source/LearningTraining/Public/LearningExternalTrainer.h
@@ -379,9 +379,10 @@ namespace UE::Learning
 			ETrainerResponse& OutResponse,
 			const TSharedPtr<FSocketTrainerServerProcess>& ExternalTrainerProcess,
 			const bool bUseExternalTrainerProcess,
-			const TCHAR* IpAddress = Trainer::DefaultIp,
+			const TCHAR* IpAddressOrHostname = Trainer::DefaultIp,
 			const uint32 Port = Trainer::DefaultPort,
-			const float InTimeout = Trainer::DefaultTimeout);
+			const float InTimeout = Trainer::DefaultTimeout,
+			const bool IsHostname = false);
 
 		~FSocketTrainer();
 
