LearningAgents: Added FrameSet and FrameRangeSet data structures.

diff --git a/Engine/Plugins/Experimental/LearningAgents/Source/Learning/Private/LearningFrameRangeSet.cpp b/Engine/Plugins/Experimental/LearningAgents/Source/Learning/Private/LearningFrameRangeSet.cpp
new file mode 100644
--- /dev/null
+++ b/Engine/Plugins/Experimental/LearningAgents/Source/Learning/Private/LearningFrameRangeSet.cpp
@@ -0,0 +1,2284 @@
+// Copyright Epic Games, Inc. All Rights Reserved.
+
+#include "LearningFrameRangeSet.h"
+#include "LearningFrameSet.h"
+
+#include "Async/ParallelFor.h"
+
+namespace UE::Learning
+{
+	namespace FrameRangeSet::Private
+	{
+		static inline void FramesCheck(const TLearningArrayView<1, const int32> Frames)
+		{
+			const int32 FrameNum = Frames.Num();
+			for (int32 FrameIdx = 0; FrameIdx < FrameNum - 1; FrameIdx++)
+			{
+				UE_LEARNING_CHECK(Frames[FrameIdx + 0] >= 0);
+				UE_LEARNING_CHECK(Frames[FrameIdx + 1] >= 0);
+				UE_LEARNING_CHECK(Frames[FrameIdx + 0] < Frames[FrameIdx + 1]);
+			}
+		}
+
+		static inline void RangesCheck(
+			const TLearningArrayView<1, const int32> Starts, 
+			const TLearningArrayView<1, const int32> Lengths)
+		{
+			UE_LEARNING_CHECK(Starts.Num() == Lengths.Num());
+
+			const int32 RangeNum = Starts.Num();
+			for (int32 RangeIdx = 0; RangeIdx < RangeNum - 1; RangeIdx++)
+			{
+				UE_LEARNING_CHECK(Lengths[RangeIdx + 0] > 0);
+				UE_LEARNING_CHECK(Lengths[RangeIdx + 1] > 0);
+				UE_LEARNING_CHECK(Starts[RangeIdx + 0] + Lengths[RangeIdx + 0] <= Starts[RangeIdx + 1]);
+			}
+		}
+
+		static inline void FrameOffsetsCheck(
+			const TLearningArrayView<1, const int32> FrameOffsets,
+			const int32 MinimumOffset,
+			const int32 MaximumOffset)
+		{
+			const int32 FrameNum = FrameOffsets.Num();
+			for (int32 FrameIdx = 0; FrameIdx < FrameNum - 1; FrameIdx++)
+			{
+				UE_LEARNING_CHECK(FrameOffsets[FrameIdx + 0] >= MinimumOffset);
+				UE_LEARNING_CHECK(FrameOffsets[FrameIdx + 1] >= MinimumOffset);
+				UE_LEARNING_CHECK(FrameOffsets[FrameIdx + 0] < MaximumOffset);
+				UE_LEARNING_CHECK(FrameOffsets[FrameIdx + 1] < MaximumOffset);
+				UE_LEARNING_CHECK(FrameOffsets[FrameIdx + 0] < FrameOffsets[FrameIdx + 1]);
+			}
+		}
+
+		static inline bool RangesEqual(
+			const TLearningArrayView<1, const int32> LhsStarts,
+			const TLearningArrayView<1, const int32> LhsLengths,
+			const TLearningArrayView<1, const int32> RhsStarts,
+			const TLearningArrayView<1, const int32> RhsLengths)
+		{
+			UE_LEARNING_CHECK(LhsStarts.Num() == LhsLengths.Num());
+			UE_LEARNING_CHECK(RhsStarts.Num() == RhsLengths.Num());
+
+			if (LhsStarts.Num() != RhsStarts.Num()) { return false; }
+
+			const int32 RangeNum = LhsStarts.Num();
+			for (int32 RangeIdx = 0; RangeIdx < RangeNum; RangeIdx++)
+			{
+				if (LhsStarts[RangeIdx] != RhsStarts[RangeIdx] || LhsLengths[RangeIdx] != RhsLengths[RangeIdx])
+				{
+					return false;
+				}
+			}
+
+			return true;
+		}
+
+		static inline bool RangesContains(
+			const TLearningArrayView<1, const int32> Starts,
+			const TLearningArrayView<1, const int32> Lengths,
+			const int32 Frame)
+		{
+			UE_LEARNING_CHECK(Starts.Num() == Lengths.Num());
+
+			const int32 RangeNum = Starts.Num();
+			for (int32 RangeIdx = 0; RangeIdx < RangeNum; RangeIdx++)
+			{
+				if (Frame >= Starts[RangeIdx] && Frame < Starts[RangeIdx] + Lengths[RangeIdx])
+				{
+					return true;
+				}
+			}
+			return false;
+		}
+
+		static inline bool RangesContainsTime(
+			const TLearningArrayView<1, const int32> Starts,
+			const TLearningArrayView<1, const int32> Lengths,
+			const float Time,
+			const float FrameDeltaTime)
+		{
+			UE_LEARNING_CHECK(Starts.Num() == Lengths.Num());
+
+			const int32 RangeNum = Starts.Num();
+			for (int32 RangeIdx = 0; RangeIdx < RangeNum; RangeIdx++)
+			{
+				const float StartTime = Starts[RangeIdx] * FrameDeltaTime;
+				const float EndTime = (StartTime + Lengths[RangeIdx] - 1) * FrameDeltaTime;
+
+				if (Time >= StartTime && Time < EndTime)
+				{
+					return true;
+				}
+			}
+			return false;
+		}
+
+		static inline bool RangesFind(
+			int32& OutRangeIdx,
+			int32& OutRangeFrame,
+			const TLearningArrayView<1, const int32> Starts,
+			const TLearningArrayView<1, const int32> Lengths,
+			const int32 Frame)
+		{
+			UE_LEARNING_CHECK(Starts.Num() == Lengths.Num());
+
+			const int32 RangeNum = Starts.Num();
+			for (int32 RangeIdx = 0; RangeIdx < RangeNum; RangeIdx++)
+			{
+				if (Frame >= Starts[RangeIdx] && Frame < Starts[RangeIdx] + Lengths[RangeIdx])
+				{
+					OutRangeIdx = RangeIdx;
+					OutRangeFrame = Frame - Starts[RangeIdx];
+					return true;
+				}
+			}
+
+			OutRangeIdx = INDEX_NONE;
+			OutRangeFrame = INDEX_NONE;
+			return false;
+		}
+
+		static inline bool RangesFindTime(
+			int32& OutRangeIdx,
+			float& OutRangeTime,
+			const TLearningArrayView<1, const int32> Starts,
+			const TLearningArrayView<1, const int32> Lengths,
+			const float Time,
+			const float FrameDeltaTime)
+		{
+			UE_LEARNING_CHECK(Starts.Num() == Lengths.Num());
+
+			const int32 RangeNum = Starts.Num();
+			for (int32 RangeIdx = 0; RangeIdx < RangeNum; RangeIdx++)
+			{
+				const float StartTime = Starts[RangeIdx] * FrameDeltaTime;
+				const float EndTime = (StartTime + Lengths[RangeIdx] - 1) * FrameDeltaTime;
+
+				if (Time >= StartTime && Time < EndTime)
+				{
+					OutRangeIdx = RangeIdx;
+					OutRangeTime = Time - StartTime;
+					return true;
+				}
+			}
+
+			OutRangeIdx = INDEX_NONE;
+			OutRangeTime = -1.0f;
+			return false;
+		}
+
+		static inline int32 FramesRangesIntersection(
+			TLearningArrayView<1, int32> OutFrames,
+			const TLearningArrayView<1, const int32> LhsFrames,
+			const TLearningArrayView<1, const int32> RhsStarts,
+			const TLearningArrayView<1, const int32> RhsLengths)
+		{
+			TRACE_CPUPROFILER_EVENT_SCOPE(Learning::FrameRangeSet::Private::FramesRangesIntersection);
+
+			FramesCheck(LhsFrames);
+			RangesCheck(RhsStarts, RhsLengths);
+
+			// Number of frames/ranges in lhs and rhs
+			const int32 LhsNum = LhsFrames.Num();
+			const int32 RhsNum = RhsStarts.Num();
+
+			// Activation state of ranges
+			bool bRhsActive = false;
+
+			// Event index for each list
+			int32 LhsIndex = 0;
+			int32 RhsIndex = 0;
+			int32 OutIndex = 0;
+
+			// While both ranges have events to process
+			while (LhsIndex < LhsNum && RhsIndex < RhsNum * 2)
+			{
+				// Are the rhs events active or inactive
+				const bool bRhsActiveNext = RhsIndex % 2 == 0;
+
+				// Time of the next lhs, and rhs events
+				const int32 LhsT = LhsFrames[LhsIndex];
+				const int32 RhsT = bRhsActiveNext ? RhsStarts[RhsIndex / 2] : RhsStarts[RhsIndex / 2] + RhsLengths[RhsIndex / 2];
+
+				// Event from lhs coming first
+				if (LhsT < RhsT)
+				{
+					// Add to output if range is active
+					if (bRhsActive)
+					{
+						OutFrames[OutIndex] = LhsT;
+						OutIndex++;
+					}
+
+					LhsIndex++;
+				}
+				// Event from rhs coming first
+				else if (RhsT < LhsT)
+				{
+					bRhsActive = bRhsActiveNext;
+					RhsIndex++;
+				}
+				// Event from lhs and rhs coming at same time
+				else
+				{
+					UE_LEARNING_CHECK(LhsT == RhsT);
+
+					bRhsActive = bRhsActiveNext;
+					RhsIndex++;
+
+					// Add to output if range is active
+					if (bRhsActive)
+					{
+						OutFrames[OutIndex] = LhsT;
+						OutIndex++;
+					}
+
+					LhsIndex++;
+				}
+			}
+
+			FramesCheck(OutFrames.Slice(0, OutIndex));
+
+			// Return number of frames added
+			return OutIndex;
+		}
+
+		static inline int32 FramesRangesDifference(
+			TLearningArrayView<1, int32> OutFrames,
+			const TLearningArrayView<1, const int32> LhsFrames,
+			const TLearningArrayView<1, const int32> RhsStarts,
+			const TLearningArrayView<1, const int32> RhsLengths)
+		{
+			TRACE_CPUPROFILER_EVENT_SCOPE(Learning::FrameRangeSet::Private::FramesRangesDifference);
+
+			FramesCheck(LhsFrames);
+			RangesCheck(RhsStarts, RhsLengths);
+
+			// Number of frames/ranges in lhs and rhs
+			const int32 LhsNum = LhsFrames.Num();
+			const int32 RhsNum = RhsStarts.Num();
+
+			// Activation state of ranges
+			bool bRhsActive = false;
+
+			// Event index for each list
+			int32 LhsIndex = 0;
+			int32 RhsIndex = 0;
+			int32 OutIndex = 0;
+
+			// While both ranges have events to process
+			while (LhsIndex < LhsNum && RhsIndex < RhsNum * 2)
+			{
+				// Are the rhs events active or inactive
+				const bool bRhsActiveNext = RhsIndex % 2 == 0;
+
+				// Time of the next lhs, and rhs events
+				const int32 LhsT = LhsFrames[LhsIndex];
+				const int32 RhsT = bRhsActiveNext ? RhsStarts[RhsIndex / 2] : RhsStarts[RhsIndex / 2] + RhsLengths[RhsIndex / 2];
+
+				// Event from lhs coming first
+				if (LhsT < RhsT)
+				{
+					// Add to output if range is active
+					if (!bRhsActive)
+					{
+						OutFrames[OutIndex] = LhsT;
+						OutIndex++;
+					}
+
+					LhsIndex++;
+				}
+				// Event from rhs coming first
+				else if (RhsT < LhsT)
+				{
+					bRhsActive = bRhsActiveNext;
+					RhsIndex++;
+				}
+				// Event from lhs and rhs coming at same time
+				else
+				{
+					UE_LEARNING_CHECK(LhsT == RhsT);
+
+					bRhsActive = bRhsActiveNext;
+					RhsIndex++;
+
+					// Add to output if range is active
+					if (!bRhsActive)
+					{
+						OutFrames[OutIndex] = LhsT;
+						OutIndex++;
+					}
+
+					LhsIndex++;
+				}
+			}
+
+			FramesCheck(OutFrames.Slice(0, OutIndex));
+
+			// Return number of events added
+			return OutIndex;
+		}
+
+
+		static inline int32 RangesUnion(
+			TLearningArrayView<1, int32> OutStarts,
+			TLearningArrayView<1, int32> OutLengths,
+			const TLearningArrayView<1, const int32> LhsStarts,
+			const TLearningArrayView<1, const int32> LhsLengths,
+			const TLearningArrayView<1, const int32> RhsStarts,
+			const TLearningArrayView<1, const int32> RhsLengths)
+		{
+			TRACE_CPUPROFILER_EVENT_SCOPE(Learning::FrameRangeSet::Private::RangesUnion);
+
+			RangesCheck(LhsStarts, LhsLengths);
+			RangesCheck(RhsStarts, RhsLengths);
+
+			if (RangesEqual(LhsStarts, LhsLengths, RhsStarts, RhsLengths))
+			{
+				Array::Copy(OutStarts.Slice(0, LhsStarts.Num()), LhsStarts);
+				Array::Copy(OutLengths.Slice(0, LhsLengths.Num()), LhsLengths);
+				return LhsLengths.Num();
+			}
+
+			// Number of ranges in lhs and rhs
+			const int32 LhsNum = LhsStarts.Num();
+			const int32 RhsNum = RhsStarts.Num();
+
+			// Activation state of each list of ranges
+			bool bOutActive = false;
+			bool bLhsActive = false;
+			bool bRhsActive = false;
+
+			// Event index for each list of ranges
+			int32 LhsIndex = 0;
+			int32 RhsIndex = 0;
+			int32 OutIndex = 0;
+
+			// While both ranges have events to process
+			while (LhsIndex < LhsNum * 2 && RhsIndex < RhsNum * 2)
+			{
+				// Are the next lhs, and rhs events active or inactive
+				const bool bLhsActiveNext = LhsIndex % 2 == 0;
+				const bool bRhsActiveNext = RhsIndex % 2 == 0;
+
+				// Time of the next lhs, and rhs events
+				const int32 LhsT = bLhsActiveNext ? LhsStarts[LhsIndex / 2] : LhsStarts[LhsIndex / 2] + LhsLengths[LhsIndex / 2];
+				const int32 RhsT = bRhsActiveNext ? RhsStarts[RhsIndex / 2] : RhsStarts[RhsIndex / 2] + RhsLengths[RhsIndex / 2];
+
+				// Event from lhs is coming first
+				if (LhsT < RhsT)
+				{
+					// Activate output
+					if (!bOutActive && bLhsActiveNext)
+					{
+						bOutActive = true;
+						OutStarts[OutIndex] = LhsT;
+					}
+					// Deactivate output
+					else if (bOutActive && !bLhsActiveNext && !bRhsActive)
+					{
+
+						bOutActive = false;
+						OutLengths[OutIndex] = LhsT - OutStarts[OutIndex];
+						OutIndex++;
+					}
+
+					bLhsActive = bLhsActiveNext;
+					LhsIndex++;
+				}
+				// Event from rhs is coming first
+				else if (RhsT < LhsT)
+				{
+					// Activate output
+					if (!bOutActive && bRhsActiveNext)
+					{
+						bOutActive = true;
+						OutStarts[OutIndex] = RhsT;
+					}
+					// Deactivate output
+					else if (bOutActive && !bLhsActive && !bRhsActiveNext)
+					{
+						bOutActive = false;
+						OutLengths[OutIndex] = RhsT - OutStarts[OutIndex];
+						OutIndex++;
+					}
+
+					bRhsActive = bRhsActiveNext;
+					RhsIndex++;
+				}
+				// Event from lhs and rhs coming at same time
+				else
+				{
+					UE_LEARNING_CHECK(LhsT == RhsT);
+
+					// Activate output
+					if (!bOutActive && (bLhsActiveNext || bRhsActiveNext))
+					{
+						bOutActive = true;
+						OutStarts[OutIndex] = LhsT;
+					}
+					// Deactivate output
+					else if (bOutActive && !(bLhsActiveNext || bRhsActiveNext))
+					{
+						bOutActive = false;
+						OutLengths[OutIndex] = LhsT - OutStarts[OutIndex];
+						OutIndex++;
+					}
+
+					bLhsActive = bLhsActiveNext;
+					bRhsActive = bRhsActiveNext;
+					LhsIndex++; RhsIndex++;
+				}
+			}
+
+			// Process any remaining lhs events
+			while (LhsIndex < LhsNum * 2)
+			{
+				UE_LEARNING_CHECK(RhsIndex == RhsNum * 2);
+
+				const bool bLhsActiveNext = LhsIndex % 2 == 0;
+				const int32 LhsT = bLhsActiveNext ? LhsStarts[LhsIndex / 2] : LhsStarts[LhsIndex / 2] + LhsLengths[LhsIndex / 2];
+
+				// Activate output
+				if (!bOutActive && bLhsActiveNext)
+				{
+					bOutActive = true;
+					OutStarts[OutIndex] = LhsT;
+				}
+				// Deactivate output
+				else if (bOutActive && !bLhsActiveNext)
+				{
+					bOutActive = false;
+					OutLengths[OutIndex] = LhsT - OutStarts[OutIndex];
+					OutIndex++;
+				}
+
+				bLhsActive = bLhsActiveNext;
+				LhsIndex++;
+			}
+
+			// Process any remaining rhs events
+			while (RhsIndex < RhsNum * 2)
+			{
+				UE_LEARNING_CHECK(LhsIndex == LhsNum * 2);
+
+				const bool bRhsActiveNext = RhsIndex % 2 == 0;
+				const int32 RhsT = bRhsActiveNext ? RhsStarts[RhsIndex / 2] : RhsStarts[RhsIndex / 2] + RhsLengths[RhsIndex / 2];
+
+				// Activate output
+				if (!bOutActive && bRhsActiveNext)
+				{
+					bOutActive = true;
+					OutStarts[OutIndex] = RhsT;
+				}
+				// Deactivate output
+				else if (bOutActive && !bRhsActiveNext)
+				{
+					bOutActive = false;
+					OutLengths[OutIndex] = RhsT - OutStarts[OutIndex];
+					OutIndex++;
+				}
+
+				bRhsActive = bRhsActiveNext;
+				RhsIndex++;
+			}
+
+			RangesCheck(OutStarts.Slice(0, OutIndex), OutLengths.Slice(0, OutIndex));
+
+			// Return number of ranges added
+			return OutIndex;
+		}
+
+		static inline int32 RangesIntersection(
+			TLearningArrayView<1, int32> OutStarts,
+			TLearningArrayView<1, int32> OutLengths,
+			const TLearningArrayView<1, const int32> LhsStarts,
+			const TLearningArrayView<1, const int32> LhsLengths,
+			const TLearningArrayView<1, const int32> RhsStarts,
+			const TLearningArrayView<1, const int32> RhsLengths)
+		{
+			TRACE_CPUPROFILER_EVENT_SCOPE(Learning::FrameRangeSet::Private::RangesIntersection);
+
+			RangesCheck(LhsStarts, LhsLengths);
+			RangesCheck(RhsStarts, RhsLengths);
+
+			if (RangesEqual(LhsStarts, LhsLengths, RhsStarts, RhsLengths))
+			{
+				Array::Copy(OutStarts.Slice(0, LhsStarts.Num()), LhsStarts);
+				Array::Copy(OutLengths.Slice(0, LhsLengths.Num()), LhsLengths);
+				return LhsLengths.Num();
+			}
+
+			// Number of ranges in lhs and rhs
+			const int32 LhsNum = LhsStarts.Num();
+			const int32 RhsNum = RhsStarts.Num();
+
+			// Activation state of each list of ranges
+			bool bOutActive = false;
+			bool bLhsActive = false;
+			bool bRhsActive = false;
+
+			// Event index for each list of ranges
+			int32 LhsIndex = 0;
+			int32 RhsIndex = 0;
+			int32 OutIndex = 0;
+
+			// While both ranges have events to process
+			while (LhsIndex < LhsNum * 2 && RhsIndex < RhsNum * 2)
+			{
+				// Are the next lhs, and rhs events active or inactive
+				const bool bLhsActiveNext = LhsIndex % 2 == 0;
+				const bool bRhsActiveNext = RhsIndex % 2 == 0;
+
+				// Time of the next lhs, and rhs events
+				const int32 LhsT = bLhsActiveNext ? LhsStarts[LhsIndex / 2] : LhsStarts[LhsIndex / 2] + LhsLengths[LhsIndex / 2];
+				const int32 RhsT = bRhsActiveNext ? RhsStarts[RhsIndex / 2] : RhsStarts[RhsIndex / 2] + RhsLengths[RhsIndex / 2];
+
+				// Event from lhs coming first
+				if (LhsT < RhsT)
+				{
+					// Activate output
+					if (!bOutActive && bRhsActive && bLhsActiveNext)
+					{
+						bOutActive = true;
+						OutStarts[OutIndex] = LhsT;
+					}
+					// Deactivate output
+					else if (bOutActive && !bLhsActiveNext)
+					{
+						bOutActive = false;
+						OutLengths[OutIndex] = LhsT - OutStarts[OutIndex];
+						OutIndex++;
+					}
+
+					bLhsActive = bLhsActiveNext;
+					LhsIndex++;
+				}
+				// Event from rhs coming first
+				else if (RhsT < LhsT)
+				{
+					// Activate output
+					if (!bOutActive && bLhsActive && bRhsActiveNext)
+					{
+						bOutActive = true;
+						OutStarts[OutIndex] = RhsT;
+					}
+					// Deactivate output
+					else if (bOutActive && !bRhsActiveNext)
+					{
+						bOutActive = false;
+						OutLengths[OutIndex] = RhsT - OutStarts[OutIndex];
+						OutIndex++;
+					}
+
+					bRhsActive = bRhsActiveNext;
+					RhsIndex++;
+				}
+				// Event from lhs and rhs coming at same time
+				else
+				{
+					UE_LEARNING_CHECK(LhsT == RhsT);
+
+					// Activate output
+					if (!bOutActive && (bLhsActiveNext && bRhsActiveNext))
+					{
+						bOutActive = true;
+						OutStarts[OutIndex] = LhsT;
+					}
+					// Deactivate output
+					else if (bOutActive && (!bLhsActiveNext || !bRhsActiveNext))
+					{
+						bOutActive = false;
+						OutLengths[OutIndex] = LhsT - OutStarts[OutIndex];
+						OutIndex++;
+					}
+
+					bLhsActive = bLhsActiveNext;
+					bRhsActive = bRhsActiveNext;
+					LhsIndex++; RhsIndex++;
+				}
+			}
+
+			RangesCheck(OutStarts.Slice(0, OutIndex), OutLengths.Slice(0, OutIndex));
+
+			// Return number of ranges added
+			return OutIndex;
+		}
+
+		static inline void RangesFrameOffsets(
+			TLearningArrayView<1, int32> OutFrameOffsets,
+			const TLearningArrayView<1, const int32> RangeLengths,
+			const int32 EntryFrameOffset)
+		{
+			UE_LEARNING_CHECK(OutFrameOffsets.Num() == RangeLengths.Num());
+			
+			int32 Offset = EntryFrameOffset;
+			const int32 RangeNum = RangeLengths.Num();
+			for (int32 RangeIdx = 0; RangeIdx < RangeNum; RangeIdx++)
+			{
+				OutFrameOffsets[RangeIdx] = Offset;
+				Offset += RangeLengths[RangeIdx];
+			}
+		}
+
+		static inline int32 RangesTotalFrameNum(const TLearningArrayView<1, const int32> RangeLengths)
+		{
+			int32 Total = 0;
+
+			const int32 RangeNum = RangeLengths.Num();
+			for (int32 RangeIdx = 0; RangeIdx < RangeNum; RangeIdx++)
+			{
+				Total += RangeLengths[RangeIdx];
+			}
+
+			return Total;
+		}
+
+		static inline int32 RangesIntersectionWithFrameOffsets(
+			TLearningArrayView<1, int32> OutStarts,
+			TLearningArrayView<1, int32> OutLengths,
+			TLearningArrayView<1, int32> OutLhsFrameOffsets,
+			TLearningArrayView<1, int32> OutRhsFrameOffsets,
+			const TLearningArrayView<1, const int32> LhsStarts,
+			const TLearningArrayView<1, const int32> LhsLengths,
+			const TLearningArrayView<1, const int32> RhsStarts,
+			const TLearningArrayView<1, const int32> RhsLengths,
+			const int32 LhsEntryFrameOffset,
+			const int32 RhsEntryFrameOffset)
+		{
+			TRACE_CPUPROFILER_EVENT_SCOPE(Learning::FrameRangeSet::Private::RangesIntersectionWithFrameOffsets);
+
+			RangesCheck(LhsStarts, LhsLengths);
+			RangesCheck(RhsStarts, RhsLengths);
+
+			if (RangesEqual(LhsStarts, LhsLengths, RhsStarts, RhsLengths))
+			{
+				Array::Copy(OutStarts.Slice(0, LhsStarts.Num()), LhsStarts);
+				Array::Copy(OutLengths.Slice(0, LhsStarts.Num()), LhsLengths);
+				RangesFrameOffsets(OutLhsFrameOffsets.Slice(0, LhsStarts.Num()), LhsLengths, LhsEntryFrameOffset);
+				RangesFrameOffsets(OutRhsFrameOffsets.Slice(0, RhsStarts.Num()), RhsLengths, RhsEntryFrameOffset);
+				FrameOffsetsCheck(OutLhsFrameOffsets.Slice(0, LhsStarts.Num()), LhsEntryFrameOffset, LhsEntryFrameOffset + RangesTotalFrameNum(LhsLengths));
+				FrameOffsetsCheck(OutRhsFrameOffsets.Slice(0, RhsStarts.Num()), RhsEntryFrameOffset, RhsEntryFrameOffset + RangesTotalFrameNum(RhsLengths));
+				return LhsLengths.Num();
+			}
+
+			// Number of ranges in lhs and rhs
+			const int32 LhsNum = LhsStarts.Num();
+			const int32 RhsNum = RhsStarts.Num();
+
+			// Activation state of each list of ranges
+			bool bOutActive = false;
+			bool bLhsActive = false;
+			bool bRhsActive = false;
+
+			// Event index for each list of ranges
+			int32 LhsIndex = 0;
+			int32 RhsIndex = 0;
+			int32 OutIndex = 0;
+
+			// Frame Offsets
+			int32 LhsFrameOffset = LhsEntryFrameOffset;
+			int32 RhsFrameOffset = RhsEntryFrameOffset;
+			int32 LhsActiveStart = INDEX_NONE;
+			int32 RhsActiveStart = INDEX_NONE;
+			int32 LhsActiveOffset = INDEX_NONE;
+			int32 RhsActiveOffset = INDEX_NONE;
+
+			// While both ranges have events to process
+			while (LhsIndex < LhsNum * 2 && RhsIndex < RhsNum * 2)
+			{
+				// Are the next lhs, and rhs events active or inactive
+				const bool bLhsActiveNext = LhsIndex % 2 == 0;
+				const bool bRhsActiveNext = RhsIndex % 2 == 0;
+
+				// Time of the next lhs, and rhs events
+				const int32 LhsT = bLhsActiveNext ? LhsStarts[LhsIndex / 2] : LhsStarts[LhsIndex / 2] + LhsLengths[LhsIndex / 2];
+				const int32 RhsT = bRhsActiveNext ? RhsStarts[RhsIndex / 2] : RhsStarts[RhsIndex / 2] + RhsLengths[RhsIndex / 2];
+
+				// Event from lhs coming first
+				if (LhsT < RhsT)
+				{
+					// Activate output
+					if (!bOutActive && bRhsActive && bLhsActiveNext)
+					{
+						bOutActive = true;
+						OutStarts[OutIndex] = LhsT;
+					}
+					// Deactivate output
+					else if (bOutActive && !bLhsActiveNext)
+					{
+						UE_LEARNING_CHECK(bLhsActive);
+						UE_LEARNING_CHECK(bRhsActive);
+						UE_LEARNING_CHECK(LhsActiveStart != INDEX_NONE);
+						UE_LEARNING_CHECK(LhsActiveOffset != INDEX_NONE);
+						UE_LEARNING_CHECK(RhsActiveStart != INDEX_NONE);
+						UE_LEARNING_CHECK(RhsActiveOffset != INDEX_NONE);
+
+						bOutActive = false;
+						OutLengths[OutIndex] = LhsT - OutStarts[OutIndex];
+						OutLhsFrameOffsets[OutIndex] = LhsActiveOffset + (OutStarts[OutIndex] - LhsActiveStart);
+						OutRhsFrameOffsets[OutIndex] = RhsActiveOffset + (OutStarts[OutIndex] - RhsActiveStart);
+						OutIndex++;
+					}
+
+					bLhsActive = bLhsActiveNext;
+					LhsActiveStart = bLhsActive ? LhsStarts[LhsIndex / 2] : INDEX_NONE;
+					LhsActiveOffset = bLhsActive ? LhsFrameOffset : INDEX_NONE;
+					LhsFrameOffset += bLhsActive ? 0 : LhsLengths[LhsIndex / 2];
+					LhsIndex++;
+				}
+				// Event from rhs coming first
+				else if (RhsT < LhsT)
+				{
+					// Activate output
+					if (!bOutActive && bLhsActive && bRhsActiveNext)
+					{
+						bOutActive = true;
+						OutStarts[OutIndex] = RhsT;
+					}
+					// Deactivate output
+					else if (bOutActive && !bRhsActiveNext)
+					{
+						UE_LEARNING_CHECK(bLhsActive);
+						UE_LEARNING_CHECK(bRhsActive);
+						UE_LEARNING_CHECK(LhsActiveStart != INDEX_NONE);
+						UE_LEARNING_CHECK(LhsActiveOffset != INDEX_NONE);
+						UE_LEARNING_CHECK(RhsActiveStart != INDEX_NONE);
+						UE_LEARNING_CHECK(RhsActiveOffset != INDEX_NONE);
+
+						bOutActive = false;
+						OutLengths[OutIndex] = RhsT - OutStarts[OutIndex];
+						OutLhsFrameOffsets[OutIndex] = LhsActiveOffset + (OutStarts[OutIndex] - LhsActiveStart);
+						OutRhsFrameOffsets[OutIndex] = RhsActiveOffset + (OutStarts[OutIndex] - RhsActiveStart);
+						OutIndex++;
+					}
+
+					bRhsActive = bRhsActiveNext;
+					RhsActiveStart = bRhsActive ? RhsStarts[RhsIndex / 2] : INDEX_NONE;
+					RhsActiveOffset = bRhsActive ? RhsFrameOffset : INDEX_NONE;
+					RhsFrameOffset += bRhsActive ? 0 : RhsLengths[RhsIndex / 2];
+					RhsIndex++;
+				}
+				// Event from lhs and rhs coming at same time
+				else
+				{
+					UE_LEARNING_CHECK(LhsT == RhsT);
+
+					// Activate output
+					if (!bOutActive && (bLhsActiveNext && bRhsActiveNext))
+					{
+						bOutActive = true;
+						OutStarts[OutIndex] = LhsT;
+					}
+					// Deactivate output
+					else if (bOutActive && (!bLhsActiveNext || !bRhsActiveNext))
+					{
+						UE_LEARNING_CHECK(bLhsActive);
+						UE_LEARNING_CHECK(bRhsActive);
+						UE_LEARNING_CHECK(LhsActiveStart != INDEX_NONE);
+						UE_LEARNING_CHECK(LhsActiveOffset != INDEX_NONE);
+						UE_LEARNING_CHECK(RhsActiveStart != INDEX_NONE);
+						UE_LEARNING_CHECK(RhsActiveOffset != INDEX_NONE);
+
+						bOutActive = false;
+						OutLengths[OutIndex] = LhsT - OutStarts[OutIndex];
+						OutLhsFrameOffsets[OutIndex] = LhsActiveOffset + (OutStarts[OutIndex] - LhsActiveStart);
+						OutRhsFrameOffsets[OutIndex] = RhsActiveOffset + (OutStarts[OutIndex] - RhsActiveStart);
+						OutIndex++;
+					}
+
+					bLhsActive = bLhsActiveNext;
+					LhsActiveStart = bLhsActive ? LhsStarts[LhsIndex / 2] : INDEX_NONE;
+					LhsActiveOffset = bLhsActive ? LhsFrameOffset : INDEX_NONE;
+					bRhsActive = bRhsActiveNext;
+					RhsActiveStart = bRhsActive ? RhsStarts[RhsIndex / 2] : INDEX_NONE;
+					RhsActiveOffset = bRhsActive ? RhsFrameOffset : INDEX_NONE;
+					LhsFrameOffset += bLhsActive ? 0 : LhsLengths[LhsIndex / 2];
+					RhsFrameOffset += bRhsActive ? 0 : RhsLengths[RhsIndex / 2];
+					LhsIndex++; RhsIndex++;
+				}
+			}
+
+			RangesCheck(OutStarts.Slice(0, OutIndex), OutLengths.Slice(0, OutIndex));
+			FrameOffsetsCheck(OutLhsFrameOffsets.Slice(0, OutIndex), LhsEntryFrameOffset, LhsEntryFrameOffset + RangesTotalFrameNum(LhsLengths));
+			FrameOffsetsCheck(OutRhsFrameOffsets.Slice(0, OutIndex), RhsEntryFrameOffset, RhsEntryFrameOffset + RangesTotalFrameNum(RhsLengths));
+
+			// Return number of ranges added
+			return OutIndex;
+		}
+
+		static inline int32 RangesDifference(
+			TLearningArrayView<1, int32> OutStarts,
+			TLearningArrayView<1, int32> OutLengths,
+			const TLearningArrayView<1, const int32> LhsStarts,
+			const TLearningArrayView<1, const int32> LhsLengths,
+			const TLearningArrayView<1, const int32> RhsStarts,
+			const TLearningArrayView<1, const int32> RhsLengths)
+		{
+			TRACE_CPUPROFILER_EVENT_SCOPE(Learning::FrameRangeSet::Private::RangesDifference);
+
+			RangesCheck(LhsStarts, LhsLengths);
+			RangesCheck(RhsStarts, RhsLengths);
+
+			if (RangesEqual(LhsStarts, LhsLengths, RhsStarts, RhsLengths))
+			{
+				return 0;
+			}
+
+			// Number of ranges in lhs and rhs
+			const int32 LhsNum = LhsStarts.Num();
+			const int32 RhsNum = RhsStarts.Num();
+
+			// Activation state of each list of ranges
+			bool bOutActive = false;
+			bool bLhsActive = false;
+			bool bRhsActive = false;
+
+			// Event index for each list of ranges
+			int32 LhsIndex = 0;
+			int32 RhsIndex = 0;
+			int32 OutIndex = 0;
+
+			// While both ranges have events to process
+			while (LhsIndex < LhsNum * 2 && RhsIndex < RhsNum * 2)
+			{
+				// Are the next lhs, and rhs events active or inactive
+				const bool bLhsActiveNext = LhsIndex % 2 == 0;
+				const bool bRhsActiveNext = RhsIndex % 2 == 0;
+
+				// Time of the next lhs, and rhs events
+				const int32 LhsT = bLhsActiveNext ? LhsStarts[LhsIndex / 2] : LhsStarts[LhsIndex / 2] + LhsLengths[LhsIndex / 2];
+				const int32 RhsT = bRhsActiveNext ? RhsStarts[RhsIndex / 2] : RhsStarts[RhsIndex / 2] + RhsLengths[RhsIndex / 2];
+
+				// Event coming from lhs first
+				if (LhsT < RhsT)
+				{
+					// Activate output
+					if (!bOutActive && !bRhsActive && bLhsActiveNext)
+					{
+						bOutActive = true;
+						OutStarts[OutIndex] = LhsT;
+					}
+					// Deactivate output
+					else if (bOutActive && !bLhsActiveNext)
+					{
+						bOutActive = false;
+						OutLengths[OutIndex] = LhsT - OutStarts[OutIndex];
+						OutIndex++;
+					}
+
+					bLhsActive = bLhsActiveNext;
+					LhsIndex++;
+				}
+				// Event coming from rhs first
+				else if (RhsT < LhsT)
+				{
+					// Activate output
+					if (!bOutActive && bLhsActive && !bRhsActiveNext)
+					{
+						bOutActive = true;
+						OutStarts[OutIndex] = RhsT;
+					}
+					// Deactivate output
+					else if (bOutActive && bRhsActiveNext)
+					{
+						bOutActive = false;
+						OutLengths[OutIndex] = RhsT - OutStarts[OutIndex];
+						OutIndex++;
+					}
+
+					bRhsActive = bRhsActiveNext;
+					RhsIndex++;
+				}
+				// Event from lhs and rhs coming at same time
+				else
+				{
+					UE_LEARNING_CHECK(LhsT == RhsT);
+
+					// Activate output
+					if (!bOutActive && bLhsActiveNext && !bRhsActiveNext)
+					{
+						bOutActive = true;
+						OutStarts[OutIndex] = LhsT;
+					}
+					// Deactivate output
+					else if (bOutActive && bRhsActiveNext)
+					{
+						bOutActive = false;
+						OutLengths[OutIndex] = LhsT - OutStarts[OutIndex];
+						OutIndex++;
+					}
+
+					bLhsActive = bLhsActiveNext;
+					bRhsActive = bRhsActiveNext;
+					LhsIndex++; RhsIndex++;
+				}
+			}
+
+			// Process any remaining lhs events
+			while (LhsIndex < LhsNum * 2)
+			{
+				const bool bLhsActiveNext = LhsIndex % 2 == 0;
+				const int32 LhsT = bLhsActiveNext ? LhsStarts[LhsIndex / 2] : LhsStarts[LhsIndex / 2] + LhsLengths[LhsIndex / 2];
+
+				// Activate output
+				if (!bOutActive && bLhsActiveNext)
+				{
+					bOutActive = true;
+					OutStarts[OutIndex] = LhsT;
+				}
+				// Deactivate output
+				else if (bOutActive && !bLhsActiveNext)
+				{
+					bOutActive = false;
+					OutLengths[OutIndex] = LhsT - OutStarts[OutIndex];
+					OutIndex++;
+				}
+
+				bLhsActive = bLhsActiveNext;
+				LhsIndex++;
+			}
+
+			RangesCheck(OutStarts.Slice(0, OutIndex), OutLengths.Slice(0, OutIndex));
+
+			// Return number of ranges added
+			return OutIndex;
+		}
+	}
+
+	void FFrameRangeSet::Check() const
+	{
+		const int32 EntryNum = GetEntryNum();
+		for (int32 EntryIdx = 0; EntryIdx < EntryNum - 1; EntryIdx++)
+		{
+			UE_LEARNING_CHECK(EntrySequences[EntryIdx + 0] < EntrySequences[EntryIdx + 1])
+		}
+
+		for (int32 EntryIdx = 0; EntryIdx < EntryNum; EntryIdx++)
+		{
+			UE_LEARNING_CHECK(EntryRangeNums[EntryIdx] > 0);
+			FrameRangeSet::Private::RangesCheck(GetEntryRangeStarts(EntryIdx), GetEntryRangeLengths(EntryIdx));
+		}
+	}
+
+	bool FFrameRangeSet::IsEmpty() const { return EntrySequences.IsEmpty(); }
+
+	void FFrameRangeSet::Empty()
+	{
+		EntrySequences.Empty();
+		EntryRangeOffsets.Empty();
+		EntryRangeNums.Empty();
+		RangeStarts.Empty();
+		RangeLengths.Empty();
+	}
+
+	int32 FFrameRangeSet::GetEntryNum() const { return EntrySequences.Num(); }
+	TLearningArrayView<1, const int32> FFrameRangeSet::GetEntrySequences() const { return EntrySequences; }
+	TLearningArrayView<1, const int32> FFrameRangeSet::GetEntryRangeNums() const { return EntryRangeNums; }
+	int32 FFrameRangeSet::GetEntrySequence(const int32 EntryIdx) const { return EntrySequences[EntryIdx]; }
+	int32 FFrameRangeSet::GetEntryRangeNum(const int32 EntryIdx) const { return EntryRangeNums[EntryIdx]; }
+
+	int32 FFrameRangeSet::GetEntryTotalFrameNum(const int32 EntryIdx) const
+	{
+		int32 Total = 0;
+		const int32 RangeNum = GetEntryRangeNum(EntryIdx);
+		for (int32 RangeIdx = 0; RangeIdx < RangeNum; RangeIdx++)
+		{
+			Total += GetEntryRangeLength(EntryIdx, RangeIdx);
+		}
+		return Total;
+	}
+
+	int32 FFrameRangeSet::GetEntryFrameOffset(const int32 EntryIdx) const
+	{
+		int32 Offset = 0;
+		for (int32 PrevEntryIdx = 0; PrevEntryIdx < EntryIdx; PrevEntryIdx++)
+		{
+			Offset += GetEntryTotalFrameNum(PrevEntryIdx);
+		}
+
+		return Offset;
+	}
+
+	TLearningArrayView<1, const int32> FFrameRangeSet::GetEntryRangeStarts(const int32 EntryIdx) const { return RangeStarts.Slice(EntryRangeOffsets[EntryIdx], EntryRangeNums[EntryIdx]); }
+	TLearningArrayView<1, const int32> FFrameRangeSet::GetEntryRangeLengths(const int32 EntryIdx) const { return RangeLengths.Slice(EntryRangeOffsets[EntryIdx], EntryRangeNums[EntryIdx]); }
+	int32 FFrameRangeSet::GetEntryRangeStart(const int32 EntryIdx, const int32 RangeIdx) const { return RangeStarts[EntryRangeOffsets[EntryIdx] + RangeIdx]; }
+	int32 FFrameRangeSet::GetEntryRangeLength(const int32 EntryIdx, const int32 RangeIdx) const { return RangeLengths[EntryRangeOffsets[EntryIdx] + RangeIdx]; }
+	
+	int32 FFrameRangeSet::GetEntryRangeFrameOffset(const int32 EntryIdx, const int32 RangeIdx) const
+	{
+		int32 Offset = GetEntryFrameOffset(EntryIdx);
+
+		for (int32 PrevRangeIdx = 0; PrevRangeIdx < RangeIdx; PrevRangeIdx++)
+		{
+			Offset += GetEntryRangeLength(EntryIdx, PrevRangeIdx);
+		}
+
+		return Offset;
+	}
+
+	float FFrameRangeSet::GetEntryRangeStartTime(const int32 EntryIdx, const int32 RangeIdx, const float FrameDeltaTime) const
+	{
+		return GetEntryRangeStart(EntryIdx, RangeIdx) * FrameDeltaTime;
+	}
+	
+	float FFrameRangeSet::GetEntryRangeEndTime(const int32 EntryIdx, const int32 RangeIdx, const float FrameDeltaTime) const
+	{
+		return (GetEntryRangeStart(EntryIdx, RangeIdx) + GetEntryRangeLength(EntryIdx, RangeIdx) - 1) * FrameDeltaTime;
+	}
+
+	float FFrameRangeSet::GetEntryRangeDuration(const int32 EntryIdx, const int32 RangeIdx, const float FrameDeltaTime) const
+	{
+		return (GetEntryRangeLength(EntryIdx, RangeIdx) - 1) * FrameDeltaTime;
+	}
+	
+	int32 FFrameRangeSet::GetTotalRangeNum() const { return RangeStarts.Num(); }
+	
+	TLearningArrayView<1, const int32> FFrameRangeSet::GetAllRangeStarts() const { return RangeStarts; }
+	TLearningArrayView<1, const int32> FFrameRangeSet::GetAllRangeLengths() const { return RangeLengths; }
+
+	int32 FFrameRangeSet::GetTotalFrameNum() const
+	{
+		int32 Total = 0;
+		const int32 RangeNum = RangeLengths.Num();
+		for (int32 RangeIdx = 0; RangeIdx < RangeNum; RangeIdx++)
+		{
+			Total += RangeLengths[RangeIdx];
+		}
+		return Total;
+	}
+
+	bool FFrameRangeSet::ContainsSequence(const int32 Sequence) const
+	{
+		return EntrySequences.ArrayView().Contains(Sequence);
+	}
+
+	bool FFrameRangeSet::Contains(const int32 Sequence, const int32 Frame) const
+	{
+		const int32 EntryIdx = FindSequenceEntry(Sequence);
+
+		return EntryIdx != INDEX_NONE && FrameRangeSet::Private::RangesContains(
+			GetEntryRangeStarts(EntryIdx),
+			GetEntryRangeLengths(EntryIdx),
+			Frame);
+	}
+
+	bool FFrameRangeSet::ContainsTime(const int32 Sequence, const float Time, const float FrameDeltaTime) const
+	{
+		const int32 EntryIdx = FindSequenceEntry(Sequence);
+
+		return EntryIdx != INDEX_NONE && FrameRangeSet::Private::RangesContainsTime(
+			GetEntryRangeStarts(EntryIdx),
+			GetEntryRangeLengths(EntryIdx),
+			Time,
+			FrameDeltaTime);
+	}
+
+	int32 FFrameRangeSet::FindSequenceEntry(const int32 Sequence) const
+	{
+		return EntrySequences.ArrayView().Find(Sequence);
+	}
+
+	bool FFrameRangeSet::Find(int32& OutEntryIdx, int32& OutRangeIdx, int32& OutRangeFrame, const int32 Sequence, const int32 Frame) const
+	{
+		const int32 EntryIdx = FindSequenceEntry(Sequence);
+
+		if (EntryIdx != INDEX_NONE)
+		{
+			int32 FoundRangeIdx = INDEX_NONE;
+			int32 FoundRangeFrame = INDEX_NONE;
+			if (FrameRangeSet::Private::RangesFind(
+				FoundRangeIdx,
+				FoundRangeFrame,
+				GetEntryRangeStarts(EntryIdx),
+				GetEntryRangeLengths(EntryIdx),
+				Frame))
+			{
+				OutEntryIdx = EntryIdx;
+				OutRangeIdx = FoundRangeIdx;
+				OutRangeFrame = FoundRangeFrame;
+				return true;
+			}
+		}
+
+		OutEntryIdx = INDEX_NONE;
+		OutRangeIdx = INDEX_NONE;
+		OutRangeFrame = INDEX_NONE;
+		return false;
+	}
+
+	bool FFrameRangeSet::FindTime(int32& OutEntryIdx, int32& OutRangeIdx, float& OutRangeTime, const int32 Sequence, const float Time, const float FrameDeltaTime) const
+	{
+		const int32 EntryIdx = FindSequenceEntry(Sequence);
+
+		if (EntryIdx != INDEX_NONE)
+		{
+			int32 FoundRangeIdx = INDEX_NONE;
+			float FoundRangeTime = -1.0f;
+			if (FrameRangeSet::Private::RangesFindTime(
+				FoundRangeIdx,
+				FoundRangeTime,
+				GetEntryRangeStarts(EntryIdx),
+				GetEntryRangeLengths(EntryIdx),
+				Time,
+				FrameDeltaTime))
+			{
+				OutEntryIdx = EntryIdx;
+				OutRangeIdx = FoundRangeIdx;
+				OutRangeTime = FoundRangeTime;
+				return true;
+			}
+		}
+
+		OutEntryIdx = INDEX_NONE;
+		OutRangeIdx = INDEX_NONE;
+		OutRangeTime = -1.0f;
+		return false;
+	}
+
+	bool FFrameRangeSet::FindFrameOffset(int32& OutEntryIdx, int32& OutRangeIdx, int32& OutRangeFrame, const int32 FrameOffset) const
+	{
+		const int32 EntryNum = GetEntryNum();
+
+		int32 Offset = 0;
+		for (int32 EntryIdx = 0; EntryIdx < EntryNum; EntryIdx++)
+		{
+			const int32 Sequence = GetEntrySequence(EntryIdx);
+			const int32 RangeNum = GetEntryRangeNum(EntryIdx);
+
+			for (int32 RangeIdx = 0; RangeIdx < RangeNum; RangeIdx++)
+			{
+				const int32 RangeStart = GetEntryRangeStart(EntryIdx, RangeIdx);
+				const int32 RangeLength = GetEntryRangeLength(EntryIdx, RangeIdx);
+
+				if (FrameOffset >= Offset && FrameOffset < Offset + RangeLength)
+				{
+					OutEntryIdx = EntryIdx;
+					OutRangeIdx = RangeIdx;
+					OutRangeFrame = RangeStart + FrameOffset - Offset;
+					return true;
+				}
+
+				Offset += RangeLength;
+			}
+		}
+
+		OutEntryIdx = INDEX_NONE;
+		OutRangeIdx = INDEX_NONE;
+		OutRangeFrame = INDEX_NONE;
+		return false;
+	}
+
+	void FFrameRangeSet::AddEntry(
+		const int32 InSequence,
+		const TLearningArrayView<1, const int32> InStarts,
+		const TLearningArrayView<1, const int32> InLengths)
+	{
+		UE_LEARNING_CHECK(InStarts.Num() == InLengths.Num());
+		UE_LEARNING_CHECK(!ContainsSequence(InSequence));
+		UE::Learning::FrameRangeSet::Private::RangesCheck(InStarts, InLengths);
+
+		if (InStarts.IsEmpty()) { return; }
+
+		const int32 CurrRangeNum = RangeStarts.Num();
+		const int32 AddRangeNum = InStarts.Num();
+		RangeStarts.SetNumUninitialized({ CurrRangeNum + AddRangeNum });
+		RangeLengths.SetNumUninitialized({ CurrRangeNum + AddRangeNum });
+		Array::Copy(RangeStarts.Slice(CurrRangeNum, AddRangeNum), InStarts);
+		Array::Copy(RangeLengths.Slice(CurrRangeNum, AddRangeNum), InLengths);
+
+		const int32 CurrEntryNum = EntrySequences.Num();
+		EntrySequences.SetNumUninitialized({ CurrEntryNum + 1 });
+		EntryRangeOffsets.SetNumUninitialized({ CurrEntryNum + 1 });
+		EntryRangeNums.SetNumUninitialized({ CurrEntryNum + 1 });
+		EntrySequences[CurrEntryNum] = InSequence;
+		EntryRangeOffsets[CurrEntryNum] = CurrRangeNum;
+		EntryRangeNums[CurrEntryNum] = AddRangeNum;
+
+		Check();
+	}
+
+	namespace FrameRangeSet
+	{
+		bool Equal(const FFrameRangeSet& Lhs, const FFrameRangeSet& Rhs)
+		{
+			return
+				Lhs.EntrySequences.Num() == Rhs.EntrySequences.Num() &&
+				Lhs.RangeStarts.Num() == Rhs.RangeStarts.Num() &&
+				Array::Equal(Lhs.EntrySequences, Rhs.EntrySequences) &&
+				Array::Equal(Lhs.EntryRangeOffsets, Rhs.EntryRangeOffsets) &&
+				Array::Equal(Lhs.EntryRangeNums, Rhs.EntryRangeNums) &&
+				Array::Equal(Lhs.RangeStarts, Rhs.RangeStarts) &&
+				Array::Equal(Lhs.RangeLengths, Rhs.RangeLengths);
+		}
+
+		void Intersection(FFrameSet& OutFrameSet, const FFrameSet& FrameSet, const FFrameRangeSet& FrameRangeSet)
+		{
+			TRACE_CPUPROFILER_EVENT_SCOPE(Learning::FrameRangeSet::Intersection);
+
+			FrameSet.Check(); FrameRangeSet.Check();
+
+			const int32 LhsEntryNum = FrameSet.GetEntryNum();
+			const int32 RhsEntryNum = FrameRangeSet.GetEntryNum();
+			const int32 LhsFrameNum = FrameSet.GetTotalFrameNum();
+
+			// Allocate potential maximum number of entries and frames we might to output
+			OutFrameSet.EntrySequences.SetNumUninitialized({ LhsEntryNum });
+			OutFrameSet.EntryFrameOffsets.SetNumUninitialized({ LhsEntryNum });
+			OutFrameSet.EntryFrameNums.SetNumUninitialized({ LhsEntryNum });
+			OutFrameSet.Frames.SetNumUninitialized({ LhsFrameNum });
+
+			// Entry index for each list of entries
+			int32 OutIndex = 0;
+			int32 LhsIndex = 0;
+			int32 RhsIndex = 0;
+
+			// Output entry index
+			int32 FrameIndex = 0;
+
+			// While both sets have entries
+			while (LhsIndex < LhsEntryNum && RhsIndex < RhsEntryNum)
+			{
+				// If entry from lhs is first
+				if (FrameSet.GetEntrySequence(LhsIndex) < FrameRangeSet.GetEntrySequence(RhsIndex))
+				{
+					LhsIndex++;
+				}
+				// If entry from rhs is first
+				else if (FrameRangeSet.GetEntrySequence(RhsIndex) < FrameSet.GetEntrySequence(LhsIndex))
+				{
+					RhsIndex++;
+				}
+				// If both contain the same entry
+				else
+				{
+					UE_LEARNING_CHECK(FrameRangeSet.GetEntrySequence(RhsIndex) == FrameSet.GetEntrySequence(LhsIndex));
+
+					// Append intersection of frames to output
+					const int32 EventNum = Private::FramesRangesIntersection(
+						OutFrameSet.Frames.Slice(FrameIndex, OutFrameSet.Frames.Num() - FrameIndex),
+						FrameSet.GetEntryFrames(LhsIndex),
+						FrameRangeSet.GetEntryRangeStarts(RhsIndex),
+						FrameRangeSet.GetEntryRangeLengths(RhsIndex));
+
+					UE_LEARNING_CHECK(EventNum <= FrameSet.GetEntryFrameNum(LhsIndex));
+
+					if (EventNum > 0)
+					{
+						OutFrameSet.EntrySequences[OutIndex] = FrameSet.GetEntrySequence(LhsIndex);
+						OutFrameSet.EntryFrameOffsets[OutIndex] = FrameIndex;
+						OutFrameSet.EntryFrameNums[OutIndex] = EventNum;
+						FrameIndex += EventNum;
+						OutIndex++;
+					}
+
+					LhsIndex++; RhsIndex++;
+				}
+			}
+
+			// Resize output to match what was added
+			OutFrameSet.EntrySequences.SetNumUninitialized({ OutIndex });
+			OutFrameSet.EntryFrameOffsets.SetNumUninitialized({ OutIndex });
+			OutFrameSet.EntryFrameNums.SetNumUninitialized({ OutIndex });
+			OutFrameSet.Frames.SetNumUninitialized({ FrameIndex });
+			OutFrameSet.Check();
+		}
+
+		void Difference(FFrameSet& OutFrameSet, const FFrameSet& FrameSet, const FFrameRangeSet& FrameRangeSet)
+		{
+			TRACE_CPUPROFILER_EVENT_SCOPE(Learning::FrameRangeSet::Difference);
+
+			FrameSet.Check(); FrameRangeSet.Check();
+
+			const int32 LhsEntryNum = FrameSet.GetEntryNum();
+			const int32 RhsEntryNum = FrameRangeSet.GetEntryNum();
+			const int32 LhsFrameNum = FrameSet.GetTotalFrameNum();
+
+			// Allocate potential maximum number of entries and frames we might to output
+			OutFrameSet.EntrySequences.SetNumUninitialized({ LhsEntryNum });
+			OutFrameSet.EntryFrameOffsets.SetNumUninitialized({ LhsEntryNum });
+			OutFrameSet.EntryFrameNums.SetNumUninitialized({ LhsEntryNum });
+			OutFrameSet.Frames.SetNumUninitialized({ LhsFrameNum });
+
+			// Entry index for each list of entries
+			int32 OutIndex = 0;
+			int32 LhsIndex = 0;
+			int32 RhsIndex = 0;
+
+			// Output event index
+			int32 FrameIndex = 0;
+
+			// While both sets have entries
+			while (LhsIndex < LhsEntryNum && RhsIndex < RhsEntryNum)
+			{
+				// If entry from lhs is first
+				if (FrameSet.GetEntrySequence(LhsIndex) < FrameRangeSet.GetEntrySequence(RhsIndex))
+				{
+					// Append frames to output
+					if (FrameSet.GetEntryFrameNum(LhsIndex) > 0)
+					{
+						OutFrameSet.EntrySequences[OutIndex] = FrameSet.GetEntrySequence(LhsIndex);
+						OutFrameSet.EntryFrameOffsets[OutIndex] = FrameIndex;
+						OutFrameSet.EntryFrameNums[OutIndex] = FrameSet.GetEntryFrameNum(LhsIndex);
+
+						Array::Copy(
+							OutFrameSet.Frames.Slice(FrameIndex, FrameSet.GetEntryFrameNum(LhsIndex)),
+							FrameSet.GetEntryFrames(LhsIndex));
+
+						FrameIndex += FrameSet.GetEntryFrameNum(LhsIndex);
+						OutIndex++;
+					}
+
+					LhsIndex++;
+				}
+				// If entry from rhs is first
+				else if (FrameRangeSet.GetEntrySequence(RhsIndex) < FrameSet.GetEntrySequence(LhsIndex))
+				{
+					RhsIndex++;
+				}
+				// If both contain the same entry
+				else
+				{
+					UE_LEARNING_CHECK(FrameRangeSet.GetEntrySequence(RhsIndex) == FrameSet.GetEntrySequence(LhsIndex));
+
+					// Append difference of frames to output
+					const int32 EventNum = Private::FramesRangesDifference(
+						OutFrameSet.Frames.Slice(FrameIndex, OutFrameSet.Frames.Num() - FrameIndex),
+						FrameSet.GetEntryFrames(LhsIndex),
+						FrameRangeSet.GetEntryRangeStarts(RhsIndex),
+						FrameRangeSet.GetEntryRangeLengths(RhsIndex));
+
+					UE_LEARNING_CHECK(EventNum <= FrameSet.GetEntryFrameNum(LhsIndex));
+
+					if (EventNum > 0)
+					{
+						OutFrameSet.EntrySequences[OutIndex] = FrameSet.GetEntrySequence(LhsIndex);
+						OutFrameSet.EntryFrameOffsets[OutIndex] = FrameIndex;
+						OutFrameSet.EntryFrameNums[OutIndex] = EventNum;
+						FrameIndex += EventNum;
+						OutIndex++;
+					}
+
+					LhsIndex++; RhsIndex++;
+				}
+			}
+
+			// Process any remaining lhs entries
+			while (LhsIndex < LhsEntryNum)
+			{
+				// Append frames to output
+				if (FrameSet.GetEntryFrameNum(LhsIndex) > 0)
+				{
+					OutFrameSet.EntrySequences[OutIndex] = FrameSet.GetEntrySequence(LhsIndex);
+					OutFrameSet.EntryFrameOffsets[OutIndex] = FrameIndex;
+					OutFrameSet.EntryFrameNums[OutIndex] = FrameSet.GetEntryFrameNum(LhsIndex);
+
+					Array::Copy(
+						OutFrameSet.Frames.Slice(FrameIndex, FrameSet.GetEntryFrameNum(LhsIndex)),
+						FrameSet.GetEntryFrames(LhsIndex));
+
+					FrameIndex += FrameSet.GetEntryFrameNum(LhsIndex);
+					OutIndex++;
+				}
+
+				LhsIndex++;
+			}
+
+			// Resize output to match what was added
+			OutFrameSet.EntrySequences.SetNumUninitialized({ OutIndex });
+			OutFrameSet.EntryFrameOffsets.SetNumUninitialized({ OutIndex });
+			OutFrameSet.EntryFrameNums.SetNumUninitialized({ OutIndex });
+			OutFrameSet.Frames.SetNumUninitialized({ FrameIndex });
+			OutFrameSet.Check();
+		}
+
+		void Union(FFrameRangeSet& Out, const FFrameRangeSet& Lhs, const FFrameRangeSet& Rhs)
+		{
+			TRACE_CPUPROFILER_EVENT_SCOPE(Learning::FrameRangeSet::Union);
+
+			Lhs.Check(); Rhs.Check();
+
+			if (Equal(Lhs, Rhs)) { Out = Lhs; return; }
+
+			// Allocate potential maximum number of entries and ranges we might to output
+
+			const int32 LhsEntryNum = Lhs.GetEntryNum();
+			const int32 RhsEntryNum = Rhs.GetEntryNum();
+			const int32 LhsRangeNum = Lhs.GetTotalRangeNum();
+			const int32 RhsRangeNum = Rhs.GetTotalRangeNum();
+
+			// Allocate potential maximum number of entries and ranges we might to output
+			Out.EntrySequences.SetNumUninitialized({ LhsEntryNum + RhsEntryNum });
+			Out.EntryRangeOffsets.SetNumUninitialized({ LhsEntryNum + RhsEntryNum });
+			Out.EntryRangeNums.SetNumUninitialized({ LhsEntryNum + RhsEntryNum });
+			Out.RangeStarts.SetNumUninitialized({ LhsRangeNum + RhsRangeNum });
+			Out.RangeLengths.SetNumUninitialized({ LhsRangeNum + RhsRangeNum });
+
+			// Entry index for each list of ranges
+			int32 OutIndex = 0;
+			int32 LhsIndex = 0;
+			int32 RhsIndex = 0;
+
+			// Output ranges index
+			int32 RangeIndex = 0;
+
+			// While both sets have entries
+			while (LhsIndex < LhsEntryNum && RhsIndex < RhsEntryNum)
+			{
+				// If entry from lhs is first
+				if (Lhs.GetEntrySequence(LhsIndex) < Rhs.GetEntrySequence(RhsIndex))
+				{
+					UE_LEARNING_CHECK(Lhs.GetEntryRangeNum(LhsIndex) > 0);
+
+					// Append subranges to output
+					Out.EntrySequences[OutIndex] = Lhs.GetEntrySequence(LhsIndex);
+					Out.EntryRangeOffsets[OutIndex] = RangeIndex;
+					Out.EntryRangeNums[OutIndex] = Lhs.GetEntryRangeNum(LhsIndex);
+
+					Array::Copy(
+						Out.RangeStarts.Slice(RangeIndex, Lhs.GetEntryRangeNum(LhsIndex)),
+						Lhs.GetEntryRangeStarts(LhsIndex));
+
+					Array::Copy(
+						Out.RangeLengths.Slice(RangeIndex, Lhs.GetEntryRangeNum(LhsIndex)),
+						Lhs.GetEntryRangeLengths(LhsIndex));
+
+					RangeIndex += Lhs.GetEntryRangeNum(LhsIndex);
+					OutIndex++;
+					LhsIndex++;
+				}
+				// If entry from rhs is first
+				else if (Rhs.GetEntrySequence(RhsIndex) < Lhs.GetEntrySequence(LhsIndex))
+				{
+					UE_LEARNING_CHECK(Rhs.GetEntryRangeNum(RhsIndex) > 0);
+
+					// Append subranges to output
+					Out.EntrySequences[OutIndex] = Rhs.GetEntrySequence(RhsIndex);
+					Out.EntryRangeOffsets[OutIndex] = RangeIndex;
+					Out.EntryRangeNums[OutIndex] = Rhs.GetEntryRangeNum(RhsIndex);
+
+					Array::Copy(
+						Out.RangeStarts.Slice(RangeIndex, Rhs.GetEntryRangeNum(RhsIndex)),
+						Rhs.GetEntryRangeStarts(RhsIndex));
+
+					Array::Copy(
+						Out.RangeLengths.Slice(RangeIndex, Rhs.GetEntryRangeNum(RhsIndex)),
+						Rhs.GetEntryRangeLengths(RhsIndex));
+
+					RangeIndex += Rhs.GetEntryRangeNum(RhsIndex);
+					OutIndex++;
+					RhsIndex++;
+				}
+				// If both contain the same entry
+				else
+				{
+					UE_LEARNING_CHECK(Lhs.GetEntryRangeNum(LhsIndex) > 0);
+					UE_LEARNING_CHECK(Rhs.GetEntryRangeNum(RhsIndex) > 0);
+
+					// Append union of subranges to output
+					const int32 RangeNum = FrameRangeSet::Private::RangesUnion(
+						Out.RangeStarts.Slice(RangeIndex, Out.RangeStarts.Num() - RangeIndex),
+						Out.RangeLengths.Slice(RangeIndex, Out.RangeLengths.Num() - RangeIndex),
+						Lhs.GetEntryRangeStarts(LhsIndex),
+						Lhs.GetEntryRangeLengths(LhsIndex),
+						Rhs.GetEntryRangeStarts(RhsIndex),
+						Rhs.GetEntryRangeLengths(RhsIndex));
+
+					UE_LEARNING_CHECK(RangeNum > 0);
+					UE_LEARNING_CHECK(RangeNum <= Lhs.GetEntryRangeNum(LhsIndex) + Rhs.GetEntryRangeNum(RhsIndex));
+
+					Out.EntrySequences[OutIndex] = Lhs.GetEntrySequence(LhsIndex);
+					Out.EntryRangeOffsets[OutIndex] = RangeIndex;
+					Out.EntryRangeNums[OutIndex] = RangeNum;
+					RangeIndex += RangeNum;
+					OutIndex++;
+					LhsIndex++;
+					RhsIndex++;
+				}
+			}
+
+			// Process any remaining lhs entries
+			while (LhsIndex < LhsEntryNum)
+			{
+				UE_LEARNING_CHECK(RhsIndex == RhsEntryNum);
+				UE_LEARNING_CHECK(Lhs.GetEntryRangeNum(LhsIndex) > 0);
+
+				// Append subranges to output
+				Out.EntrySequences[OutIndex] = Lhs.GetEntrySequence(LhsIndex);
+				Out.EntryRangeOffsets[OutIndex] = RangeIndex;
+				Out.EntryRangeNums[OutIndex] = Lhs.GetEntryRangeNum(LhsIndex);
+
+				Array::Copy(
+					Out.RangeStarts.Slice(RangeIndex, Lhs.GetEntryRangeNum(LhsIndex)),
+					Lhs.GetEntryRangeStarts(LhsIndex));
+
+				Array::Copy(
+					Out.RangeLengths.Slice(RangeIndex, Lhs.GetEntryRangeNum(LhsIndex)),
+					Lhs.GetEntryRangeLengths(LhsIndex));
+
+				RangeIndex += Lhs.GetEntryRangeNum(LhsIndex);
+				OutIndex++;
+				LhsIndex++;
+			}
+
+			// Process any remaining rhs entries
+			while (RhsIndex < RhsEntryNum)
+			{
+				UE_LEARNING_CHECK(LhsIndex == LhsEntryNum);
+				UE_LEARNING_CHECK(Rhs.GetEntryRangeNum(RhsIndex) > 0);
+
+				// Append subranges to output
+				Out.EntrySequences[OutIndex] = Rhs.GetEntrySequence(RhsIndex);
+				Out.EntryRangeOffsets[OutIndex] = RangeIndex;
+				Out.EntryRangeNums[OutIndex] = Rhs.GetEntryRangeNum(RhsIndex);
+
+				Array::Copy(
+					Out.RangeStarts.Slice(RangeIndex, Rhs.GetEntryRangeNum(RhsIndex)),
+					Rhs.GetEntryRangeStarts(RhsIndex));
+
+				Array::Copy(
+					Out.RangeLengths.Slice(RangeIndex, Rhs.GetEntryRangeNum(RhsIndex)),
+					Rhs.GetEntryRangeLengths(RhsIndex));
+
+				RangeIndex += Rhs.GetEntryRangeNum(RhsIndex);
+				OutIndex++;
+				RhsIndex++;
+			}
+
+			// Resize output to match what was added
+			Out.EntrySequences.SetNumUninitialized({ OutIndex });
+			Out.EntryRangeOffsets.SetNumUninitialized({ OutIndex });
+			Out.EntryRangeNums.SetNumUninitialized({ OutIndex });
+			Out.RangeStarts.SetNumUninitialized({ RangeIndex });
+			Out.RangeLengths.SetNumUninitialized({ RangeIndex });
+			Out.Check();
+		}
+
+		void Intersection(FFrameRangeSet& Out, const FFrameRangeSet& Lhs, const FFrameRangeSet& Rhs)
+		{
+			TRACE_CPUPROFILER_EVENT_SCOPE(Learning::FrameRangeSet::Intersection);
+
+			Lhs.Check(); Rhs.Check();
+
+			if (Equal(Lhs, Rhs)) { Out = Lhs; return; }
+
+			const int32 LhsEntryNum = Lhs.GetEntryNum();
+			const int32 RhsEntryNum = Rhs.GetEntryNum();
+			const int32 LhsRangeNum = Lhs.GetTotalRangeNum();
+			const int32 RhsRangeNum = Rhs.GetTotalRangeNum();
+
+			// Allocate potential maximum number of entries and ranges we might to output
+			Out.EntrySequences.SetNumUninitialized({ FMath::Max(LhsEntryNum, RhsEntryNum) });
+			Out.EntryRangeOffsets.SetNumUninitialized({ FMath::Max(LhsEntryNum, RhsEntryNum) });
+			Out.EntryRangeNums.SetNumUninitialized({ FMath::Max(LhsEntryNum, RhsEntryNum) });
+			Out.RangeStarts.SetNumUninitialized({ LhsRangeNum + RhsRangeNum });
+			Out.RangeLengths.SetNumUninitialized({ LhsRangeNum + RhsRangeNum });
+
+			// Anim index for each list of ranges
+			int32 OutIndex = 0;
+			int32 LhsIndex = 0;
+			int32 RhsIndex = 0;
+
+			// Output ranges index
+			int32 RangeIndex = 0;
+
+			// While both sets have entries
+			while (LhsIndex < LhsEntryNum && RhsIndex < RhsEntryNum)
+			{
+				// If entry is in lhs but not rhs skip
+				if (Lhs.GetEntrySequence(LhsIndex) < Rhs.GetEntrySequence(RhsIndex))
+				{
+					LhsIndex++;
+				}
+				// If entry is in rhs but not lhs skip
+				else if (Rhs.GetEntrySequence(RhsIndex) < Lhs.GetEntrySequence(LhsIndex))
+				{
+					RhsIndex++;
+				}
+				// If entry is in both lhs and rhs
+				else
+				{
+					UE_LEARNING_CHECK(Lhs.GetEntryRangeNum(LhsIndex) > 0);
+					UE_LEARNING_CHECK(Rhs.GetEntryRangeNum(RhsIndex) > 0);
+
+					// Append intersection of subranges to output
+					const int32 RangeNum = FrameRangeSet::Private::RangesIntersection(
+						Out.RangeStarts.Slice(RangeIndex, Out.RangeStarts.Num() - RangeIndex),
+						Out.RangeLengths.Slice(RangeIndex, Out.RangeLengths.Num() - RangeIndex),
+						Lhs.GetEntryRangeStarts(LhsIndex),
+						Lhs.GetEntryRangeLengths(LhsIndex),
+						Rhs.GetEntryRangeStarts(RhsIndex),
+						Rhs.GetEntryRangeLengths(RhsIndex));
+
+					UE_LEARNING_CHECK(RangeNum <= Lhs.GetEntryRangeNum(LhsIndex) + Rhs.GetEntryRangeNum(RhsIndex));
+
+					if (RangeNum > 0)
+					{
+						Out.EntrySequences[OutIndex] = Lhs.GetEntrySequence(LhsIndex);
+						Out.EntryRangeOffsets[OutIndex] = RangeIndex;
+						Out.EntryRangeNums[OutIndex] = RangeNum;
+						RangeIndex += RangeNum;
+						OutIndex++;
+					}
+
+					LhsIndex++; RhsIndex++;
+				}
+			}
+
+			// Resize output to match what was added
+			Out.EntrySequences.SetNumUninitialized({ OutIndex });
+			Out.EntryRangeOffsets.SetNumUninitialized({ OutIndex });
+			Out.EntryRangeNums.SetNumUninitialized({ OutIndex });
+			Out.RangeStarts.SetNumUninitialized({ RangeIndex });
+			Out.RangeLengths.SetNumUninitialized({ RangeIndex });
+			Out.Check();
+		}
+
+		void Difference(FFrameRangeSet& Out, const FFrameRangeSet& Lhs, const FFrameRangeSet& Rhs)
+		{
+			TRACE_CPUPROFILER_EVENT_SCOPE(Learning::FrameRangeSet::Difference);
+
+			Lhs.Check(); Rhs.Check();
+
+			if (Equal(Lhs, Rhs)) { Out.Empty(); return; }
+
+			const int32 LhsEntryNum = Lhs.GetEntryNum();
+			const int32 RhsEntryNum = Rhs.GetEntryNum();
+			const int32 LhsRangeNum = Lhs.GetTotalRangeNum();
+			const int32 RhsRangeNum = Rhs.GetTotalRangeNum();
+
+			// Allocate potential maximum number of entries and ranges we might to output
+			Out.EntrySequences.SetNumUninitialized({ LhsEntryNum });
+			Out.EntryRangeOffsets.SetNumUninitialized({ LhsEntryNum });
+			Out.EntryRangeNums.SetNumUninitialized({ LhsEntryNum });
+			Out.RangeStarts.SetNumUninitialized({ LhsRangeNum + RhsRangeNum });
+			Out.RangeLengths.SetNumUninitialized({ LhsRangeNum + RhsRangeNum });
+
+			// Anim index for each list of ranges
+			int32 OutIndex = 0;
+			int32 LhsIndex = 0;
+			int32 RhsIndex = 0;
+
+			// Output ranges index
+			int32 RangeIndex = 0;
+
+			// While both sets have entries
+			while (LhsIndex < LhsEntryNum && RhsIndex < RhsEntryNum)
+			{
+				// If entry from lhs is first
+				if (Lhs.GetEntrySequence(LhsIndex) < Rhs.GetEntrySequence(RhsIndex))
+				{
+					UE_LEARNING_CHECK(Lhs.GetEntryRangeNum(LhsIndex) > 0);
+
+					// Append subranges to output
+					Out.EntrySequences[OutIndex] = Lhs.GetEntrySequence(LhsIndex);
+					Out.EntryRangeOffsets[OutIndex] = RangeIndex;
+					Out.EntryRangeNums[OutIndex] = Lhs.GetEntryRangeNum(LhsIndex);
+
+					Array::Copy(
+						Out.RangeStarts.Slice(RangeIndex, Lhs.GetEntryRangeNum(LhsIndex)),
+						Lhs.GetEntryRangeStarts(LhsIndex));
+
+					Array::Copy(
+						Out.RangeLengths.Slice(RangeIndex, Lhs.GetEntryRangeNum(LhsIndex)),
+						Lhs.GetEntryRangeLengths(LhsIndex));
+
+					RangeIndex += Lhs.GetEntryRangeNum(LhsIndex);
+					OutIndex++;
+					LhsIndex++;
+				}
+				// If entry is in rhs but not lhs skip
+				else if (Rhs.GetEntrySequence(RhsIndex) < Lhs.GetEntrySequence(LhsIndex))
+				{
+					RhsIndex++;
+				}
+				// If entry is in both lhs and rhs
+				else
+				{
+					UE_LEARNING_CHECK(Lhs.GetEntryRangeNum(LhsIndex) > 0);
+					UE_LEARNING_CHECK(Rhs.GetEntryRangeNum(RhsIndex) > 0);
+
+					// Append difference of subranges to output
+					const int32 RangeNum = FrameRangeSet::Private::RangesDifference(
+						Out.RangeStarts.Slice(RangeIndex, Out.RangeStarts.Num() - RangeIndex),
+						Out.RangeLengths.Slice(RangeIndex, Out.RangeLengths.Num() - RangeIndex),
+						Lhs.GetEntryRangeStarts(LhsIndex),
+						Lhs.GetEntryRangeLengths(LhsIndex),
+						Rhs.GetEntryRangeStarts(RhsIndex),
+						Rhs.GetEntryRangeLengths(RhsIndex));
+
+					UE_LEARNING_CHECK(RangeNum <= Lhs.GetEntryRangeNum(LhsIndex) + Rhs.GetEntryRangeNum(RhsIndex));
+
+					if (RangeNum > 0)
+					{
+						Out.EntrySequences[OutIndex] = Lhs.GetEntrySequence(LhsIndex);
+						Out.EntryRangeOffsets[OutIndex] = RangeIndex;
+						Out.EntryRangeNums[OutIndex] = RangeNum;
+						RangeIndex += RangeNum;
+						OutIndex++;
+					}
+
+					LhsIndex++; RhsIndex++;
+				}
+			}
+
+			// Process any remaining lhs entries
+			while (LhsIndex < LhsEntryNum)
+			{
+				UE_LEARNING_CHECK(Lhs.GetEntryRangeNum(LhsIndex) > 0);
+
+				// Append subranges to output
+				Out.EntrySequences[OutIndex] = Lhs.GetEntrySequence(LhsIndex);
+				Out.EntryRangeOffsets[OutIndex] = RangeIndex;
+				Out.EntryRangeNums[OutIndex] = Lhs.GetEntryRangeNum(LhsIndex);
+
+				Array::Copy(
+					Out.RangeStarts.Slice(RangeIndex, Lhs.GetEntryRangeNum(LhsIndex)),
+					Lhs.GetEntryRangeStarts(LhsIndex));
+
+				Array::Copy(
+					Out.RangeLengths.Slice(RangeIndex, Lhs.GetEntryRangeNum(LhsIndex)),
+					Lhs.GetEntryRangeLengths(LhsIndex));
+
+				RangeIndex += Lhs.GetEntryRangeNum(LhsIndex);
+				OutIndex++;
+				LhsIndex++;
+			}
+
+			// Resize output to match what was added
+			Out.EntrySequences.SetNumUninitialized({ OutIndex });
+			Out.EntryRangeOffsets.SetNumUninitialized({ OutIndex });
+			Out.EntryRangeNums.SetNumUninitialized({ OutIndex });
+			Out.RangeStarts.SetNumUninitialized({ RangeIndex });
+			Out.RangeLengths.SetNumUninitialized({ RangeIndex });
+			Out.Check();
+		}
+
+		int32 IntersectionWithFrameOffsets(
+			FFrameRangeSet& Out,
+			TLearningArrayView<1, int32> OutLhsFrameOffsets,
+			TLearningArrayView<1, int32> OutRhsFrameOffsets,
+			const FFrameRangeSet& Lhs,
+			const FFrameRangeSet& Rhs)
+		{
+			TRACE_CPUPROFILER_EVENT_SCOPE(Learning::FrameRangeSet::IntersectionWithFrameOffsets);
+
+			Lhs.Check(); Rhs.Check();
+
+			if (Equal(Lhs, Rhs))
+			{
+				Out = Lhs;
+				AllRangeFrameOffsets(OutLhsFrameOffsets.Slice(0, Lhs.GetTotalRangeNum()), Lhs);
+				AllRangeFrameOffsets(OutRhsFrameOffsets.Slice(0, Rhs.GetTotalRangeNum()), Rhs);
+				return Out.GetTotalRangeNum();
+			}
+
+			const int32 LhsEntryNum = Lhs.GetEntryNum();
+			const int32 RhsEntryNum = Rhs.GetEntryNum();
+			const int32 LhsRangeNum = Lhs.GetTotalRangeNum();
+			const int32 RhsRangeNum = Rhs.GetTotalRangeNum();
+
+			// Allocate potential maximum number of entries and ranges we might to output
+			Out.EntrySequences.SetNumUninitialized({ FMath::Max(LhsEntryNum, RhsEntryNum) });
+			Out.EntryRangeOffsets.SetNumUninitialized({ FMath::Max(LhsEntryNum, RhsEntryNum) });
+			Out.EntryRangeNums.SetNumUninitialized({ FMath::Max(LhsEntryNum, RhsEntryNum) });
+			Out.RangeStarts.SetNumUninitialized({ LhsRangeNum + RhsRangeNum });
+			Out.RangeLengths.SetNumUninitialized({ LhsRangeNum + RhsRangeNum });
+
+			// Anim index for each list of ranges
+			int32 OutIndex = 0;
+			int32 LhsIndex = 0;
+			int32 RhsIndex = 0;
+
+			// Output ranges index
+			int32 RangeIndex = 0;
+
+			// Frame Offsets
+			int32 OutIntersectionIdx = 0;
+			int32 LhsFrameOffset = 0;
+			int32 RhsFrameOffset = 0;
+
+			// While both sets have entries
+			while (LhsIndex < LhsEntryNum && RhsIndex < RhsEntryNum)
+			{
+				// If entry is in lhs but not rhs skip
+				if (Lhs.GetEntrySequence(LhsIndex) < Rhs.GetEntrySequence(RhsIndex))
+				{
+					LhsFrameOffset += Lhs.GetEntryTotalFrameNum(LhsIndex);
+					LhsIndex++;
+				}
+				// If entry is in rhs but not lhs skip
+				else if (Rhs.GetEntrySequence(RhsIndex) < Lhs.GetEntrySequence(LhsIndex))
+				{
+					RhsFrameOffset += Rhs.GetEntryTotalFrameNum(RhsIndex);
+					RhsIndex++;
+				}
+				// If entry is in both lhs and rhs
+				else
+				{
+					UE_LEARNING_CHECK(Lhs.GetEntryRangeNum(LhsIndex) > 0);
+					UE_LEARNING_CHECK(Rhs.GetEntryRangeNum(RhsIndex) > 0);
+
+					// Append intersection of subranges to output
+					const int32 RangeNum = FrameRangeSet::Private::RangesIntersectionWithFrameOffsets(
+						Out.RangeStarts.Slice(RangeIndex, Out.RangeStarts.Num() - RangeIndex),
+						Out.RangeLengths.Slice(RangeIndex, Out.RangeLengths.Num() - RangeIndex),
+						OutLhsFrameOffsets.Slice(RangeIndex, OutLhsFrameOffsets.Num() - RangeIndex),
+						OutRhsFrameOffsets.Slice(RangeIndex, OutRhsFrameOffsets.Num() - RangeIndex),
+						Lhs.GetEntryRangeStarts(LhsIndex),
+						Lhs.GetEntryRangeLengths(LhsIndex),
+						Rhs.GetEntryRangeStarts(RhsIndex),
+						Rhs.GetEntryRangeLengths(RhsIndex),
+						LhsFrameOffset,
+						RhsFrameOffset);
+
+					Private::FrameOffsetsCheck(OutLhsFrameOffsets.Slice(0, RangeIndex + RangeNum), 0, LhsFrameOffset + Lhs.GetEntryTotalFrameNum(LhsIndex));
+					Private::FrameOffsetsCheck(OutRhsFrameOffsets.Slice(0, RangeIndex + RangeNum), 0, RhsFrameOffset + Rhs.GetEntryTotalFrameNum(RhsIndex));
+
+					UE_LEARNING_CHECK(RangeNum <= Lhs.GetEntryRangeNum(LhsIndex) + Rhs.GetEntryRangeNum(RhsIndex));
+
+					if (RangeNum > 0)
+					{
+						Out.EntrySequences[OutIndex] = Lhs.GetEntrySequence(LhsIndex);
+						Out.EntryRangeOffsets[OutIndex] = RangeIndex;
+						Out.EntryRangeNums[OutIndex] = RangeNum;
+						RangeIndex += RangeNum;
+						OutIndex++;
+					}
+
+					LhsFrameOffset += Lhs.GetEntryTotalFrameNum(LhsIndex);
+					RhsFrameOffset += Rhs.GetEntryTotalFrameNum(RhsIndex);
+					LhsIndex++; RhsIndex++;
+				}
+			}
+
+			// Resize output to match what was added
+			Out.EntrySequences.SetNumUninitialized({ OutIndex });
+			Out.EntryRangeOffsets.SetNumUninitialized({ OutIndex });
+			Out.EntryRangeNums.SetNumUninitialized({ OutIndex });
+			Out.RangeStarts.SetNumUninitialized({ RangeIndex });
+			Out.RangeLengths.SetNumUninitialized({ RangeIndex });
+			Out.Check();
+
+			Private::FrameOffsetsCheck(OutLhsFrameOffsets.Slice(0, RangeIndex), 0, Lhs.GetTotalFrameNum());
+			Private::FrameOffsetsCheck(OutRhsFrameOffsets.Slice(0, RangeIndex), 0, Rhs.GetTotalFrameNum());
+
+			UE_LEARNING_CHECK(RangeIndex == Out.GetTotalRangeNum());
+			return RangeIndex;
+		}
+
+		void TrimLeft(FFrameRangeSet& Out, const FFrameRangeSet& FrameRangeSet, const int32 TrimFrameNum)
+		{
+			UE_LEARNING_CHECK(TrimFrameNum >= 0);
+			FrameRangeSet.Check();
+
+			const int32 EntryNum = FrameRangeSet.GetEntryNum();
+
+			Out.Empty();
+
+			TArray<int32> RangeStartsAdded;
+			TArray<int32> RangeLengthsAdded;
+
+			int32 RangeOffset = 0;
+			for (int32 EntryIdx = 0; EntryIdx < EntryNum; EntryIdx++)
+			{
+				const int32 Sequence = FrameRangeSet.GetEntrySequence(EntryIdx);
+				const int32 RangeNum = FrameRangeSet.GetEntryRangeNum(EntryIdx);
+
+				RangeStartsAdded.Reset();
+				RangeLengthsAdded.Reset();
+				for (int32 RangeIdx = 0; RangeIdx < RangeNum; RangeIdx++)
+				{
+					const int32 RangeStart = FrameRangeSet.GetEntryRangeStart(EntryIdx, RangeIdx);
+					const int32 RangeLength = FrameRangeSet.GetEntryRangeLength(EntryIdx, RangeIdx);
+					if (RangeLength - TrimFrameNum > 0)
+					{
+						RangeStartsAdded.Add(RangeStart + TrimFrameNum);
+						RangeLengthsAdded.Add(RangeLength - TrimFrameNum);
+					}
+				}
+
+				Out.AddEntry(Sequence, RangeStartsAdded, RangeLengthsAdded);
+			}
+
+			Out.Check();
+		}
+
+		void TrimRight(FFrameRangeSet& Out, const FFrameRangeSet& FrameRangeSet, const int32 TrimFrameNum)
+		{
+			UE_LEARNING_CHECK(TrimFrameNum >= 0);
+			FrameRangeSet.Check();
+
+			const int32 EntryNum = FrameRangeSet.GetEntryNum();
+
+			Out.Empty();
+
+			TArray<int32> RangeStartsAdded;
+			TArray<int32> RangeLengthsAdded;
+
+			int32 RangeOffset = 0;
+			for (int32 EntryIdx = 0; EntryIdx < EntryNum; EntryIdx++)
+			{
+				const int32 Sequence = FrameRangeSet.GetEntrySequence(EntryIdx);
+				const int32 RangeNum = FrameRangeSet.GetEntryRangeNum(EntryIdx);
+
+				RangeStartsAdded.Reset();
+				RangeLengthsAdded.Reset();
+				for (int32 RangeIdx = 0; RangeIdx < RangeNum; RangeIdx++)
+				{
+					const int32 RangeStart = FrameRangeSet.GetEntryRangeStart(EntryIdx, RangeIdx);
+					const int32 RangeLength = FrameRangeSet.GetEntryRangeLength(EntryIdx, RangeIdx);
+					if (RangeLength - TrimFrameNum > 0)
+					{
+						RangeStartsAdded.Add(RangeStart);
+						RangeLengthsAdded.Add(RangeLength - TrimFrameNum);
+					}
+				}
+
+				Out.AddEntry(Sequence, RangeStartsAdded, RangeLengthsAdded);
+			}
+
+			Out.Check();
+		}
+
+		void TrimLeftAndRight(FFrameRangeSet& Out, const FFrameRangeSet& FrameRangeSet, const int32 TrimFrameNum)
+		{
+			UE_LEARNING_CHECK(TrimFrameNum >= 0);
+			FrameRangeSet.Check();
+
+			const int32 EntryNum = FrameRangeSet.GetEntryNum();
+
+			Out.Empty();
+
+			TArray<int32> RangeStartsAdded;
+			TArray<int32> RangeLengthsAdded;
+
+			int32 RangeOffset = 0;
+			for (int32 EntryIdx = 0; EntryIdx < EntryNum; EntryIdx++)
+			{
+				const int32 Sequence = FrameRangeSet.GetEntrySequence(EntryIdx);
+				const int32 RangeNum = FrameRangeSet.GetEntryRangeNum(EntryIdx);
+
+				RangeStartsAdded.Reset();
+				RangeLengthsAdded.Reset();
+				for (int32 RangeIdx = 0; RangeIdx < RangeNum; RangeIdx++)
+				{
+					const int32 RangeStart = FrameRangeSet.GetEntryRangeStart(EntryIdx, RangeIdx);
+					const int32 RangeLength = FrameRangeSet.GetEntryRangeLength(EntryIdx, RangeIdx);
+					if (RangeLength - 2 * TrimFrameNum > 0)
+					{
+						RangeStartsAdded.Add(RangeStart + TrimFrameNum);
+						RangeLengthsAdded.Add(RangeLength - 2 * TrimFrameNum);
+					}
+				}
+
+				Out.AddEntry(Sequence, RangeStartsAdded, RangeLengthsAdded);
+			}
+
+			Out.Check();
+		}
+
+		void MakeFrameSetFromRangeStarts(FFrameSet& OutFrameSet, const FFrameRangeSet& FrameRangeSet)
+		{
+			OutFrameSet.EntrySequences = FrameRangeSet.EntrySequences;
+			OutFrameSet.EntryFrameOffsets = FrameRangeSet.EntryRangeOffsets;
+			OutFrameSet.EntryFrameNums = FrameRangeSet.EntryRangeNums;
+			OutFrameSet.Frames = FrameRangeSet.RangeStarts;
+		}
+
+		void MakeFrameSetFromRangeEnds(FFrameSet& OutFrameSet, const FFrameRangeSet& FrameRangeSet)
+		{
+			OutFrameSet.EntrySequences = FrameRangeSet.EntrySequences;
+			OutFrameSet.EntryFrameOffsets = FrameRangeSet.EntryRangeOffsets;
+			OutFrameSet.EntryFrameNums = FrameRangeSet.EntryRangeNums;
+
+			const int32 RangeNum = FrameRangeSet.RangeStarts.Num();
+			OutFrameSet.Frames.SetNumUninitialized({ RangeNum });
+			for (int32 RangeIdx = 0; RangeIdx < RangeNum; RangeIdx++)
+			{
+				OutFrameSet.Frames[RangeIdx] = FrameRangeSet.RangeStarts[RangeIdx] + FrameRangeSet.RangeLengths[RangeIdx] - 1;
+			}
+		}
+
+		void RangesBeforeFrameSet(FFrameRangeSet& OutFrameRangeSet, const FFrameRangeSet& FrameRangeSet, const FFrameSet& FrameSet)
+		{
+			OutFrameRangeSet.Empty();
+
+			const int32 RangeSetEntryNum = FrameRangeSet.GetEntryNum();
+			const int32 FrameSetEntryNum = FrameSet.GetEntryNum();
+
+			TArray<int32> RangeStartsAdded;
+			TArray<int32> RangeLengthsAdded;
+
+			int32 RangeOffset = 0;
+			for (int32 RangeSetEntryIdx = 0; RangeSetEntryIdx < RangeSetEntryNum; RangeSetEntryIdx++)
+			{
+				const int32 Sequence = FrameRangeSet.GetEntrySequence(RangeSetEntryIdx);
+				const int32 FrameSetEntryIdx = FrameSet.FindSequenceEntry(Sequence);
+
+				if (FrameSetEntryIdx != INDEX_NONE)
+				{
+					const int32 RangeNum = FrameRangeSet.GetEntryRangeNum(RangeSetEntryIdx);
+					const int32 FrameNum = FrameSet.GetEntryFrameNum(FrameSetEntryIdx);
+
+					RangeStartsAdded.Reset();
+					RangeLengthsAdded.Reset();
+					for (int32 RangeIdx = 0; RangeIdx < RangeNum; RangeIdx++)
+					{
+						const int32 RangeStart = FrameRangeSet.GetEntryRangeStart(RangeSetEntryIdx, RangeIdx);
+						const int32 RangeLength = FrameRangeSet.GetEntryRangeLength(RangeSetEntryIdx, RangeIdx);
+
+						for (int32 FrameIdx = 0; FrameIdx < FrameNum; FrameIdx++)
+						{
+							const int32 Frame = FrameSet.GetEntryFrame(FrameSetEntryIdx, FrameIdx);
+							if (Frame > RangeStart && Frame < RangeStart + RangeLength)
+							{
+								RangeStartsAdded.Add(RangeStart);
+								RangeLengthsAdded.Add(Frame - RangeStart);
+							}
+						}
+					}
+
+					OutFrameRangeSet.AddEntry(Sequence, RangeStartsAdded, RangeLengthsAdded);
+				}
+			}
+
+			OutFrameRangeSet.Check();
+		}
+
+		void RangesAfterFrameSet(FFrameRangeSet& OutFrameRangeSet, const FFrameRangeSet& FrameRangeSet, const FFrameSet& FrameSet)
+		{
+			OutFrameRangeSet.Empty();
+
+			const int32 RangeSetEntryNum = FrameRangeSet.GetEntryNum();
+			const int32 FrameSetEntryNum = FrameSet.GetEntryNum();
+
+			TArray<int32> RangeStartsAdded;
+			TArray<int32> RangeLengthsAdded;
+
+			int32 RangeOffset = 0;
+			for (int32 RangeSetEntryIdx = 0; RangeSetEntryIdx < RangeSetEntryNum; RangeSetEntryIdx++)
+			{
+				const int32 Sequence = FrameRangeSet.GetEntrySequence(RangeSetEntryIdx);
+				const int32 FrameSetEntryIdx = FrameSet.FindSequenceEntry(Sequence);
+
+				if (FrameSetEntryIdx != INDEX_NONE)
+				{
+					const int32 RangeNum = FrameRangeSet.GetEntryRangeNum(RangeSetEntryIdx);
+					const int32 FrameNum = FrameSet.GetEntryFrameNum(FrameSetEntryIdx);
+
+					RangeStartsAdded.Reset();
+					RangeLengthsAdded.Reset();
+					for (int32 RangeIdx = 0; RangeIdx < RangeNum; RangeIdx++)
+					{
+						const int32 RangeStart = FrameRangeSet.GetEntryRangeStart(RangeSetEntryIdx, RangeIdx);
+						const int32 RangeLength = FrameRangeSet.GetEntryRangeLength(RangeSetEntryIdx, RangeIdx);
+
+						for (int32 FrameIdx = 0; FrameIdx < FrameNum; FrameIdx++)
+						{
+							const int32 Frame = FrameSet.GetEntryFrame(FrameSetEntryIdx, FrameIdx);
+							if (Frame >= RangeStart && Frame < RangeStart + RangeLength)
+							{
+								RangeStartsAdded.Add(Frame);
+								RangeLengthsAdded.Add(RangeLength - (Frame - RangeStart));
+							}
+						}
+					}
+
+					OutFrameRangeSet.AddEntry(Sequence, RangeStartsAdded, RangeLengthsAdded);
+				}
+			}
+
+			OutFrameRangeSet.Check();
+		}
+
+		void AllRangeEntries(
+			TLearningArrayView<1, int32> OutRangeEntries,
+			const FFrameRangeSet& FrameRangeSet)
+		{
+			UE_LEARNING_CHECK(OutRangeEntries.Num() == FrameRangeSet.GetTotalRangeNum());
+
+			const int32 TotalRangeNum = FrameRangeSet.GetTotalRangeNum();
+			const int32 EntryNum = FrameRangeSet.GetEntryNum();
+
+			int32 RangeOffset = 0;
+			for (int32 EntryIdx = 0; EntryIdx < EntryNum; EntryIdx++)
+			{
+				const int32 RangeNum = FrameRangeSet.GetEntryRangeNum(EntryIdx);
+				Array::Set(OutRangeEntries.Slice(RangeOffset, RangeNum), EntryIdx);
+				RangeOffset += RangeNum;
+			}
+
+			UE_LEARNING_CHECK(RangeOffset == TotalRangeNum);
+		}
+
+		void AllRangeIndices(
+			TLearningArrayView<1, int32> OutRangeIndices,
+			const FFrameRangeSet& FrameRangeSet)
+		{
+			UE_LEARNING_CHECK(OutRangeIndices.Num() == FrameRangeSet.GetTotalRangeNum());
+
+			const int32 TotalRangeNum = FrameRangeSet.GetTotalRangeNum();
+			const int32 EntryNum = FrameRangeSet.GetEntryNum();
+
+			int32 RangeOffset = 0;
+			for (int32 EntryIdx = 0; EntryIdx < EntryNum; EntryIdx++)
+			{
+				const int32 RangeNum = FrameRangeSet.GetEntryRangeNum(EntryIdx);
+				for (int32 RangeIdx = 0; RangeIdx < RangeNum; RangeIdx++)
+				{
+					OutRangeIndices[RangeOffset + RangeIdx] = RangeIdx;
+				}
+				RangeOffset += RangeNum;
+			}
+
+			UE_LEARNING_CHECK(RangeOffset == TotalRangeNum);
+		}
+
+		void AllRangeSequences(
+			TLearningArrayView<1, int32> OutRangeSequences,
+			const FFrameRangeSet& FrameRangeSet)
+		{
+			UE_LEARNING_CHECK(OutRangeSequences.Num() == FrameRangeSet.GetTotalRangeNum());
+
+			const int32 TotalRangeNum = FrameRangeSet.GetTotalRangeNum();
+			const int32 EntryNum = FrameRangeSet.GetEntryNum();
+
+			int32 RangeOffset = 0;
+			for (int32 EntryIdx = 0; EntryIdx < EntryNum; EntryIdx++)
+			{
+				const int32 RangeNum = FrameRangeSet.GetEntryRangeNum(EntryIdx);
+				Array::Set(OutRangeSequences.Slice(RangeOffset, RangeNum), FrameRangeSet.GetEntrySequence(EntryIdx));
+				RangeOffset += RangeNum;
+			}
+
+			UE_LEARNING_CHECK(RangeOffset == TotalRangeNum);
+		}
+
+		void AllRangeFrameOffsets(
+			TLearningArrayView<1, int32> OutRangeFrameOffsets,
+			const FFrameRangeSet& FrameRangeSet)
+		{
+			UE_LEARNING_CHECK(OutRangeFrameOffsets.Num() == FrameRangeSet.GetTotalRangeNum());
+
+			const int32 TotalRangeNum = FrameRangeSet.GetTotalRangeNum();
+			const int32 TotalFrameNum = FrameRangeSet.GetTotalFrameNum();
+			const int32 EntryNum = FrameRangeSet.GetEntryNum();
+
+			const int32 RangeNum = FrameRangeSet.GetTotalRangeNum();
+
+			int32 FrameOffset = 0;
+			for (int32 RangeIdx = 0; RangeIdx < RangeNum; RangeIdx++)
+			{
+				OutRangeFrameOffsets[RangeIdx] = FrameOffset;
+				FrameOffset += FrameRangeSet.RangeLengths[RangeIdx];
+			}
+
+			UE_LEARNING_CHECK(FrameOffset == TotalFrameNum);
+		}
+
+		void AllRangeStartTimes(
+			TLearningArrayView<1, float> OutRangeStartTimes,
+			const FFrameRangeSet& FrameRangeSet, 
+			const float FrameDeltaTime)
+		{
+			UE_LEARNING_CHECK(OutRangeStartTimes.Num() == FrameRangeSet.GetTotalRangeNum());
+
+			const int32 TotalRangeNum = FrameRangeSet.GetTotalRangeNum();
+			for (int32 RangeIdx = 0; RangeIdx < TotalRangeNum; RangeIdx++)
+			{
+				OutRangeStartTimes[RangeIdx] = FrameRangeSet.RangeStarts[RangeIdx] * FrameDeltaTime;
+			}
+		}
+
+		void AllRangeEndTimes(
+			TLearningArrayView<1, float> OutRangeEndTimes,
+			const FFrameRangeSet& FrameRangeSet, 
+			const float FrameDeltaTime)
+		{
+			UE_LEARNING_CHECK(OutRangeEndTimes.Num() == FrameRangeSet.GetTotalRangeNum());
+
+			const int32 TotalRangeNum = FrameRangeSet.GetTotalRangeNum();
+			for (int32 RangeIdx = 0; RangeIdx < TotalRangeNum; RangeIdx++)
+			{
+				OutRangeEndTimes[RangeIdx] = (FrameRangeSet.RangeStarts[RangeIdx] + FrameRangeSet.RangeLengths[RangeIdx] - 1) * FrameDeltaTime;
+			}
+		}
+
+		void AllRangeDurations(
+			TLearningArrayView<1, float> OutRangeDurations,
+			const FFrameRangeSet& FrameRangeSet,
+			const float FrameDeltaTime)
+		{
+			UE_LEARNING_CHECK(OutRangeDurations.Num() == FrameRangeSet.GetTotalRangeNum());
+
+			const int32 TotalRangeNum = FrameRangeSet.GetTotalRangeNum();
+			for (int32 RangeIdx = 0; RangeIdx < TotalRangeNum; RangeIdx++)
+			{
+				OutRangeDurations[RangeIdx] = (FrameRangeSet.RangeLengths[RangeIdx] - 1) * FrameDeltaTime;
+			}
+		}
+
+		void ForEachRange(
+			const FFrameRangeSet& FrameRangeSet,
+			const TFunctionRef<void(
+				const int32 TotalRangeIdx,
+				const int32 EntryIdx,
+				const int32 RangeIdx,
+				const int32 FrameOffset)> Body)
+		{
+			const int32 EntryNum = FrameRangeSet.GetEntryNum();
+
+			int32 TotalRangeIdx = 0;
+			int32 FrameOffset = 0;
+			for (int32 EntryIdx = 0; EntryIdx < EntryNum; EntryIdx++)
+			{
+				const int32 RangeNum = FrameRangeSet.GetEntryRangeNum(EntryIdx);
+				
+				for (int32 RangeIdx = 0; RangeIdx < RangeNum; RangeIdx++)
+				{
+					Body(
+						TotalRangeIdx,
+						EntryIdx,
+						RangeIdx,
+						FrameOffset);
+
+					TotalRangeIdx++;
+					FrameOffset += FrameRangeSet.GetEntryRangeLength(EntryIdx, RangeIdx);
+				}
+			}
+		}
+
+		void ParallelForEachRange(
+			const FFrameRangeSet& FrameRangeSet,
+			const TFunctionRef<void(
+				const int32 TotalRangeIdx,
+				const int32 EntryIdx,
+				const int32 RangeIdx,
+				const int32 FrameOffset)> Body)
+		{
+			const int32 TotalRangeNum = FrameRangeSet.GetTotalRangeNum();
+
+			TLearningArray<1, int32> RangeEntries;
+			TLearningArray<1, int32> RangeIndices;
+			TLearningArray<1, int32> RangeFrameOffsets;
+
+			RangeEntries.SetNumUninitialized({ TotalRangeNum });
+			RangeIndices.SetNumUninitialized({ TotalRangeNum });
+			RangeFrameOffsets.SetNumUninitialized({ TotalRangeNum });
+
+			AllRangeEntries(RangeEntries, FrameRangeSet);
+			AllRangeIndices(RangeIndices, FrameRangeSet);
+			AllRangeFrameOffsets(RangeFrameOffsets, FrameRangeSet);
+
+			ParallelFor(TotalRangeNum, [&Body, &RangeEntries, &RangeIndices, &RangeFrameOffsets](const int32 TotalRangeIdx) {
+				Body(
+					TotalRangeIdx,
+					RangeEntries[TotalRangeIdx],
+					RangeIndices[TotalRangeIdx],
+					RangeFrameOffsets[TotalRangeIdx]);
+			});
+		}
+	}
+}
diff --git a/Engine/Plugins/Experimental/LearningAgents/Source/Learning/Private/LearningFrameSet.cpp b/Engine/Plugins/Experimental/LearningAgents/Source/Learning/Private/LearningFrameSet.cpp
new file mode 100644
--- /dev/null
+++ b/Engine/Plugins/Experimental/LearningAgents/Source/Learning/Private/LearningFrameSet.cpp
@@ -0,0 +1,844 @@
+// Copyright Epic Games, Inc. All Rights Reserved.
+
+#include "LearningFrameSet.h"
+
+#include "Async/ParallelFor.h"
+
+namespace UE::Learning
+{
+	namespace FrameSet::Private
+	{
+		static inline void FramesCheck(const TLearningArrayView<1, const int32> Frames)
+		{
+			const int32 FrameNum = Frames.Num();
+			for (int32 FrameIdx = 0; FrameIdx < FrameNum - 1; FrameIdx++)
+			{
+				UE_LEARNING_CHECK(Frames[FrameIdx + 0] >= 0);
+				UE_LEARNING_CHECK(Frames[FrameIdx + 1] >= 0);
+				UE_LEARNING_CHECK(Frames[FrameIdx + 0] < Frames[FrameIdx + 1]);
+			}
+		}
+
+		static inline int32 FramesUnion(
+			TLearningArrayView<1, int32> OutFrames,
+			const TLearningArrayView<1, const int32> LhsFrames,
+			const TLearningArrayView<1, const int32> RhsFrames)
+		{
+			TRACE_CPUPROFILER_EVENT_SCOPE(Learning::FrameSet::Private::FramesUnion);
+
+			FramesCheck(LhsFrames);
+			FramesCheck(RhsFrames);
+
+			// Number of frames in lhs and rhs
+			const int32 LhsNum = LhsFrames.Num();
+			const int32 RhsNum = RhsFrames.Num();
+
+			// Event index for each list of frames
+			int32 LhsIndex = 0;
+			int32 RhsIndex = 0;
+			int32 OutIndex = 0;
+
+			// While both frames to process
+			while (LhsIndex < LhsNum && RhsIndex < RhsNum)
+			{
+				// Time of the next lhs, and rhs frames
+				const int32 LhsT = LhsFrames[LhsIndex];
+				const int32 RhsT = RhsFrames[RhsIndex];
+
+				// Frame from lhs is coming first
+				if (LhsT < RhsT)
+				{
+					OutFrames[OutIndex] = LhsT;
+					OutIndex++;
+					LhsIndex++;
+				}
+				// Frame from rhs is coming first
+				else if (RhsT < LhsT)
+				{
+					OutFrames[OutIndex] = RhsT;
+					OutIndex++;
+					RhsIndex++;
+				}
+				// Frame from lhs and rhs coming at same time
+				else
+				{
+					UE_LEARNING_CHECK(LhsT == RhsT);
+					OutFrames[OutIndex] = LhsT;
+					OutIndex++;
+					LhsIndex++;
+					RhsIndex++;
+				}
+			}
+
+			// Process any remaining lhs frames
+			while (LhsIndex < LhsNum)
+			{
+				OutFrames[OutIndex] = LhsFrames[LhsIndex];
+				OutIndex++;
+				LhsIndex++;
+			}
+
+			// Process any remaining rhs frames
+			while (RhsIndex < RhsNum)
+			{
+				OutFrames[OutIndex] = RhsFrames[RhsIndex];
+				OutIndex++;
+				RhsIndex++;
+			}
+
+			FramesCheck(OutFrames.Slice(0, OutIndex));
+
+			// Return number of frames added
+			return OutIndex;
+		}
+
+		static inline int32 FramesIntersection(
+			TLearningArrayView<1, int32> OutFrames,
+			const TLearningArrayView<1, const int32> LhsFrames,
+			const TLearningArrayView<1, const int32> RhsFrames)
+		{
+			TRACE_CPUPROFILER_EVENT_SCOPE(Learning::FrameSet::Private::FramesIntersection);
+
+			FramesCheck(LhsFrames);
+			FramesCheck(RhsFrames);
+
+			// Number of frames in lhs and rhs
+			const int32 LhsNum = LhsFrames.Num();
+			const int32 RhsNum = RhsFrames.Num();
+
+			// Event index for each list of frames
+			int32 LhsIndex = 0;
+			int32 RhsIndex = 0;
+			int32 OutIndex = 0;
+
+			// While both frames to process
+			while (LhsIndex < LhsNum && RhsIndex < RhsNum)
+			{
+				// Time of the next lhs, and rhs frames
+				const int32 LhsT = LhsFrames[LhsIndex];
+				const int32 RhsT = RhsFrames[RhsIndex];
+
+				// Frame from lhs is coming first
+				if (LhsT < RhsT)
+				{
+					LhsIndex++;
+				}
+				// Frame from rhs is coming first
+				else if (RhsT < LhsT)
+				{
+					RhsIndex++;
+				}
+				// Frame from lhs and rhs coming at same time
+				else
+				{
+					UE_LEARNING_CHECK(LhsT == RhsT);
+					OutFrames[OutIndex] = LhsT;
+					OutIndex++;
+					LhsIndex++;
+					RhsIndex++;
+				}
+			}
+
+			FramesCheck(OutFrames.Slice(0, OutIndex));
+
+			// Return number of frames added
+			return OutIndex;
+		}
+
+		static inline int32 FramesDifference(
+			TLearningArrayView<1, int32> OutFrames,
+			const TLearningArrayView<1, const int32> LhsFrames,
+			const TLearningArrayView<1, const int32> RhsFrames)
+		{
+			TRACE_CPUPROFILER_EVENT_SCOPE(Learning::FrameSet::Private::FramesDifference);
+
+			FramesCheck(LhsFrames);
+			FramesCheck(RhsFrames);
+
+			// Number of frames in lhs and rhs
+			const int32 LhsNum = LhsFrames.Num();
+			const int32 RhsNum = RhsFrames.Num();
+
+			// Event index for each list of frames
+			int32 LhsIndex = 0;
+			int32 RhsIndex = 0;
+			int32 OutIndex = 0;
+
+			// While both frames to process
+			while (LhsIndex < LhsNum && RhsIndex < RhsNum)
+			{
+				// Time of the next lhs, and rhs frames
+				const int32 LhsT = LhsFrames[LhsIndex];
+				const int32 RhsT = RhsFrames[RhsIndex];
+
+				// Frame from lhs is coming first
+				if (LhsT < RhsT)
+				{
+					OutFrames[OutIndex] = LhsT;
+					OutIndex++;
+					LhsIndex++;
+				}
+				// Frame from rhs is coming first
+				else if (RhsT < LhsT)
+				{
+					RhsIndex++;
+				}
+				// Frame from lhs and rhs coming at same time
+				else
+				{
+					UE_LEARNING_CHECK(LhsT == RhsT);
+					LhsIndex++;
+					RhsIndex++;
+				}
+			}
+
+			// Process any remaining lhs frames
+			while (LhsIndex < LhsNum)
+			{
+				OutFrames[OutIndex] = LhsFrames[LhsIndex];
+				OutIndex++;
+				LhsIndex++;
+			}
+
+			FramesCheck(OutFrames.Slice(0, OutIndex));
+
+			// Return number of frames added
+			return OutIndex;
+		}
+
+	}
+
+	void FFrameSet::Check() const
+	{
+		const int32 EntryNum = EntrySequences.Num();
+
+		for (int32 EntryIdx = 0; EntryIdx < EntryNum - 1; EntryIdx++)
+		{
+			UE_LEARNING_CHECK(EntrySequences[EntryIdx + 0] < EntrySequences[EntryIdx + 1]);
+		}
+
+		for (int32 EntryIdx = 0; EntryIdx < EntryNum; EntryIdx++)
+		{
+			UE_LEARNING_CHECK(EntryFrameNums[EntryIdx] > 0);
+			FrameSet::Private::FramesCheck(GetEntryFrames(EntryIdx));
+		}
+	}
+
+	bool FFrameSet::IsEmpty() const { return EntrySequences.IsEmpty(); }
+	
+	void FFrameSet::Empty()
+	{
+		EntrySequences.Empty();
+		EntryFrameOffsets.Empty();
+		EntryFrameNums.Empty();
+		Frames.Empty();
+	}
+
+	int32 FFrameSet::GetEntryNum() const { return EntrySequences.Num(); }
+	TLearningArrayView<1, const int32> FFrameSet::GetEntrySequences() const { return EntrySequences; }
+	TLearningArrayView<1, const int32> FFrameSet::GetEntryFrameNums() const { return EntryFrameNums; }
+	int32 FFrameSet::GetEntrySequence(const int32 EntryIdx) const { return EntrySequences[EntryIdx]; }
+	int32 FFrameSet::GetEntryFrameNum(const int32 EntryIdx) const { return EntryFrameNums[EntryIdx]; }
+	TLearningArrayView<1, const int32> FFrameSet::GetEntryFrames(const int32 EntryIdx) const { return Frames.Slice(EntryFrameOffsets[EntryIdx], EntryFrameNums[EntryIdx]); }
+	int32 FFrameSet::GetEntryFrame(const int32 EntryIdx, const int32 FrameIdx) const { return Frames[EntryFrameOffsets[EntryIdx] + FrameIdx]; }
+	float FFrameSet::GetEntryFrameTime(const int32 EntryIdx, const int32 FrameIdx, const float FrameDeltaTime) const { return GetEntryFrame(EntryIdx, FrameIdx) * FrameDeltaTime; }
+	int32 FFrameSet::GetEntryFrameOffset(const int32 EntryIdx) const { return EntryFrameOffsets[EntryIdx]; }
+	int32 FFrameSet::GetTotalFrameNum() const { return Frames.Num(); }
+	
+	bool FFrameSet::ContainsSequence(const int32 Sequence) const
+	{
+		return EntrySequences.ArrayView().Contains(Sequence);
+	}
+
+	bool FFrameSet::Contains(const int32 Sequence, const int32 Frame) const
+	{
+		const int32 EntryIdx = FindSequenceEntry(Sequence);
+
+		return EntryIdx != INDEX_NONE && GetEntryFrames(EntryIdx).Contains(Frame);
+	}
+
+	int32 FFrameSet::FindSequenceEntry(const int32 Sequence) const
+	{
+		return EntrySequences.ArrayView().Find(Sequence);
+	}
+
+	bool FFrameSet::Find(int32& OutEntryIdx, int32& OutFrameIdx, const int32 Sequence, const int32 Frame) const
+	{
+		const int32 EntryIdx = FindSequenceEntry(Sequence);
+
+		if (EntryIdx != INDEX_NONE)
+		{
+			const int32 FoundFrameIdx = GetEntryFrames(EntryIdx).ArrayView().Find(Frame);
+			if (FoundFrameIdx != INDEX_NONE)
+			{
+				OutEntryIdx = EntryIdx;
+				OutFrameIdx = FoundFrameIdx;
+				return true;
+			}
+		}
+
+		OutEntryIdx = INDEX_NONE;
+		OutFrameIdx = INDEX_NONE;
+		return false;
+	}
+
+	bool FFrameSet::FindNearest(int32& OutEntryIdx, int32& OutFrameIdx, int32& OutFrameDifference, const int32 Sequence, const int32 Frame) const
+	{
+		const int32 EntryIdx = FindSequenceEntry(Sequence);
+
+		if (EntryIdx != INDEX_NONE)
+		{
+			OutFrameDifference = INT32_MAX;
+			OutEntryIdx = INDEX_NONE;
+			OutFrameIdx = INDEX_NONE;
+
+			const int32 FrameNum = GetEntryFrameNum(EntryIdx);
+			for (int32 FrameIdx = 0; FrameIdx < FrameNum; FrameIdx++)
+			{
+				const int32 FrameDifference = GetEntryFrame(EntryIdx, FrameIdx) - Frame;
+
+				if (OutFrameDifference == INT32_MAX || FMath::Abs(FrameDifference) < OutFrameDifference)
+				{
+					OutFrameDifference = FrameDifference;
+					OutEntryIdx = EntryIdx;
+					OutFrameIdx = FrameIdx;
+				}
+			}
+
+			return OutFrameDifference != INT32_MAX;
+		}
+
+		OutFrameDifference = INT32_MAX;
+		OutEntryIdx = INDEX_NONE;
+		OutFrameIdx = INDEX_NONE;
+		return false;
+	}
+
+	bool FFrameSet::FindNearestInRange(int32& OutEntryIdx, int32& OutFrameIdx, int32& OutFrameDifference, const int32 Sequence, const int32 Frame, const int32 RangeStart, const int32 RangeLength) const
+	{
+		UE_LEARNING_CHECK(Frame >= RangeStart && Frame < RangeStart + RangeLength);
+
+		const int32 EntryIdx = FindSequenceEntry(Sequence);
+
+		if (EntryIdx != INDEX_NONE)
+		{
+			OutFrameDifference = INT32_MAX;
+			OutEntryIdx = INDEX_NONE;
+			OutFrameIdx = INDEX_NONE;
+
+			const int32 FrameNum = GetEntryFrameNum(EntryIdx);
+			for (int32 FrameIdx = 0; FrameIdx < FrameNum; FrameIdx++)
+			{
+				const int32 EntryFrame = GetEntryFrame(EntryIdx, FrameIdx);
+				if (EntryFrame >= RangeStart && EntryFrame < RangeStart + RangeLength)
+				{
+					const int32 FrameDifference = EntryFrame - Frame;
+
+					if (OutFrameDifference == INT32_MAX || FMath::Abs(FrameDifference) < OutFrameDifference)
+					{
+						OutFrameDifference = FrameDifference;
+						OutEntryIdx = EntryIdx;
+						OutFrameIdx = FrameIdx;
+					}
+				}
+			}
+
+			return OutFrameDifference != INT32_MAX;
+		}
+
+		OutFrameDifference = INT32_MAX;
+		OutEntryIdx = INDEX_NONE;
+		OutFrameIdx = INDEX_NONE;
+		return false;
+	}
+
+	bool FFrameSet::FindFrameOffset(int32& OutEntryIdx, int32& OutFrameIdx, const int32 FrameOffset) const
+	{
+		int32 Offset = 0;
+
+		const int32 EntryNum = GetEntryNum();
+		for (int32 EntryIdx = 0; EntryIdx < EntryNum; EntryIdx++)
+		{
+			const int32 FrameNum = GetEntryFrameNum(EntryIdx);
+
+			if (FrameOffset >= Offset && FrameOffset < FrameNum)
+			{
+				OutEntryIdx = EntryIdx;
+				OutFrameIdx = FrameOffset - Offset;
+				return true;
+			}
+
+			Offset += FrameNum;
+		}
+
+		OutEntryIdx = INDEX_NONE;
+		OutFrameIdx = INDEX_NONE;
+		return false;
+	}
+
+	void FFrameSet::AddEntry(
+		const int32 InSequence,
+		const TLearningArrayView<1, const int32> InFrames)
+	{
+		UE_LEARNING_CHECK(!EntrySequences.ArrayView().Contains(InSequence));
+		UE::Learning::FrameSet::Private::FramesCheck(InFrames);
+
+		if (InFrames.IsEmpty()) { return; }
+
+		const int32 CurrFrameOffset = Frames.Num();
+		const int32 AddFrameNum = InFrames.Num();
+		Frames.SetNumUninitialized({ CurrFrameOffset + AddFrameNum });
+		Array::Copy(Frames.Slice(CurrFrameOffset, AddFrameNum), InFrames);
+
+		const int32 CurrEntryNum = EntrySequences.Num();
+		EntrySequences.SetNumUninitialized({ CurrEntryNum + 1 });
+		EntryFrameOffsets.SetNumUninitialized({ CurrEntryNum + 1 });
+		EntryFrameNums.SetNumUninitialized({ CurrEntryNum + 1 });
+		EntrySequences[CurrEntryNum] = InSequence;
+		EntryFrameOffsets[CurrEntryNum] = CurrFrameOffset;
+		EntryFrameNums[CurrEntryNum] = AddFrameNum;
+
+		Check();
+	}
+
+	namespace FrameSet
+	{
+		bool Equal(const FFrameSet& Lhs, const FFrameSet& Rhs)
+		{
+			return
+				Lhs.EntrySequences.Num() == Rhs.EntrySequences.Num() &&
+				Lhs.Frames.Num() == Rhs.Frames.Num() &&
+				Array::Equal(Lhs.EntrySequences, Rhs.EntrySequences) &&
+				Array::Equal(Lhs.EntryFrameOffsets, Rhs.EntryFrameOffsets) &&
+				Array::Equal(Lhs.EntryFrameNums, Rhs.EntryFrameNums) &&
+				Array::Equal(Lhs.Frames, Rhs.Frames);
+		}
+
+		void Union(FFrameSet& Out, const FFrameSet& Lhs, const FFrameSet& Rhs)
+		{
+			TRACE_CPUPROFILER_EVENT_SCOPE(Learning::FrameSet::Union);
+
+			if (Equal(Lhs, Rhs)) { Out = Lhs; return; }
+
+			const int32 LhsEntryNum = Lhs.GetEntryNum();
+			const int32 RhsEntryNum = Rhs.GetEntryNum();
+			const int32 LhsFrameNum = Lhs.GetTotalFrameNum();
+			const int32 RhsFrameNum = Rhs.GetTotalFrameNum();
+
+			// Allocate potential maximum number of entries and frames we might to output
+			Out.EntrySequences.SetNumUninitialized({ LhsEntryNum + RhsEntryNum });
+			Out.EntryFrameOffsets.SetNumUninitialized({ LhsEntryNum + RhsEntryNum });
+			Out.EntryFrameNums.SetNumUninitialized({ LhsEntryNum + RhsEntryNum });
+			Out.Frames.SetNumUninitialized({ LhsFrameNum + RhsFrameNum });
+
+			// Entry index for each list of frames
+			int32 OutIndex = 0;
+			int32 LhsIndex = 0;
+			int32 RhsIndex = 0;
+
+			// Output entry index
+			int32 EventIndex = 0;
+
+			// While both sets have entries
+			while (LhsIndex < LhsEntryNum && RhsIndex < RhsEntryNum)
+			{
+				// If entry from lhs is first
+				if (Lhs.GetEntrySequence(LhsIndex) < Rhs.GetEntrySequence(RhsIndex))
+				{
+					// Append frames to output
+					if (Lhs.GetEntryFrameNum(LhsIndex) > 0)
+					{
+						Out.EntrySequences[OutIndex] = Lhs.GetEntrySequence(LhsIndex);
+						Out.EntryFrameOffsets[OutIndex] = EventIndex;
+						Out.EntryFrameNums[OutIndex] = Lhs.GetEntryFrameNum(LhsIndex);
+						Array::Copy(Out.Frames.Slice(EventIndex, Lhs.GetEntryFrameNum(LhsIndex)), Lhs.GetEntryFrames(LhsIndex));
+						EventIndex += Lhs.GetEntryFrameNum(LhsIndex);
+						OutIndex++;
+					}
+
+					LhsIndex++;
+				}
+				// If entry from rhs is first
+				else if (Rhs.GetEntrySequence(RhsIndex) < Lhs.GetEntrySequence(LhsIndex))
+				{
+					// Append frames to output
+					if (Rhs.GetEntryFrameNum(RhsIndex) > 0)
+					{
+						Out.EntrySequences[OutIndex] = Rhs.GetEntrySequence(RhsIndex);
+						Out.EntryFrameOffsets[OutIndex] = EventIndex;
+						Out.EntryFrameNums[OutIndex] = Rhs.GetEntryFrameNum(RhsIndex);
+						Array::Copy(Out.Frames.Slice(EventIndex, Rhs.GetEntryFrameNum(RhsIndex)), Rhs.GetEntryFrames(RhsIndex));
+						EventIndex += Rhs.GetEntryFrameNum(RhsIndex);
+						OutIndex++;
+					}
+
+					RhsIndex++;
+				}
+				// If both contain the same entry
+				else
+				{
+					UE_LEARNING_CHECK(Rhs.GetEntrySequence(RhsIndex) == Lhs.GetEntrySequence(LhsIndex));
+
+					// Append union to output
+					const int32 FrameNum = Private::FramesUnion(
+						Out.Frames.Slice(EventIndex, LhsFrameNum + RhsFrameNum - EventIndex),
+						Lhs.GetEntryFrames(LhsIndex),
+						Rhs.GetEntryFrames(RhsIndex));
+
+					UE_LEARNING_CHECK(FrameNum <= Lhs.GetEntryFrameNum(LhsIndex) + Rhs.GetEntryFrameNum(RhsIndex));
+
+					if (FrameNum > 0)
+					{
+						Out.EntrySequences[OutIndex] = Lhs.GetEntrySequence(LhsIndex);
+						Out.EntryFrameOffsets[OutIndex] = EventIndex;
+						Out.EntryFrameNums[OutIndex] = FrameNum;
+						EventIndex += FrameNum;
+						OutIndex++;
+					}
+
+					LhsIndex++; RhsIndex++;
+				}
+			}
+
+			// Process any remaining lhs entries
+			while (LhsIndex < LhsEntryNum)
+			{
+				// Append frames to output
+				if (Lhs.GetEntryFrameNum(LhsIndex) > 0)
+				{
+					Out.EntrySequences[OutIndex] = Lhs.GetEntrySequence(LhsIndex);
+					Out.EntryFrameOffsets[OutIndex] = EventIndex;
+					Out.EntryFrameNums[OutIndex] = Lhs.GetEntryFrameNum(LhsIndex);
+					Array::Copy(Out.Frames.Slice(EventIndex, Lhs.GetEntryFrameNum(LhsIndex)), Lhs.GetEntryFrames(LhsIndex));
+					EventIndex += Lhs.GetEntryFrameNum(LhsIndex);
+					OutIndex++;
+				}
+
+				LhsIndex++;
+			}
+
+			// Process any remaining rhs entries
+			while (RhsIndex < RhsEntryNum)
+			{
+				// Append frames to output
+				if (Rhs.GetEntryFrameNum(RhsIndex) > 0)
+				{
+					Out.EntrySequences[OutIndex] = Rhs.GetEntrySequence(RhsIndex);
+					Out.EntryFrameOffsets[OutIndex] = EventIndex;
+					Out.EntryFrameNums[OutIndex] = Rhs.GetEntryFrameNum(RhsIndex);
+					Array::Copy(Out.Frames.Slice(EventIndex, Rhs.GetEntryFrameNum(RhsIndex)), Rhs.GetEntryFrames(RhsIndex));
+					EventIndex += Rhs.GetEntryFrameNum(RhsIndex);
+					OutIndex++;
+				}
+
+				RhsIndex++;
+			}
+
+			// Resize output to match what was added
+			Out.EntrySequences.SetNumUninitialized({ OutIndex });
+			Out.EntryFrameOffsets.SetNumUninitialized({ OutIndex });
+			Out.EntryFrameNums.SetNumUninitialized({ OutIndex });
+			Out.Frames.SetNumUninitialized({ EventIndex });
+			Out.Check();
+		}
+
+		void Intersection(FFrameSet& Out, const FFrameSet& Lhs, const FFrameSet& Rhs)
+		{
+			TRACE_CPUPROFILER_EVENT_SCOPE(Learning::FrameSet::Intersection);
+
+			if (Equal(Lhs, Rhs)) { Out = Lhs; return; }
+
+			const int32 LhsEntryNum = Lhs.GetEntryNum();
+			const int32 RhsEntryNum = Rhs.GetEntryNum();
+			const int32 LhsFrameNum = Lhs.GetTotalFrameNum();
+			const int32 RhsFrameNum = Rhs.GetTotalFrameNum();
+
+			// Allocate potential maximum number of entries and frames we might to output
+			Out.EntrySequences.SetNumUninitialized({ FMath::Max(LhsEntryNum, RhsEntryNum) });
+			Out.EntryFrameOffsets.SetNumUninitialized({ FMath::Max(LhsEntryNum, RhsEntryNum) });
+			Out.EntryFrameNums.SetNumUninitialized({ FMath::Max(LhsEntryNum, RhsEntryNum) });
+			Out.Frames.SetNumUninitialized({ LhsFrameNum + RhsFrameNum });
+
+			// Entry index for each list of frames
+			int32 OutIndex = 0;
+			int32 LhsIndex = 0;
+			int32 RhsIndex = 0;
+
+			// Output entry index
+			int32 EventIndex = 0;
+
+			// While both sets have entries
+			while (LhsIndex < LhsEntryNum && RhsIndex < RhsEntryNum)
+			{
+				// If entry from lhs is first
+				if (Lhs.GetEntrySequence(LhsIndex) < Rhs.GetEntrySequence(RhsIndex))
+				{
+					LhsIndex++;
+				}
+				// If entry from rhs is first
+				else if (Rhs.GetEntrySequence(RhsIndex) < Lhs.GetEntrySequence(LhsIndex))
+				{
+					RhsIndex++;
+				}
+				// If both contain the same entry
+				else
+				{
+					UE_LEARNING_CHECK(Rhs.GetEntrySequence(RhsIndex) == Lhs.GetEntrySequence(LhsIndex));
+
+					// Append intersection to output
+					const int32 FrameNum = Private::FramesIntersection(
+						Out.Frames.Slice(EventIndex, LhsFrameNum + RhsFrameNum - EventIndex),
+						Lhs.GetEntryFrames(LhsIndex),
+						Rhs.GetEntryFrames(RhsIndex));
+
+					UE_LEARNING_CHECK(FrameNum <= Lhs.GetEntryFrameNum(LhsIndex) + Rhs.GetEntryFrameNum(RhsIndex));
+
+					if (FrameNum > 0)
+					{
+						Out.EntrySequences[OutIndex] = Lhs.GetEntrySequence(LhsIndex);
+						Out.EntryFrameOffsets[OutIndex] = EventIndex;
+						Out.EntryFrameNums[OutIndex] = FrameNum;
+						EventIndex += FrameNum;
+						OutIndex++;
+					}
+
+					LhsIndex++; RhsIndex++;
+				}
+			}
+
+			// Resize output to match what was added
+			Out.EntrySequences.SetNumUninitialized({ OutIndex });
+			Out.EntryFrameOffsets.SetNumUninitialized({ OutIndex });
+			Out.EntryFrameNums.SetNumUninitialized({ OutIndex });
+			Out.Frames.SetNumUninitialized({ EventIndex });
+			Out.Check();
+		}
+
+		void Difference(FFrameSet& Out, const FFrameSet& Lhs, const FFrameSet& Rhs)
+		{
+			TRACE_CPUPROFILER_EVENT_SCOPE(Learning::FrameSet::Difference);
+
+			if (Equal(Lhs, Rhs)) { Out = FFrameSet(); return; }
+
+			const int32 LhsEntryNum = Lhs.GetEntryNum();
+			const int32 RhsEntryNum = Rhs.GetEntryNum();
+			const int32 LhsFrameNum = Lhs.GetTotalFrameNum();
+			const int32 RhsFrameNum = Rhs.GetTotalFrameNum();
+
+			// Allocate potential maximum number of entries and frames we might to output
+			Out.EntrySequences.SetNumUninitialized({ LhsEntryNum });
+			Out.EntryFrameOffsets.SetNumUninitialized({ LhsEntryNum });
+			Out.EntryFrameNums.SetNumUninitialized({ LhsEntryNum });
+			Out.Frames.SetNumUninitialized({ LhsFrameNum + RhsFrameNum });
+
+			// Entry index for each list of frames
+			int32 OutIndex = 0;
+			int32 LhsIndex = 0;
+			int32 RhsIndex = 0;
+
+			// Output entry index
+			int32 EventIndex = 0;
+
+			// While both sets have entries
+			while (LhsIndex < LhsEntryNum && RhsIndex < RhsEntryNum)
+			{
+				// If entry from lhs is first
+				if (Lhs.GetEntrySequence(LhsIndex) < Rhs.GetEntrySequence(RhsIndex))
+				{
+					// Append frames to output
+					if (Lhs.GetEntryFrameNum(LhsIndex) > 0)
+					{
+						Out.EntrySequences[OutIndex] = Lhs.GetEntrySequence(LhsIndex);
+						Out.EntryFrameOffsets[OutIndex] = EventIndex;
+						Out.EntryFrameNums[OutIndex] = Lhs.GetEntryFrameNum(LhsIndex);
+						Array::Copy(Out.Frames.Slice(EventIndex, Lhs.GetEntryFrameNum(LhsIndex)), Lhs.GetEntryFrames(LhsIndex));
+						EventIndex += Lhs.GetEntryFrameNum(LhsIndex);
+						OutIndex++;
+					}
+
+					LhsIndex++;
+				}
+				// If entry from rhs is first
+				else if (Rhs.GetEntrySequence(RhsIndex) < Lhs.GetEntrySequence(LhsIndex))
+				{
+					RhsIndex++;
+				}
+				// If both contain the same entry
+				else
+				{
+					UE_LEARNING_CHECK(Rhs.GetEntrySequence(RhsIndex) == Lhs.GetEntrySequence(LhsIndex));
+
+					// Append difference to output
+					const int32 FrameNum = Private::FramesDifference(
+						Out.Frames.Slice(EventIndex, LhsFrameNum + RhsFrameNum - EventIndex),
+						Lhs.GetEntryFrames(LhsIndex),
+						Rhs.GetEntryFrames(RhsIndex));
+
+					UE_LEARNING_CHECK(FrameNum <= Lhs.GetEntryFrameNum(LhsIndex) + Rhs.GetEntryFrameNum(RhsIndex));
+
+					if (FrameNum > 0)
+					{
+						Out.EntrySequences[OutIndex] = Lhs.GetEntrySequence(LhsIndex);
+						Out.EntryFrameOffsets[OutIndex] = EventIndex;
+						Out.EntryFrameNums[OutIndex] = FrameNum;
+						EventIndex += FrameNum;
+						OutIndex++;
+					}
+
+					LhsIndex++; RhsIndex++;
+				}
+			}
+
+			// Process any remaining lhs entries
+			while (LhsIndex < LhsEntryNum)
+			{
+				// Append frames to output
+				if (Lhs.GetEntryFrameNum(LhsIndex) > 0)
+				{
+					Out.EntrySequences[OutIndex] = Lhs.GetEntrySequence(LhsIndex);
+					Out.EntryFrameOffsets[OutIndex] = EventIndex;
+					Out.EntryFrameNums[OutIndex] = Lhs.GetEntryFrameNum(LhsIndex);
+					Array::Copy(Out.Frames.Slice(EventIndex, Lhs.GetEntryFrameNum(LhsIndex)), Lhs.GetEntryFrames(LhsIndex));
+					EventIndex += Lhs.GetEntryFrameNum(LhsIndex);
+					OutIndex++;
+				}
+
+				LhsIndex++;
+			}
+
+			// Resize output to match what was added
+			Out.EntrySequences.SetNumUninitialized({ OutIndex });
+			Out.EntryFrameOffsets.SetNumUninitialized({ OutIndex });
+			Out.EntryFrameNums.SetNumUninitialized({ OutIndex });
+			Out.Frames.SetNumUninitialized({ EventIndex });
+			Out.Check();
+		}
+
+		void AllFrameEntries(
+			TLearningArrayView<1, int32> OutFrameEntries,
+			const FFrameSet& FrameSet)
+		{
+			UE_LEARNING_CHECK(OutFrameEntries.Num() == FrameSet.GetTotalFrameNum());
+
+			const int32 TotalFrameNum = FrameSet.GetTotalFrameNum();
+			const int32 EntryNum = FrameSet.GetEntryNum();
+
+			int32 FrameOffset = 0;
+			for (int32 EntryIdx = 0; EntryIdx < EntryNum; EntryIdx++)
+			{
+				const int32 FrameNum = FrameSet.GetEntryFrameNum(EntryIdx);
+				Array::Set(OutFrameEntries.Slice(FrameOffset, FrameNum), EntryIdx);
+				FrameOffset += FrameNum;
+			}
+
+			UE_LEARNING_CHECK(FrameOffset == TotalFrameNum);
+		}
+
+		void AllFrameIndices(
+			TLearningArrayView<1, int32> OutFrameIndices,
+			const FFrameSet& FrameSet)
+		{
+			UE_LEARNING_CHECK(OutFrameIndices.Num() == FrameSet.GetTotalFrameNum());
+
+			const int32 TotalFrameNum = FrameSet.GetTotalFrameNum();
+			const int32 EntryNum = FrameSet.GetEntryNum();
+
+			int32 FrameOffset = 0;
+			for (int32 EntryIdx = 0; EntryIdx < EntryNum; EntryIdx++)
+			{
+				const int32 FrameNum = FrameSet.GetEntryFrameNum(EntryIdx);
+				for (int32 FrameIdx = 0; FrameIdx < FrameNum; FrameIdx++)
+				{
+					OutFrameIndices[FrameOffset + FrameIdx] = FrameIdx;
+				}
+				FrameOffset += FrameNum;
+			}
+
+			UE_LEARNING_CHECK(FrameOffset == TotalFrameNum);
+		}
+
+		void AllFrameSequences(
+			TLearningArrayView<1, int32> OutFrameSequences,
+			const FFrameSet& FrameSet)
+		{
+			UE_LEARNING_CHECK(OutFrameSequences.Num() == FrameSet.GetTotalFrameNum());
+
+			const int32 TotalFrameNum = FrameSet.GetTotalFrameNum();
+			const int32 EntryNum = FrameSet.GetEntryNum();
+
+			int32 FrameOffset = 0;
+			for (int32 EntryIdx = 0; EntryIdx < EntryNum; EntryIdx++)
+			{
+				const int32 FrameNum = FrameSet.GetEntryFrameNum(EntryIdx);
+				Array::Set(OutFrameSequences.Slice(FrameOffset, FrameNum), FrameSet.GetEntrySequence(EntryIdx));
+				FrameOffset += FrameNum;
+			}
+
+			UE_LEARNING_CHECK(FrameOffset == TotalFrameNum);
+		}
+
+		void AllFrameTimes(
+			TLearningArrayView<1, float> OutFrameTimes,
+			const FFrameSet& FrameSet,
+			const float FrameDeltaTime)
+		{
+			UE_LEARNING_CHECK(OutFrameTimes.Num() == FrameSet.GetTotalFrameNum());
+
+			const int32 TotalFrameNum = FrameSet.GetTotalFrameNum();
+			for (int32 FrameIdx = 0; FrameIdx < TotalFrameNum; FrameIdx++)
+			{
+				OutFrameTimes[FrameIdx] = FrameSet.Frames[FrameIdx] * FrameDeltaTime;
+			}
+		}
+
+		void ForEachFrame(
+			const FFrameSet& FrameSet,
+			const TFunctionRef<void(
+				const int32 TotalFrameIdx,
+				const int32 EntryIdx,
+				const int32 FrameIdx)> Body)
+		{
+			const int32 EntryNum = FrameSet.GetEntryNum();
+			for (int32 EntryIdx = 0; EntryIdx < EntryNum; EntryIdx++)
+			{
+				const int32 FrameNum = FrameSet.GetEntryFrameNum(EntryIdx);
+				const int32 FrameOffset = FrameSet.GetEntryFrameOffset(EntryIdx);
+				for (int32 FrameIdx = 0; FrameIdx < FrameNum; FrameIdx++)
+				{
+					Body(
+						FrameOffset + FrameIdx,
+						EntryIdx,
+						FrameIdx);
+				}
+			}
+		}
+
+		void ParallelForEachFrame(
+			const FFrameSet& FrameSet,
+			const TFunctionRef<void(
+				const int32 TotalFrameIdx,
+				const int32 EntryIdx,
+				const int32 FrameIdx)> Body)
+		{
+			const int32 TotalFrameNum = FrameSet.GetTotalFrameNum();
+
+			TLearningArray<1, int32> FrameEntries;
+			TLearningArray<1, int32> FrameIndices;
+
+			FrameEntries.SetNumUninitialized({ TotalFrameNum });
+			FrameIndices.SetNumUninitialized({ TotalFrameNum });
+
+			AllFrameEntries(FrameEntries, FrameSet);
+			AllFrameIndices(FrameIndices, FrameSet);
+
+			ParallelFor(TotalFrameNum, [&Body, &FrameEntries, &FrameIndices](const int32 TotalFrameIdx) {
+				Body(
+					TotalFrameIdx,
+					FrameEntries[TotalFrameIdx],
+					FrameIndices[TotalFrameIdx]);
+				});
+		}
+
+	}
+}
+
diff --git a/Engine/Plugins/Experimental/LearningAgents/Source/Learning/Private/LearningRandom.cpp b/Engine/Plugins/Experimental/LearningAgents/Source/Learning/Private/LearningRandom.cpp
--- a/Engine/Plugins/Experimental/LearningAgents/Source/Learning/Private/LearningRandom.cpp
+++ b/Engine/Plugins/Experimental/LearningAgents/Source/Learning/Private/LearningRandom.cpp
@@ -2,6 +2,9 @@
 
 #include "LearningRandom.h"
 
+#include "LearningFrameSet.h"
+#include "LearningFrameRangeSet.h"
+
 #if UE_LEARNING_ISPC
 #include "Learning.ispc.generated.h"
 #endif
@@ -565,6 +568,31 @@ namespace UE::Learning::Random
 #endif
 	}
 
+	void FrameInFrameSet(
+		int32& OutEntryIdx,
+		int32& OutFrameIdx,
+		const FFrameSet& FrameSet,
+		const uint32 State)
+	{
+		const int32 TotalFrameNum = FrameSet.GetTotalFrameNum();
+		const int32 FrameOffset = IntInRange(0x26810f69 ^ State, 0, TotalFrameNum - 1);
+		const bool bFound = FrameSet.FindFrameOffset(OutEntryIdx, OutFrameIdx, FrameOffset);
+		UE_LEARNING_CHECK(bFound);
+	}
+
+	void FrameInFrameRangeSet(
+		int32& OutEntryIdx,
+		int32& OutRangeIdx,
+		int32& OutRangeFrame,
+		const FFrameRangeSet& FrameRangeSet,
+		const uint32 State)
+	{
+		const int32 TotalFrameNum = FrameRangeSet.GetTotalFrameNum();
+		const int32 FrameOffset = IntInRange(0x76ad14bb ^ State, 0, TotalFrameNum - 1);
+		const bool bFound = FrameRangeSet.FindFrameOffset(OutEntryIdx, OutRangeIdx, OutRangeFrame, FrameOffset);
+		UE_LEARNING_CHECK(bFound);
+	}
+
 	////////////
 
 	uint32 SampleInt(uint32& State)
@@ -845,4 +873,35 @@ namespace UE::Learning::Random
 			Masked,
 			Scale);
 	}
+
+
+	void SampleFrameInFrameSet(
+		int32& OutEntryIdx,
+		int32& OutFrameIdx,
+		uint32& State,
+		const FFrameSet& FrameSet)
+	{
+		State = Int(State ^ 0xfbfed01a);
+		FrameInFrameSet(
+			OutEntryIdx,
+			OutFrameIdx,
+			FrameSet,
+			State ^ 0xef6c39f5);
+	}
+
+	void SampleFrameInFrameRangeSet(
+		int32& OutEntryIdx,
+		int32& OutRangeIdx,
+		int32& OutRangeFrame,
+		uint32& State,
+		const FFrameRangeSet& FrameRangeSet)
+	{
+		State = Int(State ^ 0x54f54418);
+		FrameInFrameRangeSet(
+			OutEntryIdx,
+			OutRangeIdx,
+			OutRangeFrame,
+			FrameRangeSet,
+			State ^ 0xec5ae556);
+	}
 }
\ No newline at end of file
diff --git a/Engine/Plugins/Experimental/LearningAgents/Source/Learning/Public/LearningFrameRangeSet.h b/Engine/Plugins/Experimental/LearningAgents/Source/Learning/Public/LearningFrameRangeSet.h
new file mode 100644
--- /dev/null
+++ b/Engine/Plugins/Experimental/LearningAgents/Source/Learning/Public/LearningFrameRangeSet.h
@@ -0,0 +1,256 @@
+// Copyright Epic Games, Inc. All Rights Reserved.
+
+#pragma once
+
+#include "LearningArray.h"
+
+namespace UE::Learning
+{
+	struct FFrameSet;
+
+	/**
+	 * A FFrameRangeSet represents a set of ranges within a set of sequences. This can be useful for encoding things like a collection of tags or 
+	 * labels for different parts of some database of animations, replays, or recordings.
+	 *
+	 * This data structure stores that information in a way that makes performing operations such as union, intersection, difference etc, efficient.
+	 * 
+	 * The way it works is by storing a sorted array of "entries", where each entry has a corresponding sequence, the number of ranges in the set for
+	 * that sequence, and an offset into two large arrays of range starts and stops. The sub-ranges of the large arrays of range starts and stops are 
+	 * also sorted.
+	 *
+	 * Having the entries sorted by sequence, and the ranges sorted, allows for efficient set operations via tape-merge algorithms.
+	 *
+	 * Also provided are some algorithms for getting "offsets", which correspond to the index associated with a frame if you were to
+	 * flatten this data-structure into one large array of frames by concatenating all ranges together into one large block.
+	 *
+	 * This data-structure is related to the FrameSet data-structure, which essentially acts the same way but stores individual frames instead of
+	 * frame ranges.
+	 */
+	struct LEARNING_API FFrameRangeSet
+	{
+	public:
+
+		/** Check if the FrameRangeSet is well-formed (i.e. correctly sorted without duplicate entries) */
+		void Check() const;
+
+		/**
+		 * Adds the given ranges associated with the given sequence to the set. Assumes this sequence (and no sequences with a larger index)
+		 * are already added to the set.
+		 */
+		void AddEntry(
+			const int32 InSequence,
+			const TLearningArrayView<1, const int32> InStarts,
+			const TLearningArrayView<1, const int32> InLengths);
+
+		/** True if the FrameRangeSet is Empty, otherwise false */
+		bool IsEmpty() const;
+
+		/** Empties the FrameRangeSet */
+		void Empty();
+
+		/** Gets the number of entries in the FrameRangeSet */
+		int32 GetEntryNum() const;
+
+		/** Returns an array of sequences for each entry */
+		TLearningArrayView<1, const int32> GetEntrySequences() const;
+
+		/** Returns an array for the number of ranges in each entry */
+		TLearningArrayView<1, const int32> GetEntryRangeNums() const;
+
+		/** Gets the sequence associated with a given entry */
+		int32 GetEntrySequence(const int32 EntryIdx) const;
+
+		/** Gets the number of ranges for a given entry */
+		int32 GetEntryRangeNum(const int32 EntryIdx) const;
+
+		/** Computes the total number of frames given by all ranges in an entry */
+		int32 GetEntryTotalFrameNum(const int32 EntryIdx) const;
+
+		/** Computes the flat offset of a particular entry if you were to flatten this data structure */
+		int32 GetEntryFrameOffset(const int32 EntryIdx) const;
+
+		/** Gets all the range starts associated with a given entry */
+		TLearningArrayView<1, const int32> GetEntryRangeStarts(const int32 EntryIdx) const;
+
+		/** Gets all the range lengths associated with a given entry */
+		TLearningArrayView<1, const int32> GetEntryRangeLengths(const int32 EntryIdx) const;
+		
+		/** Gets the range start associated with a given entry and range index */
+		int32 GetEntryRangeStart(const int32 EntryIdx, const int32 RangeIdx) const;
+
+		/** Gets the range length associated with a given entry and range index */
+		int32 GetEntryRangeLength(const int32 EntryIdx, const int32 RangeIdx) const;
+
+		/** Computes the flat offset associated with a given entry and range index */
+		int32 GetEntryRangeFrameOffset(const int32 EntryIdx, const int32 RangeIdx) const;
+
+		/** Gets the range start time associated with a given entry and range index */
+		float GetEntryRangeStartTime(const int32 EntryIdx, const int32 RangeIdx, const float FrameDeltaTime) const;
+
+		/** Gets the range end time associated with a given entry and range index */
+		float GetEntryRangeEndTime(const int32 EntryIdx, const int32 RangeIdx, const float FrameDeltaTime) const;
+
+		/** Gets the range duration associated with a given entry and range index */
+		float GetEntryRangeDuration(const int32 EntryIdx, const int32 RangeIdx, const float FrameDeltaTime) const;
+
+		/** Computes the total number of ranges across all entries */
+		int32 GetTotalRangeNum() const;
+
+		/** Gets an array of the range starts for all ranges */
+		TLearningArrayView<1, const int32> GetAllRangeStarts() const;
+
+		/** Gets an array of the range lengths for all ranges */
+		TLearningArrayView<1, const int32> GetAllRangeLengths() const;
+
+		/** Computes the total number of frames across all entries and ranges */
+		int32 GetTotalFrameNum() const;
+
+		/** Checks if this FrameRangeSet contains a given sequence */
+		bool ContainsSequence(const int32 Sequence) const;
+
+		/** Checks if this FrameRangeSet contains a given sequence and frame in that sequence */
+		bool Contains(const int32 Sequence, const int32 Frame) const;
+
+		/** Checks if this FrameRangeSet contains a given sequence and time in that sequence */
+		bool ContainsTime(const int32 Sequence, const float Time, const float FrameDeltaTime) const;
+
+		/** Finds the entry index associated with a given sequence */
+		int32 FindSequenceEntry(const int32 Sequence) const;
+
+		/** Finds the entry index and range index associated with a given sequence and frame in that sequence */
+		bool Find(int32& OutEntryIdx, int32& OutRangeIdx, int32& OutRangeFrame, const int32 Sequence, const int32 Frame) const;
+
+		/** Finds the entry index and range index associated with a given sequence and time in that sequence */
+		bool FindTime(int32& OutEntryIdx, int32& OutRangeIdx, float& OutRangeTime, const int32 Sequence, const float Time, const float FrameDeltaTime) const;
+
+		/** Finds entry index, frame index, and frame in that range, associated with some flat frame offset  */
+		bool FindFrameOffset(int32& OutEntryIdx, int32& OutRangeIdx, int32& OutRangeFrame, const int32 FrameOffset) const;
+
+	public:
+
+		/** Array of sequences associated with each entry */
+		TLearningArray<1, int32> EntrySequences;
+
+		/** Array of offsets into the RangeStarts and RangeLengths arrays associated with each entry */
+		TLearningArray<1, int32> EntryRangeOffsets;
+
+		/** Array of the number of ranges associated with each entry */
+		TLearningArray<1, int32> EntryRangeNums;
+
+		/** Large array of all range starts for all entries, indexed using offsets from EntryRangeOffsets */
+		TLearningArray<1, int32> RangeStarts;
+
+		/** Large array of all range lengths for all entries, indexed using offsets from EntryRangeOffsets */
+		TLearningArray<1, int32> RangeLengths;
+	};
+
+	namespace FrameRangeSet
+	{
+		/** Checks if two frame range sets are equal */
+		LEARNING_API bool Equal(const FFrameRangeSet& Lhs, const FFrameRangeSet& Rhs);
+
+		/** Computes the intersection of a frame set and frame range set */
+		LEARNING_API void Intersection(FFrameSet& OutFrameSet, const FFrameSet& FrameSet, const FFrameRangeSet& FrameRangeSet);
+
+		/** Computes the difference of a frame set and frame range set */
+		LEARNING_API void Difference(FFrameSet& OutFrameSet, const FFrameSet& FrameSet, const FFrameRangeSet& FrameRangeSet);
+
+		/** Computes the union of two frame range sets */
+		LEARNING_API void Union(FFrameRangeSet& Out, const FFrameRangeSet& Lhs, const FFrameRangeSet& Rhs);
+
+		/** Computes the intersection of two frame range sets */
+		LEARNING_API void Intersection(FFrameRangeSet& Out, const FFrameRangeSet& Lhs, const FFrameRangeSet& Rhs);
+
+		/** Computes the difference of two frame range sets */
+		LEARNING_API void Difference(FFrameRangeSet& Out, const FFrameRangeSet& Lhs, const FFrameRangeSet& Rhs);
+
+		/**
+		 * Computes the intersection of two frame range sets while recording the flat frame offsets in the lhs and rhs, associated with each range
+		 * that gets added to the output. The OutLhsFrameOffsets and OutRhsFrameOffsets outputs should be pre-allocated to at least the maximum of the
+		 * number of ranges in Lhs and Rhs. Returns the number of ranges in the added to the output frame range set.
+		 */
+		LEARNING_API int32 IntersectionWithFrameOffsets(
+			FFrameRangeSet& Out, 
+			TLearningArrayView<1, int32> OutLhsFrameOffsets,
+			TLearningArrayView<1, int32> OutRhsFrameOffsets,
+			const FFrameRangeSet& Lhs,
+			const FFrameRangeSet& Rhs);
+
+		/** Creates a new frame range set with the left-hand-side (i.e. the start) of all ranges trimmed by the given number of frames */
+		LEARNING_API void TrimLeft(FFrameRangeSet& Out, const FFrameRangeSet& FrameRangeSet, const int32 TrimFrameNum);
+
+		/** Creates a new frame range set with the right-hand-side (i.e. the end) of all ranges trimmed by the given number of frames */
+		LEARNING_API void TrimRight(FFrameRangeSet& Out, const FFrameRangeSet& FrameRangeSet, const int32 TrimFrameNum);
+
+		/** Creates a new frame range set with both the start and the end of all ranges trimmed by the given number of frames */
+		LEARNING_API void TrimLeftAndRight(FFrameRangeSet& Out, const FFrameRangeSet& FrameRangeSet, const int32 TrimFrameNum);
+
+		/** Create a frame set from all of the range starts of a given frame range set */
+		LEARNING_API void MakeFrameSetFromRangeStarts(FFrameSet& OutFrameSet, const FFrameRangeSet& FrameRangeSet);
+
+		/** Create a frame set from all of the range ends of a given frame range set */
+		LEARNING_API void MakeFrameSetFromRangeEnds(FFrameSet& OutFrameSet, const FFrameRangeSet& FrameRangeSet);
+
+		/** Trim a frame range set to just the period of time before the first occurrence of a frame from a given frame set */
+		LEARNING_API void RangesBeforeFrameSet(FFrameRangeSet& OutFrameRangeSet, const FFrameRangeSet& FrameRangeSet, const FFrameSet& FrameSet);
+
+		/** Trim a frame range set to just the period of time after the first occurrence of a frame from a given frame set */
+		LEARNING_API void RangesAfterFrameSet(FFrameRangeSet& OutFrameRangeSet, const FFrameRangeSet& FrameRangeSet, const FFrameSet& FrameSet);
+
+		/** Computes the entry indices associated with every range in the set */
+		LEARNING_API void AllRangeEntries(
+			TLearningArrayView<1, int32> OutRangeEntries,
+			const FFrameRangeSet& FrameRangeSet);
+
+		/** Computes the range indices associated with every range in the set */
+		LEARNING_API void AllRangeIndices(
+			TLearningArrayView<1, int32> OutRangeIndices,
+			const FFrameRangeSet& FrameRangeSet);
+
+		/** Computes the sequences associated with every range in the set */
+		LEARNING_API void AllRangeSequences(
+			TLearningArrayView<1, int32> OutRangeSequences,
+			const FFrameRangeSet& FrameRangeSet);
+
+		/** Computes the flat frame offsets associated with every range in the set */
+		LEARNING_API void AllRangeFrameOffsets(
+			TLearningArrayView<1, int32> OutRangeFrameOffsets,
+			const FFrameRangeSet& FrameRangeSet);
+
+		/** Computes the start times associated with every range in the set */
+		LEARNING_API void AllRangeStartTimes(
+			TLearningArrayView<1, float> OutRangeStartTimes,
+			const FFrameRangeSet& FrameRangeSet,
+			const float FrameDeltaTime);
+
+		/** Computes the end times associated with every range in the set */
+		LEARNING_API void AllRangeEndTimes(
+			TLearningArrayView<1, float> OutRangeEndTimes,
+			const FFrameRangeSet& FrameRangeSet,
+			const float FrameDeltaTime);
+
+		/** Computes the durations associated with every range in the set */
+		LEARNING_API void AllRangeDurations(
+			TLearningArrayView<1, float> OutRangeDurations,
+			const FFrameRangeSet& FrameRangeSet,
+			const float FrameDeltaTime);
+
+		/** Iterates over every range in the set and calls the provided callback */
+		LEARNING_API void ForEachRange(
+			const FFrameRangeSet& FrameRangeSet,
+			const TFunctionRef<void(
+				const int32 TotalRangeIdx,
+				const int32 EntryIdx,
+				const int32 RangeIdx,
+				const int32 FrameOffset)> Body);
+
+		/** Iterates over every range in the set in parallel and calls the provided callback */
+		LEARNING_API void ParallelForEachRange(
+			const FFrameRangeSet& FrameRangeSet,
+			const TFunctionRef<void(
+				const int32 TotalRangeIdx,
+				const int32 EntryIdx,
+				const int32 RangeIdx,
+				const int32 FrameOffset)> Body);
+	}
+}
\ No newline at end of file
diff --git a/Engine/Plugins/Experimental/LearningAgents/Source/Learning/Public/LearningFrameSet.h b/Engine/Plugins/Experimental/LearningAgents/Source/Learning/Public/LearningFrameSet.h
new file mode 100644
--- /dev/null
+++ b/Engine/Plugins/Experimental/LearningAgents/Source/Learning/Public/LearningFrameSet.h
@@ -0,0 +1,170 @@
+// Copyright Epic Games, Inc. All Rights Reserved.
+
+#pragma once
+
+#include "LearningArray.h"
+
+namespace UE::Learning
+{
+	/**
+	 * A FrameSet represents a set of frames within a set of sequences. This can be useful for encoding things like a collection of single-frame 
+	 * events.
+	 * 
+	 * This data structure stores that information in a way that makes performing operations such as union, intersection, difference etc, efficient. 
+	 * 
+	 * The way it works is by storing a sorted array of "entries", where each entry has a corresponding sequence, the number of frames in the set for 
+	 * that sequence, and an offset into a large array of frames. The sub-ranges of the large array of frames corresponding to each entry is also 
+	 * sorted.
+	 * 
+	 * Having the entries sorted by sequence, and the frames sorted, allows for efficient set operations via tape-merge algorithms.
+	 * 
+	 * Also provided are some algorithms for getting "offsets", which correspond to the index associated with a particular thing if you were to
+	 * flatten this data-structure into one large array of frames.
+	 * 
+	 * This data-structure is related to the FrameRangeSet data-structure, which essentially acts the same way but stores frame ranges instead of 
+	 * individual frames.
+	 */
+	struct LEARNING_API FFrameSet
+	{
+	public:
+
+		/** Check if the FrameSet is well-formed (i.e. correctly sorted without duplicate entries) */
+		void Check() const;
+
+		/**
+		 * Adds the given frames associated with the given sequence to the set. Assumes this sequence (and no sequences with a larger index)
+		 * are already added to the set.
+		 */
+		void AddEntry(
+			const int32 InSequence,
+			const TLearningArrayView<1, const int32> InFrames);
+
+		/** True if the FrameSet is Empty, otherwise false */
+		bool IsEmpty() const;
+
+		/** Empties the FrameSet */
+		void Empty();
+
+		/** Gets the number of entries in the FrameSet */
+		int32 GetEntryNum() const;
+
+		/** Returns and array of sequences corresponding to each entry */
+		TLearningArrayView<1, const int32> GetEntrySequences() const;
+
+		/** Returns an array containing the number of frames corresponding in each entry */
+		TLearningArrayView<1, const int32> GetEntryFrameNums() const;
+
+		/** Gets the sequence associated with a given entry */
+		int32 GetEntrySequence(const int32 EntryIdx) const;
+
+		/** Gets the number of frames for a given entry */
+		int32 GetEntryFrameNum(const int32 EntryIdx) const;
+
+		/** Gets the array of frames associated with a given entry */
+		TLearningArrayView<1, const int32> GetEntryFrames(const int32 EntryIdx) const;
+
+		/** Gets the frame number for an entry and entry frame index */
+		int32 GetEntryFrame(const int32 EntryIdx, const int32 FrameIdx) const;
+
+		/** Gets the frame time associated with an entry and entry frame index */
+		float GetEntryFrameTime(const int32 EntryIdx, const int32 FrameIdx, const float FrameDeltaTime) const;
+
+		/** Gets the flat offset associated with a particular entry if you were to flatten this data structure */
+		int32 GetEntryFrameOffset(const int32 EntryIdx) const;
+
+		/** Gets the total number of frames if you were to flatten this data structure */
+		int32 GetTotalFrameNum() const;
+
+		/** Checks if this FrameSet contains a given sequence */
+		bool ContainsSequence(const int32 Sequence) const;
+
+		/** Checks if this FrameSet contains a given sequence and frame in that sequence */
+		bool Contains(const int32 Sequence, const int32 Frame) const;
+
+		/** Finds the entry index associated with a given sequence */
+		int32 FindSequenceEntry(const int32 Sequence) const;
+
+		/** Finds the entry index and frame index associated with a given sequence and frame in that sequence */
+		bool Find(int32& OutEntryIdx, int32& OutFrameIdx, const int32 Sequence, const int32 Frame) const;
+
+		/** Finds the nearest entry index, frame index, and the frame difference, for some given sequence and frame in that sequence */
+		bool FindNearest(int32& OutEntryIdx, int32& OutFrameIdx, int32& OutFrameDifference, const int32 Sequence, const int32 Frame) const;
+
+		/** 
+		 * Finds the nearest entry index, frame index, and the frame difference, for some given sequence and frame in that sequence, 
+		 * limited to a given range
+		 */
+		bool FindNearestInRange(
+			int32& OutEntryIdx, int32& OutFrameIdx, int32& OutFrameDifference, 
+			const int32 Sequence, const int32 Frame, const int32 RangeStart, const int32 RangeLength) const;
+
+		/** Finds entry index and frame index associated with some flat offset  */
+		bool FindFrameOffset(int32& OutEntryIdx, int32& OutFrameIdx, const int32 FrameOffset) const;
+
+	public:
+
+		/** Array of sequences associated with each entry */
+		TLearningArray<1, int32> EntrySequences;
+
+		/** Array of offsets into the Frames array associated with each entry */
+		TLearningArray<1, int32> EntryFrameOffsets;
+
+		/** Array of the number of frames associated with each entry */
+		TLearningArray<1, int32> EntryFrameNums;
+
+		/** Large array of all frames for all entries, indexed using offsets from EntryFrameOffsets */
+		TLearningArray<1, int32> Frames;
+	};
+
+	namespace FrameSet
+	{
+		/** Checks if two frame sets are equal */
+		LEARNING_API bool Equal(const FFrameSet& Lhs, const FFrameSet& Rhs);
+
+		/** Computes the union of two frame sets */
+		LEARNING_API void Union(FFrameSet& Out, const FFrameSet& Lhs, const FFrameSet& Rhs);
+
+		/** Computes the intersection of two frame sets */
+		LEARNING_API void Intersection(FFrameSet& Out, const FFrameSet& Lhs, const FFrameSet& Rhs);
+
+		/** Computes the difference of two frame sets */
+		LEARNING_API void Difference(FFrameSet& Out, const FFrameSet& Lhs, const FFrameSet& Rhs);
+
+		/** Computes the entry indices associated with every frame in the set */
+		LEARNING_API void AllFrameEntries(
+			TLearningArrayView<1, int32> OutFrameEntries,
+			const FFrameSet& FrameSet);
+
+		/** Computes the frame indices associated with every frame in the set */
+		LEARNING_API void AllFrameIndices(
+			TLearningArrayView<1, int32> OutFrameIndices,
+			const FFrameSet& FrameSet);
+
+		/** Computes the sequences associated with every frame in the set */
+		LEARNING_API void AllFrameSequences(
+			TLearningArrayView<1, int32> OutFrameSequences,
+			const FFrameSet& FrameSet);
+
+		/** Computes the times associated with every frame in the set */
+		LEARNING_API void AllFrameTimes(
+			TLearningArrayView<1, float> OutFrameTimes,
+			const FFrameSet& FrameSet,
+			const float FrameDeltaTime);
+
+		/** Iterates over every frame in the set and calls the provided callback */
+		LEARNING_API void ForEachFrame(
+			const FFrameSet& FrameSet,
+			const TFunctionRef<void(
+				const int32 TotalFrameIdx,
+				const int32 EntryIdx,
+				const int32 FrameIdx)> Body);
+
+		/** Iterates over every frame in the set in parallel and calls the provided callback */
+		LEARNING_API void ParallelForEachFrame(
+			const FFrameSet& FrameSet,
+			const TFunctionRef<void(
+				const int32 TotalFrameIdx,
+				const int32 EntryIdx,
+				const int32 RangeIdx)> Body);
+	}
+}
\ No newline at end of file
diff --git a/Engine/Plugins/Experimental/LearningAgents/Source/Learning/Public/LearningRandom.h b/Engine/Plugins/Experimental/LearningAgents/Source/Learning/Public/LearningRandom.h
--- a/Engine/Plugins/Experimental/LearningAgents/Source/Learning/Public/LearningRandom.h
+++ b/Engine/Plugins/Experimental/LearningAgents/Source/Learning/Public/LearningRandom.h
@@ -7,6 +7,12 @@
 #include "Math/Vector.h"
 #include "Math/Quat.h"
 
+namespace UE::Learning
+{
+	struct FFrameSet;
+	struct FFrameRangeSet;
+}
+
 /**
 * This is a simple hashing-based pseudo-random-number-generator
 * that separates out the generation of random numbers from the
@@ -109,14 +115,6 @@ namespace UE::Learning::Random
 		const FVector Axis0 = FVector::ForwardVector,
 		const FVector Axis1 = FVector::RightVector);
 
-	LEARNING_API void PlanarDirectionVelocity(
-		FVector& OutDirection,
-		FVector& OutVelocity,
-		const uint32 State,
-		const float VelocityScale = 100.0f,
-		const FVector Axis0 = FVector::ForwardVector,
-		const FVector Axis1 = FVector::RightVector);
-
 	LEARNING_API FQuat Rotation(const uint32 Input);
 
 	LEARNING_API void IntArray(
@@ -208,6 +206,19 @@ namespace UE::Learning::Random
 		const TLearningArrayView<1, const bool> Masked,
 		const float Scale);
 
+	LEARNING_API void FrameInFrameSet(
+		int32& OutEntryIdx,
+		int32& OutFrameIdx,
+		const FFrameSet& FrameSet,
+		const uint32 State);
+
+	LEARNING_API void FrameInFrameRangeSet(
+		int32& OutEntryIdx,
+		int32& OutRangeIdx,
+		int32& OutRangeFrame,
+		const FFrameRangeSet& FrameRangeSet,
+		const uint32 State);
+
 	/////////////////////////////////////////////////////////
 	//
 	// Here we provide a more normal interface which updates
@@ -359,4 +370,17 @@ namespace UE::Learning::Random
 		const TLearningArrayView<1, const float> Logits,
 		const TLearningArrayView<1, const bool> Masked,
 		const float Scale);
+
+	LEARNING_API void SampleFrameInFrameSet(
+		int32& OutEntryIdx,
+		int32& OutFrameIdx,
+		uint32& State,
+		const FFrameSet& FrameSet);
+
+	LEARNING_API void SampleFrameInFrameRangeSet(
+		int32& OutEntryIdx,
+		int32& OutRangeIdx,
+		int32& OutRangeFrame,
+		uint32& State,
+		const FFrameRangeSet& FrameRangeSet);
 }
\ No newline at end of file
