[LearningAgents] Trivial - Fixes an issue where different gyms templates are spaced one extra apart

diff --git a/Engine/Plugins/Experimental/LearningAgents/Source/LearningAgentsTraining/Private/LearningAgentsGymsManager.cpp b/Engine/Plugins/Experimental/LearningAgents/Source/LearningAgentsTraining/Private/LearningAgentsGymsManager.cpp
--- a/Engine/Plugins/Experimental/LearningAgents/Source/LearningAgentsTraining/Private/LearningAgentsGymsManager.cpp
+++ b/Engine/Plugins/Experimental/LearningAgents/Source/LearningAgentsTraining/Private/LearningAgentsGymsManager.cpp
@@ -65,7 +65,7 @@ void ALearningAgentsGymsManager::SpawnGyms()
 				CurrentSpawnCount++;
 			}
 		}
-		TemplateSpawnLocation.X += (GymsSpacing + (MaxBounds.X - MinBounds.X)) * (GridSize);
+		TemplateSpawnLocation.X += (GymsSpacing + (MaxBounds.X - MinBounds.X)) * (GridSize - 1);
 	}
 	
 	if (TObjectPtr<UNavigationSystemV1> NavigationSystem = FNavigationSystem::GetCurrent<UNavigationSystemV1>(GetWorld()))
