Learning Agents: Fix hostname option

#rnx

diff --git a/Engine/Plugins/Experimental/LearningAgents/Source/LearningTraining/Private/LearningExternalTrainer.cpp b/Engine/Plugins/Experimental/LearningAgents/Source/LearningTraining/Private/LearningExternalTrainer.cpp
--- a/Engine/Plugins/Experimental/LearningAgents/Source/LearningTraining/Private/LearningExternalTrainer.cpp
+++ b/Engine/Plugins/Experimental/LearningAgents/Source/LearningTraining/Private/LearningExternalTrainer.cpp
@@ -681,7 +681,6 @@ namespace UE::Learning
 		ISocketSubsystem* SocketSubsystem = ISocketSubsystem::Get(PLATFORM_SOCKETSUBSYSTEM);
 		checkf(SocketSubsystem, TEXT("Could not get socket subsystem"));
 
-		bool bIsValid = false;
 		TSharedRef<FInternetAddr> Address = SocketSubsystem->CreateInternetAddr();
 		Address->SetPort(Port);
 
@@ -719,14 +718,15 @@ namespace UE::Learning
 		}
 		else
 		{
+			bool bIsValid = false;
 			Address->SetIp(IpAddressOrHostname, bIsValid);
-		}
 
-		if (!bIsValid)
-		{
-			UE_LOG(LogLearning, Error, TEXT("Invalid Ip Address \"%s\"..."), IpAddressOrHostname);
-			OutResponse = ETrainerResponse::Unexpected;
-			return;
+			if (!bIsValid)
+			{
+				UE_LOG(LogLearning, Error, TEXT("Invalid Ip Address \"%s\"..."), IpAddressOrHostname);
+				OutResponse = ETrainerResponse::Unexpected;
+				return;
+			}
 		}
 
 		Socket = FTcpSocketBuilder(TEXT("LearningTrainerSocket")).AsBlocking().Build();
