[LearningAgents] EntitiesManager Stochastic Episode Spawn

diff --git a/Engine/Plugins/Experimental/LearningAgents/Source/LearningAgentsTraining/Private/LearningAgentsEntitiesManagerComponent.cpp b/Engine/Plugins/Experimental/LearningAgents/Source/LearningAgentsTraining/Private/LearningAgentsEntitiesManagerComponent.cpp
--- a/Engine/Plugins/Experimental/LearningAgents/Source/LearningAgentsTraining/Private/LearningAgentsEntitiesManagerComponent.cpp
+++ b/Engine/Plugins/Experimental/LearningAgents/Source/LearningAgentsTraining/Private/LearningAgentsEntitiesManagerComponent.cpp
@@ -18,9 +18,14 @@ void ULearningAgentsEntitiesManagerComponent::InitializeLearningComponent()
 
 		if (Gym->GetLocalRole() == ROLE_Authority)
 		{
-			for (FEntityInfo& EntityInfo : Entities)
+			for (FLearningAgentsEntityInfo& EntityInfo : Entities)
 			{
-				SpawnEntitiesAtRandomLocations(EntityInfo.EntityClass, EntityInfo.EntitySpawnZOffset, EntityInfo.EpisodeStartEntitySpawnCount);
+				int32 SpawnCount = 0;
+				if (ensure(Gym->GetRandomStream()))
+				{
+					SpawnCount = Gym->GetRandomStream()->RandRange(EntityInfo.EpisodeEntitySpawnCountMin, EntityInfo.EpisodeEntitySpawnCountMax);
+				}
+				SpawnEntitiesAtRandomLocations(EntityInfo.EntityClass, EntityInfo.EntitySpawnZOffset, SpawnCount);
 			}
 		}
 
@@ -47,9 +52,14 @@ void ULearningAgentsEntitiesManagerComponent::ResetLearningComponent()
 		}
 		for (TPair<FName, FSpawnedEntitiesInfo>& UniqueTypeEntities : EntitiesPool)
 		{
+			int32 SpawnCount = 0;
+			if (ensure(Gym->GetRandomStream()))
+			{
+				SpawnCount = Gym->GetRandomStream()->RandRange(UniqueTypeEntities.Value.EntityInfo.EpisodeEntitySpawnCountMin, UniqueTypeEntities.Value.EntityInfo.EpisodeEntitySpawnCountMax);
+			}
 			for (int32 i = 0; i < UniqueTypeEntities.Value.SpawnedEntities.Num(); ++i)
 			{
-				if (i < UniqueTypeEntities.Value.EntityInfo.EpisodeStartEntitySpawnCount)
+				if (i < SpawnCount)
 				{
 					// Make sure transform is applied before reset to wipe all data (ex: distance traveled)
 					if (TObjectPtr<AActor> ActorEntity = Cast<AActor>(UniqueTypeEntities.Value.SpawnedEntities[i].GetObject()))
@@ -61,6 +71,8 @@ void ULearningAgentsEntitiesManagerComponent::ResetLearningComponent()
 
 					ILearningAgentsEntityInterface::Execute_ResetEntity(UniqueTypeEntities.Value.SpawnedEntities[i].GetObject(), Gym);
 					ILearningAgentsEntityInterface::Execute_EnableEntity(UniqueTypeEntities.Value.SpawnedEntities[i].GetObject());
+
+					SpawnEntitiesAtRandomLocations(UniqueTypeEntities.Value.EntityInfo.EntityClass, UniqueTypeEntities.Value.EntityInfo.EntitySpawnZOffset, SpawnCount - i - 1);
 				}
 				else
 				{
@@ -165,10 +177,9 @@ TArray<TScriptInterface<ILearningAgentsEntityInterface>> ULearningAgentsEntities
 	else
 	{
 		FSpawnedEntitiesInfo NewSpawnEntitiesInfo;
-		FEntityInfo NewEntityInfo;
+		FLearningAgentsEntityInfo NewEntityInfo;
 		NewEntityInfo.EntityClass = EntityClass;
 		NewEntityInfo.EntitySpawnZOffset = EntitySpawnZOffset;
-		NewEntityInfo.EpisodeStartEntitySpawnCount = 0;
 		NewSpawnEntitiesInfo.EntityInfo = NewEntityInfo;
 
 		for (int32 Index = 0; Index < SpawnCount; ++Index)
@@ -189,13 +200,14 @@ TArray<TScriptInterface<ILearningAgentsEntityInterface>> ULearningAgentsEntities
 void ULearningAgentsEntitiesManagerComponent::BeginPlay()
 {
 	Super::BeginPlay();
-	CheckEntityClasses();
-
-	for (FEntityInfo& EntityInfo : Entities)
+	if (CheckEntityClasses())
 	{
-		FSpawnedEntitiesInfo SpawnedEntitiesInfo;
-		SpawnedEntitiesInfo.EntityInfo = EntityInfo;
-		EntitiesPool.Add(EntityInfo.EntityClass->GetFName(), SpawnedEntitiesInfo);
+		for (FLearningAgentsEntityInfo& EntityInfo : Entities)
+		{
+			FSpawnedEntitiesInfo SpawnedEntitiesInfo;
+			SpawnedEntitiesInfo.EntityInfo = EntityInfo;
+			EntitiesPool.Add(EntityInfo.EntityClass->GetFName(), SpawnedEntitiesInfo);
+		}
 	}
 }
 
@@ -234,7 +246,7 @@ bool ULearningAgentsEntitiesManagerComponent::CheckEntityClasses() const
 		return false;
 	}
 
-	for (const FEntityInfo& EntityInfo : Entities)
+	for (const FLearningAgentsEntityInfo& EntityInfo : Entities)
 	{
 		if (!EntityInfo.EntityClass)
 		{
diff --git a/Engine/Plugins/Experimental/LearningAgents/Source/LearningAgentsTraining/Public/LearningAgentsEntitiesManagerComponent.h b/Engine/Plugins/Experimental/LearningAgents/Source/LearningAgentsTraining/Public/LearningAgentsEntitiesManagerComponent.h
--- a/Engine/Plugins/Experimental/LearningAgents/Source/LearningAgentsTraining/Public/LearningAgentsEntitiesManagerComponent.h
+++ b/Engine/Plugins/Experimental/LearningAgents/Source/LearningAgentsTraining/Public/LearningAgentsEntitiesManagerComponent.h
@@ -10,10 +10,10 @@
 class ILearningAgentsEntityInterface;
 
 /**
- * FEntityInfo holds spawn information of a single entity type.
+ * FLearningAgentsEntityInfo holds spawn information of a single entity type.
  */
 USTRUCT(BlueprintType)
-struct FEntityInfo
+struct FLearningAgentsEntityInfo
 {
 	GENERATED_BODY()
 
@@ -21,13 +21,21 @@ struct FEntityInfo
 	UPROPERTY(EditDefaultsOnly, BlueprintReadWrite, Category = "LearningAgents")
 	TSubclassOf<AActor> EntityClass;
 
+	/** Specify the spawner data class. */
+	UPROPERTY(EditDefaultsOnly, BlueprintReadWrite, Category = "LearningAgents")
+	TSubclassOf<UObject> SpawnerDataClass;
+
 	/** Specify the Z offset to apply to spawn locations. */
 	UPROPERTY(EditDefaultsOnly, BlueprintReadWrite, Category = "LearningAgents")
 	float EntitySpawnZOffset = 0.0f;
 
-	/** The number of entities to spawn at the start of an episode. */
-	UPROPERTY(EditDefaultsOnly, BlueprintReadWrite, Category = "LearningAgents")
-	int32 EpisodeStartEntitySpawnCount = 0;
+	/** The min number of entities to spawn at the start of an episode. */
+	UPROPERTY(EditDefaultsOnly, BlueprintReadWrite, Category = "LearningAgents", meta = (ClampMin = "0.0", UIMin = "0.0"))
+	int32 EpisodeEntitySpawnCountMin = 0;
+
+	/** The max number of entities to spawn at the start of an episode. */
+	UPROPERTY(EditDefaultsOnly, BlueprintReadWrite, Category = "LearningAgents", meta = (ClampMin = "0.0", UIMin = "0.0"))
+	int32 EpisodeEntitySpawnCountMax = 0;
 };
 
 /**
@@ -39,7 +47,7 @@ struct FSpawnedEntitiesInfo
 	GENERATED_BODY()
 
 	UPROPERTY(VisibleAnywhere, Category = "LearningAgents")
-	FEntityInfo EntityInfo;
+	FLearningAgentsEntityInfo EntityInfo;
 
 	UPROPERTY(VisibleAnywhere, Category = "LearningAgents")
 	TArray<TScriptInterface<ILearningAgentsEntityInterface>> SpawnedEntities;
@@ -56,7 +64,7 @@ class LEARNINGAGENTSTRAINING_API ULearningAgentsEntitiesManagerComponent : publi
 public:
 	/** Specify the entity types to spawn. */
 	UPROPERTY(EditDefaultsOnly, Category = "LearningAgents")
-	TArray<FEntityInfo> Entities;
+	TArray<FLearningAgentsEntityInfo> Entities;
 
 	/** Initializes the entities that the component manages. */
 	virtual void InitializeLearningComponent();
