Learning Agents: Expose policy parallelization to BP

diff --git a/Engine/Plugins/Experimental/LearningAgents/Source/LearningAgents/Private/LearningAgentsPolicy.cpp b/Engine/Plugins/Experimental/LearningAgents/Source/LearningAgents/Private/LearningAgentsPolicy.cpp
--- a/Engine/Plugins/Experimental/LearningAgents/Source/LearningAgents/Private/LearningAgentsPolicy.cpp
+++ b/Engine/Plugins/Experimental/LearningAgents/Source/LearningAgents/Private/LearningAgentsPolicy.cpp
@@ -361,11 +361,13 @@ void ULearningAgentsPolicy::SetupPolicy(
 	}
 
 	// Create Encoder / Policy / Decoder Objects
+	UE::Learning::FNeuralNetworkInferenceSettings InferenceSettings;
+	InferenceSettings.bParallelEvaluation = PolicySettings.bUseParallelEvaluation;
 
 	EncoderObject = MakeShared<UE::Learning::FNeuralNetworkFunction>(
 		Manager->GetMaxAgentNum(),
 		EncoderNetwork->NeuralNetworkData->GetNetwork(),
-		UE::Learning::FNeuralNetworkInferenceSettings());
+		InferenceSettings);
 
 	PolicyObject = MakeShared<UE::Learning::FNeuralNetworkPolicy>(
 		Manager->GetMaxAgentNum(),
@@ -373,12 +375,12 @@ void ULearningAgentsPolicy::SetupPolicy(
 		ActionEncodedVectorSize,
 		MemoryStateSize,
 		PolicyNetwork->NeuralNetworkData->GetNetwork(),
-		UE::Learning::FNeuralNetworkInferenceSettings());
+		InferenceSettings);
 
 	DecoderObject = MakeShared<UE::Learning::FNeuralNetworkFunction>(
 		Manager->GetMaxAgentNum(),
 		DecoderNetwork->NeuralNetworkData->GetNetwork(),
-		UE::Learning::FNeuralNetworkInferenceSettings());
+		InferenceSettings);
 
 	// State Variables
 
diff --git a/Engine/Plugins/Experimental/LearningAgents/Source/LearningAgents/Public/LearningAgentsPolicy.h b/Engine/Plugins/Experimental/LearningAgents/Source/LearningAgents/Public/LearningAgentsPolicy.h
--- a/Engine/Plugins/Experimental/LearningAgents/Source/LearningAgents/Public/LearningAgentsPolicy.h
+++ b/Engine/Plugins/Experimental/LearningAgents/Source/LearningAgents/Public/LearningAgentsPolicy.h
@@ -49,6 +49,10 @@ public:
 	/** Activation function to use on hidden layers of the policy network */
 	UPROPERTY(EditAnywhere, Category = "LearningAgents")
 	ELearningAgentsActivationFunction ActivationFunction = ELearningAgentsActivationFunction::ELU;
+
+	/** If true, policy can be evaluated on multiple threads to speed up large agent counts. */
+	UPROPERTY(EditAnywhere, Category = "LearningAgents")
+	bool bUseParallelEvaluation = true;
 };
 
 /** A policy that maps from observations to actions. */
