[LearningAgents] Changed AActor reference in FLearningAgentsEntityInfo to UObject reference

#rnx

diff --git a/Engine/Plugins/Experimental/LearningAgents/Source/LearningAgentsTraining/Private/LearningAgentsEntitiesManagerComponent.cpp b/Engine/Plugins/Experimental/LearningAgents/Source/LearningAgentsTraining/Private/LearningAgentsEntitiesManagerComponent.cpp
--- a/Engine/Plugins/Experimental/LearningAgents/Source/LearningAgentsTraining/Private/LearningAgentsEntitiesManagerComponent.cpp
+++ b/Engine/Plugins/Experimental/LearningAgents/Source/LearningAgentsTraining/Private/LearningAgentsEntitiesManagerComponent.cpp
@@ -28,16 +28,6 @@ void ULearningAgentsEntitiesManagerComponent::InitializeLearningComponent()
 				SpawnEntitiesAtRandomLocations(EntityInfo.EntityClass, EntityInfo.EntitySpawnZOffset, SpawnCount);
 			}
 		}
-
-		// This is in a separate loop because initialization may be dependent on when all entities are spawned 
-		for (TPair<FName, FSpawnedEntitiesInfo>& UniqueTypeEntities : EntitiesPool)
-		{
-			for (TScriptInterface<ILearningAgentsEntityInterface>& Entity : UniqueTypeEntities.Value.SpawnedEntities)
-			{
-				ILearningAgentsEntityInterface::Execute_InitializeEntity(Entity.GetObject(), Gym);
-			}
-		}
-
 	}
 }
 
@@ -72,7 +62,10 @@ void ULearningAgentsEntitiesManagerComponent::ResetLearningComponent()
 					ILearningAgentsEntityInterface::Execute_ResetEntity(UniqueTypeEntities.Value.SpawnedEntities[i].GetObject(), Gym);
 					ILearningAgentsEntityInterface::Execute_EnableEntity(UniqueTypeEntities.Value.SpawnedEntities[i].GetObject());
 
-					SpawnEntitiesAtRandomLocations(UniqueTypeEntities.Value.EntityInfo.EntityClass, UniqueTypeEntities.Value.EntityInfo.EntitySpawnZOffset, SpawnCount - i - 1);
+					if (int32 AdditionalSpawnCount = SpawnCount - i - 1)
+					{
+						SpawnEntitiesAtRandomLocations(UniqueTypeEntities.Value.EntityInfo.EntityClass, UniqueTypeEntities.Value.EntityInfo.EntitySpawnZOffset, SpawnCount - i - 1);
+					}
 				}
 				else
 				{
@@ -83,9 +76,9 @@ void ULearningAgentsEntitiesManagerComponent::ResetLearningComponent()
 	}
 }
 
-TScriptInterface<ILearningAgentsEntityInterface> ULearningAgentsEntitiesManagerComponent::SpawnEntitiesAtRandomLocations(TSubclassOf<AActor> EntityClass, float EntitySpawnZOffset, int32 SpawnCount)
+TScriptInterface<ILearningAgentsEntityInterface> ULearningAgentsEntitiesManagerComponent::SpawnEntitiesAtRandomLocations(TSubclassOf<UObject> EntityClass, float EntitySpawnZOffset, int32 SpawnCount)
 {
-	if (SpawnCount < 0)
+	if (SpawnCount <= 0)
 	{
 		UE_LOG(LogLearning, Warning, TEXT("Spawn Count must be greater than 0 in %s!"), *EntityClass->GetName());
 		return TScriptInterface<ILearningAgentsEntityInterface>();
@@ -106,7 +99,7 @@ TScriptInterface<ILearningAgentsEntityInterface> ULearningAgentsEntitiesManagerC
 	}
 }
 
-TScriptInterface<ILearningAgentsEntityInterface> ULearningAgentsEntitiesManagerComponent::SpawnEntityAtProjectedLocation(TSubclassOf<AActor> EntityClass, float EntitySpawnZOffset, const FTransform& InTransform)
+TScriptInterface<ILearningAgentsEntityInterface> ULearningAgentsEntitiesManagerComponent::SpawnEntityAtProjectedLocation(TSubclassOf<UObject> EntityClass, float EntitySpawnZOffset, const FTransform& InTransform)
 {
 	FTransform Transform = InTransform;
 	ProjectTransform(Transform);
@@ -123,12 +116,12 @@ TScriptInterface<ILearningAgentsEntityInterface> ULearningAgentsEntitiesManagerC
 	}
 }
 
-TArray<TScriptInterface<ILearningAgentsEntityInterface>> ULearningAgentsEntitiesManagerComponent::SpawnEntities(TSubclassOf<AActor> EntityClass, float EntitySpawnZOffset, int32 SpawnCount, const FTransform& InTransform)
+TArray<TScriptInterface<ILearningAgentsEntityInterface>> ULearningAgentsEntitiesManagerComponent::SpawnEntities(TSubclassOf<UObject> EntityClass, float EntitySpawnZOffset, int32 SpawnCount, const FTransform& InTransform)
 {
 	TArray<TScriptInterface<ILearningAgentsEntityInterface>> SpawnedEntities;
 	SpawnedEntities.SetNumUninitialized(SpawnCount);
 
-	if (SpawnCount < 0)
+	if (SpawnCount <= 0)
 	{
 		UE_LOG(LogLearning, Warning, TEXT("Spawn Count must be greater than 0 in %s!"), *EntityClass->GetName());
 		return SpawnedEntities;
@@ -164,13 +157,20 @@ TArray<TScriptInterface<ILearningAgentsEntityInterface>> ULearningAgentsEntities
 		// Allocate new object and add to pool because there's not enough in pool to pull from.
 		for (int32 Index = SuccessfulSpawnCount; SuccessfulSpawnCount < SpawnCount; ++SuccessfulSpawnCount)
 		{
-			TObjectPtr<AActor> Entity = GetWorld()->SpawnActor<AActor>(EntityClass, Transform.GetLocation(), Transform.GetRotation().Rotator(), SpawnParameters);
-			if (TObjectPtr<ALearningAgentsGymBase> Gym = Cast<ALearningAgentsGymBase>(GetOwner()))
+			if (EntityClass->IsChildOf(AActor::StaticClass()))
 			{
-				ILearningAgentsEntityInterface::Execute_InitializeEntity(Entity, Gym);
+				TObjectPtr<AActor> Entity = GetWorld()->SpawnActor<AActor>(EntityClass, Transform.GetLocation(), Transform.GetRotation().Rotator(), SpawnParameters);
+				if (TObjectPtr<ALearningAgentsGymBase> Gym = Cast<ALearningAgentsGymBase>(GetOwner()))
+				{
+					ILearningAgentsEntityInterface::Execute_InitializeEntity(Entity, Gym);
+				}
+				SpawnEntitiesInfo->SpawnedEntities.Add(Entity);
+				SpawnedEntities.Add(TScriptInterface<ILearningAgentsEntityInterface>(Entity));
+			}
+			else
+			{
+				UE_LOG(LogLearning, Warning, TEXT("Could not spawn entity from class %s because it's not an AActor!"), *EntityClass->GetName());
 			}
-			SpawnEntitiesInfo->SpawnedEntities.Add(Entity);
-			SpawnedEntities.Add(TScriptInterface<ILearningAgentsEntityInterface>(Entity));
 		}
 	}
 	// New entity spawned during an episode
@@ -184,13 +184,20 @@ TArray<TScriptInterface<ILearningAgentsEntityInterface>> ULearningAgentsEntities
 
 		for (int32 Index = 0; Index < SpawnCount; ++Index)
 		{
-			TObjectPtr<AActor> Entity = GetWorld()->SpawnActor<AActor>(EntityClass, Transform.GetLocation(), Transform.GetRotation().Rotator(), SpawnParameters);
-			if (TObjectPtr<ALearningAgentsGymBase> Gym = Cast<ALearningAgentsGymBase>(GetOwner()))
+			if (EntityClass->IsChildOf(AActor::StaticClass()))
+			{
+				TObjectPtr<AActor> Entity = GetWorld()->SpawnActor<AActor>(EntityClass, Transform.GetLocation(), Transform.GetRotation().Rotator(), SpawnParameters);
+				if (TObjectPtr<ALearningAgentsGymBase> Gym = Cast<ALearningAgentsGymBase>(GetOwner()))
+				{
+					ILearningAgentsEntityInterface::Execute_InitializeEntity(Entity, Gym);
+				}
+				NewSpawnEntitiesInfo.SpawnedEntities.Add(Entity);
+				SpawnedEntities.Add(TScriptInterface<ILearningAgentsEntityInterface>(Entity));
+			}
+			else
 			{
-				ILearningAgentsEntityInterface::Execute_InitializeEntity(Entity, Gym);
+				UE_LOG(LogLearning, Warning, TEXT("Could not spawn entity from class %s because it's not an AActor!"), *EntityClass->GetName());
 			}
-			NewSpawnEntitiesInfo.SpawnedEntities.Add(Entity);
-			SpawnedEntities.Add(TScriptInterface<ILearningAgentsEntityInterface>(Entity));
 		}
 		EntitiesPool.Add(EntityClass->GetFName(), NewSpawnEntitiesInfo);
 	}
diff --git a/Engine/Plugins/Experimental/LearningAgents/Source/LearningAgentsTraining/Public/LearningAgentsEntitiesManagerComponent.h b/Engine/Plugins/Experimental/LearningAgents/Source/LearningAgentsTraining/Public/LearningAgentsEntitiesManagerComponent.h
--- a/Engine/Plugins/Experimental/LearningAgents/Source/LearningAgentsTraining/Public/LearningAgentsEntitiesManagerComponent.h
+++ b/Engine/Plugins/Experimental/LearningAgents/Source/LearningAgentsTraining/Public/LearningAgentsEntitiesManagerComponent.h
@@ -2,6 +2,7 @@
 
 #pragma once
 
+#include <tuple>
 #include "CoreMinimal.h"
 #include "Components/ActorComponent.h"
 #include "LearningAgentsLearningComponentInterface.h"
@@ -19,11 +20,7 @@ struct FLearningAgentsEntityInfo
 
 	/** Specify the entity class to spawn. */
 	UPROPERTY(EditDefaultsOnly, BlueprintReadWrite, Category = "LearningAgents")
-	TSubclassOf<AActor> EntityClass;
-
-	/** Specify the spawner data class. */
-	UPROPERTY(EditDefaultsOnly, BlueprintReadWrite, Category = "LearningAgents")
-	TSubclassOf<UObject> SpawnerDataClass;
+	TSubclassOf<UObject> EntityClass;
 
 	/** Specify the Z offset to apply to spawn locations. */
 	UPROPERTY(EditDefaultsOnly, BlueprintReadWrite, Category = "LearningAgents")
@@ -74,15 +71,15 @@ public:
 
 	/** Spawns pooled entities at random locations. */
 	UFUNCTION(BlueprintCallable, Category = "LearningAgents")
-	TScriptInterface<ILearningAgentsEntityInterface> SpawnEntitiesAtRandomLocations(TSubclassOf<AActor> EntityClass, float EntitySpawnZOffset, int32 SpawnCount);
+	TScriptInterface<ILearningAgentsEntityInterface> SpawnEntitiesAtRandomLocations(TSubclassOf<UObject> EntityClass, float EntitySpawnZOffset, int32 SpawnCount);
 
 	/** Spawns a single pooled entity with a specified transform projected in the gym. */
 	UFUNCTION(BlueprintCallable, Category = "LearningAgents")
-	TScriptInterface<ILearningAgentsEntityInterface> SpawnEntityAtProjectedLocation(TSubclassOf<AActor> EntityClass, float EntitySpawnZOffset, const FTransform& Transform);
+	TScriptInterface<ILearningAgentsEntityInterface> SpawnEntityAtProjectedLocation(TSubclassOf<UObject> EntityClass, float EntitySpawnZOffset, const FTransform& Transform);
 
 	/** Spawns multiple pooled entities. */
 	UFUNCTION(BlueprintCallable, Category = "LearningAgents")
-	TArray<TScriptInterface<ILearningAgentsEntityInterface>> SpawnEntities(TSubclassOf<AActor> EntityClass, float EntitySpawnZOffset, int32 SpawnCount, const FTransform& Transform);
+	TArray<TScriptInterface<ILearningAgentsEntityInterface>> SpawnEntities(TSubclassOf<UObject> EntityClass, float EntitySpawnZOffset, int32 SpawnCount, const FTransform& Transform);
 
 protected:
 	virtual void BeginPlay() override;
