LearningAgents: Removed special macros for checks and profiling to better fit coding standard.

diff --git a/Engine/Plugins/Experimental/LearningAgents/Source/Learning/Private/LearningAction.cpp b/Engine/Plugins/Experimental/LearningAgents/Source/Learning/Private/LearningAction.cpp
--- a/Engine/Plugins/Experimental/LearningAgents/Source/Learning/Private/LearningAction.cpp
+++ b/Engine/Plugins/Experimental/LearningAgents/Source/Learning/Private/LearningAction.cpp
@@ -160,8 +160,8 @@ namespace UE::Learning::Action
 
 	FSchemaElement FSchema::CreateContinuous(const FSchemaContinuousParameters Parameters, const FName Tag)
 	{
-		UE_LEARNING_CHECK(Parameters.Num >= 0);
-		UE_LEARNING_CHECK(Parameters.Scale >= 0.0f);
+		check(Parameters.Num >= 0);
+		check(Parameters.Scale >= 0.0f);
 
 		FContinuousData ElementData;
 		ElementData.Num = Parameters.Num;
@@ -180,8 +180,8 @@ namespace UE::Learning::Action
 
 	FSchemaElement FSchema::CreateDiscreteExclusive(const FSchemaDiscreteExclusiveParameters Parameters, const FName Tag)
 	{
-		UE_LEARNING_CHECK(Parameters.PriorProbabilities.Num() == Parameters.Num);
-		UE_LEARNING_CHECK(Private::CheckPriorProbabilitiesExclusive(Parameters.PriorProbabilities));
+		check(Parameters.PriorProbabilities.Num() == Parameters.Num);
+		check(Private::CheckPriorProbabilitiesExclusive(Parameters.PriorProbabilities));
 
 		FDiscreteExclusiveData ElementData;
 		ElementData.Num = Parameters.Num;
@@ -202,8 +202,8 @@ namespace UE::Learning::Action
 
 	FSchemaElement FSchema::CreateDiscreteInclusive(const FSchemaDiscreteInclusiveParameters Parameters, const FName Tag)
 	{
-		UE_LEARNING_CHECK(Parameters.PriorProbabilities.Num() == Parameters.Num);
-		UE_LEARNING_CHECK(Private::CheckPriorProbabilitiesInclusive(Parameters.PriorProbabilities));
+		check(Parameters.PriorProbabilities.Num() == Parameters.Num);
+		check(Private::CheckPriorProbabilitiesInclusive(Parameters.PriorProbabilities));
 
 		FDiscreteInclusiveData ElementData;
 		ElementData.Num = Parameters.Num;
@@ -224,9 +224,9 @@ namespace UE::Learning::Action
 
 	FSchemaElement FSchema::CreateAnd(const FSchemaAndParameters Parameters, const FName Tag)
 	{
-		UE_LEARNING_CHECK(Parameters.Elements.Num() == Parameters.ElementNames.Num());
-		UE_LEARNING_CHECK(!Private::ContainsDuplicates(Parameters.ElementNames));
-		UE_LEARNING_CHECK(Private::CheckAllValid(*this, Parameters.Elements));
+		check(Parameters.Elements.Num() == Parameters.ElementNames.Num());
+		check(!Private::ContainsDuplicates(Parameters.ElementNames));
+		check(Private::CheckAllValid(*this, Parameters.Elements));
 
 		FAndData ElementData;
 		ElementData.Num = Parameters.Elements.Num();
@@ -248,11 +248,11 @@ namespace UE::Learning::Action
 
 	FSchemaElement FSchema::CreateOrExclusive(const FSchemaOrExclusiveParameters Parameters, const FName Tag)
 	{
-		UE_LEARNING_CHECK(Parameters.Elements.Num() == Parameters.ElementNames.Num());
-		UE_LEARNING_CHECK(!Private::ContainsDuplicates(Parameters.ElementNames));
-		UE_LEARNING_CHECK(Private::CheckAllValid(*this, Parameters.Elements));
-		UE_LEARNING_CHECK(Parameters.PriorProbabilities.Num() == Parameters.Elements.Num());
-		UE_LEARNING_CHECK(Private::CheckPriorProbabilitiesExclusive(Parameters.PriorProbabilities));
+		check(Parameters.Elements.Num() == Parameters.ElementNames.Num());
+		check(!Private::ContainsDuplicates(Parameters.ElementNames));
+		check(Private::CheckAllValid(*this, Parameters.Elements));
+		check(Parameters.PriorProbabilities.Num() == Parameters.Elements.Num());
+		check(Private::CheckPriorProbabilitiesExclusive(Parameters.PriorProbabilities));
 
 		FOrExclusiveData ElementData;
 		ElementData.Num = Parameters.Elements.Num();
@@ -276,11 +276,11 @@ namespace UE::Learning::Action
 
 	FSchemaElement FSchema::CreateOrInclusive(const FSchemaOrInclusiveParameters Parameters, const FName Tag)
 	{
-		UE_LEARNING_CHECK(Parameters.Elements.Num() == Parameters.ElementNames.Num());
-		UE_LEARNING_CHECK(!Private::ContainsDuplicates(Parameters.ElementNames));
-		UE_LEARNING_CHECK(Private::CheckAllValid(*this, Parameters.Elements));
-		UE_LEARNING_CHECK(Parameters.PriorProbabilities.Num() == Parameters.Elements.Num());
-		UE_LEARNING_CHECK(Private::CheckPriorProbabilitiesInclusive(Parameters.PriorProbabilities));
+		check(Parameters.Elements.Num() == Parameters.ElementNames.Num());
+		check(!Private::ContainsDuplicates(Parameters.ElementNames));
+		check(Private::CheckAllValid(*this, Parameters.Elements));
+		check(Parameters.PriorProbabilities.Num() == Parameters.Elements.Num());
+		check(Private::CheckPriorProbabilitiesInclusive(Parameters.PriorProbabilities));
 
 		FOrInclusiveData ElementData;
 		ElementData.Num = Parameters.Elements.Num();
@@ -304,8 +304,8 @@ namespace UE::Learning::Action
 
 	FSchemaElement FSchema::CreateArray(const FSchemaArrayParameters Parameters, const FName Tag)
 	{
-		UE_LEARNING_CHECK(IsValid(Parameters.Element));
-		UE_LEARNING_CHECK(Parameters.Num >= 0);
+		check(IsValid(Parameters.Element));
+		check(Parameters.Num >= 0);
 
 		FArrayData ElementData;
 		ElementData.Num = Parameters.Num;
@@ -327,7 +327,7 @@ namespace UE::Learning::Action
 
 	FSchemaElement FSchema::CreateEncoding(const FSchemaEncodingParameters Parameters, const FName Tag)
 	{
-		UE_LEARNING_CHECK(IsValid(Parameters.Element));
+		check(IsValid(Parameters.Element));
 
 		FEncodingData ElementData;
 		ElementData.EncodingSize = Parameters.EncodingSize;
@@ -356,43 +356,43 @@ namespace UE::Learning::Action
 
 	EType FSchema::GetType(const FSchemaElement Element) const
 	{
-		UE_LEARNING_CHECK(IsValid(Element));
+		check(IsValid(Element));
 		return Types[Element.Index];
 	}
 
 	FName FSchema::GetTag(const FSchemaElement Element) const
 	{
-		UE_LEARNING_CHECK(IsValid(Element));
+		check(IsValid(Element));
 		return Tags[Element.Index];
 	}
 
 	int32 FSchema::GetEncodedVectorSize(const FSchemaElement Element) const
 	{
-		UE_LEARNING_CHECK(IsValid(Element));
+		check(IsValid(Element));
 		return EncodedVectorSizes[Element.Index];
 	}
 
 	int32 FSchema::GetActionVectorSize(const FSchemaElement Element) const
 	{
-		UE_LEARNING_CHECK(IsValid(Element));
+		check(IsValid(Element));
 		return ActionVectorSizes[Element.Index];
 	}
 
 	int32 FSchema::GetActionDistributionVectorSize(const FSchemaElement Element) const
 	{
-		UE_LEARNING_CHECK(IsValid(Element));
+		check(IsValid(Element));
 		return ActionDistributionVectorSizes[Element.Index];
 	}
 
 	int32 FSchema::GetActionModifierVectorSize(const FSchemaElement Element) const
 	{
-		UE_LEARNING_CHECK(IsValid(Element));
+		check(IsValid(Element));
 		return ActionModifierVectorSizes[Element.Index];
 	}
 
 	FSchemaContinuousParameters FSchema::GetContinuous(const FSchemaElement Element) const
 	{
-		UE_LEARNING_CHECK(IsValid(Element) && GetType(Element) == EType::Continuous);
+		check(IsValid(Element) && GetType(Element) == EType::Continuous);
 		const FContinuousData& ElementData = ContinuousData[TypeDataIndices[Element.Index]];
 
 		FSchemaContinuousParameters Parameters;
@@ -403,7 +403,7 @@ namespace UE::Learning::Action
 
 	FSchemaDiscreteExclusiveParameters FSchema::GetDiscreteExclusive(const FSchemaElement Element) const
 	{
-		UE_LEARNING_CHECK(IsValid(Element) && GetType(Element) == EType::DiscreteExclusive);
+		check(IsValid(Element) && GetType(Element) == EType::DiscreteExclusive);
 		const FDiscreteExclusiveData& ElementData = DiscreteExclusiveData[TypeDataIndices[Element.Index]];
 
 		FSchemaDiscreteExclusiveParameters Parameters;
@@ -414,7 +414,7 @@ namespace UE::Learning::Action
 
 	FSchemaDiscreteInclusiveParameters FSchema::GetDiscreteInclusive(const FSchemaElement Element) const
 	{
-		UE_LEARNING_CHECK(IsValid(Element) && GetType(Element) == EType::DiscreteInclusive);
+		check(IsValid(Element) && GetType(Element) == EType::DiscreteInclusive);
 		const FDiscreteInclusiveData& ElementData = DiscreteInclusiveData[TypeDataIndices[Element.Index]];
 
 		FSchemaDiscreteInclusiveParameters Parameters;
@@ -425,7 +425,7 @@ namespace UE::Learning::Action
 
 	FSchemaAndParameters FSchema::GetAnd(const FSchemaElement Element) const
 	{
-		UE_LEARNING_CHECK(IsValid(Element) && GetType(Element) == EType::And);
+		check(IsValid(Element) && GetType(Element) == EType::And);
 		const FAndData& ElementData = AndData[TypeDataIndices[Element.Index]];
 
 		FSchemaAndParameters Parameters;
@@ -436,7 +436,7 @@ namespace UE::Learning::Action
 
 	FSchemaOrExclusiveParameters FSchema::GetOrExclusive(const FSchemaElement Element) const
 	{
-		UE_LEARNING_CHECK(IsValid(Element) && GetType(Element) == EType::OrExclusive);
+		check(IsValid(Element) && GetType(Element) == EType::OrExclusive);
 		const FOrExclusiveData& ElementData = OrExclusiveData[TypeDataIndices[Element.Index]];
 
 		FSchemaOrExclusiveParameters Parameters;
@@ -448,7 +448,7 @@ namespace UE::Learning::Action
 
 	FSchemaOrInclusiveParameters FSchema::GetOrInclusive(const FSchemaElement Element) const
 	{
-		UE_LEARNING_CHECK(IsValid(Element) && GetType(Element) == EType::OrInclusive);
+		check(IsValid(Element) && GetType(Element) == EType::OrInclusive);
 		const FOrInclusiveData& ElementData = OrInclusiveData[TypeDataIndices[Element.Index]];
 
 		FSchemaOrInclusiveParameters Parameters;
@@ -460,7 +460,7 @@ namespace UE::Learning::Action
 
 	FSchemaArrayParameters FSchema::GetArray(const FSchemaElement Element) const
 	{
-		UE_LEARNING_CHECK(IsValid(Element) && GetType(Element) == EType::Array);
+		check(IsValid(Element) && GetType(Element) == EType::Array);
 		const FArrayData& ElementData = ArrayData[TypeDataIndices[Element.Index]];
 
 		FSchemaArrayParameters Parameters;
@@ -471,7 +471,7 @@ namespace UE::Learning::Action
 
 	FSchemaEncodingParameters FSchema::GetEncoding(const FSchemaElement Element) const
 	{
-		UE_LEARNING_CHECK(IsValid(Element) && GetType(Element) == EType::Encoding);
+		check(IsValid(Element) && GetType(Element) == EType::Encoding);
 		const FEncodingData& ElementData = EncodingData[TypeDataIndices[Element.Index]];
 
 		FSchemaEncodingParameters Parameters;
@@ -618,9 +618,9 @@ namespace UE::Learning::Action
 
 	FObjectElement FObject::CreateAnd(const FObjectAndParameters Parameters, const FName Tag)
 	{
-		UE_LEARNING_CHECK(Parameters.Elements.Num() == Parameters.ElementNames.Num());
-		UE_LEARNING_CHECK(!Private::ContainsDuplicates(Parameters.ElementNames));
-		UE_LEARNING_CHECK(Private::CheckAllValid(*this, Parameters.Elements));
+		check(Parameters.Elements.Num() == Parameters.ElementNames.Num());
+		check(!Private::ContainsDuplicates(Parameters.ElementNames));
+		check(Private::CheckAllValid(*this, Parameters.Elements));
 
 		const int32 Index = Types.Add(EType::And);
 		Tags.Add(Tag);
@@ -642,7 +642,7 @@ namespace UE::Learning::Action
 
 	FObjectElement FObject::CreateOrExclusive(const FObjectOrExclusiveParameters Parameters, const FName Tag)
 	{
-		UE_LEARNING_CHECK(IsValid(Parameters.Element));
+		check(IsValid(Parameters.Element));
 
 		const int32 Index = Types.Add(EType::OrExclusive);
 		Tags.Add(Tag);
@@ -664,9 +664,9 @@ namespace UE::Learning::Action
 
 	FObjectElement FObject::CreateOrInclusive(const FObjectOrInclusiveParameters Parameters, const FName Tag)
 	{
-		UE_LEARNING_CHECK(Parameters.Elements.Num() == Parameters.ElementNames.Num());
-		UE_LEARNING_CHECK(!Private::ContainsDuplicates(Parameters.ElementNames));
-		UE_LEARNING_CHECK(Private::CheckAllValid(*this, Parameters.Elements));
+		check(Parameters.Elements.Num() == Parameters.ElementNames.Num());
+		check(!Private::ContainsDuplicates(Parameters.ElementNames));
+		check(Private::CheckAllValid(*this, Parameters.Elements));
 
 		const int32 Index = Types.Add(EType::OrInclusive);
 		Tags.Add(Tag);
@@ -688,7 +688,7 @@ namespace UE::Learning::Action
 
 	FObjectElement FObject::CreateArray(const FObjectArrayParameters Parameters, const FName Tag)
 	{
-		UE_LEARNING_CHECK(Private::CheckAllValid(*this, Parameters.Elements));
+		check(Private::CheckAllValid(*this, Parameters.Elements));
 
 		const int32 Index = Types.Add(EType::Array);
 		Tags.Add(Tag);
@@ -713,7 +713,7 @@ namespace UE::Learning::Action
 
 	FObjectElement FObject::CreateEncoding(const FObjectEncodingParameters Parameters, const FName Tag)
 	{
-		UE_LEARNING_CHECK(IsValid(Parameters.Element));
+		check(IsValid(Parameters.Element));
 
 		const int32 Index = Types.Add(EType::Encoding);
 		Tags.Add(Tag);
@@ -740,19 +740,19 @@ namespace UE::Learning::Action
 
 	EType FObject::GetType(const FObjectElement Element) const
 	{
-		UE_LEARNING_CHECK(IsValid(Element));
+		check(IsValid(Element));
 		return Types[Element.Index];
 	}
 	
 	FName FObject::GetTag(const FObjectElement Element) const
 	{
-		UE_LEARNING_CHECK(IsValid(Element));
+		check(IsValid(Element));
 		return Tags[Element.Index];
 	}
 
 	FObjectContinuousParameters FObject::GetContinuous(const FObjectElement Element) const
 	{
-		UE_LEARNING_CHECK(IsValid(Element) && GetType(Element) == EType::Continuous);
+		check(IsValid(Element) && GetType(Element) == EType::Continuous);
 
 		FObjectContinuousParameters Parameters;
 		Parameters.Values = TArrayView<const float>(ContinuousValues.GetData() + ContinuousDataOffsets[Element.Index], ContinuousDataNums[Element.Index]);
@@ -761,7 +761,7 @@ namespace UE::Learning::Action
 
 	FObjectDiscreteExclusiveParameters FObject::GetDiscreteExclusive(const FObjectElement Element) const
 	{
-		UE_LEARNING_CHECK(IsValid(Element) && GetType(Element) == EType::DiscreteExclusive);
+		check(IsValid(Element) && GetType(Element) == EType::DiscreteExclusive);
 
 		FObjectDiscreteExclusiveParameters Parameters;
 		Parameters.DiscreteIndex = DiscreteValues[DiscreteDataOffsets[Element.Index]];
@@ -770,7 +770,7 @@ namespace UE::Learning::Action
 
 	FObjectDiscreteInclusiveParameters FObject::GetDiscreteInclusive(const FObjectElement Element) const
 	{
-		UE_LEARNING_CHECK(IsValid(Element) && GetType(Element) == EType::DiscreteInclusive);
+		check(IsValid(Element) && GetType(Element) == EType::DiscreteInclusive);
 
 		FObjectDiscreteInclusiveParameters Parameters;
 		Parameters.DiscreteIndices = TArrayView<const int32>(DiscreteValues.GetData() + DiscreteDataOffsets[Element.Index], DiscreteDataNums[Element.Index]);
@@ -779,7 +779,7 @@ namespace UE::Learning::Action
 
 	FObjectAndParameters FObject::GetAnd(const FObjectElement Element) const
 	{
-		UE_LEARNING_CHECK(IsValid(Element) && GetType(Element) == EType::And);
+		check(IsValid(Element) && GetType(Element) == EType::And);
 
 		FObjectAndParameters Parameters;
 		Parameters.ElementNames = TArrayView<const FName>(SubElementNames.GetData() + ElementDataOffsets[Element.Index], ElementDataNums[Element.Index]);
@@ -789,7 +789,7 @@ namespace UE::Learning::Action
 
 	FObjectOrExclusiveParameters FObject::GetOrExclusive(const FObjectElement Element) const
 	{
-		UE_LEARNING_CHECK(IsValid(Element) && GetType(Element) == EType::OrExclusive);
+		check(IsValid(Element) && GetType(Element) == EType::OrExclusive);
 
 		FObjectOrExclusiveParameters Parameters;
 		Parameters.ElementName = SubElementNames[ElementDataOffsets[Element.Index]];
@@ -799,7 +799,7 @@ namespace UE::Learning::Action
 
 	FObjectOrInclusiveParameters FObject::GetOrInclusive(const FObjectElement Element) const
 	{
-		UE_LEARNING_CHECK(IsValid(Element) && GetType(Element) == EType::OrInclusive);
+		check(IsValid(Element) && GetType(Element) == EType::OrInclusive);
 
 		FObjectOrInclusiveParameters Parameters;
 		Parameters.ElementNames = TArrayView<const FName>(SubElementNames.GetData() + ElementDataOffsets[Element.Index], ElementDataNums[Element.Index]);
@@ -809,7 +809,7 @@ namespace UE::Learning::Action
 
 	FObjectArrayParameters FObject::GetArray(const FObjectElement Element) const
 	{
-		UE_LEARNING_CHECK(IsValid(Element) && GetType(Element) == EType::Array);
+		check(IsValid(Element) && GetType(Element) == EType::Array);
 		
 		FObjectArrayParameters Parameters;
 		Parameters.Elements = TArrayView<const FObjectElement>(SubElementObjects.GetData() + ElementDataOffsets[Element.Index], ElementDataNums[Element.Index]);
@@ -818,7 +818,7 @@ namespace UE::Learning::Action
 
 	FObjectEncodingParameters FObject::GetEncoding(const FObjectElement Element) const
 	{
-		UE_LEARNING_CHECK(IsValid(Element) && GetType(Element) == EType::Encoding);
+		check(IsValid(Element) && GetType(Element) == EType::Encoding);
 
 		FObjectEncodingParameters Parameters;
 		Parameters.Element = SubElementObjects[ElementDataOffsets[Element.Index]];
@@ -963,9 +963,9 @@ namespace UE::Learning::Action
 
 	FModifierElement FModifier::CreateAnd(const FModifierAndParameters Parameters, const FName Tag)
 	{
-		UE_LEARNING_CHECK(Parameters.Elements.Num() == Parameters.ElementNames.Num());
-		UE_LEARNING_CHECK(!Private::ContainsDuplicates(Parameters.ElementNames));
-		UE_LEARNING_CHECK(Private::CheckAllValid(*this, Parameters.Elements));
+		check(Parameters.Elements.Num() == Parameters.ElementNames.Num());
+		check(!Private::ContainsDuplicates(Parameters.ElementNames));
+		check(Private::CheckAllValid(*this, Parameters.Elements));
 
 		const int32 Index = Types.Add(EType::And);
 		Tags.Add(Tag);
@@ -990,10 +990,10 @@ namespace UE::Learning::Action
 
 	FModifierElement FModifier::CreateOrExclusive(const FModifierOrExclusiveParameters Parameters, const FName Tag)
 	{
-		UE_LEARNING_CHECK(Parameters.Elements.Num() == Parameters.ElementNames.Num());
-		UE_LEARNING_CHECK(!Private::ContainsDuplicates(Parameters.ElementNames));
-		UE_LEARNING_CHECK(!Private::ContainsDuplicates(Parameters.MaskedElements));
-		UE_LEARNING_CHECK(Private::CheckAllValid(*this, Parameters.Elements));
+		check(Parameters.Elements.Num() == Parameters.ElementNames.Num());
+		check(!Private::ContainsDuplicates(Parameters.ElementNames));
+		check(!Private::ContainsDuplicates(Parameters.MaskedElements));
+		check(Private::CheckAllValid(*this, Parameters.Elements));
 
 		const int32 Index = Types.Add(EType::OrExclusive);
 		Tags.Add(Tag);
@@ -1020,10 +1020,10 @@ namespace UE::Learning::Action
 
 	FModifierElement FModifier::CreateOrInclusive(const FModifierOrInclusiveParameters Parameters, const FName Tag)
 	{
-		UE_LEARNING_CHECK(Parameters.Elements.Num() == Parameters.ElementNames.Num());
-		UE_LEARNING_CHECK(!Private::ContainsDuplicates(Parameters.ElementNames));
-		UE_LEARNING_CHECK(!Private::ContainsDuplicates(Parameters.MaskedElements));
-		UE_LEARNING_CHECK(Private::CheckAllValid(*this, Parameters.Elements));
+		check(Parameters.Elements.Num() == Parameters.ElementNames.Num());
+		check(!Private::ContainsDuplicates(Parameters.ElementNames));
+		check(!Private::ContainsDuplicates(Parameters.MaskedElements));
+		check(Private::CheckAllValid(*this, Parameters.Elements));
 
 		const int32 Index = Types.Add(EType::OrInclusive);
 		Tags.Add(Tag);
@@ -1050,7 +1050,7 @@ namespace UE::Learning::Action
 
 	FModifierElement FModifier::CreateArray(const FModifierArrayParameters Parameters, const FName Tag)
 	{
-		UE_LEARNING_CHECK(Private::CheckAllValid(*this, Parameters.Elements));
+		check(Private::CheckAllValid(*this, Parameters.Elements));
 
 		const int32 Index = Types.Add(EType::Array);
 		Tags.Add(Tag);
@@ -1078,7 +1078,7 @@ namespace UE::Learning::Action
 
 	FModifierElement FModifier::CreateEncoding(const FModifierEncodingParameters Parameters, const FName Tag)
 	{
-		UE_LEARNING_CHECK(IsValid(Parameters.Element));
+		check(IsValid(Parameters.Element));
 
 		const int32 Index = Types.Add(EType::Encoding);
 		Tags.Add(Tag);
@@ -1108,19 +1108,19 @@ namespace UE::Learning::Action
 
 	EType FModifier::GetType(const FModifierElement Element) const
 	{
-		UE_LEARNING_CHECK(IsValid(Element));
+		check(IsValid(Element));
 		return Types[Element.Index];
 	}
 
 	FName FModifier::GetTag(const FModifierElement Element) const
 	{
-		UE_LEARNING_CHECK(IsValid(Element));
+		check(IsValid(Element));
 		return Tags[Element.Index];
 	}
 
 	FModifierContinuousParameters FModifier::GetContinuous(const FModifierElement Element) const
 	{
-		UE_LEARNING_CHECK(IsValid(Element) && GetType(Element) == EType::Continuous);
+		check(IsValid(Element) && GetType(Element) == EType::Continuous);
 
 		FModifierContinuousParameters Parameters;
 		Parameters.Masked = TArrayView<const bool>(ContinuousMaskeds.GetData() + ContinuousDataOffsets[Element.Index], ContinuousDataNums[Element.Index]);
@@ -1130,7 +1130,7 @@ namespace UE::Learning::Action
 
 	FModifierDiscreteExclusiveParameters FModifier::GetDiscreteExclusive(const FModifierElement Element) const
 	{
-		UE_LEARNING_CHECK(IsValid(Element) && GetType(Element) == EType::DiscreteExclusive);
+		check(IsValid(Element) && GetType(Element) == EType::DiscreteExclusive);
 
 		FModifierDiscreteExclusiveParameters Parameters;
 		Parameters.MaskedIndices = TArrayView<const int32>(DiscreteValues.GetData() + DiscreteDataOffsets[Element.Index], DiscreteDataNums[Element.Index]);
@@ -1139,7 +1139,7 @@ namespace UE::Learning::Action
 
 	FModifierDiscreteInclusiveParameters FModifier::GetDiscreteInclusive(const FModifierElement Element) const
 	{
-		UE_LEARNING_CHECK(IsValid(Element) && GetType(Element) == EType::DiscreteInclusive);
+		check(IsValid(Element) && GetType(Element) == EType::DiscreteInclusive);
 
 		FModifierDiscreteInclusiveParameters Parameters;
 		Parameters.MaskedIndices = TArrayView<const int32>(DiscreteValues.GetData() + DiscreteDataOffsets[Element.Index], DiscreteDataNums[Element.Index]);
@@ -1148,7 +1148,7 @@ namespace UE::Learning::Action
 
 	FModifierAndParameters FModifier::GetAnd(const FModifierElement Element) const
 	{
-		UE_LEARNING_CHECK(IsValid(Element) && GetType(Element) == EType::And);
+		check(IsValid(Element) && GetType(Element) == EType::And);
 
 		FModifierAndParameters Parameters;
 		Parameters.ElementNames = TArrayView<const FName>(SubElementNames.GetData() + ElementDataOffsets[Element.Index], ElementDataNums[Element.Index]);
@@ -1158,7 +1158,7 @@ namespace UE::Learning::Action
 
 	FModifierOrExclusiveParameters FModifier::GetOrExclusive(const FModifierElement Element) const
 	{
-		UE_LEARNING_CHECK(IsValid(Element) && GetType(Element) == EType::OrExclusive);
+		check(IsValid(Element) && GetType(Element) == EType::OrExclusive);
 
 		FModifierOrExclusiveParameters Parameters;
 		Parameters.ElementNames = TArrayView<const FName>(SubElementNames.GetData() + ElementDataOffsets[Element.Index], ElementDataNums[Element.Index]);
@@ -1169,7 +1169,7 @@ namespace UE::Learning::Action
 
 	FModifierOrInclusiveParameters FModifier::GetOrInclusive(const FModifierElement Element) const
 	{
-		UE_LEARNING_CHECK(IsValid(Element) && GetType(Element) == EType::OrInclusive);
+		check(IsValid(Element) && GetType(Element) == EType::OrInclusive);
 
 		FModifierOrInclusiveParameters Parameters;
 		Parameters.ElementNames = TArrayView<const FName>(SubElementNames.GetData() + ElementDataOffsets[Element.Index], ElementDataNums[Element.Index]);
@@ -1180,7 +1180,7 @@ namespace UE::Learning::Action
 
 	FModifierArrayParameters FModifier::GetArray(const FModifierElement Element) const
 	{
-		UE_LEARNING_CHECK(IsValid(Element) && GetType(Element) == EType::Array);
+		check(IsValid(Element) && GetType(Element) == EType::Array);
 
 		FModifierArrayParameters Parameters;
 		Parameters.Elements = TArrayView<const FModifierElement>(SubElementModifiers.GetData() + ElementDataOffsets[Element.Index], ElementDataNums[Element.Index]);
@@ -1189,7 +1189,7 @@ namespace UE::Learning::Action
 
 	FModifierEncodingParameters FModifier::GetEncoding(const FModifierElement Element) const
 	{
-		UE_LEARNING_CHECK(IsValid(Element) && GetType(Element) == EType::Encoding);
+		check(IsValid(Element) && GetType(Element) == EType::Encoding);
 
 		FModifierEncodingParameters Parameters;
 		Parameters.Element = SubElementModifiers[ElementDataOffsets[Element.Index]];
@@ -1263,7 +1263,7 @@ namespace UE::Learning::Action
 			case EEncodingActivationFunction::ELU: return NNE::RuntimeBasic::FModelBuilder::EActivationFunction::ELU;
 			case EEncodingActivationFunction::TanH: return NNE::RuntimeBasic::FModelBuilder::EActivationFunction::TanH;
 			case EEncodingActivationFunction::GELU: return NNE::RuntimeBasic::FModelBuilder::EActivationFunction::GELU;
-			default: UE_LEARNING_NOT_IMPLEMENTED(); return NNE::RuntimeBasic::FModelBuilder::EActivationFunction::ReLU;
+			default: checkNoEntry(); return NNE::RuntimeBasic::FModelBuilder::EActivationFunction::ReLU;
 			}
 		}
 
@@ -1312,7 +1312,7 @@ namespace UE::Learning::Action
 		const FSchemaElement SchemaElement,
 		const int32 Salt)
 	{
-		UE_LEARNING_CHECK(Schema.IsValid(SchemaElement));
+		check(Schema.IsValid(SchemaElement));
 		const EType SchemaElementType = Schema.GetType(SchemaElement);
 
 		const int32 Hash = Private::HashCombine({ Salt, Private::HashInt((int32)SchemaElementType) });
@@ -1359,7 +1359,7 @@ namespace UE::Learning::Action
 
 		default:
 		{
-			UE_LEARNING_NOT_IMPLEMENTED();
+			checkNoEntry();
 			return 0;
 		}
 
@@ -1372,8 +1372,8 @@ namespace UE::Learning::Action
 		const FSchema& SchemaB,
 		const FSchemaElement SchemaElementB)
 	{
-		UE_LEARNING_CHECK(SchemaA.IsValid(SchemaElementA));
-		UE_LEARNING_CHECK(SchemaB.IsValid(SchemaElementB));
+		check(SchemaA.IsValid(SchemaElementA));
+		check(SchemaB.IsValid(SchemaElementB));
 
 		const EType SchemaElementTypeA = SchemaA.GetType(SchemaElementA);
 		const EType SchemaElementTypeB = SchemaB.GetType(SchemaElementB);
@@ -1459,13 +1459,13 @@ namespace UE::Learning::Action
 
 		case EType::Encoding:
 		{
-			UE_LEARNING_CHECKF(false, TEXT("Encoding elements should always be forwarded..."));
+			checkf(false, TEXT("Encoding elements should always be forwarded..."));
 			return false;
 		}
 
 		default:
 		{
-			UE_LEARNING_NOT_IMPLEMENTED();
+			checkNoEntry();
 			return false;
 		}
 
@@ -1644,7 +1644,7 @@ namespace UE::Learning::Action
 				NNE::RuntimeBasic::FModelBuilder::EWeightInitializationType::KaimingGaussian; break;
 			case EWeightInitialization::KaimingUniform: LinearLayerSettings.WeightInitializationSettings.Type =
 				NNE::RuntimeBasic::FModelBuilder::EWeightInitializationType::KaimingUniform; break;
-			default: UE_LEARNING_NOT_IMPLEMENTED();
+			default: checkNoEntry();
 			}
 
 			OutElement = Builder.MakeSequence({
@@ -1665,15 +1665,15 @@ namespace UE::Learning::Action
 
 		default:
 		{
-			UE_LEARNING_NOT_IMPLEMENTED();
+			checkNoEntry();
 		}
 		}
 
-		UE_LEARNING_CHECKF(OutElement.GetInputSize() == Schema.GetEncodedVectorSize(SchemaElement),
+		checkf(OutElement.GetInputSize() == Schema.GetEncodedVectorSize(SchemaElement),
 			TEXT("Decoder Network Input unexpected size. Got %i, expected %i according to Schema."),
 			OutElement.GetInputSize(), Schema.GetEncodedVectorSize(SchemaElement));
 
-		UE_LEARNING_CHECKF(OutElement.GetOutputSize() == Schema.GetActionDistributionVectorSize(SchemaElement),
+		checkf(OutElement.GetOutputSize() == Schema.GetActionDistributionVectorSize(SchemaElement),
 			TEXT("Decoder Network Output unexpected size. Got %i, expected %i according to Schema."),
 			OutElement.GetOutputSize(), Schema.GetActionDistributionVectorSize(SchemaElement));
 	}
@@ -1687,7 +1687,7 @@ namespace UE::Learning::Action
 		const FNetworkSettings& NetworkSettings,
 		const uint32 Seed)
 	{
-		UE_LEARNING_CHECK(Schema.IsValid(SchemaElement));
+		check(Schema.IsValid(SchemaElement));
 
 		NNE::RuntimeBasic::FModelBuilder Builder(Seed);
 		NNE::RuntimeBasic::FModelBuilderElement Element;
@@ -1704,7 +1704,7 @@ namespace UE::Learning::Action
 		const FSchemaElement SchemaElement,
 		const float ActionNoiseScale)
 	{
-		UE_LEARNING_CHECK(Schema.IsValid(SchemaElement));
+		check(Schema.IsValid(SchemaElement));
 
 		const EType SchemaElementType = Schema.GetType(SchemaElement);
 
@@ -1715,9 +1715,9 @@ namespace UE::Learning::Action
 		case EType::Continuous:
 		{
 			const int32 ValueNum = Schema.GetContinuous(SchemaElement).Num;
-			UE_LEARNING_CHECK(ValueNum == OutActionVector.Num());
-			UE_LEARNING_CHECK(ValueNum * 2 == ActionDistributionVector.Num());
-			UE_LEARNING_CHECK(1 + ValueNum * 2 == ActionModifierVector.Num());
+			check(ValueNum == OutActionVector.Num());
+			check(ValueNum * 2 == ActionDistributionVector.Num());
+			check(1 + ValueNum * 2 == ActionModifierVector.Num());
 
 			if (ActionModifierVector[0])
 			{
@@ -1753,9 +1753,9 @@ namespace UE::Learning::Action
 		case EType::DiscreteExclusive:
 		{
 			const int32 ValueNum = Schema.GetDiscreteExclusive(SchemaElement).Num;
-			UE_LEARNING_CHECK(ValueNum == OutActionVector.Num());
-			UE_LEARNING_CHECK(ValueNum == ActionDistributionVector.Num());
-			UE_LEARNING_CHECK(1 + ValueNum == ActionModifierVector.Num());
+			check(ValueNum == OutActionVector.Num());
+			check(ValueNum == ActionDistributionVector.Num());
+			check(1 + ValueNum == ActionModifierVector.Num());
 
 			if (ActionModifierVector[0])
 			{
@@ -1765,7 +1765,7 @@ namespace UE::Learning::Action
 				{
 					Masked[ValueIdx] = ActionModifierVector[1 + ValueIdx] == 1.0f;
 				}
-				UE_LEARNING_CHECK(Private::CheckExclusiveMaskValid(Masked));
+				check(Private::CheckExclusiveMaskValid(Masked));
 
 				Random::SampleDistributionMultinoulliMasked(
 					OutActionVector,
@@ -1789,9 +1789,9 @@ namespace UE::Learning::Action
 		case EType::DiscreteInclusive:
 		{
 			const int32 ValueNum = Schema.GetDiscreteInclusive(SchemaElement).Num;
-			UE_LEARNING_CHECK(ValueNum == OutActionVector.Num());
-			UE_LEARNING_CHECK(ValueNum == ActionDistributionVector.Num());
-			UE_LEARNING_CHECK(1 + ValueNum == ActionModifierVector.Num());
+			check(ValueNum == OutActionVector.Num());
+			check(ValueNum == ActionDistributionVector.Num());
+			check(1 + ValueNum == ActionModifierVector.Num());
 
 			if (ActionModifierVector[0])
 			{
@@ -1849,9 +1849,9 @@ namespace UE::Learning::Action
 				SubElementActionModifierVectorOffset += SubElementActionModifierVectorSize;
 			}
 
-			UE_LEARNING_CHECK(SubElementActionVectorOffset == OutActionVector.Num());
-			UE_LEARNING_CHECK(SubElementActionDistributionVectorOffset == ActionDistributionVector.Num());
-			UE_LEARNING_CHECK(SubElementActionModifierVectorOffset == ActionModifierVector.Num());
+			check(SubElementActionVectorOffset == OutActionVector.Num());
+			check(SubElementActionDistributionVectorOffset == ActionDistributionVector.Num());
+			check(SubElementActionModifierVectorOffset == ActionModifierVector.Num());
 
 			break;
 		}
@@ -1865,9 +1865,9 @@ namespace UE::Learning::Action
 			const int32 SubElementActionModifierVectorTotal = Private::GetTotalActionModifierVectorSize(Schema, Parameters.Elements);
 			const int32 ElementNum = Parameters.Elements.Num();
 
-			UE_LEARNING_CHECK(SubElementActionVectorMax + ElementNum == OutActionVector.Num());
-			UE_LEARNING_CHECK(SubElementActionDistributionVectorTotal + ElementNum == ActionDistributionVector.Num());
-			UE_LEARNING_CHECK(1 + ElementNum + SubElementActionModifierVectorTotal == ActionModifierVector.Num());
+			check(SubElementActionVectorMax + ElementNum == OutActionVector.Num());
+			check(SubElementActionDistributionVectorTotal + ElementNum == ActionDistributionVector.Num());
+			check(1 + ElementNum + SubElementActionModifierVectorTotal == ActionModifierVector.Num());
 
 			// Zero main part of vector
 			Array::Zero(OutActionVector.Slice(0, SubElementActionVectorMax));
@@ -1880,7 +1880,7 @@ namespace UE::Learning::Action
 				{
 					Masked[ElementIdx] = ActionModifierVector[1 + ElementIdx] == 1.0f;
 				}
-				UE_LEARNING_CHECK(Private::CheckExclusiveMaskValid(Masked));
+				check(Private::CheckExclusiveMaskValid(Masked));
 
 				// Sample which sub-element to generate
 				Random::SampleDistributionMultinoulliMasked(
@@ -1911,7 +1911,7 @@ namespace UE::Learning::Action
 				const int32 SubElementActionDistributionVectorSize = Schema.GetActionDistributionVectorSize(SubElement);
 				const int32 SubElementActionModifierVectorSize = Schema.GetActionModifierVectorSize(SubElement);
 
-				UE_LEARNING_CHECK(SubElementActionVectorSize <= SubElementActionVectorMax);
+				check(SubElementActionVectorSize <= SubElementActionVectorMax);
 
 				if (OutActionVector[SubElementActionVectorMax + SubElementIdx])
 				{
@@ -1932,9 +1932,9 @@ namespace UE::Learning::Action
 				SubElementActionModifierVectorOffset += SubElementActionModifierVectorSize;
 			}
 
-			UE_LEARNING_CHECK(SubElementsSampled == 1); // Exactly one sub-element should have been sampled
-			UE_LEARNING_CHECK(SubElementActionDistributionVectorOffset == SubElementActionDistributionVectorTotal);
-			UE_LEARNING_CHECK(SubElementActionModifierVectorOffset == 1 + ElementNum + SubElementActionModifierVectorTotal);
+			check(SubElementsSampled == 1); // Exactly one sub-element should have been sampled
+			check(SubElementActionDistributionVectorOffset == SubElementActionDistributionVectorTotal);
+			check(SubElementActionModifierVectorOffset == 1 + ElementNum + SubElementActionModifierVectorTotal);
 
 			break;
 		}
@@ -1948,9 +1948,9 @@ namespace UE::Learning::Action
 			const int32 SubElementActionModifierVectorTotal = Private::GetTotalActionModifierVectorSize(Schema, Parameters.Elements);
 			const int32 ElementNum = Parameters.Elements.Num();
 
-			UE_LEARNING_CHECK(SubElementActionVectorTotal + ElementNum == OutActionVector.Num());
-			UE_LEARNING_CHECK(SubElementActionDistributionVectorTotal + ElementNum == ActionDistributionVector.Num());
-			UE_LEARNING_CHECK(1 + ElementNum + SubElementActionModifierVectorTotal == ActionModifierVector.Num());
+			check(SubElementActionVectorTotal + ElementNum == OutActionVector.Num());
+			check(SubElementActionDistributionVectorTotal + ElementNum == ActionDistributionVector.Num());
+			check(1 + ElementNum + SubElementActionModifierVectorTotal == ActionModifierVector.Num());
 
 			// Zero main part of vector
 			Array::Zero(OutActionVector.Slice(0, SubElementActionVectorTotal));
@@ -2011,9 +2011,9 @@ namespace UE::Learning::Action
 				SubElementActionModifierVectorOffset += SubElementActionModifierVectorSize;
 			}
 
-			UE_LEARNING_CHECK(SubElementActionVectorOffset == SubElementActionVectorTotal);
-			UE_LEARNING_CHECK(SubElementActionDistributionVectorOffset == SubElementActionDistributionVectorTotal);
-			UE_LEARNING_CHECK(SubElementActionModifierVectorOffset == 1 + ElementNum + SubElementActionModifierVectorTotal);
+			check(SubElementActionVectorOffset == SubElementActionVectorTotal);
+			check(SubElementActionDistributionVectorOffset == SubElementActionDistributionVectorTotal);
+			check(SubElementActionModifierVectorOffset == 1 + ElementNum + SubElementActionModifierVectorTotal);
 
 			break;
 		}
@@ -2026,9 +2026,9 @@ namespace UE::Learning::Action
 			const int32 SubElementActionDistributionVectorSize = Schema.GetActionDistributionVectorSize(Parameters.Element);
 			const int32 SubElementActionModifierVectorSize = Schema.GetActionModifierVectorSize(Parameters.Element);
 
-			UE_LEARNING_CHECK(SubElementActionVectorSize * Parameters.Num == OutActionVector.Num());
-			UE_LEARNING_CHECK(SubElementActionDistributionVectorSize * Parameters.Num == ActionDistributionVector.Num());
-			UE_LEARNING_CHECK(1 + SubElementActionModifierVectorSize * Parameters.Num == ActionModifierVector.Num());
+			check(SubElementActionVectorSize * Parameters.Num == OutActionVector.Num());
+			check(SubElementActionDistributionVectorSize * Parameters.Num == ActionDistributionVector.Num());
+			check(1 + SubElementActionModifierVectorSize * Parameters.Num == ActionModifierVector.Num());
 
 			for (int32 ElementIdx = 0; ElementIdx < Parameters.Num; ElementIdx++)
 			{
@@ -2072,15 +2072,15 @@ namespace UE::Learning::Action
 		const FObject& Object,
 		const FObjectElement ObjectElement)
 	{
-		UE_LEARNING_CHECK(Schema.IsValid(SchemaElement));
-		UE_LEARNING_CHECK(Object.IsValid(ObjectElement));
-		UE_LEARNING_CHECK(OutActionVector.Num() == Schema.GetActionVectorSize(SchemaElement));
+		check(Schema.IsValid(SchemaElement));
+		check(Object.IsValid(ObjectElement));
+		check(OutActionVector.Num() == Schema.GetActionVectorSize(SchemaElement));
 
 		// Check that the types match
 
 		const EType SchemaElementType = Schema.GetType(SchemaElement);
 		const EType ObjectElementType = Object.GetType(ObjectElement);
-		UE_LEARNING_CHECK(ObjectElementType == SchemaElementType);
+		check(ObjectElementType == SchemaElementType);
 
 		// Zero Action Vector
 
@@ -2100,9 +2100,9 @@ namespace UE::Learning::Action
 			const FSchemaContinuousParameters SchemaParameters = Schema.GetContinuous(SchemaElement);
 
 			TArrayView<const float> ActionValues = Object.GetContinuous(ObjectElement).Values;
-			UE_LEARNING_CHECK(Schema.GetActionVectorSize(SchemaElement) == ActionValues.Num());
-			UE_LEARNING_CHECK(Schema.GetActionVectorSize(SchemaElement) == OutActionVector.Num());
-			UE_LEARNING_CHECK(Schema.GetActionVectorSize(SchemaElement) == SchemaParameters.Num);
+			check(Schema.GetActionVectorSize(SchemaElement) == ActionValues.Num());
+			check(Schema.GetActionVectorSize(SchemaElement) == OutActionVector.Num());
+			check(Schema.GetActionVectorSize(SchemaElement) == SchemaParameters.Num);
 
 			// Copy in and scale the values from the action object
 
@@ -2120,8 +2120,8 @@ namespace UE::Learning::Action
 		case EType::DiscreteExclusive:
 		{
 			int32 ActionValue = Object.GetDiscreteExclusive(ObjectElement).DiscreteIndex;
-			UE_LEARNING_CHECK(Schema.GetActionVectorSize(SchemaElement) > ActionValue && ActionValue >= 0);
-			UE_LEARNING_CHECK(Schema.GetActionVectorSize(SchemaElement) == OutActionVector.Num());
+			check(Schema.GetActionVectorSize(SchemaElement) > ActionValue && ActionValue >= 0);
+			check(Schema.GetActionVectorSize(SchemaElement) == OutActionVector.Num());
 
 			// Set the single value in the action vector
 
@@ -2132,14 +2132,14 @@ namespace UE::Learning::Action
 		case EType::DiscreteInclusive:
 		{
 			TArrayView<const int32> ActionValues = Object.GetDiscreteInclusive(ObjectElement).DiscreteIndices;
-			UE_LEARNING_CHECK(Schema.GetActionVectorSize(SchemaElement) >= ActionValues.Num());
-			UE_LEARNING_CHECK(Schema.GetActionVectorSize(SchemaElement) == OutActionVector.Num());
+			check(Schema.GetActionVectorSize(SchemaElement) >= ActionValues.Num());
+			check(Schema.GetActionVectorSize(SchemaElement) == OutActionVector.Num());
 
 			// Set values in the action vector
 
 			for (int32 ActionValueIdx = 0; ActionValueIdx < ActionValues.Num(); ActionValueIdx++)
 			{
-				UE_LEARNING_CHECK(Schema.GetActionVectorSize(SchemaElement) > ActionValues[ActionValueIdx] && ActionValues[ActionValueIdx] >= 0);
+				check(Schema.GetActionVectorSize(SchemaElement) > ActionValues[ActionValueIdx] && ActionValues[ActionValueIdx] >= 0);
 				OutActionVector[ActionValues[ActionValueIdx]] = 1.0f;
 			}
 
@@ -2152,7 +2152,7 @@ namespace UE::Learning::Action
 
 			const FSchemaAndParameters SchemaParameters = Schema.GetAnd(SchemaElement);
 			const FObjectAndParameters ObjectParameters = Object.GetAnd(ObjectElement);
-			UE_LEARNING_CHECK(SchemaParameters.Elements.Num() == ObjectParameters.Elements.Num());
+			check(SchemaParameters.Elements.Num() == ObjectParameters.Elements.Num());
 
 			// Set the Sub-elements
 
@@ -2161,7 +2161,7 @@ namespace UE::Learning::Action
 			for (int32 SchemaElementIdx = 0; SchemaElementIdx < SchemaParameters.Elements.Num(); SchemaElementIdx++)
 			{
 				const int32 ObjectElementIndex = ObjectParameters.ElementNames.Find(SchemaParameters.ElementNames[SchemaElementIdx]);
-				UE_LEARNING_CHECK(ObjectElementIndex != INDEX_NONE);
+				check(ObjectElementIndex != INDEX_NONE);
 
 				const int32 SubElementSize = Schema.GetActionVectorSize(SchemaParameters.Elements[SchemaElementIdx]);
 
@@ -2175,7 +2175,7 @@ namespace UE::Learning::Action
 				SubElementOffset += SubElementSize;
 			}
 
-			UE_LEARNING_CHECK(SubElementOffset == OutActionVector.Num());
+			check(SubElementOffset == OutActionVector.Num());
 			return;
 		}
 
@@ -2187,7 +2187,7 @@ namespace UE::Learning::Action
 			const FObjectOrExclusiveParameters ObjectParameters = Object.GetOrExclusive(ObjectElement);
 
 			const int32 SchemaElementIndex = SchemaParameters.ElementNames.Find(ObjectParameters.ElementName);
-			UE_LEARNING_CHECK(SchemaElementIndex != INDEX_NONE);
+			check(SchemaElementIndex != INDEX_NONE);
 
 			// Set the sub-element
 
@@ -2206,7 +2206,7 @@ namespace UE::Learning::Action
 
 			OutActionVector[MaxSubElementSize + SchemaElementIndex] = 1.0f;
 
-			UE_LEARNING_CHECK(OutActionVector.Num() == MaxSubElementSize + SchemaParameters.Elements.Num());
+			check(OutActionVector.Num() == MaxSubElementSize + SchemaParameters.Elements.Num());
 			return;
 		}
 
@@ -2216,7 +2216,7 @@ namespace UE::Learning::Action
 
 			const FSchemaOrInclusiveParameters SchemaParameters = Schema.GetOrInclusive(SchemaElement);
 			const FObjectOrInclusiveParameters ObjectParameters = Object.GetOrInclusive(ObjectElement);
-			UE_LEARNING_CHECK(ObjectParameters.Elements.Num() <= SchemaParameters.Elements.Num());
+			check(ObjectParameters.Elements.Num() <= SchemaParameters.Elements.Num());
 
 			// Update sub-elements
 
@@ -2225,7 +2225,7 @@ namespace UE::Learning::Action
 			for (int32 ObjectElementIdx = 0; ObjectElementIdx < ObjectParameters.Elements.Num(); ObjectElementIdx++)
 			{
 				const int32 SchemaElementIdx = SchemaParameters.ElementNames.Find(ObjectParameters.ElementNames[ObjectElementIdx]);
-				UE_LEARNING_CHECK(SchemaElementIdx != INDEX_NONE);
+				check(SchemaElementIdx != INDEX_NONE);
 
 				const int32 SubElementSize = Schema.GetActionVectorSize(SchemaParameters.Elements[SchemaElementIdx]);
 
@@ -2241,12 +2241,12 @@ namespace UE::Learning::Action
 
 			// Set Mask
 
-			UE_LEARNING_CHECK(SubElementOffset + SchemaParameters.Elements.Num() == OutActionVector.Num());
+			check(SubElementOffset + SchemaParameters.Elements.Num() == OutActionVector.Num());
 
 			for (int32 ObjectElementIdx = 0; ObjectElementIdx < ObjectParameters.Elements.Num(); ObjectElementIdx++)
 			{
 				const int32 SchemaElementIdx = SchemaParameters.ElementNames.Find(ObjectParameters.ElementNames[ObjectElementIdx]);
-				UE_LEARNING_CHECK(SchemaElementIdx != INDEX_NONE);
+				check(SchemaElementIdx != INDEX_NONE);
 
 				OutActionVector[SubElementOffset + SchemaElementIdx] = 1.0f;
 			}
@@ -2260,7 +2260,7 @@ namespace UE::Learning::Action
 
 			const FSchemaArrayParameters SchemaParameters = Schema.GetArray(SchemaElement);
 			const FObjectArrayParameters ObjectParameters = Object.GetArray(ObjectElement);
-			UE_LEARNING_CHECK(SchemaParameters.Num == ObjectParameters.Elements.Num());
+			check(SchemaParameters.Num == ObjectParameters.Elements.Num());
 
 			// Update sub-elements
 
@@ -2296,7 +2296,7 @@ namespace UE::Learning::Action
 
 		default:
 		{
-			UE_LEARNING_NOT_IMPLEMENTED();
+			checkNoEntry();
 			return;
 		}
 
@@ -2310,7 +2310,7 @@ namespace UE::Learning::Action
 		const FSchemaElement SchemaElement,
 		const TLearningArrayView<1, const float> ActionVector)
 	{
-		UE_LEARNING_CHECK(Schema.IsValid(SchemaElement));
+		check(Schema.IsValid(SchemaElement));
 
 		// Check that the types match
 
@@ -2320,7 +2320,7 @@ namespace UE::Learning::Action
 		// Get Action Vector Size
 
 		const int32 ActionVectorSize = ActionVector.Num();
-		UE_LEARNING_CHECK(ActionVectorSize == Schema.GetActionVectorSize(SchemaElement));
+		check(ActionVectorSize == Schema.GetActionVectorSize(SchemaElement));
 
 		// Logic for each specific element type
 
@@ -2336,7 +2336,7 @@ namespace UE::Learning::Action
 		case EType::Continuous:
 		{
 			const FSchemaContinuousParameters SchemaParameters = Schema.GetContinuous(SchemaElement);
-			UE_LEARNING_CHECK(ActionVectorSize == SchemaParameters.Num);
+			check(ActionVectorSize == SchemaParameters.Num);
 
 			const int32 ValueNum = SchemaParameters.Num;
 			const float ValueScale = FMath::Max(SchemaParameters.Scale, UE_SMALL_NUMBER);
@@ -2354,20 +2354,20 @@ namespace UE::Learning::Action
 
 		case EType::DiscreteExclusive:
 		{
-			UE_LEARNING_CHECK(ActionVectorSize == Schema.GetDiscreteExclusive(SchemaElement).Num);
+			check(ActionVectorSize == Schema.GetDiscreteExclusive(SchemaElement).Num);
 
 			// Find Index
 			int32 ExclusiveIndex = INDEX_NONE;
 			for (int32 Idx = 0; Idx < ActionVectorSize; Idx++)
 			{
-				UE_LEARNING_CHECK(ActionVector[Idx] == 0.0f || ActionVector[Idx] == 1.0f);
+				check(ActionVector[Idx] == 0.0f || ActionVector[Idx] == 1.0f);
 				if (ActionVector[Idx])
 				{
 					ExclusiveIndex = Idx;
 					break;
 				}
 			}
-			UE_LEARNING_CHECK(ExclusiveIndex != INDEX_NONE);
+			check(ExclusiveIndex != INDEX_NONE);
 
 			OutObjectElement = OutObject.CreateDiscreteExclusive({ ExclusiveIndex }, SchemaElementTag);
 			return;
@@ -2375,14 +2375,14 @@ namespace UE::Learning::Action
 
 		case EType::DiscreteInclusive:
 		{
-			UE_LEARNING_CHECK(ActionVectorSize == Schema.GetDiscreteInclusive(SchemaElement).Num);
+			check(ActionVectorSize == Schema.GetDiscreteInclusive(SchemaElement).Num);
 
 			// Find Indices
 			TArray<int32, TInlineAllocator<8>> InclusiveIndices;
 			InclusiveIndices.Reserve(ActionVectorSize);
 			for (int32 Idx = 0; Idx < ActionVectorSize; Idx++)
 			{
-				UE_LEARNING_CHECK(ActionVector[Idx] == 0.0f || ActionVector[Idx] == 1.0f);
+				check(ActionVector[Idx] == 0.0f || ActionVector[Idx] == 1.0f);
 				if (ActionVector[Idx])
 				{
 					InclusiveIndices.Add(Idx);
@@ -2416,7 +2416,7 @@ namespace UE::Learning::Action
 
 				SubElementOffset += SubElementSize;
 			}
-			UE_LEARNING_CHECK(SubElementOffset == ActionVectorSize);
+			check(SubElementOffset == ActionVectorSize);
 
 			OutObjectElement = OutObject.CreateAnd({ Parameters.ElementNames, SubElements }, SchemaElementTag);
 			return;
@@ -2433,14 +2433,14 @@ namespace UE::Learning::Action
 			int32 SchemaElementIndex = INDEX_NONE;
 			for (int32 SubElementIdx = 0; SubElementIdx < Parameters.Elements.Num(); SubElementIdx++)
 			{
-				UE_LEARNING_CHECK(ActionVector[MaxSubElementSize + SubElementIdx] == 0.0f || ActionVector[MaxSubElementSize + SubElementIdx] == 1.0f);
+				check(ActionVector[MaxSubElementSize + SubElementIdx] == 0.0f || ActionVector[MaxSubElementSize + SubElementIdx] == 1.0f);
 				if (ActionVector[MaxSubElementSize + SubElementIdx])
 				{
 					SchemaElementIndex = SubElementIdx;
 					break;
 				}
 			}
-			UE_LEARNING_CHECK(SchemaElementIndex != INDEX_NONE);
+			check(SchemaElementIndex != INDEX_NONE);
 
 			// Create sub-element
 
@@ -2478,7 +2478,7 @@ namespace UE::Learning::Action
 			{
 				const int32 SubElementSize = Schema.GetActionVectorSize(Parameters.Elements[SubElementIdx]);
 
-				UE_LEARNING_CHECK(ActionVector[TotalSubElementSize + SubElementIdx] == 0.0f || ActionVector[TotalSubElementSize + SubElementIdx] == 1.0f);
+				check(ActionVector[TotalSubElementSize + SubElementIdx] == 0.0f || ActionVector[TotalSubElementSize + SubElementIdx] == 1.0f);
 				if (ActionVector[TotalSubElementSize + SubElementIdx])
 				{
 					FObjectElement SubElement;
@@ -2495,7 +2495,7 @@ namespace UE::Learning::Action
 
 				SubElementOffset += SubElementSize;
 			}
-			UE_LEARNING_CHECK(SubElementOffset + Parameters.Elements.Num() == ActionVectorSize);
+			check(SubElementOffset + Parameters.Elements.Num() == ActionVectorSize);
 
 			OutObjectElement = OutObject.CreateOrInclusive({ SubElementNames, SubElements }, SchemaElementTag);
 			return;
@@ -2544,7 +2544,7 @@ namespace UE::Learning::Action
 
 		default:
 		{
-			UE_LEARNING_NOT_IMPLEMENTED();
+			checkNoEntry();
 			OutObjectElement = FObjectElement();
 			return;
 		}
@@ -2558,15 +2558,15 @@ namespace UE::Learning::Action
 		const FModifier& Modifier,
 		const FModifierElement ModifierElement)
 	{
-		UE_LEARNING_CHECK(Schema.IsValid(SchemaElement));
-		UE_LEARNING_CHECK(Modifier.IsValid(ModifierElement));
-		UE_LEARNING_CHECK(OutActionModifierVector.Num() == Schema.GetActionModifierVectorSize(SchemaElement));
+		check(Schema.IsValid(SchemaElement));
+		check(Modifier.IsValid(ModifierElement));
+		check(OutActionModifierVector.Num() == Schema.GetActionModifierVectorSize(SchemaElement));
 
 		// Check that the types match
 
 		const EType SchemaElementType = Schema.GetType(SchemaElement);
 		const EType ModifierElementType = Modifier.GetType(ModifierElement);
-		UE_LEARNING_CHECK(ModifierElementType == EType::Null || ModifierElementType == SchemaElementType);
+		check(ModifierElementType == EType::Null || ModifierElementType == SchemaElementType);
 
 		// Zero Action Modifier Vector and return early if we have a null Type
 
@@ -2589,7 +2589,7 @@ namespace UE::Learning::Action
 		case EType::Null:
 		{
 			// This should never be reached
-			UE_LEARNING_NOT_IMPLEMENTED();
+			checkNoEntry();
 			break;
 		}
 
@@ -2602,9 +2602,9 @@ namespace UE::Learning::Action
 
 			const TArrayView<const bool> Masked = Modifier.GetContinuous(ModifierElement).Masked;
 			const TArrayView<const float> MaskedValues = Modifier.GetContinuous(ModifierElement).MaskedValues;
-			UE_LEARNING_CHECK(Masked.Num() == ValueNum);
-			UE_LEARNING_CHECK(MaskedValues.Num() == ValueNum);
-			UE_LEARNING_CHECK(Schema.GetActionModifierVectorSize(SchemaElement) == 1 + Masked.Num() + MaskedValues.Num());
+			check(Masked.Num() == ValueNum);
+			check(MaskedValues.Num() == ValueNum);
+			check(Schema.GetActionModifierVectorSize(SchemaElement) == 1 + Masked.Num() + MaskedValues.Num());
 
 			for (int32 ValueIdx = 0; ValueIdx < ValueNum; ValueIdx++)
 			{
@@ -2618,11 +2618,11 @@ namespace UE::Learning::Action
 		case EType::DiscreteExclusive:
 		{
 			const TArrayView<const int32> MaskIndices = Modifier.GetDiscreteExclusive(ModifierElement).MaskedIndices;
-			UE_LEARNING_CHECK(Schema.GetDiscreteExclusive(SchemaElement).Num >= MaskIndices.Num());
+			check(Schema.GetDiscreteExclusive(SchemaElement).Num >= MaskIndices.Num());
 
 			for (int32 MaskIndicesIdx = 0; MaskIndicesIdx < MaskIndices.Num(); MaskIndicesIdx++)
 			{
-				UE_LEARNING_CHECK(Schema.GetDiscreteExclusive(SchemaElement).Num > MaskIndices[MaskIndicesIdx] && MaskIndices[MaskIndicesIdx] >= 0);
+				check(Schema.GetDiscreteExclusive(SchemaElement).Num > MaskIndices[MaskIndicesIdx] && MaskIndices[MaskIndicesIdx] >= 0);
 				OutActionModifierVector[1 + MaskIndices[MaskIndicesIdx]] = 1.0f;
 			}
 
@@ -2632,11 +2632,11 @@ namespace UE::Learning::Action
 		case EType::DiscreteInclusive:
 		{
 			const TArrayView<const int32> MaskIndices = Modifier.GetDiscreteInclusive(ModifierElement).MaskedIndices;
-			UE_LEARNING_CHECK(Schema.GetDiscreteInclusive(SchemaElement).Num >= MaskIndices.Num());
+			check(Schema.GetDiscreteInclusive(SchemaElement).Num >= MaskIndices.Num());
 
 			for (int32 MaskIndicesIdx = 0; MaskIndicesIdx < MaskIndices.Num(); MaskIndicesIdx++)
 			{
-				UE_LEARNING_CHECK(Schema.GetDiscreteInclusive(SchemaElement).Num > MaskIndices[MaskIndicesIdx] && MaskIndices[MaskIndicesIdx] >= 0);
+				check(Schema.GetDiscreteInclusive(SchemaElement).Num > MaskIndices[MaskIndicesIdx] && MaskIndices[MaskIndicesIdx] >= 0);
 				OutActionModifierVector[1 + MaskIndices[MaskIndicesIdx]] = 1.0f;
 			}
 
@@ -2648,7 +2648,7 @@ namespace UE::Learning::Action
 			const FSchemaAndParameters SchemaParameters = Schema.GetAnd(SchemaElement);
 			const FModifierAndParameters ModifierParameters = Modifier.GetAnd(ModifierElement);
 
-			UE_LEARNING_CHECK(OutActionModifierVector.Num() == 1 + Private::GetTotalActionModifierVectorSize(Schema, SchemaParameters.Elements));
+			check(OutActionModifierVector.Num() == 1 + Private::GetTotalActionModifierVectorSize(Schema, SchemaParameters.Elements));
 
 			// Set the Sub-elements
 
@@ -2672,7 +2672,7 @@ namespace UE::Learning::Action
 				SubElementOffset += SubElementSize;
 			}
 
-			UE_LEARNING_CHECK(SubElementOffset == OutActionModifierVector.Num());
+			check(SubElementOffset == OutActionModifierVector.Num());
 			return;
 		}
 
@@ -2681,14 +2681,14 @@ namespace UE::Learning::Action
 			const FSchemaOrExclusiveParameters SchemaParameters = Schema.GetOrExclusive(SchemaElement);
 			const FModifierOrExclusiveParameters ModifierParameters = Modifier.GetOrExclusive(ModifierElement);
 
-			UE_LEARNING_CHECK(OutActionModifierVector.Num() == 1 + SchemaParameters.Elements.Num() + Private::GetTotalActionModifierVectorSize(Schema, SchemaParameters.Elements));
+			check(OutActionModifierVector.Num() == 1 + SchemaParameters.Elements.Num() + Private::GetTotalActionModifierVectorSize(Schema, SchemaParameters.Elements));
 
 			// Set the Mask
 
 			for (int32 MaskElementIdx = 0; MaskElementIdx < ModifierParameters.MaskedElements.Num(); MaskElementIdx++)
 			{
 				const int32 SchemaMaskElementIdx = SchemaParameters.ElementNames.Find(ModifierParameters.MaskedElements[MaskElementIdx]);
-				UE_LEARNING_CHECK(SchemaMaskElementIdx != INDEX_NONE);
+				check(SchemaMaskElementIdx != INDEX_NONE);
 				OutActionModifierVector[1 + SchemaMaskElementIdx] = 1.0f;
 			}
 
@@ -2714,7 +2714,7 @@ namespace UE::Learning::Action
 				SubElementOffset += SubElementSize;
 			}
 
-			UE_LEARNING_CHECK(SubElementOffset == OutActionModifierVector.Num());
+			check(SubElementOffset == OutActionModifierVector.Num());
 			return;
 		}
 
@@ -2723,14 +2723,14 @@ namespace UE::Learning::Action
 			const FSchemaOrInclusiveParameters SchemaParameters = Schema.GetOrInclusive(SchemaElement);
 			const FModifierOrInclusiveParameters ModifierParameters = Modifier.GetOrInclusive(ModifierElement);
 
-			UE_LEARNING_CHECK(OutActionModifierVector.Num() == 1 + SchemaParameters.Elements.Num() + Private::GetTotalActionModifierVectorSize(Schema, SchemaParameters.Elements));
+			check(OutActionModifierVector.Num() == 1 + SchemaParameters.Elements.Num() + Private::GetTotalActionModifierVectorSize(Schema, SchemaParameters.Elements));
 
 			// Set the Mask
 
 			for (int32 MaskElementIdx = 0; MaskElementIdx < ModifierParameters.MaskedElements.Num(); MaskElementIdx++)
 			{
 				const int32 SchemaMaskElementIdx = SchemaParameters.ElementNames.Find(ModifierParameters.MaskedElements[MaskElementIdx]);
-				UE_LEARNING_CHECK(SchemaMaskElementIdx != INDEX_NONE);
+				check(SchemaMaskElementIdx != INDEX_NONE);
 				OutActionModifierVector[1 + SchemaMaskElementIdx] = 1.0f;
 			}
 
@@ -2756,7 +2756,7 @@ namespace UE::Learning::Action
 				SubElementOffset += SubElementSize;
 			}
 
-			UE_LEARNING_CHECK(SubElementOffset == OutActionModifierVector.Num());
+			check(SubElementOffset == OutActionModifierVector.Num());
 			return;
 		}
 
@@ -2766,7 +2766,7 @@ namespace UE::Learning::Action
 
 			const FSchemaArrayParameters SchemaParameters = Schema.GetArray(SchemaElement);
 			const FModifierArrayParameters ModifierParameters = Modifier.GetArray(ModifierElement);
-			UE_LEARNING_CHECK(SchemaParameters.Num == ModifierParameters.Elements.Num());
+			check(SchemaParameters.Num == ModifierParameters.Elements.Num());
 
 			// Update sub-elements
 
@@ -2804,7 +2804,7 @@ namespace UE::Learning::Action
 
 		default:
 		{
-			UE_LEARNING_NOT_IMPLEMENTED();
+			checkNoEntry();
 			return;
 		}
 
@@ -2818,7 +2818,7 @@ namespace UE::Learning::Action
 		const FSchemaElement SchemaElement,
 		const TLearningArrayView<1, const float> ActionModifierVector)
 	{
-		UE_LEARNING_CHECK(Schema.IsValid(SchemaElement));
+		check(Schema.IsValid(SchemaElement));
 
 		// Get Type and Tag
 
@@ -2828,12 +2828,12 @@ namespace UE::Learning::Action
 		// Get Action Modifier Vector Size
 
 		const int32 ActionModifierVectorSize = ActionModifierVector.Num();
-		UE_LEARNING_CHECK(ActionModifierVectorSize == Schema.GetActionModifierVectorSize(SchemaElement));
+		check(ActionModifierVectorSize == Schema.GetActionModifierVectorSize(SchemaElement));
 
 		// We always have at least one element in the ActionModifierVector which says if the element is provided
 		// if this first value is zero then it means nothing below is masked and we always just return the null element
 
-		UE_LEARNING_CHECK(ActionModifierVectorSize > 0);
+		check(ActionModifierVectorSize > 0);
 
 		if (ActionModifierVector[0] == 0.0f)
 		{
@@ -2842,7 +2842,7 @@ namespace UE::Learning::Action
 		}
 		else
 		{
-			UE_LEARNING_CHECK(ActionModifierVector[0] == 1.0f);
+			check(ActionModifierVector[0] == 1.0f);
 		}
 
 		// Logic for each specific element type
@@ -2859,7 +2859,7 @@ namespace UE::Learning::Action
 		case EType::Continuous:
 		{
 			const FSchemaContinuousParameters SchemaParameters = Schema.GetContinuous(SchemaElement);
-			UE_LEARNING_CHECK(ActionModifierVectorSize == 1 + 2 * SchemaParameters.Num);
+			check(ActionModifierVectorSize == 1 + 2 * SchemaParameters.Num);
 
 			const int32 ValueNum = SchemaParameters.Num;
 
@@ -2885,14 +2885,14 @@ namespace UE::Learning::Action
 		{
 			const int32 ValueNum = Schema.GetDiscreteExclusive(SchemaElement).Num;
 
-			UE_LEARNING_CHECK(ActionModifierVectorSize == 1 + ValueNum);
+			check(ActionModifierVectorSize == 1 + ValueNum);
 
 			// Find Indices
 			TArray<int32, TInlineAllocator<8>> MaskedIndices;
 			MaskedIndices.Reserve(ValueNum);
 			for (int32 Idx = 0; Idx < ValueNum; Idx++)
 			{
-				UE_LEARNING_CHECK(ActionModifierVector[1 + Idx] == 0.0f || ActionModifierVector[1 + Idx] == 1.0f);
+				check(ActionModifierVector[1 + Idx] == 0.0f || ActionModifierVector[1 + Idx] == 1.0f);
 				if (ActionModifierVector[1 + Idx] == 1.0f)
 				{
 					MaskedIndices.Add(Idx);
@@ -2907,14 +2907,14 @@ namespace UE::Learning::Action
 		{
 			const int32 ValueNum = Schema.GetDiscreteInclusive(SchemaElement).Num;
 
-			UE_LEARNING_CHECK(ActionModifierVectorSize == 1 + ValueNum);
+			check(ActionModifierVectorSize == 1 + ValueNum);
 
 			// Find Indices
 			TArray<int32, TInlineAllocator<8>> MaskedIndices;
 			MaskedIndices.Reserve(ValueNum);
 			for (int32 Idx = 0; Idx < ValueNum; Idx++)
 			{
-				UE_LEARNING_CHECK(ActionModifierVector[1 + Idx] == 0.0f || ActionModifierVector[1 + Idx] == 1.0f);
+				check(ActionModifierVector[1 + Idx] == 0.0f || ActionModifierVector[1 + Idx] == 1.0f);
 				if (ActionModifierVector[1 + Idx] == 1.0f)
 				{
 					MaskedIndices.Add(Idx);
@@ -2948,7 +2948,7 @@ namespace UE::Learning::Action
 
 				SubElementOffset += SubElementSize;
 			}
-			UE_LEARNING_CHECK(SubElementOffset == ActionModifierVectorSize);
+			check(SubElementOffset == ActionModifierVectorSize);
 
 			OutModifierElement = OutModifier.CreateAnd({ Parameters.ElementNames, SubElements }, SchemaElementTag);
 			return;
@@ -2966,7 +2966,7 @@ namespace UE::Learning::Action
 
 			for (int32 Idx = 0; Idx < SubElementNum; Idx++)
 			{
-				UE_LEARNING_CHECK(ActionModifierVector[1 + Idx] == 0.0f || ActionModifierVector[1 + Idx] == 1.0f);
+				check(ActionModifierVector[1 + Idx] == 0.0f || ActionModifierVector[1 + Idx] == 1.0f);
 				if (ActionModifierVector[1 + Idx] == 1.0f)
 				{
 					MaskedElements.Add(Parameters.ElementNames[Idx]);
@@ -2992,7 +2992,7 @@ namespace UE::Learning::Action
 
 				SubElementOffset += SubElementSize;
 			}
-			UE_LEARNING_CHECK(SubElementOffset == ActionModifierVectorSize);
+			check(SubElementOffset == ActionModifierVectorSize);
 
 			OutModifierElement = OutModifier.CreateOrExclusive({ Parameters.ElementNames, SubElements, MaskedElements }, SchemaElementTag);
 			return;
@@ -3010,7 +3010,7 @@ namespace UE::Learning::Action
 
 			for (int32 Idx = 0; Idx < SubElementNum; Idx++)
 			{
-				UE_LEARNING_CHECK(ActionModifierVector[1 + Idx] == 0.0f || ActionModifierVector[1 + Idx] == 1.0f);
+				check(ActionModifierVector[1 + Idx] == 0.0f || ActionModifierVector[1 + Idx] == 1.0f);
 				if (ActionModifierVector[1 + Idx] == 1.0f)
 				{
 					MaskedElements.Add(Parameters.ElementNames[Idx]);
@@ -3036,7 +3036,7 @@ namespace UE::Learning::Action
 
 				SubElementOffset += SubElementSize;
 			}
-			UE_LEARNING_CHECK(SubElementOffset == ActionModifierVectorSize);
+			check(SubElementOffset == ActionModifierVectorSize);
 
 			OutModifierElement = OutModifier.CreateOrInclusive({ Parameters.ElementNames, SubElements, MaskedElements }, SchemaElementTag);
 			return;
@@ -3086,7 +3086,7 @@ namespace UE::Learning::Action
 
 		default:
 		{
-			UE_LEARNING_NOT_IMPLEMENTED();
+			checkNoEntry();
 			OutModifierElement = FModifierElement();
 			return;
 		}
diff --git a/Engine/Plugins/Experimental/LearningAgents/Source/Learning/Private/LearningAdamOptimizer.cpp b/Engine/Plugins/Experimental/LearningAgents/Source/Learning/Private/LearningAdamOptimizer.cpp
--- a/Engine/Plugins/Experimental/LearningAgents/Source/Learning/Private/LearningAdamOptimizer.cpp
+++ b/Engine/Plugins/Experimental/LearningAgents/Source/Learning/Private/LearningAdamOptimizer.cpp
@@ -14,7 +14,7 @@ namespace UE::Learning
 
 	void FAdamOptimizer::Resize(const int32 SampleNum, const int32 DimensionsNum)
 	{
-		UE_LEARNING_TRACE_CPUPROFILER_EVENT_SCOPE(Learning::FAdamOptimizer::Resize);
+		TRACE_CPUPROFILER_EVENT_SCOPE(Learning::FAdamOptimizer::Resize);
 
 		Iterations.SetNumUninitialized({ DimensionsNum });
 		Estimate.SetNumUninitialized({ DimensionsNum });
@@ -29,7 +29,7 @@ namespace UE::Learning
 
 	void FAdamOptimizer::Sample(TLearningArrayView<2, float> OutSamples)
 	{
-		UE_LEARNING_TRACE_CPUPROFILER_EVENT_SCOPE(Learning::FAdamOptimizer::Sample);
+		TRACE_CPUPROFILER_EVENT_SCOPE(Learning::FAdamOptimizer::Sample);
 
 		const int32 SampleNum = OutSamples.Num<0>();
 		const int32 DimNum = OutSamples.Num<1>();
@@ -65,7 +65,7 @@ namespace UE::Learning
 
 	void FAdamOptimizer::Reset(TLearningArrayView<2, float> OutSamples, const TLearningArrayView<1, const float> InitialGuess)
 	{
-		UE_LEARNING_TRACE_CPUPROFILER_EVENT_SCOPE(Learning::FAdamOptimizer::Reset);
+		TRACE_CPUPROFILER_EVENT_SCOPE(Learning::FAdamOptimizer::Reset);
 
 		Array::Zero(Iterations);
 		Array::Copy(Estimate, InitialGuess);
@@ -83,7 +83,7 @@ namespace UE::Learning
 		const TLearningArrayView<1, const float> Losses,
 		const ELogSetting LogSettings)
 	{
-		UE_LEARNING_TRACE_CPUPROFILER_EVENT_SCOPE(Learning::FAdamOptimizer::Update);
+		TRACE_CPUPROFILER_EVENT_SCOPE(Learning::FAdamOptimizer::Update);
 
 		const int32 SampleNum = InOutSamples.Num<0>();
 		const int32 DimNum = InOutSamples.Num<1>();
@@ -92,7 +92,7 @@ namespace UE::Learning
 		Array::Check(Losses);
 
 		{
-			UE_LEARNING_TRACE_CPUPROFILER_EVENT_SCOPE(Learning::FAdamOptimizer::Update::EstimateGradient);
+			TRACE_CPUPROFILER_EVENT_SCOPE(Learning::FAdamOptimizer::Update::EstimateGradient);
 
 #if UE_LEARNING_ISPC
 			ispc::LearningEstimateGradient(
@@ -130,7 +130,7 @@ namespace UE::Learning
 		}
 
 		{
-			UE_LEARNING_TRACE_CPUPROFILER_EVENT_SCOPE(Learning::FAdamOptimizer::Update::UpdateEstimate);
+			TRACE_CPUPROFILER_EVENT_SCOPE(Learning::FAdamOptimizer::Update::UpdateEstimate);
 
 #if UE_LEARNING_ISPC
 			ispc::LearningUpdateAdamEstimate(
diff --git a/Engine/Plugins/Experimental/LearningAgents/Source/Learning/Private/LearningArray.cpp b/Engine/Plugins/Experimental/LearningAgents/Source/Learning/Private/LearningArray.cpp
--- a/Engine/Plugins/Experimental/LearningAgents/Source/Learning/Private/LearningArray.cpp
+++ b/Engine/Plugins/Experimental/LearningAgents/Source/Learning/Private/LearningArray.cpp
@@ -52,16 +52,16 @@ namespace UE::Learning
 
 	void SlicedParallelFor(const int32 Num, const int32 MinSliceElementNum, const TFunctionRef<void(int32 Start, int32 Length)> Body)
 	{
-		UE_LEARNING_TRACE_CPUPROFILER_EVENT_SCOPE(Learning::SlicedParallelFor);
+		TRACE_CPUPROFILER_EVENT_SCOPE(Learning::SlicedParallelFor);
 
-		UE_LEARNING_CHECK(Num >= 0);
+		check(Num >= 0);
 
 		if (Num <= 0)
 		{
 			return;
 		}
 
-		UE_LEARNING_CHECK(MinSliceElementNum >= 1);
+		check(MinSliceElementNum >= 1);
 
 		const int32 IdealTaskNum = ParallelForImpl::GetNumberOfThreadTasks(Num, MinSliceElementNum, EParallelForFlags::None);
 		const int32 IdealSliceLength = FMath::DivideAndRoundUp(Num, IdealTaskNum);
diff --git a/Engine/Plugins/Experimental/LearningAgents/Source/Learning/Private/LearningCMAOptimizer.cpp b/Engine/Plugins/Experimental/LearningAgents/Source/Learning/Private/LearningCMAOptimizer.cpp
--- a/Engine/Plugins/Experimental/LearningAgents/Source/Learning/Private/LearningCMAOptimizer.cpp
+++ b/Engine/Plugins/Experimental/LearningAgents/Source/Learning/Private/LearningCMAOptimizer.cpp
@@ -20,7 +20,7 @@ namespace UE::Learning
 
 	void FCMAOptimizer::Resize(const int32 SampleNum, const int32 DimensionsNum)
 	{
-		UE_LEARNING_TRACE_CPUPROFILER_EVENT_SCOPE(Learning::FCMAOptimizer::Resize);
+		TRACE_CPUPROFILER_EVENT_SCOPE(Learning::FCMAOptimizer::Resize);
 
 		PathSigma.SetNumUninitialized({ DimensionsNum });
 		PathCovariance.SetNumUninitialized({ DimensionsNum });
@@ -60,7 +60,7 @@ namespace UE::Learning
 		TLearningArrayView<2, float> OutSamples,
 		const TLearningArrayView<1, const float> InitialGuess)
 	{
-		UE_LEARNING_TRACE_CPUPROFILER_EVENT_SCOPE(Learning::FCMAOptimizer::Reset);
+		TRACE_CPUPROFILER_EVENT_SCOPE(Learning::FCMAOptimizer::Reset);
 
 		const int32 SampleNum = Weights.Num<0>();
 		const int32 DimNum = InitialGuess.Num<0>();
@@ -142,7 +142,7 @@ namespace UE::Learning
 		const TLearningArrayView<1, const float> Losses,
 		const ELogSetting LogSettings)
 	{
-		UE_LEARNING_TRACE_CPUPROFILER_EVENT_SCOPE(Learning::FCMAOptimizer::Update);
+		TRACE_CPUPROFILER_EVENT_SCOPE(Learning::FCMAOptimizer::Update);
 
 		Array::Check(InOutSamples);
 		Array::Check(Losses);
@@ -184,7 +184,7 @@ namespace UE::Learning
 		// Update Mean as weighted sum of controls
 
 		{
-			UE_LEARNING_TRACE_CPUPROFILER_EVENT_SCOPE(Learning::FCMAOptimizer::Update::AdjustMean);
+			TRACE_CPUPROFILER_EVENT_SCOPE(Learning::FCMAOptimizer::Update::AdjustMean);
 
 #if UE_LEARNING_ISPC
 			ispc::LearningAdjustCMAMean(
@@ -215,7 +215,7 @@ namespace UE::Learning
 		// Compute Update Direction
 
 		{
-			UE_LEARNING_TRACE_CPUPROFILER_EVENT_SCOPE(Learning::FCMAOptimizer::Update::ComputeUpdateDirection);
+			TRACE_CPUPROFILER_EVENT_SCOPE(Learning::FCMAOptimizer::Update::ComputeUpdateDirection);
 
 #if UE_LEARNING_ISPC
 			ispc::LearningComputeCMAUpdateDirection(
@@ -234,7 +234,7 @@ namespace UE::Learning
 		bool bHsig = false;
 
 		{
-			UE_LEARNING_TRACE_CPUPROFILER_EVENT_SCOPE(Learning::FCMAOptimizer::Update::UpdatePaths);
+			TRACE_CPUPROFILER_EVENT_SCOPE(Learning::FCMAOptimizer::Update::UpdatePaths);
 
 			float SumSquarePs = 0.0f;
 
@@ -265,7 +265,7 @@ namespace UE::Learning
 		// Update Covariance and Sigma
 
 		{
-			UE_LEARNING_TRACE_CPUPROFILER_EVENT_SCOPE(Learning::FCMAOptimizer::Update::UpdateCovariance);
+			TRACE_CPUPROFILER_EVENT_SCOPE(Learning::FCMAOptimizer::Update::UpdateCovariance);
 
 			// Original reference code here squares the bool... no idea why
 			const float C1A = C1 * (1 - (1 - FMath::Square((float)bHsig)) * CC * (2 - CC));
@@ -299,7 +299,7 @@ namespace UE::Learning
 
 		// Update Sampling Transform
 		{
-			UE_LEARNING_TRACE_CPUPROFILER_EVENT_SCOPE(Learning::FCMAOptimizer::Update::SamplingTransform);
+			TRACE_CPUPROFILER_EVENT_SCOPE(Learning::FCMAOptimizer::Update::SamplingTransform);
 
 			Eigen::SelfAdjointEigenSolver<Eigen::MatrixXf> EigenSolver(InEigenColMatrix(Covariance));
 			check(EigenSolver.info() == Eigen::Success);
diff --git a/Engine/Plugins/Experimental/LearningAgents/Source/Learning/Private/LearningCompletion.cpp b/Engine/Plugins/Experimental/LearningAgents/Source/Learning/Private/LearningCompletion.cpp
--- a/Engine/Plugins/Experimental/LearningAgents/Source/Learning/Private/LearningCompletion.cpp
+++ b/Engine/Plugins/Experimental/LearningAgents/Source/Learning/Private/LearningCompletion.cpp
@@ -51,7 +51,7 @@ namespace UE::Learning
 			case ECompletionMode::Running: return TEXT("Running");
 			case ECompletionMode::Truncated: return TEXT("Truncated");
 			case ECompletionMode::Terminated: return TEXT("Terminated");
-			default: UE_LEARNING_NOT_IMPLEMENTED(); return TEXT("Unknown");
+			default: checkNoEntry(); return TEXT("Unknown");
 			}
 		}
 
diff --git a/Engine/Plugins/Experimental/LearningAgents/Source/Learning/Private/LearningCritic.cpp b/Engine/Plugins/Experimental/LearningAgents/Source/Learning/Private/LearningCritic.cpp
--- a/Engine/Plugins/Experimental/LearningAgents/Source/Learning/Private/LearningCritic.cpp
+++ b/Engine/Plugins/Experimental/LearningAgents/Source/Learning/Private/LearningCritic.cpp
@@ -13,8 +13,8 @@ namespace UE::Learning
 		: ObservationEncodedNum(InObservationEncodedNum)
 		, MemoryStateNum(InMemoryStateNum)
 	{
-		UE_LEARNING_CHECK(InNeuralNetwork->GetInputSize() == ObservationEncodedNum + MemoryStateNum);
-		UE_LEARNING_CHECK(InNeuralNetwork->GetOutputSize() == 1);
+		check(InNeuralNetwork->GetInputSize() == ObservationEncodedNum + MemoryStateNum);
+		check(InNeuralNetwork->GetOutputSize() == 1);
 
 		Input.SetNumUninitialized({ InMaxInstanceNum, ObservationEncodedNum + MemoryStateNum });
 		Array::Zero(Input);
@@ -28,7 +28,7 @@ namespace UE::Learning
 		const TLearningArrayView<2, const float> InputMemoryState,
 		const FIndexSet Instances)
 	{
-		UE_LEARNING_TRACE_CPUPROFILER_EVENT_SCOPE(Learning::FNeuralNetworkCritic::Evaluate);
+		TRACE_CPUPROFILER_EVENT_SCOPE(Learning::FNeuralNetworkCritic::Evaluate);
 
 		Array::Check(InputObservationVectorsEncoded, Instances);
 		Array::Check(InputMemoryState, Instances);
diff --git a/Engine/Plugins/Experimental/LearningAgents/Source/Learning/Private/LearningEigen.h b/Engine/Plugins/Experimental/LearningAgents/Source/Learning/Private/LearningEigen.h
--- a/Engine/Plugins/Experimental/LearningAgents/Source/Learning/Private/LearningEigen.h
+++ b/Engine/Plugins/Experimental/LearningAgents/Source/Learning/Private/LearningEigen.h
@@ -22,7 +22,7 @@
 // In general there is not much harm in enabling this since it does
 // not perform expensive checks however this can still be disabled 
 // in editor mode to give some small improvement to performance
-#define UE_LEARNING_EIGEN_DEBUG_ENABLE UE_LEARNING_DEVELOPMENT
+#define UE_LEARNING_EIGEN_DEBUG_ENABLE (!UE_BUILD_SHIPPING)
 //#define UE_LEARNING_EIGEN_DEBUG_ENABLE 0
 
 #if UE_LEARNING_EIGEN_DEBUG_ENABLE
@@ -38,7 +38,7 @@
 	#undef eigen_assert
 	#endif
 
-	#define eigen_assert(x) UE_LEARNING_CHECK(x)
+	#define eigen_assert(x) check(x)
 #else
 	#ifdef EIGEN_MPL2_ONLY
 	#undef EIGEN_MPL2_ONLY
diff --git a/Engine/Plugins/Experimental/LearningAgents/Source/Learning/Private/LearningFrameRangeSet.cpp b/Engine/Plugins/Experimental/LearningAgents/Source/Learning/Private/LearningFrameRangeSet.cpp
--- a/Engine/Plugins/Experimental/LearningAgents/Source/Learning/Private/LearningFrameRangeSet.cpp
+++ b/Engine/Plugins/Experimental/LearningAgents/Source/Learning/Private/LearningFrameRangeSet.cpp
@@ -14,9 +14,9 @@ namespace UE::Learning
 			const int32 FrameNum = Frames.Num();
 			for (int32 FrameIdx = 0; FrameIdx < FrameNum - 1; FrameIdx++)
 			{
-				UE_LEARNING_CHECK(Frames[FrameIdx + 0] >= 0);
-				UE_LEARNING_CHECK(Frames[FrameIdx + 1] >= 0);
-				UE_LEARNING_CHECK(Frames[FrameIdx + 0] < Frames[FrameIdx + 1]);
+				check(Frames[FrameIdx + 0] >= 0);
+				check(Frames[FrameIdx + 1] >= 0);
+				check(Frames[FrameIdx + 0] < Frames[FrameIdx + 1]);
 			}
 		}
 
@@ -24,14 +24,14 @@ namespace UE::Learning
 			const TLearningArrayView<1, const int32> Starts, 
 			const TLearningArrayView<1, const int32> Lengths)
 		{
-			UE_LEARNING_CHECK(Starts.Num() == Lengths.Num());
+			check(Starts.Num() == Lengths.Num());
 
 			const int32 RangeNum = Starts.Num();
 			for (int32 RangeIdx = 0; RangeIdx < RangeNum - 1; RangeIdx++)
 			{
-				UE_LEARNING_CHECK(Lengths[RangeIdx + 0] > 0);
-				UE_LEARNING_CHECK(Lengths[RangeIdx + 1] > 0);
-				UE_LEARNING_CHECK(Starts[RangeIdx + 0] + Lengths[RangeIdx + 0] <= Starts[RangeIdx + 1]);
+				check(Lengths[RangeIdx + 0] > 0);
+				check(Lengths[RangeIdx + 1] > 0);
+				check(Starts[RangeIdx + 0] + Lengths[RangeIdx + 0] <= Starts[RangeIdx + 1]);
 			}
 		}
 
@@ -43,11 +43,11 @@ namespace UE::Learning
 			const int32 FrameNum = FrameOffsets.Num();
 			for (int32 FrameIdx = 0; FrameIdx < FrameNum - 1; FrameIdx++)
 			{
-				UE_LEARNING_CHECK(FrameOffsets[FrameIdx + 0] >= MinimumOffset);
-				UE_LEARNING_CHECK(FrameOffsets[FrameIdx + 1] >= MinimumOffset);
-				UE_LEARNING_CHECK(FrameOffsets[FrameIdx + 0] < MaximumOffset);
-				UE_LEARNING_CHECK(FrameOffsets[FrameIdx + 1] < MaximumOffset);
-				UE_LEARNING_CHECK(FrameOffsets[FrameIdx + 0] < FrameOffsets[FrameIdx + 1]);
+				check(FrameOffsets[FrameIdx + 0] >= MinimumOffset);
+				check(FrameOffsets[FrameIdx + 1] >= MinimumOffset);
+				check(FrameOffsets[FrameIdx + 0] < MaximumOffset);
+				check(FrameOffsets[FrameIdx + 1] < MaximumOffset);
+				check(FrameOffsets[FrameIdx + 0] < FrameOffsets[FrameIdx + 1]);
 			}
 		}
 
@@ -57,8 +57,8 @@ namespace UE::Learning
 			const TLearningArrayView<1, const int32> RhsStarts,
 			const TLearningArrayView<1, const int32> RhsLengths)
 		{
-			UE_LEARNING_CHECK(LhsStarts.Num() == LhsLengths.Num());
-			UE_LEARNING_CHECK(RhsStarts.Num() == RhsLengths.Num());
+			check(LhsStarts.Num() == LhsLengths.Num());
+			check(RhsStarts.Num() == RhsLengths.Num());
 
 			if (LhsStarts.Num() != RhsStarts.Num()) { return false; }
 
@@ -79,7 +79,7 @@ namespace UE::Learning
 			const TLearningArrayView<1, const int32> Lengths,
 			const int32 Frame)
 		{
-			UE_LEARNING_CHECK(Starts.Num() == Lengths.Num());
+			check(Starts.Num() == Lengths.Num());
 
 			const int32 RangeNum = Starts.Num();
 			for (int32 RangeIdx = 0; RangeIdx < RangeNum; RangeIdx++)
@@ -98,7 +98,7 @@ namespace UE::Learning
 			const float Time,
 			const float FrameDeltaTime)
 		{
-			UE_LEARNING_CHECK(Starts.Num() == Lengths.Num());
+			check(Starts.Num() == Lengths.Num());
 
 			const int32 RangeNum = Starts.Num();
 			for (int32 RangeIdx = 0; RangeIdx < RangeNum; RangeIdx++)
@@ -121,7 +121,7 @@ namespace UE::Learning
 			const TLearningArrayView<1, const int32> Lengths,
 			const int32 Frame)
 		{
-			UE_LEARNING_CHECK(Starts.Num() == Lengths.Num());
+			check(Starts.Num() == Lengths.Num());
 
 			const int32 RangeNum = Starts.Num();
 			for (int32 RangeIdx = 0; RangeIdx < RangeNum; RangeIdx++)
@@ -147,7 +147,7 @@ namespace UE::Learning
 			const float Time,
 			const float FrameDeltaTime)
 		{
-			UE_LEARNING_CHECK(Starts.Num() == Lengths.Num());
+			check(Starts.Num() == Lengths.Num());
 
 			const int32 RangeNum = Starts.Num();
 			for (int32 RangeIdx = 0; RangeIdx < RangeNum; RangeIdx++)
@@ -222,7 +222,7 @@ namespace UE::Learning
 				// Event from lhs and rhs coming at same time
 				else
 				{
-					UE_LEARNING_CHECK(LhsT == RhsT);
+					check(LhsT == RhsT);
 
 					bRhsActive = bRhsActiveNext;
 					RhsIndex++;
@@ -298,7 +298,7 @@ namespace UE::Learning
 				// Event from lhs and rhs coming at same time
 				else
 				{
-					UE_LEARNING_CHECK(LhsT == RhsT);
+					check(LhsT == RhsT);
 
 					bRhsActive = bRhsActiveNext;
 					RhsIndex++;
@@ -410,7 +410,7 @@ namespace UE::Learning
 				// Event from lhs and rhs coming at same time
 				else
 				{
-					UE_LEARNING_CHECK(LhsT == RhsT);
+					check(LhsT == RhsT);
 
 					// Activate output
 					if (!bOutActive && (bLhsActiveNext || bRhsActiveNext))
@@ -435,7 +435,7 @@ namespace UE::Learning
 			// Process any remaining lhs events
 			while (LhsIndex < LhsNum * 2)
 			{
-				UE_LEARNING_CHECK(RhsIndex == RhsNum * 2);
+				check(RhsIndex == RhsNum * 2);
 
 				const bool bLhsActiveNext = LhsIndex % 2 == 0;
 				const int32 LhsT = bLhsActiveNext ? LhsStarts[LhsIndex / 2] : LhsStarts[LhsIndex / 2] + LhsLengths[LhsIndex / 2];
@@ -461,7 +461,7 @@ namespace UE::Learning
 			// Process any remaining rhs events
 			while (RhsIndex < RhsNum * 2)
 			{
-				UE_LEARNING_CHECK(LhsIndex == LhsNum * 2);
+				check(LhsIndex == LhsNum * 2);
 
 				const bool bRhsActiveNext = RhsIndex % 2 == 0;
 				const int32 RhsT = bRhsActiveNext ? RhsStarts[RhsIndex / 2] : RhsStarts[RhsIndex / 2] + RhsLengths[RhsIndex / 2];
@@ -578,7 +578,7 @@ namespace UE::Learning
 				// Event from lhs and rhs coming at same time
 				else
 				{
-					UE_LEARNING_CHECK(LhsT == RhsT);
+					check(LhsT == RhsT);
 
 					// Activate output
 					if (!bOutActive && (bLhsActiveNext && bRhsActiveNext))
@@ -611,7 +611,7 @@ namespace UE::Learning
 			const TLearningArrayView<1, const int32> RangeLengths,
 			const int32 EntryFrameOffset)
 		{
-			UE_LEARNING_CHECK(OutFrameOffsets.Num() == RangeLengths.Num());
+			check(OutFrameOffsets.Num() == RangeLengths.Num());
 			
 			int32 Offset = EntryFrameOffset;
 			const int32 RangeNum = RangeLengths.Num();
@@ -708,12 +708,12 @@ namespace UE::Learning
 					// Deactivate output
 					else if (bOutActive && !bLhsActiveNext)
 					{
-						UE_LEARNING_CHECK(bLhsActive);
-						UE_LEARNING_CHECK(bRhsActive);
-						UE_LEARNING_CHECK(LhsActiveStart != INDEX_NONE);
-						UE_LEARNING_CHECK(LhsActiveOffset != INDEX_NONE);
-						UE_LEARNING_CHECK(RhsActiveStart != INDEX_NONE);
-						UE_LEARNING_CHECK(RhsActiveOffset != INDEX_NONE);
+						check(bLhsActive);
+						check(bRhsActive);
+						check(LhsActiveStart != INDEX_NONE);
+						check(LhsActiveOffset != INDEX_NONE);
+						check(RhsActiveStart != INDEX_NONE);
+						check(RhsActiveOffset != INDEX_NONE);
 
 						bOutActive = false;
 						OutLengths[OutIndex] = LhsT - OutStarts[OutIndex];
@@ -740,12 +740,12 @@ namespace UE::Learning
 					// Deactivate output
 					else if (bOutActive && !bRhsActiveNext)
 					{
-						UE_LEARNING_CHECK(bLhsActive);
-						UE_LEARNING_CHECK(bRhsActive);
-						UE_LEARNING_CHECK(LhsActiveStart != INDEX_NONE);
-						UE_LEARNING_CHECK(LhsActiveOffset != INDEX_NONE);
-						UE_LEARNING_CHECK(RhsActiveStart != INDEX_NONE);
-						UE_LEARNING_CHECK(RhsActiveOffset != INDEX_NONE);
+						check(bLhsActive);
+						check(bRhsActive);
+						check(LhsActiveStart != INDEX_NONE);
+						check(LhsActiveOffset != INDEX_NONE);
+						check(RhsActiveStart != INDEX_NONE);
+						check(RhsActiveOffset != INDEX_NONE);
 
 						bOutActive = false;
 						OutLengths[OutIndex] = RhsT - OutStarts[OutIndex];
@@ -763,7 +763,7 @@ namespace UE::Learning
 				// Event from lhs and rhs coming at same time
 				else
 				{
-					UE_LEARNING_CHECK(LhsT == RhsT);
+					check(LhsT == RhsT);
 
 					// Activate output
 					if (!bOutActive && (bLhsActiveNext && bRhsActiveNext))
@@ -774,12 +774,12 @@ namespace UE::Learning
 					// Deactivate output
 					else if (bOutActive && (!bLhsActiveNext || !bRhsActiveNext))
 					{
-						UE_LEARNING_CHECK(bLhsActive);
-						UE_LEARNING_CHECK(bRhsActive);
-						UE_LEARNING_CHECK(LhsActiveStart != INDEX_NONE);
-						UE_LEARNING_CHECK(LhsActiveOffset != INDEX_NONE);
-						UE_LEARNING_CHECK(RhsActiveStart != INDEX_NONE);
-						UE_LEARNING_CHECK(RhsActiveOffset != INDEX_NONE);
+						check(bLhsActive);
+						check(bRhsActive);
+						check(LhsActiveStart != INDEX_NONE);
+						check(LhsActiveOffset != INDEX_NONE);
+						check(RhsActiveStart != INDEX_NONE);
+						check(RhsActiveOffset != INDEX_NONE);
 
 						bOutActive = false;
 						OutLengths[OutIndex] = LhsT - OutStarts[OutIndex];
@@ -894,7 +894,7 @@ namespace UE::Learning
 				// Event from lhs and rhs coming at same time
 				else
 				{
-					UE_LEARNING_CHECK(LhsT == RhsT);
+					check(LhsT == RhsT);
 
 					// Activate output
 					if (!bOutActive && bLhsActiveNext && !bRhsActiveNext)
@@ -952,12 +952,12 @@ namespace UE::Learning
 		const int32 EntryNum = GetEntryNum();
 		for (int32 EntryIdx = 0; EntryIdx < EntryNum - 1; EntryIdx++)
 		{
-			UE_LEARNING_CHECK(EntrySequences[EntryIdx + 0] < EntrySequences[EntryIdx + 1])
+			check(EntrySequences[EntryIdx + 0] < EntrySequences[EntryIdx + 1])
 		}
 
 		for (int32 EntryIdx = 0; EntryIdx < EntryNum; EntryIdx++)
 		{
-			UE_LEARNING_CHECK(EntryRangeNums[EntryIdx] > 0);
+			check(EntryRangeNums[EntryIdx] > 0);
 			FrameRangeSet::Private::RangesCheck(GetEntryRangeStarts(EntryIdx), GetEntryRangeLengths(EntryIdx));
 		}
 	}
@@ -1175,8 +1175,8 @@ namespace UE::Learning
 		const TLearningArrayView<1, const int32> InStarts,
 		const TLearningArrayView<1, const int32> InLengths)
 	{
-		UE_LEARNING_CHECK(InStarts.Num() == InLengths.Num());
-		UE_LEARNING_CHECK(!ContainsSequence(InSequence));
+		check(InStarts.Num() == InLengths.Num());
+		check(!ContainsSequence(InSequence));
 		UE::Learning::FrameRangeSet::Private::RangesCheck(InStarts, InLengths);
 
 		if (InStarts.IsEmpty()) { return; }
@@ -1253,7 +1253,7 @@ namespace UE::Learning
 				// If both contain the same entry
 				else
 				{
-					UE_LEARNING_CHECK(FrameRangeSet.GetEntrySequence(RhsIndex) == FrameSet.GetEntrySequence(LhsIndex));
+					check(FrameRangeSet.GetEntrySequence(RhsIndex) == FrameSet.GetEntrySequence(LhsIndex));
 
 					// Append intersection of frames to output
 					const int32 EventNum = Private::FramesRangesIntersection(
@@ -1262,7 +1262,7 @@ namespace UE::Learning
 						FrameRangeSet.GetEntryRangeStarts(RhsIndex),
 						FrameRangeSet.GetEntryRangeLengths(RhsIndex));
 
-					UE_LEARNING_CHECK(EventNum <= FrameSet.GetEntryFrameNum(LhsIndex));
+					check(EventNum <= FrameSet.GetEntryFrameNum(LhsIndex));
 
 					if (EventNum > 0)
 					{
@@ -1340,7 +1340,7 @@ namespace UE::Learning
 				// If both contain the same entry
 				else
 				{
-					UE_LEARNING_CHECK(FrameRangeSet.GetEntrySequence(RhsIndex) == FrameSet.GetEntrySequence(LhsIndex));
+					check(FrameRangeSet.GetEntrySequence(RhsIndex) == FrameSet.GetEntrySequence(LhsIndex));
 
 					// Append difference of frames to output
 					const int32 EventNum = Private::FramesRangesDifference(
@@ -1349,7 +1349,7 @@ namespace UE::Learning
 						FrameRangeSet.GetEntryRangeStarts(RhsIndex),
 						FrameRangeSet.GetEntryRangeLengths(RhsIndex));
 
-					UE_LEARNING_CHECK(EventNum <= FrameSet.GetEntryFrameNum(LhsIndex));
+					check(EventNum <= FrameSet.GetEntryFrameNum(LhsIndex));
 
 					if (EventNum > 0)
 					{
@@ -1429,7 +1429,7 @@ namespace UE::Learning
 				// If entry from lhs is first
 				if (Lhs.GetEntrySequence(LhsIndex) < Rhs.GetEntrySequence(RhsIndex))
 				{
-					UE_LEARNING_CHECK(Lhs.GetEntryRangeNum(LhsIndex) > 0);
+					check(Lhs.GetEntryRangeNum(LhsIndex) > 0);
 
 					// Append subranges to output
 					Out.EntrySequences[OutIndex] = Lhs.GetEntrySequence(LhsIndex);
@@ -1451,7 +1451,7 @@ namespace UE::Learning
 				// If entry from rhs is first
 				else if (Rhs.GetEntrySequence(RhsIndex) < Lhs.GetEntrySequence(LhsIndex))
 				{
-					UE_LEARNING_CHECK(Rhs.GetEntryRangeNum(RhsIndex) > 0);
+					check(Rhs.GetEntryRangeNum(RhsIndex) > 0);
 
 					// Append subranges to output
 					Out.EntrySequences[OutIndex] = Rhs.GetEntrySequence(RhsIndex);
@@ -1473,8 +1473,8 @@ namespace UE::Learning
 				// If both contain the same entry
 				else
 				{
-					UE_LEARNING_CHECK(Lhs.GetEntryRangeNum(LhsIndex) > 0);
-					UE_LEARNING_CHECK(Rhs.GetEntryRangeNum(RhsIndex) > 0);
+					check(Lhs.GetEntryRangeNum(LhsIndex) > 0);
+					check(Rhs.GetEntryRangeNum(RhsIndex) > 0);
 
 					// Append union of subranges to output
 					const int32 RangeNum = FrameRangeSet::Private::RangesUnion(
@@ -1485,8 +1485,8 @@ namespace UE::Learning
 						Rhs.GetEntryRangeStarts(RhsIndex),
 						Rhs.GetEntryRangeLengths(RhsIndex));
 
-					UE_LEARNING_CHECK(RangeNum > 0);
-					UE_LEARNING_CHECK(RangeNum <= Lhs.GetEntryRangeNum(LhsIndex) + Rhs.GetEntryRangeNum(RhsIndex));
+					check(RangeNum > 0);
+					check(RangeNum <= Lhs.GetEntryRangeNum(LhsIndex) + Rhs.GetEntryRangeNum(RhsIndex));
 
 					Out.EntrySequences[OutIndex] = Lhs.GetEntrySequence(LhsIndex);
 					Out.EntryRangeOffsets[OutIndex] = RangeIndex;
@@ -1501,8 +1501,8 @@ namespace UE::Learning
 			// Process any remaining lhs entries
 			while (LhsIndex < LhsEntryNum)
 			{
-				UE_LEARNING_CHECK(RhsIndex == RhsEntryNum);
-				UE_LEARNING_CHECK(Lhs.GetEntryRangeNum(LhsIndex) > 0);
+				check(RhsIndex == RhsEntryNum);
+				check(Lhs.GetEntryRangeNum(LhsIndex) > 0);
 
 				// Append subranges to output
 				Out.EntrySequences[OutIndex] = Lhs.GetEntrySequence(LhsIndex);
@@ -1525,8 +1525,8 @@ namespace UE::Learning
 			// Process any remaining rhs entries
 			while (RhsIndex < RhsEntryNum)
 			{
-				UE_LEARNING_CHECK(LhsIndex == LhsEntryNum);
-				UE_LEARNING_CHECK(Rhs.GetEntryRangeNum(RhsIndex) > 0);
+				check(LhsIndex == LhsEntryNum);
+				check(Rhs.GetEntryRangeNum(RhsIndex) > 0);
 
 				// Append subranges to output
 				Out.EntrySequences[OutIndex] = Rhs.GetEntrySequence(RhsIndex);
@@ -1599,8 +1599,8 @@ namespace UE::Learning
 				// If entry is in both lhs and rhs
 				else
 				{
-					UE_LEARNING_CHECK(Lhs.GetEntryRangeNum(LhsIndex) > 0);
-					UE_LEARNING_CHECK(Rhs.GetEntryRangeNum(RhsIndex) > 0);
+					check(Lhs.GetEntryRangeNum(LhsIndex) > 0);
+					check(Rhs.GetEntryRangeNum(RhsIndex) > 0);
 
 					// Append intersection of subranges to output
 					const int32 RangeNum = FrameRangeSet::Private::RangesIntersection(
@@ -1611,7 +1611,7 @@ namespace UE::Learning
 						Rhs.GetEntryRangeStarts(RhsIndex),
 						Rhs.GetEntryRangeLengths(RhsIndex));
 
-					UE_LEARNING_CHECK(RangeNum <= Lhs.GetEntryRangeNum(LhsIndex) + Rhs.GetEntryRangeNum(RhsIndex));
+					check(RangeNum <= Lhs.GetEntryRangeNum(LhsIndex) + Rhs.GetEntryRangeNum(RhsIndex));
 
 					if (RangeNum > 0)
 					{
@@ -1669,7 +1669,7 @@ namespace UE::Learning
 				// If entry from lhs is first
 				if (Lhs.GetEntrySequence(LhsIndex) < Rhs.GetEntrySequence(RhsIndex))
 				{
-					UE_LEARNING_CHECK(Lhs.GetEntryRangeNum(LhsIndex) > 0);
+					check(Lhs.GetEntryRangeNum(LhsIndex) > 0);
 
 					// Append subranges to output
 					Out.EntrySequences[OutIndex] = Lhs.GetEntrySequence(LhsIndex);
@@ -1696,8 +1696,8 @@ namespace UE::Learning
 				// If entry is in both lhs and rhs
 				else
 				{
-					UE_LEARNING_CHECK(Lhs.GetEntryRangeNum(LhsIndex) > 0);
-					UE_LEARNING_CHECK(Rhs.GetEntryRangeNum(RhsIndex) > 0);
+					check(Lhs.GetEntryRangeNum(LhsIndex) > 0);
+					check(Rhs.GetEntryRangeNum(RhsIndex) > 0);
 
 					// Append difference of subranges to output
 					const int32 RangeNum = FrameRangeSet::Private::RangesDifference(
@@ -1708,7 +1708,7 @@ namespace UE::Learning
 						Rhs.GetEntryRangeStarts(RhsIndex),
 						Rhs.GetEntryRangeLengths(RhsIndex));
 
-					UE_LEARNING_CHECK(RangeNum <= Lhs.GetEntryRangeNum(LhsIndex) + Rhs.GetEntryRangeNum(RhsIndex));
+					check(RangeNum <= Lhs.GetEntryRangeNum(LhsIndex) + Rhs.GetEntryRangeNum(RhsIndex));
 
 					if (RangeNum > 0)
 					{
@@ -1726,7 +1726,7 @@ namespace UE::Learning
 			// Process any remaining lhs entries
 			while (LhsIndex < LhsEntryNum)
 			{
-				UE_LEARNING_CHECK(Lhs.GetEntryRangeNum(LhsIndex) > 0);
+				check(Lhs.GetEntryRangeNum(LhsIndex) > 0);
 
 				// Append subranges to output
 				Out.EntrySequences[OutIndex] = Lhs.GetEntrySequence(LhsIndex);
@@ -1817,8 +1817,8 @@ namespace UE::Learning
 				// If entry is in both lhs and rhs
 				else
 				{
-					UE_LEARNING_CHECK(Lhs.GetEntryRangeNum(LhsIndex) > 0);
-					UE_LEARNING_CHECK(Rhs.GetEntryRangeNum(RhsIndex) > 0);
+					check(Lhs.GetEntryRangeNum(LhsIndex) > 0);
+					check(Rhs.GetEntryRangeNum(RhsIndex) > 0);
 
 					// Append intersection of subranges to output
 					const int32 RangeNum = FrameRangeSet::Private::RangesIntersectionWithFrameOffsets(
@@ -1836,7 +1836,7 @@ namespace UE::Learning
 					Private::FrameOffsetsCheck(OutLhsFrameOffsets.Slice(0, RangeIndex + RangeNum), 0, LhsFrameOffset + Lhs.GetEntryTotalFrameNum(LhsIndex));
 					Private::FrameOffsetsCheck(OutRhsFrameOffsets.Slice(0, RangeIndex + RangeNum), 0, RhsFrameOffset + Rhs.GetEntryTotalFrameNum(RhsIndex));
 
-					UE_LEARNING_CHECK(RangeNum <= Lhs.GetEntryRangeNum(LhsIndex) + Rhs.GetEntryRangeNum(RhsIndex));
+					check(RangeNum <= Lhs.GetEntryRangeNum(LhsIndex) + Rhs.GetEntryRangeNum(RhsIndex));
 
 					if (RangeNum > 0)
 					{
@@ -1864,13 +1864,13 @@ namespace UE::Learning
 			Private::FrameOffsetsCheck(OutLhsFrameOffsets.Slice(0, RangeIndex), 0, Lhs.GetTotalFrameNum());
 			Private::FrameOffsetsCheck(OutRhsFrameOffsets.Slice(0, RangeIndex), 0, Rhs.GetTotalFrameNum());
 
-			UE_LEARNING_CHECK(RangeIndex == Out.GetTotalRangeNum());
+			check(RangeIndex == Out.GetTotalRangeNum());
 			return RangeIndex;
 		}
 
 		void TrimLeft(FFrameRangeSet& Out, const FFrameRangeSet& FrameRangeSet, const int32 TrimFrameNum)
 		{
-			UE_LEARNING_CHECK(TrimFrameNum >= 0);
+			check(TrimFrameNum >= 0);
 			FrameRangeSet.Check();
 
 			const int32 EntryNum = FrameRangeSet.GetEntryNum();
@@ -1907,7 +1907,7 @@ namespace UE::Learning
 
 		void TrimRight(FFrameRangeSet& Out, const FFrameRangeSet& FrameRangeSet, const int32 TrimFrameNum)
 		{
-			UE_LEARNING_CHECK(TrimFrameNum >= 0);
+			check(TrimFrameNum >= 0);
 			FrameRangeSet.Check();
 
 			const int32 EntryNum = FrameRangeSet.GetEntryNum();
@@ -1944,7 +1944,7 @@ namespace UE::Learning
 
 		void TrimLeftAndRight(FFrameRangeSet& Out, const FFrameRangeSet& FrameRangeSet, const int32 TrimFrameNum)
 		{
-			UE_LEARNING_CHECK(TrimFrameNum >= 0);
+			check(TrimFrameNum >= 0);
 			FrameRangeSet.Check();
 
 			const int32 EntryNum = FrameRangeSet.GetEntryNum();
@@ -2097,7 +2097,7 @@ namespace UE::Learning
 			TLearningArrayView<1, int32> OutRangeEntries,
 			const FFrameRangeSet& FrameRangeSet)
 		{
-			UE_LEARNING_CHECK(OutRangeEntries.Num() == FrameRangeSet.GetTotalRangeNum());
+			check(OutRangeEntries.Num() == FrameRangeSet.GetTotalRangeNum());
 
 			const int32 TotalRangeNum = FrameRangeSet.GetTotalRangeNum();
 			const int32 EntryNum = FrameRangeSet.GetEntryNum();
@@ -2110,14 +2110,14 @@ namespace UE::Learning
 				RangeOffset += RangeNum;
 			}
 
-			UE_LEARNING_CHECK(RangeOffset == TotalRangeNum);
+			check(RangeOffset == TotalRangeNum);
 		}
 
 		void AllRangeIndices(
 			TLearningArrayView<1, int32> OutRangeIndices,
 			const FFrameRangeSet& FrameRangeSet)
 		{
-			UE_LEARNING_CHECK(OutRangeIndices.Num() == FrameRangeSet.GetTotalRangeNum());
+			check(OutRangeIndices.Num() == FrameRangeSet.GetTotalRangeNum());
 
 			const int32 TotalRangeNum = FrameRangeSet.GetTotalRangeNum();
 			const int32 EntryNum = FrameRangeSet.GetEntryNum();
@@ -2133,14 +2133,14 @@ namespace UE::Learning
 				RangeOffset += RangeNum;
 			}
 
-			UE_LEARNING_CHECK(RangeOffset == TotalRangeNum);
+			check(RangeOffset == TotalRangeNum);
 		}
 
 		void AllRangeSequences(
 			TLearningArrayView<1, int32> OutRangeSequences,
 			const FFrameRangeSet& FrameRangeSet)
 		{
-			UE_LEARNING_CHECK(OutRangeSequences.Num() == FrameRangeSet.GetTotalRangeNum());
+			check(OutRangeSequences.Num() == FrameRangeSet.GetTotalRangeNum());
 
 			const int32 TotalRangeNum = FrameRangeSet.GetTotalRangeNum();
 			const int32 EntryNum = FrameRangeSet.GetEntryNum();
@@ -2153,14 +2153,14 @@ namespace UE::Learning
 				RangeOffset += RangeNum;
 			}
 
-			UE_LEARNING_CHECK(RangeOffset == TotalRangeNum);
+			check(RangeOffset == TotalRangeNum);
 		}
 
 		void AllRangeFrameOffsets(
 			TLearningArrayView<1, int32> OutRangeFrameOffsets,
 			const FFrameRangeSet& FrameRangeSet)
 		{
-			UE_LEARNING_CHECK(OutRangeFrameOffsets.Num() == FrameRangeSet.GetTotalRangeNum());
+			check(OutRangeFrameOffsets.Num() == FrameRangeSet.GetTotalRangeNum());
 
 			const int32 TotalRangeNum = FrameRangeSet.GetTotalRangeNum();
 			const int32 TotalFrameNum = FrameRangeSet.GetTotalFrameNum();
@@ -2175,7 +2175,7 @@ namespace UE::Learning
 				FrameOffset += FrameRangeSet.RangeLengths[RangeIdx];
 			}
 
-			UE_LEARNING_CHECK(FrameOffset == TotalFrameNum);
+			check(FrameOffset == TotalFrameNum);
 		}
 
 		void AllRangeStartTimes(
@@ -2183,7 +2183,7 @@ namespace UE::Learning
 			const FFrameRangeSet& FrameRangeSet, 
 			const float FrameDeltaTime)
 		{
-			UE_LEARNING_CHECK(OutRangeStartTimes.Num() == FrameRangeSet.GetTotalRangeNum());
+			check(OutRangeStartTimes.Num() == FrameRangeSet.GetTotalRangeNum());
 
 			const int32 TotalRangeNum = FrameRangeSet.GetTotalRangeNum();
 			for (int32 RangeIdx = 0; RangeIdx < TotalRangeNum; RangeIdx++)
@@ -2197,7 +2197,7 @@ namespace UE::Learning
 			const FFrameRangeSet& FrameRangeSet, 
 			const float FrameDeltaTime)
 		{
-			UE_LEARNING_CHECK(OutRangeEndTimes.Num() == FrameRangeSet.GetTotalRangeNum());
+			check(OutRangeEndTimes.Num() == FrameRangeSet.GetTotalRangeNum());
 
 			const int32 TotalRangeNum = FrameRangeSet.GetTotalRangeNum();
 			for (int32 RangeIdx = 0; RangeIdx < TotalRangeNum; RangeIdx++)
@@ -2211,7 +2211,7 @@ namespace UE::Learning
 			const FFrameRangeSet& FrameRangeSet,
 			const float FrameDeltaTime)
 		{
-			UE_LEARNING_CHECK(OutRangeDurations.Num() == FrameRangeSet.GetTotalRangeNum());
+			check(OutRangeDurations.Num() == FrameRangeSet.GetTotalRangeNum());
 
 			const int32 TotalRangeNum = FrameRangeSet.GetTotalRangeNum();
 			for (int32 RangeIdx = 0; RangeIdx < TotalRangeNum; RangeIdx++)
diff --git a/Engine/Plugins/Experimental/LearningAgents/Source/Learning/Private/LearningFrameSet.cpp b/Engine/Plugins/Experimental/LearningAgents/Source/Learning/Private/LearningFrameSet.cpp
--- a/Engine/Plugins/Experimental/LearningAgents/Source/Learning/Private/LearningFrameSet.cpp
+++ b/Engine/Plugins/Experimental/LearningAgents/Source/Learning/Private/LearningFrameSet.cpp
@@ -13,9 +13,9 @@ namespace UE::Learning
 			const int32 FrameNum = Frames.Num();
 			for (int32 FrameIdx = 0; FrameIdx < FrameNum - 1; FrameIdx++)
 			{
-				UE_LEARNING_CHECK(Frames[FrameIdx + 0] >= 0);
-				UE_LEARNING_CHECK(Frames[FrameIdx + 1] >= 0);
-				UE_LEARNING_CHECK(Frames[FrameIdx + 0] < Frames[FrameIdx + 1]);
+				check(Frames[FrameIdx + 0] >= 0);
+				check(Frames[FrameIdx + 1] >= 0);
+				check(Frames[FrameIdx + 0] < Frames[FrameIdx + 1]);
 			}
 		}
 
@@ -62,7 +62,7 @@ namespace UE::Learning
 				// Frame from lhs and rhs coming at same time
 				else
 				{
-					UE_LEARNING_CHECK(LhsT == RhsT);
+					check(LhsT == RhsT);
 					OutFrames[OutIndex] = LhsT;
 					OutIndex++;
 					LhsIndex++;
@@ -131,7 +131,7 @@ namespace UE::Learning
 				// Frame from lhs and rhs coming at same time
 				else
 				{
-					UE_LEARNING_CHECK(LhsT == RhsT);
+					check(LhsT == RhsT);
 					OutFrames[OutIndex] = LhsT;
 					OutIndex++;
 					LhsIndex++;
@@ -186,7 +186,7 @@ namespace UE::Learning
 				// Frame from lhs and rhs coming at same time
 				else
 				{
-					UE_LEARNING_CHECK(LhsT == RhsT);
+					check(LhsT == RhsT);
 					LhsIndex++;
 					RhsIndex++;
 				}
@@ -214,12 +214,12 @@ namespace UE::Learning
 
 		for (int32 EntryIdx = 0; EntryIdx < EntryNum - 1; EntryIdx++)
 		{
-			UE_LEARNING_CHECK(EntrySequences[EntryIdx + 0] < EntrySequences[EntryIdx + 1]);
+			check(EntrySequences[EntryIdx + 0] < EntrySequences[EntryIdx + 1]);
 		}
 
 		for (int32 EntryIdx = 0; EntryIdx < EntryNum; EntryIdx++)
 		{
-			UE_LEARNING_CHECK(EntryFrameNums[EntryIdx] > 0);
+			check(EntryFrameNums[EntryIdx] > 0);
 			FrameSet::Private::FramesCheck(GetEntryFrames(EntryIdx));
 		}
 	}
@@ -316,7 +316,7 @@ namespace UE::Learning
 
 	bool FFrameSet::FindNearestInRange(int32& OutEntryIdx, int32& OutFrameIdx, int32& OutFrameDifference, const int32 Sequence, const int32 Frame, const int32 RangeStart, const int32 RangeLength) const
 	{
-		UE_LEARNING_CHECK(Frame >= RangeStart && Frame < RangeStart + RangeLength);
+		check(Frame >= RangeStart && Frame < RangeStart + RangeLength);
 
 		const int32 EntryIdx = FindSequenceEntry(Sequence);
 
@@ -380,7 +380,7 @@ namespace UE::Learning
 		const int32 InSequence,
 		const TLearningArrayView<1, const int32> InFrames)
 	{
-		UE_LEARNING_CHECK(!EntrySequences.ArrayView().Contains(InSequence));
+		check(!EntrySequences.ArrayView().Contains(InSequence));
 		UE::Learning::FrameSet::Private::FramesCheck(InFrames);
 
 		if (InFrames.IsEmpty()) { return; }
@@ -477,7 +477,7 @@ namespace UE::Learning
 				// If both contain the same entry
 				else
 				{
-					UE_LEARNING_CHECK(Rhs.GetEntrySequence(RhsIndex) == Lhs.GetEntrySequence(LhsIndex));
+					check(Rhs.GetEntrySequence(RhsIndex) == Lhs.GetEntrySequence(LhsIndex));
 
 					// Append union to output
 					const int32 FrameNum = Private::FramesUnion(
@@ -485,7 +485,7 @@ namespace UE::Learning
 						Lhs.GetEntryFrames(LhsIndex),
 						Rhs.GetEntryFrames(RhsIndex));
 
-					UE_LEARNING_CHECK(FrameNum <= Lhs.GetEntryFrameNum(LhsIndex) + Rhs.GetEntryFrameNum(RhsIndex));
+					check(FrameNum <= Lhs.GetEntryFrameNum(LhsIndex) + Rhs.GetEntryFrameNum(RhsIndex));
 
 					if (FrameNum > 0)
 					{
@@ -583,7 +583,7 @@ namespace UE::Learning
 				// If both contain the same entry
 				else
 				{
-					UE_LEARNING_CHECK(Rhs.GetEntrySequence(RhsIndex) == Lhs.GetEntrySequence(LhsIndex));
+					check(Rhs.GetEntrySequence(RhsIndex) == Lhs.GetEntrySequence(LhsIndex));
 
 					// Append intersection to output
 					const int32 FrameNum = Private::FramesIntersection(
@@ -591,7 +591,7 @@ namespace UE::Learning
 						Lhs.GetEntryFrames(LhsIndex),
 						Rhs.GetEntryFrames(RhsIndex));
 
-					UE_LEARNING_CHECK(FrameNum <= Lhs.GetEntryFrameNum(LhsIndex) + Rhs.GetEntryFrameNum(RhsIndex));
+					check(FrameNum <= Lhs.GetEntryFrameNum(LhsIndex) + Rhs.GetEntryFrameNum(RhsIndex));
 
 					if (FrameNum > 0)
 					{
@@ -666,7 +666,7 @@ namespace UE::Learning
 				// If both contain the same entry
 				else
 				{
-					UE_LEARNING_CHECK(Rhs.GetEntrySequence(RhsIndex) == Lhs.GetEntrySequence(LhsIndex));
+					check(Rhs.GetEntrySequence(RhsIndex) == Lhs.GetEntrySequence(LhsIndex));
 
 					// Append difference to output
 					const int32 FrameNum = Private::FramesDifference(
@@ -674,7 +674,7 @@ namespace UE::Learning
 						Lhs.GetEntryFrames(LhsIndex),
 						Rhs.GetEntryFrames(RhsIndex));
 
-					UE_LEARNING_CHECK(FrameNum <= Lhs.GetEntryFrameNum(LhsIndex) + Rhs.GetEntryFrameNum(RhsIndex));
+					check(FrameNum <= Lhs.GetEntryFrameNum(LhsIndex) + Rhs.GetEntryFrameNum(RhsIndex));
 
 					if (FrameNum > 0)
 					{
@@ -718,7 +718,7 @@ namespace UE::Learning
 			TLearningArrayView<1, int32> OutFrameEntries,
 			const FFrameSet& FrameSet)
 		{
-			UE_LEARNING_CHECK(OutFrameEntries.Num() == FrameSet.GetTotalFrameNum());
+			check(OutFrameEntries.Num() == FrameSet.GetTotalFrameNum());
 
 			const int32 TotalFrameNum = FrameSet.GetTotalFrameNum();
 			const int32 EntryNum = FrameSet.GetEntryNum();
@@ -731,14 +731,14 @@ namespace UE::Learning
 				FrameOffset += FrameNum;
 			}
 
-			UE_LEARNING_CHECK(FrameOffset == TotalFrameNum);
+			check(FrameOffset == TotalFrameNum);
 		}
 
 		void AllFrameIndices(
 			TLearningArrayView<1, int32> OutFrameIndices,
 			const FFrameSet& FrameSet)
 		{
-			UE_LEARNING_CHECK(OutFrameIndices.Num() == FrameSet.GetTotalFrameNum());
+			check(OutFrameIndices.Num() == FrameSet.GetTotalFrameNum());
 
 			const int32 TotalFrameNum = FrameSet.GetTotalFrameNum();
 			const int32 EntryNum = FrameSet.GetEntryNum();
@@ -754,14 +754,14 @@ namespace UE::Learning
 				FrameOffset += FrameNum;
 			}
 
-			UE_LEARNING_CHECK(FrameOffset == TotalFrameNum);
+			check(FrameOffset == TotalFrameNum);
 		}
 
 		void AllFrameSequences(
 			TLearningArrayView<1, int32> OutFrameSequences,
 			const FFrameSet& FrameSet)
 		{
-			UE_LEARNING_CHECK(OutFrameSequences.Num() == FrameSet.GetTotalFrameNum());
+			check(OutFrameSequences.Num() == FrameSet.GetTotalFrameNum());
 
 			const int32 TotalFrameNum = FrameSet.GetTotalFrameNum();
 			const int32 EntryNum = FrameSet.GetEntryNum();
@@ -774,7 +774,7 @@ namespace UE::Learning
 				FrameOffset += FrameNum;
 			}
 
-			UE_LEARNING_CHECK(FrameOffset == TotalFrameNum);
+			check(FrameOffset == TotalFrameNum);
 		}
 
 		void AllFrameTimes(
@@ -782,7 +782,7 @@ namespace UE::Learning
 			const FFrameSet& FrameSet,
 			const float FrameDeltaTime)
 		{
-			UE_LEARNING_CHECK(OutFrameTimes.Num() == FrameSet.GetTotalFrameNum());
+			check(OutFrameTimes.Num() == FrameSet.GetTotalFrameNum());
 
 			const int32 TotalFrameNum = FrameSet.GetTotalFrameNum();
 			for (int32 FrameIdx = 0; FrameIdx < TotalFrameNum; FrameIdx++)
diff --git a/Engine/Plugins/Experimental/LearningAgents/Source/Learning/Private/LearningKMeans.cpp b/Engine/Plugins/Experimental/LearningAgents/Source/Learning/Private/LearningKMeans.cpp
--- a/Engine/Plugins/Experimental/LearningAgents/Source/Learning/Private/LearningKMeans.cpp
+++ b/Engine/Plugins/Experimental/LearningAgents/Source/Learning/Private/LearningKMeans.cpp
@@ -72,7 +72,7 @@ namespace UE::Learning::KMeans
 		const bool bParallelEvaluation,
 		const uint16 MinParallelBatchSize)
 	{
-		UE_LEARNING_TRACE_CPUPROFILER_EVENT_SCOPE(Learning::KMeans::UpdateAssignmentsFromCenters);
+		TRACE_CPUPROFILER_EVENT_SCOPE(Learning::KMeans::UpdateAssignmentsFromCenters);
 
 		const int32 SampleNum = Samples.Num<0>();
 		const int32 ClusterNum = Centers.Num<0>();
@@ -115,8 +115,8 @@ namespace UE::Learning::KMeans
 					}
 				}
 
-				UE_LEARNING_CHECK(BestDistance != MAX_flt);
-				UE_LEARNING_CHECK(BestIndex != INDEX_NONE);
+				check(BestDistance != MAX_flt);
+				check(BestIndex != INDEX_NONE);
 
 				OutAssignments[SampleIdx] = BestIndex;
 			}
@@ -141,7 +141,7 @@ namespace UE::Learning::KMeans
 		const bool bParallelEvaluation,
 		const uint16 MinParallelBatchSize)
 	{
-		UE_LEARNING_TRACE_CPUPROFILER_EVENT_SCOPE(Learning::KMeans::UpdateAssignmentsFromBounds);
+		TRACE_CPUPROFILER_EVENT_SCOPE(Learning::KMeans::UpdateAssignmentsFromBounds);
 
 		const int32 SampleNum = Samples.Num<0>();
 		const int32 ClusterNum = Mins.Num<0>();
@@ -186,8 +186,8 @@ namespace UE::Learning::KMeans
 					}
 				}
 
-				UE_LEARNING_CHECK(BestDistance != MAX_flt);
-				UE_LEARNING_CHECK(BestIndex != INDEX_NONE);
+				check(BestDistance != MAX_flt);
+				check(BestIndex != INDEX_NONE);
 
 				OutAssignments[SampleIdx] = BestIndex;
 			}
@@ -208,7 +208,7 @@ namespace UE::Learning::KMeans
 		TLearningArrayView<1, int32> OutAssignmentCounts,
 		const TLearningArrayView<1, const int32> Assignments)
 	{
-		UE_LEARNING_TRACE_CPUPROFILER_EVENT_SCOPE(Learning::KMeans::CountClusterAssignments);
+		TRACE_CPUPROFILER_EVENT_SCOPE(Learning::KMeans::CountClusterAssignments);
 
 		const int32 SampleNum = Assignments.Num();
 
@@ -226,7 +226,7 @@ namespace UE::Learning::KMeans
 		const TLearningArrayView<1, const int32> AssignmentCounts,
 		const TLearningArrayView<2, const float> Samples)
 	{
-		UE_LEARNING_TRACE_CPUPROFILER_EVENT_SCOPE(Learning::KMeans::UpdateCenters);
+		TRACE_CPUPROFILER_EVENT_SCOPE(Learning::KMeans::UpdateCenters);
 
 		const int32 SampleNum = Samples.Num<0>();
 		const int32 ClusterNum = OutCenters.Num<0>();
@@ -264,7 +264,7 @@ namespace UE::Learning::KMeans
 		const TLearningArrayView<1, const int32> Assignments,
 		const TLearningArrayView<2, const float> Samples)
 	{
-		UE_LEARNING_TRACE_CPUPROFILER_EVENT_SCOPE(Learning::KMeans::UpdateBounds);
+		TRACE_CPUPROFILER_EVENT_SCOPE(Learning::KMeans::UpdateBounds);
 
 		const int32 SampleNum = Samples.Num<0>();
 		const int32 ClusterNum = OutMins.Num<0>();
@@ -306,7 +306,7 @@ namespace UE::Learning::KMeans
 		const TLearningArrayView<1, const int32> AssignmentCounts,
 		const TLearningArrayView<2, const float> Samples)
 	{
-		UE_LEARNING_TRACE_CPUPROFILER_EVENT_SCOPE(Learning::KMeans::ComputeClusteredIndex);
+		TRACE_CPUPROFILER_EVENT_SCOPE(Learning::KMeans::ComputeClusteredIndex);
 
 		const int32 SampleNum = Samples.Num<0>();
 		const int32 ClusterNum = OutClusterStarts.Num();
@@ -351,7 +351,7 @@ namespace UE::Learning::KMeans
 		// Validate Number of assignments matched
 		for (int32 ClusterIdx = 0; ClusterIdx < ClusterNum; ClusterIdx++)
 		{
-			UE_LEARNING_CHECK(OutClusterLengths[ClusterIdx] == AssignmentCounts[ClusterIdx]);
+			check(OutClusterLengths[ClusterIdx] == AssignmentCounts[ClusterIdx]);
 		}
 #endif
 	}
diff --git a/Engine/Plugins/Experimental/LearningAgents/Source/Learning/Private/LearningNeuralNetwork.cpp b/Engine/Plugins/Experimental/LearningAgents/Source/Learning/Private/LearningNeuralNetwork.cpp
--- a/Engine/Plugins/Experimental/LearningAgents/Source/Learning/Private/LearningNeuralNetwork.cpp
+++ b/Engine/Plugins/Experimental/LearningAgents/Source/Learning/Private/LearningNeuralNetwork.cpp
@@ -249,7 +249,7 @@ namespace UE::Learning
 
 	void FNeuralNetworkInference::ReloadModelInstances(NNE::IModelCPU& InModel, const int32 InInputSize, const int32 InOutputSize)
 	{
-		UE_LEARNING_TRACE_CPUPROFILER_EVENT_SCOPE(Learning::FNeuralNetworkInference::ReloadModelInstances);
+		TRACE_CPUPROFILER_EVENT_SCOPE(Learning::FNeuralNetworkInference::ReloadModelInstances);
 
 		InputSize = InInputSize;
 		OutputSize = InOutputSize;
@@ -293,15 +293,15 @@ namespace UE::Learning
 
 	void FNeuralNetworkInference::Evaluate(TLearningArrayView<2, float> Output, const TLearningArrayView<2, const float> Input)
 	{
-		UE_LEARNING_TRACE_CPUPROFILER_EVENT_SCOPE(Learning::FNeuralNetworkInference::Evaluate);
+		TRACE_CPUPROFILER_EVENT_SCOPE(Learning::FNeuralNetworkInference::Evaluate);
 
-		UE_LEARNING_CHECK(Output.Num<0>() == Input.Num<0>());
-		UE_LEARNING_CHECK(Input.Num<1>() == InputSize);
-		UE_LEARNING_CHECK(Output.Num<1>() == OutputSize);
+		check(Output.Num<0>() == Input.Num<0>());
+		check(Input.Num<1>() == InputSize);
+		check(Output.Num<1>() == OutputSize);
 
 		const int32 BatchSize = Output.Num<0>();
 
-		UE_LEARNING_CHECK(BatchSize <= MaxBatchSize);
+		check(BatchSize <= MaxBatchSize);
 
 		if (BatchSize == 0) { return; }
 
@@ -313,7 +313,7 @@ namespace UE::Learning
 			const int32 InferenceInstanceBatchSize = FMath::DivideAndRoundUp(BatchSize, IdealInferenceInstanceNum);
 			const int32 InferenceInstanceNum = FMath::DivideAndRoundUp(BatchSize, InferenceInstanceBatchSize);
 
-			UE_LEARNING_CHECK(InferenceInstanceNum <= ModelInstances.Num());
+			check(InferenceInstanceNum <= ModelInstances.Num());
 
 			ParallelFor(InferenceInstanceNum, [this, Input, Output, BatchSize, InferenceInstanceBatchSize](int32 InferenceInstanceIdx)
 			{
@@ -396,8 +396,8 @@ namespace UE::Learning
 	{
 		if (NeuralNetwork != NewNeuralNetwork)
 		{
-			UE_LEARNING_CHECK(NeuralNetwork->GetInputSize() == NewNeuralNetwork->GetInputSize());
-			UE_LEARNING_CHECK(NeuralNetwork->GetOutputSize() == NewNeuralNetwork->GetOutputSize());
+			check(NeuralNetwork->GetInputSize() == NewNeuralNetwork->GetInputSize());
+			check(NeuralNetwork->GetOutputSize() == NewNeuralNetwork->GetOutputSize());
 			NeuralNetwork = NewNeuralNetwork;
 			NeuralNetworkInference = NeuralNetwork->CreateInferenceObject(MaxInstanceNum, InferenceSettings);
 		}
diff --git a/Engine/Plugins/Experimental/LearningAgents/Source/Learning/Private/LearningObservation.cpp b/Engine/Plugins/Experimental/LearningAgents/Source/Learning/Private/LearningObservation.cpp
--- a/Engine/Plugins/Experimental/LearningAgents/Source/Learning/Private/LearningObservation.cpp
+++ b/Engine/Plugins/Experimental/LearningAgents/Source/Learning/Private/LearningObservation.cpp
@@ -78,8 +78,8 @@ namespace UE::Learning::Observation
 
 	FSchemaElement FSchema::CreateContinuous(const FSchemaContinuousParameters Parameters, const FName Tag)
 	{
-		UE_LEARNING_CHECK(Parameters.Num >= 0);
-		UE_LEARNING_CHECK(Parameters.Scale >= 0.0f);
+		check(Parameters.Num >= 0);
+		check(Parameters.Scale >= 0.0f);
 
 		FContinuousData ElementData;
 		ElementData.Num = Parameters.Num;
@@ -96,9 +96,9 @@ namespace UE::Learning::Observation
 
 	FSchemaElement FSchema::CreateAnd(const FSchemaAndParameters Parameters, const FName Tag)
 	{
-		UE_LEARNING_CHECK(Parameters.Elements.Num() == Parameters.ElementNames.Num());
-		UE_LEARNING_CHECK(!Private::ContainsDuplicates(Parameters.ElementNames));
-		UE_LEARNING_CHECK(Private::CheckAllValid(*this, Parameters.Elements));
+		check(Parameters.Elements.Num() == Parameters.ElementNames.Num());
+		check(!Private::ContainsDuplicates(Parameters.ElementNames));
+		check(Private::CheckAllValid(*this, Parameters.Elements));
 
 		FAndData ElementData;
 		ElementData.Num = Parameters.Elements.Num();
@@ -118,9 +118,9 @@ namespace UE::Learning::Observation
 
 	FSchemaElement FSchema::CreateOrExclusive(const FSchemaOrExclusiveParameters Parameters, const FName Tag)
 	{
-		UE_LEARNING_CHECK(Parameters.Elements.Num() == Parameters.ElementNames.Num());
-		UE_LEARNING_CHECK(!Private::ContainsDuplicates(Parameters.ElementNames));
-		UE_LEARNING_CHECK(Private::CheckAllValid(*this, Parameters.Elements));
+		check(Parameters.Elements.Num() == Parameters.ElementNames.Num());
+		check(!Private::ContainsDuplicates(Parameters.ElementNames));
+		check(Private::CheckAllValid(*this, Parameters.Elements));
 
 		FOrExclusiveData ElementData;
 		ElementData.Num = Parameters.Elements.Num();
@@ -141,9 +141,9 @@ namespace UE::Learning::Observation
 
 	FSchemaElement FSchema::CreateOrInclusive(const FSchemaOrInclusiveParameters Parameters, const FName Tag)
 	{
-		UE_LEARNING_CHECK(Parameters.Elements.Num() == Parameters.ElementNames.Num());
-		UE_LEARNING_CHECK(!Private::ContainsDuplicates(Parameters.ElementNames));
-		UE_LEARNING_CHECK(Private::CheckAllValid(*this, Parameters.Elements));
+		check(Parameters.Elements.Num() == Parameters.ElementNames.Num());
+		check(!Private::ContainsDuplicates(Parameters.ElementNames));
+		check(Private::CheckAllValid(*this, Parameters.Elements));
 
 		FOrInclusiveData ElementData;
 		ElementData.Num = Parameters.Elements.Num();
@@ -166,8 +166,8 @@ namespace UE::Learning::Observation
 
 	FSchemaElement FSchema::CreateArray(const FSchemaArrayParameters Parameters, const FName Tag)
 	{
-		UE_LEARNING_CHECK(Parameters.Num >= 0);
-		UE_LEARNING_CHECK(IsValid(Parameters.Element));
+		check(Parameters.Num >= 0);
+		check(IsValid(Parameters.Element));
 
 		FArrayData ElementData;
 		ElementData.Num = Parameters.Num;
@@ -187,7 +187,7 @@ namespace UE::Learning::Observation
 
 	FSchemaElement FSchema::CreateSet(const FSchemaSetParameters Parameters, const FName Tag)
 	{
-		UE_LEARNING_CHECK(IsValid(Parameters.Element));
+		check(IsValid(Parameters.Element));
 
 		FSetData ElementData;
 		ElementData.MaxNum = Parameters.MaxNum;
@@ -210,7 +210,7 @@ namespace UE::Learning::Observation
 
 	FSchemaElement FSchema::CreateEncoding(const FSchemaEncodingParameters Parameters, const FName Tag)
 	{
-		UE_LEARNING_CHECK(IsValid(Parameters.Element));
+		check(IsValid(Parameters.Element));
 
 		FEncodingData ElementData;
 		ElementData.ElementIndex = SubElementObjects.Num();
@@ -237,31 +237,31 @@ namespace UE::Learning::Observation
 
 	EType FSchema::GetType(const FSchemaElement Element) const
 	{
-		UE_LEARNING_CHECK(IsValid(Element));
+		check(IsValid(Element));
 		return Types[Element.Index];
 	}
 
 	FName FSchema::GetTag(const FSchemaElement Element) const
 	{
-		UE_LEARNING_CHECK(IsValid(Element));
+		check(IsValid(Element));
 		return Tags[Element.Index];
 	}
 
 	int32 FSchema::GetObservationVectorSize(const FSchemaElement Element) const
 	{
-		UE_LEARNING_CHECK(IsValid(Element));
+		check(IsValid(Element));
 		return ObservationVectorSizes[Element.Index];
 	}
 
 	int32 FSchema::GetEncodedVectorSize(const FSchemaElement Element) const
 	{
-		UE_LEARNING_CHECK(IsValid(Element));
+		check(IsValid(Element));
 		return EncodedVectorSizes[Element.Index];
 	}
 
 	FSchemaContinuousParameters FSchema::GetContinuous(const FSchemaElement Element) const
 	{
-		UE_LEARNING_CHECK(IsValid(Element) && GetType(Element) == EType::Continuous);
+		check(IsValid(Element) && GetType(Element) == EType::Continuous);
 		const FContinuousData& ElementData = ContinuousData[TypeDataIndices[Element.Index]];
 
 		FSchemaContinuousParameters Parameters;
@@ -272,7 +272,7 @@ namespace UE::Learning::Observation
 
 	FSchemaAndParameters FSchema::GetAnd(const FSchemaElement Element) const
 	{
-		UE_LEARNING_CHECK(IsValid(Element) && GetType(Element) == EType::And);
+		check(IsValid(Element) && GetType(Element) == EType::And);
 		const FAndData& ElementData = AndData[TypeDataIndices[Element.Index]];
 
 		FSchemaAndParameters Parameters;
@@ -283,7 +283,7 @@ namespace UE::Learning::Observation
 
 	FSchemaOrExclusiveParameters FSchema::GetOrExclusive(const FSchemaElement Element) const
 	{
-		UE_LEARNING_CHECK(IsValid(Element) && GetType(Element) == EType::OrExclusive);
+		check(IsValid(Element) && GetType(Element) == EType::OrExclusive);
 		const FOrExclusiveData& ElementData = OrExclusiveData[TypeDataIndices[Element.Index]];
 
 		FSchemaOrExclusiveParameters Parameters;
@@ -295,7 +295,7 @@ namespace UE::Learning::Observation
 
 	FSchemaOrInclusiveParameters FSchema::GetOrInclusive(const FSchemaElement Element) const
 	{
-		UE_LEARNING_CHECK(IsValid(Element) && GetType(Element) == EType::OrInclusive);
+		check(IsValid(Element) && GetType(Element) == EType::OrInclusive);
 		const FOrInclusiveData& ElementData = OrInclusiveData[TypeDataIndices[Element.Index]];
 
 		FSchemaOrInclusiveParameters Parameters;
@@ -309,7 +309,7 @@ namespace UE::Learning::Observation
 
 	FSchemaArrayParameters FSchema::GetArray(const FSchemaElement Element) const
 	{
-		UE_LEARNING_CHECK(IsValid(Element) && GetType(Element) == EType::Array);
+		check(IsValid(Element) && GetType(Element) == EType::Array);
 		const FArrayData& ElementData = ArrayData[TypeDataIndices[Element.Index]];
 
 		FSchemaArrayParameters Parameters;
@@ -320,7 +320,7 @@ namespace UE::Learning::Observation
 
 	FSchemaSetParameters FSchema::GetSet(const FSchemaElement Element) const
 	{
-		UE_LEARNING_CHECK(IsValid(Element) && GetType(Element) == EType::Set);
+		check(IsValid(Element) && GetType(Element) == EType::Set);
 		const FSetData& ElementData = SetData[TypeDataIndices[Element.Index]];
 
 		FSchemaSetParameters Parameters;
@@ -334,7 +334,7 @@ namespace UE::Learning::Observation
 
 	FSchemaEncodingParameters FSchema::GetEncoding(const FSchemaElement Element) const
 	{
-		UE_LEARNING_CHECK(IsValid(Element) && GetType(Element) == EType::Encoding);
+		check(IsValid(Element) && GetType(Element) == EType::Encoding);
 		const FEncodingData& ElementData = EncodingData[TypeDataIndices[Element.Index]];
 
 		FSchemaEncodingParameters Parameters;
@@ -431,9 +431,9 @@ namespace UE::Learning::Observation
 
 	FObjectElement FObject::CreateAnd(const FObjectAndParameters Parameters, const FName Tag)
 	{
-		UE_LEARNING_CHECK(Parameters.Elements.Num() == Parameters.ElementNames.Num());
-		UE_LEARNING_CHECK(!Private::ContainsDuplicates(Parameters.ElementNames));
-		UE_LEARNING_CHECK(Private::CheckAllValid(*this, Parameters.Elements));
+		check(Parameters.Elements.Num() == Parameters.ElementNames.Num());
+		check(!Private::ContainsDuplicates(Parameters.ElementNames));
+		check(Private::CheckAllValid(*this, Parameters.Elements));
 
 		const int32 Index = Types.Add(EType::And);
 		Tags.Add(Tag);
@@ -452,7 +452,7 @@ namespace UE::Learning::Observation
 
 	FObjectElement FObject::CreateOrExclusive(const FObjectOrExclusiveParameters Parameters, const FName Tag)
 	{
-		UE_LEARNING_CHECK(IsValid(Parameters.Element));
+		check(IsValid(Parameters.Element));
 
 		const int32 Index = Types.Add(EType::OrExclusive);
 		Tags.Add(Tag);
@@ -471,9 +471,9 @@ namespace UE::Learning::Observation
 
 	FObjectElement FObject::CreateOrInclusive(const FObjectOrInclusiveParameters Parameters, const FName Tag)
 	{
-		UE_LEARNING_CHECK(Parameters.Elements.Num() == Parameters.ElementNames.Num());
-		UE_LEARNING_CHECK(!Private::ContainsDuplicates(Parameters.ElementNames));
-		UE_LEARNING_CHECK(Private::CheckAllValid(*this, Parameters.Elements));
+		check(Parameters.Elements.Num() == Parameters.ElementNames.Num());
+		check(!Private::ContainsDuplicates(Parameters.ElementNames));
+		check(Private::CheckAllValid(*this, Parameters.Elements));
 
 		const int32 Index = Types.Add(EType::OrInclusive);
 		Tags.Add(Tag);
@@ -492,7 +492,7 @@ namespace UE::Learning::Observation
 
 	FObjectElement FObject::CreateArray(const FObjectArrayParameters Parameters, const FName Tag)
 	{
-		UE_LEARNING_CHECK(Private::CheckAllValid(*this, Parameters.Elements));
+		check(Private::CheckAllValid(*this, Parameters.Elements));
 
 		const int32 Index = Types.Add(EType::Array);
 		Tags.Add(Tag);
@@ -514,7 +514,7 @@ namespace UE::Learning::Observation
 
 	FObjectElement FObject::CreateSet(const FObjectSetParameters Parameters, const FName Tag)
 	{
-		UE_LEARNING_CHECK(Private::CheckAllValid(*this, Parameters.Elements));
+		check(Private::CheckAllValid(*this, Parameters.Elements));
 
 		const int32 Index = Types.Add(EType::Set);
 		Tags.Add(Tag);
@@ -536,7 +536,7 @@ namespace UE::Learning::Observation
 
 	FObjectElement FObject::CreateEncoding(const FObjectEncodingParameters Parameters, const FName Tag)
 	{
-		UE_LEARNING_CHECK(IsValid(Parameters.Element));
+		check(IsValid(Parameters.Element));
 
 		const int32 Index = Types.Add(EType::Encoding);
 		Tags.Add(Tag);
@@ -560,19 +560,19 @@ namespace UE::Learning::Observation
 
 	EType FObject::GetType(const FObjectElement Element) const
 	{
-		UE_LEARNING_CHECK(IsValid(Element));
+		check(IsValid(Element));
 		return Types[Element.Index];
 	}
 
 	FName FObject::GetTag(const FObjectElement Element) const
 	{
-		UE_LEARNING_CHECK(IsValid(Element));
+		check(IsValid(Element));
 		return Tags[Element.Index];
 	}
 
 	FObjectContinuousParameters FObject::GetContinuous(const FObjectElement Element) const
 	{
-		UE_LEARNING_CHECK(IsValid(Element) && GetType(Element) == EType::Continuous);
+		check(IsValid(Element) && GetType(Element) == EType::Continuous);
 
 		FObjectContinuousParameters Parameters;
 		Parameters.Values = TArrayView<const float>(ContinuousValues.GetData() + ContinuousDataOffsets[Element.Index], ContinuousDataNums[Element.Index]);
@@ -581,7 +581,7 @@ namespace UE::Learning::Observation
 
 	FObjectAndParameters FObject::GetAnd(const FObjectElement Element) const
 	{
-		UE_LEARNING_CHECK(IsValid(Element) && GetType(Element) == EType::And);
+		check(IsValid(Element) && GetType(Element) == EType::And);
 
 		FObjectAndParameters Parameters;
 		Parameters.ElementNames = TArrayView<const FName>(SubElementNames.GetData() + SubElementDataOffsets[Element.Index], SubElementDataNums[Element.Index]);
@@ -591,7 +591,7 @@ namespace UE::Learning::Observation
 
 	FObjectOrExclusiveParameters FObject::GetOrExclusive(const FObjectElement Element) const
 	{
-		UE_LEARNING_CHECK(IsValid(Element) && GetType(Element) == EType::OrExclusive);
+		check(IsValid(Element) && GetType(Element) == EType::OrExclusive);
 
 		FObjectOrExclusiveParameters Parameters;
 		Parameters.ElementName = SubElementNames[SubElementDataOffsets[Element.Index]];
@@ -601,7 +601,7 @@ namespace UE::Learning::Observation
 
 	FObjectOrInclusiveParameters FObject::GetOrInclusive(const FObjectElement Element) const
 	{
-		UE_LEARNING_CHECK(IsValid(Element) && GetType(Element) == EType::OrInclusive);
+		check(IsValid(Element) && GetType(Element) == EType::OrInclusive);
 
 		FObjectOrInclusiveParameters Parameters;
 		Parameters.ElementNames = TArrayView<const FName>(SubElementNames.GetData() + SubElementDataOffsets[Element.Index], SubElementDataNums[Element.Index]);
@@ -611,7 +611,7 @@ namespace UE::Learning::Observation
 
 	FObjectArrayParameters FObject::GetArray(const FObjectElement Element) const
 	{
-		UE_LEARNING_CHECK(IsValid(Element) && GetType(Element) == EType::Array);
+		check(IsValid(Element) && GetType(Element) == EType::Array);
 
 		FObjectArrayParameters Parameters;
 		Parameters.Elements = TArrayView<const FObjectElement>(SubElementObjects.GetData() + SubElementDataOffsets[Element.Index], SubElementDataNums[Element.Index]);
@@ -620,7 +620,7 @@ namespace UE::Learning::Observation
 
 	FObjectSetParameters FObject::GetSet(const FObjectElement Element) const
 	{
-		UE_LEARNING_CHECK(IsValid(Element) && GetType(Element) == EType::Set);
+		check(IsValid(Element) && GetType(Element) == EType::Set);
 
 		FObjectSetParameters Parameters;
 		Parameters.Elements = TArrayView<const FObjectElement>(SubElementObjects.GetData() + SubElementDataOffsets[Element.Index], SubElementDataNums[Element.Index]);
@@ -629,7 +629,7 @@ namespace UE::Learning::Observation
 
 	FObjectEncodingParameters FObject::GetEncoding(const FObjectElement Element) const
 	{
-		UE_LEARNING_CHECK(IsValid(Element) && GetType(Element) == EType::Encoding);
+		check(IsValid(Element) && GetType(Element) == EType::Encoding);
 
 		FObjectEncodingParameters Parameters;
 		Parameters.Element = SubElementObjects[SubElementDataOffsets[Element.Index]];
@@ -692,7 +692,7 @@ namespace UE::Learning::Observation
 			case EEncodingActivationFunction::ELU: return NNE::RuntimeBasic::FModelBuilder::EActivationFunction::ELU;
 			case EEncodingActivationFunction::TanH: return NNE::RuntimeBasic::FModelBuilder::EActivationFunction::TanH;
 			case EEncodingActivationFunction::GELU: return NNE::RuntimeBasic::FModelBuilder::EActivationFunction::GELU;
-			default: UE_LEARNING_NOT_IMPLEMENTED(); return NNE::RuntimeBasic::FModelBuilder::EActivationFunction::ReLU;
+			default: checkNoEntry(); return NNE::RuntimeBasic::FModelBuilder::EActivationFunction::ReLU;
 			}
 		}
 
@@ -741,7 +741,7 @@ namespace UE::Learning::Observation
 		const FSchemaElement SchemaElement,
 		const int32 Salt)
 	{
-		UE_LEARNING_CHECK(Schema.IsValid(SchemaElement));
+		check(Schema.IsValid(SchemaElement));
 		const EType SchemaElementType = Schema.GetType(SchemaElement);
 
 		const int32 Hash = Private::HashCombine({ Salt, Private::HashInt((int32)SchemaElementType) });
@@ -790,7 +790,7 @@ namespace UE::Learning::Observation
 
 		default:
 		{
-			UE_LEARNING_NOT_IMPLEMENTED();
+			checkNoEntry();
 			return 0;
 		}
 
@@ -803,8 +803,8 @@ namespace UE::Learning::Observation
 		const FSchema& SchemaB,
 		const FSchemaElement SchemaElementB)
 	{
-		UE_LEARNING_CHECK(SchemaA.IsValid(SchemaElementA));
-		UE_LEARNING_CHECK(SchemaB.IsValid(SchemaElementB));
+		check(SchemaA.IsValid(SchemaElementA));
+		check(SchemaB.IsValid(SchemaElementB));
 
 		const EType SchemaElementTypeA = SchemaA.GetType(SchemaElementA);
 		const EType SchemaElementTypeB = SchemaB.GetType(SchemaElementB);
@@ -894,13 +894,13 @@ namespace UE::Learning::Observation
 
 		case EType::Encoding:
 		{
-			UE_LEARNING_CHECKF(false, TEXT("Encoding elements should always be forwarded..."));
+			checkf(false, TEXT("Encoding elements should always be forwarded..."));
 			return false;
 		}
 
 		default:
 		{
-			UE_LEARNING_NOT_IMPLEMENTED();
+			checkNoEntry();
 			return false;
 		}
 
@@ -978,7 +978,7 @@ namespace UE::Learning::Observation
 					NNE::RuntimeBasic::FModelBuilder::EWeightInitializationType::KaimingGaussian; break;
 				case EWeightInitialization::KaimingUniform: LinearLayerSettings.WeightInitializationSettings.Type =
 					NNE::RuntimeBasic::FModelBuilder::EWeightInitializationType::KaimingUniform; break;
-				default: UE_LEARNING_NOT_IMPLEMENTED();
+				default: checkNoEntry();
 				}
 
 				BuilderEncoders.Emplace(Builder.MakeLinearLayer(SubElementEncodedSize, Parameters.EncodingSize, LinearLayerSettings));
@@ -1017,7 +1017,7 @@ namespace UE::Learning::Observation
 					NNE::RuntimeBasic::FModelBuilder::EWeightInitializationType::KaimingGaussian; break;
 				case EWeightInitialization::KaimingUniform: LinearLayerSettings.WeightInitializationSettings.Type =
 					NNE::RuntimeBasic::FModelBuilder::EWeightInitializationType::KaimingUniform; break;
-				default: UE_LEARNING_NOT_IMPLEMENTED();
+				default: checkNoEntry();
 				}
 
 				BuilderQueryLayers.Emplace(Builder.MakeLinearLayer(SubElementEncodedSize, Parameters.AttentionHeadNum * Parameters.AttentionEncodingSize, LinearLayerSettings));
@@ -1067,7 +1067,7 @@ namespace UE::Learning::Observation
 				NNE::RuntimeBasic::FModelBuilder::EWeightInitializationType::KaimingGaussian; break;
 			case EWeightInitialization::KaimingUniform: LinearLayerSettings.WeightInitializationSettings.Type =
 				NNE::RuntimeBasic::FModelBuilder::EWeightInitializationType::KaimingUniform; break;
-			default: UE_LEARNING_NOT_IMPLEMENTED();
+			default: checkNoEntry();
 			}
 
 			OutElement = Builder.MakeAggregateSet(
@@ -1103,7 +1103,7 @@ namespace UE::Learning::Observation
 				NNE::RuntimeBasic::FModelBuilder::EWeightInitializationType::KaimingGaussian; break;
 			case EWeightInitialization::KaimingUniform: LinearLayerSettings.WeightInitializationSettings.Type =
 				NNE::RuntimeBasic::FModelBuilder::EWeightInitializationType::KaimingUniform; break;
-			default: UE_LEARNING_NOT_IMPLEMENTED();
+			default: checkNoEntry();
 			}
 
 			OutElement = Builder.MakeSequence({
@@ -1123,15 +1123,15 @@ namespace UE::Learning::Observation
 
 		default:
 		{
-			UE_LEARNING_NOT_IMPLEMENTED();
+			checkNoEntry();
 		}
 		}
 
-		UE_LEARNING_CHECKF(OutElement.GetInputSize() == Schema.GetObservationVectorSize(SchemaElement),
+		checkf(OutElement.GetInputSize() == Schema.GetObservationVectorSize(SchemaElement),
 			TEXT("Encoder Network Input unexpected size for %s. Got %i, expected %i according to Schema."),
 			*Schema.GetTag(SchemaElement).ToString(), OutElement.GetInputSize(), Schema.GetObservationVectorSize(SchemaElement));
 
-		UE_LEARNING_CHECKF(OutElement.GetOutputSize() == Schema.GetEncodedVectorSize(SchemaElement),
+		checkf(OutElement.GetOutputSize() == Schema.GetEncodedVectorSize(SchemaElement),
 			TEXT("Encoder Network Output unexpected size for %s. Got %i, expected %i according to Schema."),
 			*Schema.GetTag(SchemaElement).ToString(), OutElement.GetOutputSize(), Schema.GetEncodedVectorSize(SchemaElement));
 	}
@@ -1145,7 +1145,7 @@ namespace UE::Learning::Observation
 		const FNetworkSettings& NetworkSettings,
 		const uint32 Seed)
 	{
-		UE_LEARNING_CHECK(Schema.IsValid(SchemaElement));
+		check(Schema.IsValid(SchemaElement));
 
 		NNE::RuntimeBasic::FModelBuilder Builder(Seed);
 		NNE::RuntimeBasic::FModelBuilderElement Element;
@@ -1160,15 +1160,15 @@ namespace UE::Learning::Observation
 		const FObject& Object,
 		const FObjectElement ObjectElement)
 	{
-		UE_LEARNING_CHECK(Schema.IsValid(SchemaElement));
-		UE_LEARNING_CHECK(Object.IsValid(ObjectElement));
-		UE_LEARNING_CHECK(OutObservationVector.Num() == Schema.GetObservationVectorSize(SchemaElement));
+		check(Schema.IsValid(SchemaElement));
+		check(Object.IsValid(ObjectElement));
+		check(OutObservationVector.Num() == Schema.GetObservationVectorSize(SchemaElement));
 
 		// Check that the types match
 
 		const EType SchemaElementType = Schema.GetType(SchemaElement);
 		const EType ObjectElementType = Object.GetType(ObjectElement);
-		UE_LEARNING_CHECK(ObjectElementType == SchemaElementType);
+		check(ObjectElementType == SchemaElementType);
 
 		// Zero Observation Vector
 
@@ -1188,9 +1188,9 @@ namespace UE::Learning::Observation
 			const FSchemaContinuousParameters SchemaParameters = Schema.GetContinuous(SchemaElement);
 
 			const TArrayView<const float> ObservationValues = Object.GetContinuous(ObjectElement).Values;
-			UE_LEARNING_CHECK(Schema.GetObservationVectorSize(SchemaElement) == ObservationValues.Num());
-			UE_LEARNING_CHECK(Schema.GetObservationVectorSize(SchemaElement) == OutObservationVector.Num());
-			UE_LEARNING_CHECK(Schema.GetObservationVectorSize(SchemaElement) == SchemaParameters.Num);
+			check(Schema.GetObservationVectorSize(SchemaElement) == ObservationValues.Num());
+			check(Schema.GetObservationVectorSize(SchemaElement) == OutObservationVector.Num());
+			check(Schema.GetObservationVectorSize(SchemaElement) == SchemaParameters.Num);
 
 			// Copy in and scale the values from the observation object
 
@@ -1211,7 +1211,7 @@ namespace UE::Learning::Observation
 
 			const FSchemaAndParameters SchemaParameters = Schema.GetAnd(SchemaElement);
 			const FObjectAndParameters ObjectParameters = Object.GetAnd(ObjectElement);
-			UE_LEARNING_CHECK(SchemaParameters.Elements.Num() == ObjectParameters.Elements.Num());
+			check(SchemaParameters.Elements.Num() == ObjectParameters.Elements.Num());
 
 			// Update Sub-elements
 
@@ -1220,7 +1220,7 @@ namespace UE::Learning::Observation
 			for (int32 SchemaElementIdx = 0; SchemaElementIdx < SchemaParameters.Elements.Num(); SchemaElementIdx++)
 			{
 				const int32 ObjectElementIndex = ObjectParameters.ElementNames.Find(SchemaParameters.ElementNames[SchemaElementIdx]);
-				UE_LEARNING_CHECK(ObjectElementIndex != INDEX_NONE);
+				check(ObjectElementIndex != INDEX_NONE);
 
 				const int32 SubElementSize = Schema.GetObservationVectorSize(SchemaParameters.Elements[SchemaElementIdx]);
 
@@ -1234,7 +1234,7 @@ namespace UE::Learning::Observation
 				SubElementOffset += SubElementSize;
 			}
 
-			UE_LEARNING_CHECK(SubElementOffset == OutObservationVector.Num());
+			check(SubElementOffset == OutObservationVector.Num());
 			return;
 		}
 
@@ -1246,7 +1246,7 @@ namespace UE::Learning::Observation
 			const FObjectOrExclusiveParameters ObjectParameters = Object.GetOrExclusive(ObjectElement);
 
 			const int32 SchemaElementIndex = SchemaParameters.ElementNames.Find(ObjectParameters.ElementName);
-			UE_LEARNING_CHECK(SchemaElementIndex != INDEX_NONE);
+			check(SchemaElementIndex != INDEX_NONE);
 
 			// Update sub-element
 
@@ -1265,7 +1265,7 @@ namespace UE::Learning::Observation
 
 			OutObservationVector[MaxSubElementSize + SchemaElementIndex] = 1.0f;
 
-			UE_LEARNING_CHECK(OutObservationVector.Num() == MaxSubElementSize + SchemaParameters.Elements.Num());
+			check(OutObservationVector.Num() == MaxSubElementSize + SchemaParameters.Elements.Num());
 			return;
 		}
 
@@ -1275,7 +1275,7 @@ namespace UE::Learning::Observation
 
 			const FSchemaOrInclusiveParameters SchemaParameters = Schema.GetOrInclusive(SchemaElement);
 			const FObjectOrInclusiveParameters ObjectParameters = Object.GetOrInclusive(ObjectElement);
-			UE_LEARNING_CHECK(ObjectParameters.Elements.Num() <= SchemaParameters.Elements.Num());
+			check(ObjectParameters.Elements.Num() <= SchemaParameters.Elements.Num());
 
 			// Update sub-elements
 
@@ -1301,12 +1301,12 @@ namespace UE::Learning::Observation
 
 			// Set Mask
 
-			UE_LEARNING_CHECK(SubElementOffset + SchemaParameters.Elements.Num() == OutObservationVector.Num());
+			check(SubElementOffset + SchemaParameters.Elements.Num() == OutObservationVector.Num());
 
 			for (int32 ObjectElementIdx = 0; ObjectElementIdx < ObjectParameters.Elements.Num(); ObjectElementIdx++)
 			{
 				const int32 SchemaElementIdx = SchemaParameters.ElementNames.Find(ObjectParameters.ElementNames[ObjectElementIdx]);
-				UE_LEARNING_CHECK(SchemaElementIdx != INDEX_NONE);
+				check(SchemaElementIdx != INDEX_NONE);
 
 				OutObservationVector[SubElementOffset + SchemaElementIdx] = 1.0f;
 			}
@@ -1320,7 +1320,7 @@ namespace UE::Learning::Observation
 
 			const FSchemaArrayParameters SchemaParameters = Schema.GetArray(SchemaElement);
 			const FObjectArrayParameters ObjectParameters = Object.GetArray(ObjectElement);
-			UE_LEARNING_CHECK(SchemaParameters.Num == ObjectParameters.Elements.Num());
+			check(SchemaParameters.Num == ObjectParameters.Elements.Num());
 
 			// Update sub-elements
 
@@ -1345,7 +1345,7 @@ namespace UE::Learning::Observation
 
 			const FSchemaSetParameters SchemaParameters = Schema.GetSet(SchemaElement);
 			const FObjectSetParameters ObjectParameters = Object.GetSet(ObjectElement);
-			UE_LEARNING_CHECK(SchemaParameters.MaxNum >= ObjectParameters.Elements.Num());
+			check(SchemaParameters.MaxNum >= ObjectParameters.Elements.Num());
 
 			// Update sub-elements
 
@@ -1371,7 +1371,7 @@ namespace UE::Learning::Observation
 
 			Array::Set(OutObservationVector.Slice(SubElementOffset, ObjectParameters.Elements.Num()), 1.0f);
 
-			UE_LEARNING_CHECK(SubElementOffset + SchemaParameters.MaxNum == OutObservationVector.Num());
+			check(SubElementOffset + SchemaParameters.MaxNum == OutObservationVector.Num());
 			return;
 		}
 
@@ -1392,7 +1392,7 @@ namespace UE::Learning::Observation
 
 		default:
 		{
-			UE_LEARNING_NOT_IMPLEMENTED();
+			checkNoEntry();
 			return;
 		}
 
@@ -1406,7 +1406,7 @@ namespace UE::Learning::Observation
 		const FSchemaElement SchemaElement,
 		const TLearningArrayView<1, const float> ObservationVector)
 	{
-		UE_LEARNING_CHECK(Schema.IsValid(SchemaElement));
+		check(Schema.IsValid(SchemaElement));
 
 		// Check that the types match
 
@@ -1416,7 +1416,7 @@ namespace UE::Learning::Observation
 		// Get Observation Vector Size
 
 		const int32 ObservationVectorSize = ObservationVector.Num();
-		UE_LEARNING_CHECK(ObservationVectorSize == Schema.GetObservationVectorSize(SchemaElement));
+		check(ObservationVectorSize == Schema.GetObservationVectorSize(SchemaElement));
 
 		// Logic for each specific element type
 
@@ -1432,7 +1432,7 @@ namespace UE::Learning::Observation
 		case EType::Continuous:
 		{
 			const FSchemaContinuousParameters SchemaParameters = Schema.GetContinuous(SchemaElement);
-			UE_LEARNING_CHECK(ObservationVectorSize == SchemaParameters.Num);
+			check(ObservationVectorSize == SchemaParameters.Num);
 
 			const int32 ValueNum = SchemaParameters.Num;
 			const float ValueScale = FMath::Max(SchemaParameters.Scale, UE_SMALL_NUMBER);
@@ -1471,7 +1471,7 @@ namespace UE::Learning::Observation
 
 				SubElementOffset += SubElementSize;
 			}
-			UE_LEARNING_CHECK(SubElementOffset == ObservationVectorSize);
+			check(SubElementOffset == ObservationVectorSize);
 
 			OutObjectElement = OutObject.CreateAnd({ Parameters.ElementNames, SubElements }, SchemaElementTag);
 			return;
@@ -1488,14 +1488,14 @@ namespace UE::Learning::Observation
 			int32 SchemaElementIndex = INDEX_NONE;
 			for (int32 SubElementIdx = 0; SubElementIdx < Parameters.Elements.Num(); SubElementIdx++)
 			{
-				UE_LEARNING_CHECK(ObservationVector[MaxSubElementSize + SubElementIdx] == 0.0f || ObservationVector[MaxSubElementSize + SubElementIdx] == 1.0f);
+				check(ObservationVector[MaxSubElementSize + SubElementIdx] == 0.0f || ObservationVector[MaxSubElementSize + SubElementIdx] == 1.0f);
 				if (ObservationVector[MaxSubElementSize + SubElementIdx])
 				{
 					SchemaElementIndex = SubElementIdx;
 					break;
 				}
 			}
-			UE_LEARNING_CHECK(SchemaElementIndex != INDEX_NONE);
+			check(SchemaElementIndex != INDEX_NONE);
 
 			// Create sub-element
 
@@ -1533,7 +1533,7 @@ namespace UE::Learning::Observation
 			{
 				const int32 SubElementSize = Schema.GetObservationVectorSize(Parameters.Elements[SubElementIdx]);
 
-				UE_LEARNING_CHECK(
+				check(
 					ObservationVector[TotalSubElementSize + SubElementIdx] == 0.0f || 
 					ObservationVector[TotalSubElementSize + SubElementIdx] == 1.0f);
 
@@ -1553,7 +1553,7 @@ namespace UE::Learning::Observation
 
 				SubElementOffset += SubElementSize;
 			}
-			UE_LEARNING_CHECK(SubElementOffset + Parameters.Elements.Num() == ObservationVectorSize);
+			check(SubElementOffset + Parameters.Elements.Num() == ObservationVectorSize);
 
 			OutObjectElement = OutObject.CreateOrInclusive({ SubElementNames, SubElements }, SchemaElementTag);
 			return;
@@ -1597,7 +1597,7 @@ namespace UE::Learning::Observation
 
 			for (int32 SubElementIdx = 0; SubElementIdx < Parameters.MaxNum; SubElementIdx++)
 			{
-				UE_LEARNING_CHECK(
+				check(
 					ObservationVector[SubElementSize * Parameters.MaxNum + SubElementIdx] == 0.0f || 
 					ObservationVector[SubElementSize * Parameters.MaxNum + SubElementIdx] == 1.0f);
 
@@ -1640,7 +1640,7 @@ namespace UE::Learning::Observation
 
 		default:
 		{
-			UE_LEARNING_NOT_IMPLEMENTED();
+			checkNoEntry();
 			OutObjectElement = FObjectElement();
 			return;
 		}
@@ -1654,12 +1654,12 @@ namespace UE::Learning::Observation
 		const FSchemaElement SchemaElement,
 		const float NoiseScale)
 	{
-		UE_LEARNING_CHECK(Schema.IsValid(SchemaElement));
+		check(Schema.IsValid(SchemaElement));
 
 		const EType SchemaElementType = Schema.GetType(SchemaElement);
 
 		const int32 ObservationVectorSize = InOutObservationVector.Num();
-		UE_LEARNING_CHECK(ObservationVectorSize == Schema.GetObservationVectorSize(SchemaElement));
+		check(ObservationVectorSize == Schema.GetObservationVectorSize(SchemaElement));
 
 		switch (SchemaElementType)
 		{
@@ -1672,7 +1672,7 @@ namespace UE::Learning::Observation
 		case EType::Continuous:
 		{
 			const FSchemaContinuousParameters SchemaParameters = Schema.GetContinuous(SchemaElement);
-			UE_LEARNING_CHECK(ObservationVectorSize == SchemaParameters.Num);
+			check(ObservationVectorSize == SchemaParameters.Num);
 
 			TLearningArray<1, float, TInlineAllocator<32>> NoiseValues;
 			NoiseValues.SetNumUninitialized({ SchemaParameters.Num });
@@ -1703,7 +1703,7 @@ namespace UE::Learning::Observation
 
 				SubElementOffset += SubElementSize;
 			}
-			UE_LEARNING_CHECK(SubElementOffset == ObservationVectorSize);
+			check(SubElementOffset == ObservationVectorSize);
 
 			return;
 		}
@@ -1716,14 +1716,14 @@ namespace UE::Learning::Observation
 			int32 SchemaElementIndex = INDEX_NONE;
 			for (int32 SubElementIdx = 0; SubElementIdx < Parameters.Elements.Num(); SubElementIdx++)
 			{
-				UE_LEARNING_CHECK(InOutObservationVector[MaxSubElementSize + SubElementIdx] == 0.0f || InOutObservationVector[MaxSubElementSize + SubElementIdx] == 1.0f);
+				check(InOutObservationVector[MaxSubElementSize + SubElementIdx] == 0.0f || InOutObservationVector[MaxSubElementSize + SubElementIdx] == 1.0f);
 				if (InOutObservationVector[MaxSubElementSize + SubElementIdx])
 				{
 					SchemaElementIndex = SubElementIdx;
 					break;
 				}
 			}
-			UE_LEARNING_CHECK(SchemaElementIndex != INDEX_NONE);
+			check(SchemaElementIndex != INDEX_NONE);
 
 			const int32 SubElementSize = Schema.GetObservationVectorSize(Parameters.Elements[SchemaElementIndex]);
 
@@ -1747,7 +1747,7 @@ namespace UE::Learning::Observation
 			{
 				const int32 SubElementSize = Schema.GetObservationVectorSize(Parameters.Elements[SubElementIdx]);
 
-				UE_LEARNING_CHECK(
+				check(
 					InOutObservationVector[TotalSubElementSize + SubElementIdx] == 0.0f ||
 					InOutObservationVector[TotalSubElementSize + SubElementIdx] == 1.0f);
 
@@ -1763,7 +1763,7 @@ namespace UE::Learning::Observation
 
 				SubElementOffset += SubElementSize;
 			}
-			UE_LEARNING_CHECK(SubElementOffset + Parameters.Elements.Num() == ObservationVectorSize);
+			check(SubElementOffset + Parameters.Elements.Num() == ObservationVectorSize);
 
 			return;
 		}
@@ -1793,7 +1793,7 @@ namespace UE::Learning::Observation
 
 			for (int32 SubElementIdx = 0; SubElementIdx < Parameters.MaxNum; SubElementIdx++)
 			{
-				UE_LEARNING_CHECK(
+				check(
 					InOutObservationVector[SubElementSize * Parameters.MaxNum + SubElementIdx] == 0.0f ||
 					InOutObservationVector[SubElementSize * Parameters.MaxNum + SubElementIdx] == 1.0f);
 
@@ -1827,7 +1827,7 @@ namespace UE::Learning::Observation
 
 		default:
 		{
-			UE_LEARNING_NOT_IMPLEMENTED();
+			checkNoEntry();
 			return;
 		}
 		}
diff --git a/Engine/Plugins/Experimental/LearningAgents/Source/Learning/Private/LearningOptimizationPlanner.cpp b/Engine/Plugins/Experimental/LearningAgents/Source/Learning/Private/LearningOptimizationPlanner.cpp
--- a/Engine/Plugins/Experimental/LearningAgents/Source/Learning/Private/LearningOptimizationPlanner.cpp
+++ b/Engine/Plugins/Experimental/LearningAgents/Source/Learning/Private/LearningOptimizationPlanner.cpp
@@ -30,7 +30,7 @@ namespace UE::Learning
 			const TLearningArrayView<3, const float> ActionVectors,
 			const FIndexSet Instances)
 		{
-			UE_LEARNING_TRACE_CPUPROFILER_EVENT_SCOPE(Learning::OptimizationPlanner::RunPlan);
+			TRACE_CPUPROFILER_EVENT_SCOPE(Learning::OptimizationPlanner::RunPlan);
 
 			const int32 StepNum = ActionVectors.Num<0>();
 
@@ -99,7 +99,7 @@ namespace UE::Learning
 			const TFunctionRef<void(const FIndexSet Instances)> RewardFunction,
 			const FIndexSet SampleInstances)
 		{
-			UE_LEARNING_TRACE_CPUPROFILER_EVENT_SCOPE(Learning::OptimizationPlanner::ComputePlanLosses);
+			TRACE_CPUPROFILER_EVENT_SCOPE(Learning::OptimizationPlanner::ComputePlanLosses);
 
 			// Zero Losses
 
@@ -162,13 +162,13 @@ namespace UE::Learning
 			FRWLock* ActionVectorsLock,
 			TAtomic<bool>* bActionVectorsUpdatedFlag)
 		{
-			UE_LEARNING_TRACE_CPUPROFILER_EVENT_SCOPE(Learning::OptimizationPlanner::Plan);
+			TRACE_CPUPROFILER_EVENT_SCOPE(Learning::OptimizationPlanner::Plan);
 
-			UE_LEARNING_CHECKF(OptimizationPlannerBuffer.Samples.Num<0>() == SampleInstances.Num(),
+			checkf(OptimizationPlannerBuffer.Samples.Num<0>() == SampleInstances.Num(),
 				TEXT("Number of samples in OptimizationPlannerBuffer (%i) must match the number of SampleInstances (%i)."),
 				OptimizationPlannerBuffer.Samples.Num<0>(), SampleInstances.Num());
 
-			UE_LEARNING_CHECKF(OptimizationPlannerBuffer.Samples.Num<1>() == InOutActionVectors.Num<0>(),
+			checkf(OptimizationPlannerBuffer.Samples.Num<1>() == InOutActionVectors.Num<0>(),
 				TEXT("Number of steps in OptimizationPlannerBuffer (%i) must match the number of steps in InOutActionVectors (%i)."),
 				OptimizationPlannerBuffer.Samples.Num<1>(), InOutActionVectors.Num<0>());
 
@@ -245,21 +245,21 @@ namespace UE::Learning
 			FRWLock* ActionVectorsLock,
 			TAtomic<bool>* bActionVectorsUpdatedFlag)
 		{
-			UE_LEARNING_TRACE_CPUPROFILER_EVENT_SCOPE(Learning::OptimizationPlanner::PlanWindowed);
+			TRACE_CPUPROFILER_EVENT_SCOPE(Learning::OptimizationPlanner::PlanWindowed);
 
-			UE_LEARNING_CHECKF(WindowStepNum <= InOutActionVectors.Num<0>(),
+			checkf(WindowStepNum <= InOutActionVectors.Num<0>(),
 				TEXT("WindowStepNum (%i) must be less than or equal to the number of steps in the action vectors (%i)"),
 				WindowStepNum, InOutActionVectors.Num<0>());
 
-			UE_LEARNING_CHECKF(OptimizationPlannerBuffer.Losses.Num() == SampleInstances.Num(),
+			checkf(OptimizationPlannerBuffer.Losses.Num() == SampleInstances.Num(),
 				TEXT("Number of samples in OptimizationPlannerBuffer (%i) must match the number of SampleInstances (%i)."),
 				OptimizationPlannerBuffer.Losses.Num(), SampleInstances.Num());
 
-			UE_LEARNING_CHECKF(OptimizationPlannerBuffer.Samples.Num<1>() == WindowStepNum,
+			checkf(OptimizationPlannerBuffer.Samples.Num<1>() == WindowStepNum,
 				TEXT("Number of steps in OptimizationPlannerBuffer (%i) must match WindowStepNum (%i)."),
 				OptimizationPlannerBuffer.Samples.Num<1>(), WindowStepNum);
 
-			UE_LEARNING_CHECK(!SampleInstances.Contains(PlanInstance));
+			check(!SampleInstances.Contains(PlanInstance));
 
 			const int32 TotalStepNum = InOutActionVectors.Num<0>();
 
diff --git a/Engine/Plugins/Experimental/LearningAgents/Source/Learning/Private/LearningPCA.cpp b/Engine/Plugins/Experimental/LearningAgents/Source/Learning/Private/LearningPCA.cpp
--- a/Engine/Plugins/Experimental/LearningAgents/Source/Learning/Private/LearningPCA.cpp
+++ b/Engine/Plugins/Experimental/LearningAgents/Source/Learning/Private/LearningPCA.cpp
@@ -44,7 +44,7 @@ namespace UE::Learning
 		const ELogSetting LogSettings,
 		FProgress* Progress)
 	{
-		UE_LEARNING_TRACE_CPUPROFILER_EVENT_SCOPE(Learning::FPCAEncoder::Fit);
+		TRACE_CPUPROFILER_EVENT_SCOPE(Learning::FPCAEncoder::Fit);
 
 		Array::Check(Data);
 
@@ -129,7 +129,7 @@ namespace UE::Learning
 				UE_LOG(LogLearning, Display, TEXT("--------------------"));
 			}
 
-			UE_LEARNING_CHECK(ComponentsToKeepNum > 0);
+			check(ComponentsToKeepNum > 0);
 
 			if (LogSettings != ELogSetting::Silent)
 			{
@@ -168,7 +168,7 @@ namespace UE::Learning
 		const TLearningArrayView<2, const float> Data,
 		FProgress* Progress) const
 	{
-		UE_LEARNING_TRACE_CPUPROFILER_EVENT_SCOPE(Learning::FPCAEncoder::Transform);
+		TRACE_CPUPROFILER_EVENT_SCOPE(Learning::FPCAEncoder::Transform);
 
 		if (Progress)
 		{
@@ -180,7 +180,7 @@ namespace UE::Learning
 
 		SlicedParallelFor(Data.Num<0>(), 128, [&](const int32 RowStart, const int32 RowSliceNum)
 		{
-			UE_LEARNING_TRACE_CPUPROFILER_EVENT_SCOPE(Learning::PCA::Transform::Batch);
+			TRACE_CPUPROFILER_EVENT_SCOPE(Learning::PCA::Transform::Batch);
 
 			OutEigenMatrix(OutData.Slice(RowStart, RowSliceNum)).noalias() =
 				(InEigenMatrix(Matrix).transpose() * InEigenMatrix(Data.Slice(RowStart, RowSliceNum)).transpose()).transpose();
@@ -198,9 +198,9 @@ namespace UE::Learning
 		TLearningArrayView<1, float> OutData,
 		const TLearningArrayView<1, const float> Data) const
 	{
-		UE_LEARNING_TRACE_CPUPROFILER_EVENT_SCOPE(Learning::FPCAEncoder::Transform);
-		UE_LEARNING_CHECK(Data.Num() == DimensionNum());
-		UE_LEARNING_CHECK(OutData.Num() == FeatureNum());
+		TRACE_CPUPROFILER_EVENT_SCOPE(Learning::FPCAEncoder::Transform);
+		check(Data.Num() == DimensionNum());
+		check(OutData.Num() == FeatureNum());
 
 		Array::Check(Data);
 
@@ -213,9 +213,9 @@ namespace UE::Learning
 		TLearningArrayView<1, float> OutData,
 		const TLearningArrayView<1, const float> Data) const
 	{
-		UE_LEARNING_TRACE_CPUPROFILER_EVENT_SCOPE(Learning::FPCAEncoder::InverseTransform);
-		UE_LEARNING_CHECK(Data.Num() == FeatureNum());
-		UE_LEARNING_CHECK(OutData.Num() == DimensionNum());
+		TRACE_CPUPROFILER_EVENT_SCOPE(Learning::FPCAEncoder::InverseTransform);
+		check(Data.Num() == FeatureNum());
+		check(OutData.Num() == DimensionNum());
 
 		Array::Check(Data);
 
diff --git a/Engine/Plugins/Experimental/LearningAgents/Source/Learning/Private/LearningPSOOptimizer.cpp b/Engine/Plugins/Experimental/LearningAgents/Source/Learning/Private/LearningPSOOptimizer.cpp
--- a/Engine/Plugins/Experimental/LearningAgents/Source/Learning/Private/LearningPSOOptimizer.cpp
+++ b/Engine/Plugins/Experimental/LearningAgents/Source/Learning/Private/LearningPSOOptimizer.cpp
@@ -14,7 +14,7 @@ namespace UE::Learning
 
 	void FPSOOptimizer::Resize(const int32 SampleNum, const int32 DimensionsNum)
 	{
-		UE_LEARNING_TRACE_CPUPROFILER_EVENT_SCOPE(Learning::FPSOOptimizer::Resize);
+		TRACE_CPUPROFILER_EVENT_SCOPE(Learning::FPSOOptimizer::Resize);
 
 		LocalBestPositions.SetNumUninitialized({ SampleNum, DimensionsNum });
 		LocalBestLoss.SetNumUninitialized({ SampleNum });
@@ -28,7 +28,7 @@ namespace UE::Learning
 		TLearningArrayView<2, float> OutSamples,
 		const TLearningArrayView<1, const float> InitialGuess)
 	{
-		UE_LEARNING_TRACE_CPUPROFILER_EVENT_SCOPE(Learning::FPSOOptimizer::Reset);
+		TRACE_CPUPROFILER_EVENT_SCOPE(Learning::FPSOOptimizer::Reset);
 
 		const int32 SampleNum = OutSamples.Num<0>();
 		const int32 DimensionNum = OutSamples.Num<1>();
@@ -61,7 +61,7 @@ namespace UE::Learning
 		const TLearningArrayView<1, const float> Losses,
 		const ELogSetting LogSettings)
 	{
-		UE_LEARNING_TRACE_CPUPROFILER_EVENT_SCOPE(Learning::FPSOOptimizer::Update);
+		TRACE_CPUPROFILER_EVENT_SCOPE(Learning::FPSOOptimizer::Update);
 
 		const int32 SampleNum = InOutSamples.Num<0>();
 		const int32 DimNum = InOutSamples.Num<1>();
@@ -89,7 +89,7 @@ namespace UE::Learning
 		// performing it instead at the beginning of the next loop iteration.
 
 		{
-			UE_LEARNING_TRACE_CPUPROFILER_EVENT_SCOPE(Learning::FPSOOptimizer::Update::UpdateBest);
+			TRACE_CPUPROFILER_EVENT_SCOPE(Learning::FPSOOptimizer::Update::UpdateBest);
 
 #if UE_LEARNING_ISPC
 			ispc::LearningUpdatePSOBest(
@@ -120,7 +120,7 @@ namespace UE::Learning
 		}
 
 		{
-			UE_LEARNING_TRACE_CPUPROFILER_EVENT_SCOPE(Learning::FPSOOptimizer::Update::UpdateVelocitiesSamples);
+			TRACE_CPUPROFILER_EVENT_SCOPE(Learning::FPSOOptimizer::Update::UpdateVelocitiesSamples);
 
 #if UE_LEARNING_ISPC
 			ispc::LearningUpdatePSOVelocitiesSamples(
diff --git a/Engine/Plugins/Experimental/LearningAgents/Source/Learning/Private/LearningPolicy.cpp b/Engine/Plugins/Experimental/LearningAgents/Source/Learning/Private/LearningPolicy.cpp
--- a/Engine/Plugins/Experimental/LearningAgents/Source/Learning/Private/LearningPolicy.cpp
+++ b/Engine/Plugins/Experimental/LearningAgents/Source/Learning/Private/LearningPolicy.cpp
@@ -15,8 +15,8 @@ namespace UE::Learning
 		, ActionEncodedNum(InActionEncodedNum)
 		, MemoryStateNum(InMemoryStateNum)
 	{
-		UE_LEARNING_CHECK(InNeuralNetwork->GetInputSize() == ObservationEncodedNum + MemoryStateNum);
-		UE_LEARNING_CHECK(InNeuralNetwork->GetOutputSize() == ActionEncodedNum + MemoryStateNum);
+		check(InNeuralNetwork->GetInputSize() == ObservationEncodedNum + MemoryStateNum);
+		check(InNeuralNetwork->GetOutputSize() == ActionEncodedNum + MemoryStateNum);
 
 		Input.SetNumUninitialized({ InMaxInstanceNum, ObservationEncodedNum + MemoryStateNum });
 		Output.SetNumUninitialized({ InMaxInstanceNum, ActionEncodedNum + MemoryStateNum });
@@ -34,7 +34,7 @@ namespace UE::Learning
 		const TLearningArrayView<2, const float> InputMemoryState, 
 		const FIndexSet Instances)
 	{
-		UE_LEARNING_TRACE_CPUPROFILER_EVENT_SCOPE(Learning::FNeuralNetworkPolicy::Evaluate);
+		TRACE_CPUPROFILER_EVENT_SCOPE(Learning::FNeuralNetworkPolicy::Evaluate);
 
 		const int32 InstanceNum = Instances.Num();
 
diff --git a/Engine/Plugins/Experimental/LearningAgents/Source/Learning/Private/LearningRandom.cpp b/Engine/Plugins/Experimental/LearningAgents/Source/Learning/Private/LearningRandom.cpp
--- a/Engine/Plugins/Experimental/LearningAgents/Source/Learning/Private/LearningRandom.cpp
+++ b/Engine/Plugins/Experimental/LearningAgents/Source/Learning/Private/LearningRandom.cpp
@@ -145,7 +145,7 @@ namespace UE::Learning::Random
 		TLearningArrayView<1, uint32> Output,
 		const uint32 State)
 	{
-		UE_LEARNING_TRACE_CPUPROFILER_EVENT_SCOPE(Learning::Random::IntArray);
+		TRACE_CPUPROFILER_EVENT_SCOPE(Learning::Random::IntArray);
 
 		const int32 ElementNum = Output.Num();
 
@@ -167,7 +167,7 @@ namespace UE::Learning::Random
 		const uint32 State,
 		const FIndexSet Indices)
 	{
-		UE_LEARNING_TRACE_CPUPROFILER_EVENT_SCOPE(Learning::Random::IntArray);
+		TRACE_CPUPROFILER_EVENT_SCOPE(Learning::Random::IntArray);
 
 		if (Indices.IsSlice() && Indices.Num() > 4)
 		{
@@ -188,7 +188,7 @@ namespace UE::Learning::Random
 		TLearningArrayView<1, float> Output,
 		const uint32 State)
 	{
-		UE_LEARNING_TRACE_CPUPROFILER_EVENT_SCOPE(Learning::Random::FloatArray);
+		TRACE_CPUPROFILER_EVENT_SCOPE(Learning::Random::FloatArray);
 
 		const int32 ElementNum = Output.Num();
 
@@ -211,7 +211,7 @@ namespace UE::Learning::Random
 		const float Min,
 		const float Max)
 	{
-		UE_LEARNING_TRACE_CPUPROFILER_EVENT_SCOPE(Learning::Random::UniformArray);
+		TRACE_CPUPROFILER_EVENT_SCOPE(Learning::Random::UniformArray);
 
 		const int32 ElementNum = Output.Num();
 
@@ -236,7 +236,7 @@ namespace UE::Learning::Random
 		const float Mean,
 		const float Std)
 	{
-		UE_LEARNING_TRACE_CPUPROFILER_EVENT_SCOPE(Learning::Random::GaussianArray);
+		TRACE_CPUPROFILER_EVENT_SCOPE(Learning::Random::GaussianArray);
 
 		const int32 ElementNum = Output.Num();
 
@@ -262,7 +262,7 @@ namespace UE::Learning::Random
 		const float Std,
 		const float Clip)
 	{
-		UE_LEARNING_TRACE_CPUPROFILER_EVENT_SCOPE(Learning::Random::ClippedGaussianArray);
+		TRACE_CPUPROFILER_EVENT_SCOPE(Learning::Random::ClippedGaussianArray);
 
 		const int32 ElementNum = Output.Num();
 
@@ -281,7 +281,7 @@ namespace UE::Learning::Random
 		const FVector Axis0,
 		const FVector Axis1)
 	{
-		UE_LEARNING_TRACE_CPUPROFILER_EVENT_SCOPE(Learning::Random::PlanarClippedGaussianArray);
+		TRACE_CPUPROFILER_EVENT_SCOPE(Learning::Random::PlanarClippedGaussianArray);
 
 		const int32 ElementNum = Output.Num();
 
@@ -297,7 +297,7 @@ namespace UE::Learning::Random
 		const FVector Axis0,
 		const FVector Axis1)
 	{
-		UE_LEARNING_TRACE_CPUPROFILER_EVENT_SCOPE(Learning::Random::PlanarDirectionArray);
+		TRACE_CPUPROFILER_EVENT_SCOPE(Learning::Random::PlanarDirectionArray);
 
 		const int32 ElementNum = Output.Num();
 
@@ -314,7 +314,7 @@ namespace UE::Learning::Random
 		const TLearningArrayView<1, const float> LogStd,
 		const float Scale)
 	{
-		UE_LEARNING_TRACE_CPUPROFILER_EVENT_SCOPE(Learning::Random::DistributionIndependantNormal);
+		TRACE_CPUPROFILER_EVENT_SCOPE(Learning::Random::DistributionIndependantNormal);
 
 		const int32 ElementNum = Output.Num();
 
@@ -353,7 +353,7 @@ namespace UE::Learning::Random
 		const TLearningArrayView<1, const float> MaskedValues,
 		const float Scale)
 	{
-		UE_LEARNING_TRACE_CPUPROFILER_EVENT_SCOPE(Learning::Random::DistributionIndependantNormalMasked);
+		TRACE_CPUPROFILER_EVENT_SCOPE(Learning::Random::DistributionIndependantNormalMasked);
 
 		const int32 ElementNum = Output.Num();
 
@@ -391,9 +391,9 @@ namespace UE::Learning::Random
 		const TLearningArrayView<1, const float> Logits,
 		const float Scale)
 	{
-		UE_LEARNING_TRACE_CPUPROFILER_EVENT_SCOPE(Learning::Random::DistributionMultinoulli);
+		TRACE_CPUPROFILER_EVENT_SCOPE(Learning::Random::DistributionMultinoulli);
 
-		UE_LEARNING_CHECK(Output.Num() == Logits.Num());
+		check(Output.Num() == Logits.Num());
 
 		const int32 ElementNum = Output.Num();
 		if (ElementNum == 0)
@@ -429,7 +429,7 @@ namespace UE::Learning::Random
 			}
 		}
 
-		UE_LEARNING_CHECK(MaxIdx != INDEX_NONE);
+		check(MaxIdx != INDEX_NONE);
 
 		Array::Zero(Output);
 		Output[MaxIdx] = 1.0f;
@@ -442,10 +442,10 @@ namespace UE::Learning::Random
 		const TLearningArrayView<1, const bool> Masked,
 		const float Scale)
 	{
-		UE_LEARNING_TRACE_CPUPROFILER_EVENT_SCOPE(Learning::Random::DistributionMultinoulliMasked);
+		TRACE_CPUPROFILER_EVENT_SCOPE(Learning::Random::DistributionMultinoulliMasked);
 
-		UE_LEARNING_CHECK(Output.Num() == Logits.Num());
-		UE_LEARNING_CHECK(Logits.Num() == Masked.Num());
+		check(Output.Num() == Logits.Num());
+		check(Logits.Num() == Masked.Num());
 
 		const int32 ElementNum = Output.Num();
 		if (ElementNum == 0)
@@ -484,7 +484,7 @@ namespace UE::Learning::Random
 			}
 		}
 
-		UE_LEARNING_CHECK(MaxIdx != INDEX_NONE);
+		check(MaxIdx != INDEX_NONE);
 
 		Array::Zero(Output);
 		Output[MaxIdx] = 1.0f;
@@ -496,9 +496,9 @@ namespace UE::Learning::Random
 		const TLearningArrayView<1, const float> Logits,
 		const float Scale)
 	{
-		UE_LEARNING_TRACE_CPUPROFILER_EVENT_SCOPE(Learning::Random::DistributionBernoulli);
+		TRACE_CPUPROFILER_EVENT_SCOPE(Learning::Random::DistributionBernoulli);
 
-		UE_LEARNING_CHECK(Output.Num() == Logits.Num());
+		check(Output.Num() == Logits.Num());
 
 		const int32 ElementNum = Output.Num();
 
@@ -535,10 +535,10 @@ namespace UE::Learning::Random
 		const TLearningArrayView<1, const bool> Masked,
 		const float Scale)
 	{
-		UE_LEARNING_TRACE_CPUPROFILER_EVENT_SCOPE(Learning::Random::DistributionBernoulliMasked);
+		TRACE_CPUPROFILER_EVENT_SCOPE(Learning::Random::DistributionBernoulliMasked);
 
-		UE_LEARNING_CHECK(Output.Num() == Logits.Num());
-		UE_LEARNING_CHECK(Logits.Num() == Masked.Num());
+		check(Output.Num() == Logits.Num());
+		check(Logits.Num() == Masked.Num());
 
 		const int32 ElementNum = Output.Num();
 
@@ -577,7 +577,7 @@ namespace UE::Learning::Random
 		const int32 TotalFrameNum = FrameSet.GetTotalFrameNum();
 		const int32 FrameOffset = IntInRange(0x26810f69 ^ State, 0, TotalFrameNum - 1);
 		const bool bFound = FrameSet.FindFrameOffset(OutEntryIdx, OutFrameIdx, FrameOffset);
-		UE_LEARNING_CHECK(bFound);
+		check(bFound);
 	}
 
 	void FrameInFrameRangeSet(
@@ -590,7 +590,7 @@ namespace UE::Learning::Random
 		const int32 TotalFrameNum = FrameRangeSet.GetTotalFrameNum();
 		const int32 FrameOffset = IntInRange(0x76ad14bb ^ State, 0, TotalFrameNum - 1);
 		const bool bFound = FrameRangeSet.FindFrameOffset(OutEntryIdx, OutRangeIdx, OutRangeFrame, FrameOffset);
-		UE_LEARNING_CHECK(bFound);
+		check(bFound);
 	}
 
 	////////////
diff --git a/Engine/Plugins/Experimental/LearningAgents/Source/Learning/Public/LearningArray.h b/Engine/Plugins/Experimental/LearningAgents/Source/Learning/Public/LearningArray.h
--- a/Engine/Plugins/Experimental/LearningAgents/Source/Learning/Public/LearningArray.h
+++ b/Engine/Plugins/Experimental/LearningAgents/Source/Learning/Public/LearningArray.h
@@ -15,74 +15,6 @@
 //#define UE_LEARNING_ISPC 0
 #endif
 
-// If to enable profiling in Learning
-#ifndef UE_LEARNING_PROFILE
-#define UE_LEARNING_PROFILE (!UE_BUILD_SHIPPING)
-//#define UE_LEARNING_PROFILE 0
-#endif
-
-// Learning development mode enables various runtime sanity checks such as those described below. 
-// However, given that tools built with Learning are often designed to be performant in the editor, 
-// these checks can be disabled once the systems are stable even when not shipping.
-#ifndef UE_LEARNING_DEVELOPMENT
-#define UE_LEARNING_DEVELOPMENT (!UE_BUILD_SHIPPING)
-//#define UE_LEARNING_DEVELOPMENT 0
-#endif
-
-// There are various checks in Learning to see if arrays contain nan/inf/-inf
-// which can be disabled or enabled with this flag. These will be disabled in 
-// Shipping mode but in certain cases you may want to disable them in editor 
-// mode too to improve performance.
-#ifndef UE_LEARNING_ARRAY_CHECK_VALUES
-#define UE_LEARNING_ARRAY_CHECK_VALUES UE_LEARNING_DEVELOPMENT
-#endif
-
-// There are various checks in Learning to see if arrays are the right shapes
-// which can be disabled or enabled with this flag. These will be disabled in 
-// Shipping mode but in certain cases you may want to disable them in editor 
-// mode too to improve performance as some tight loops do contain shape checks.
-#ifndef UE_LEARNING_ARRAY_CHECK_SHAPE
-#define UE_LEARNING_ARRAY_CHECK_SHAPE UE_LEARNING_DEVELOPMENT
-#endif
-
-
-// Profiling Macros
-
-#if UE_LEARNING_PROFILE == 1
-#define UE_LEARNING_TRACE_CPUPROFILER_EVENT_SCOPE(Name) TRACE_CPUPROFILER_EVENT_SCOPE(Name)
-#else
-#define UE_LEARNING_TRACE_CPUPROFILER_EVENT_SCOPE(Name)
-#endif
-
-
-// Assert Macros
-
-#if UE_LEARNING_DEVELOPMENT == 1
-#define UE_LEARNING_CHECK(Expr) check(Expr)
-#define UE_LEARNING_CHECKF(Expr, Format, ...) checkf(Expr, Format, ##__VA_ARGS__)
-#else
-#define UE_LEARNING_CHECK(...)
-#define UE_LEARNING_CHECKF(...)
-#endif
-
-#define UE_LEARNING_NOT_IMPLEMENTED() checkf(false, TEXT("Not Implemented!"))
-
-#if UE_LEARNING_ARRAY_CHECK_SHAPE == 1
-#define UE_LEARNING_ARRAY_SHAPE_CHECK(Expr) check(Expr)
-#define UE_LEARNING_ARRAY_SHAPE_CHECKF(Expr, Format, ...) checkf(Expr, Format, ##__VA_ARGS__)
-#else
-#define UE_LEARNING_ARRAY_SHAPE_CHECK(...)
-#define UE_LEARNING_ARRAY_SHAPE_CHECKF(...)
-#endif
-
-#if UE_LEARNING_ARRAY_CHECK_VALUES == 1
-#define UE_LEARNING_ARRAY_VALUE_CHECK(Expr) check(Expr)
-#define UE_LEARNING_ARRAY_VALUE_CHECKF(Expr, Format, ...) checkf(Expr, Format, ##__VA_ARGS__)
-#else
-#define UE_LEARNING_ARRAY_VALUE_CHECK(...)
-#define UE_LEARNING_ARRAY_VALUE_CHECKF(...)
-#endif
-
 // Learning Multi-Dimensional Array Types
 
 template<uint32 DimNum>
@@ -97,8 +29,6 @@ using TLearningConstArrayView = TConstMultiArrayView<DimNum, ElementType>;
 template<uint32 DimNum, typename ElementType, typename Allocator = FDefaultAllocator64>
 using TLearningArray = TMultiArray<DimNum, ElementType, Allocator>;
 
-
-
 // Learning Array Functions
 
 namespace UE::Learning
@@ -200,15 +130,12 @@ namespace UE::Learning
 
 		inline FIndexSet Slice(const int32 InStart, const int32 InNum) const
 		{
-
-#if UE_LEARNING_DEVELOPMENT
 			if (IsSlice())
 			{
-				UE_LEARNING_CHECK(InStart >= 0);
-				UE_LEARNING_CHECK(InNum >= 0);
-				UE_LEARNING_CHECK(InStart + InNum <= SliceNum);
+				check(InStart >= 0);
+				check(InNum >= 0);
+				check(InStart + InNum <= SliceNum);
 			}
-#endif
 
 			return IsSlice() ? FIndexSet(SliceStart + InStart, InNum) : Indices.Slice(InStart, InNum);
 		}
@@ -257,9 +184,9 @@ namespace UE::Learning
 		}
 
 		inline bool IsSlice() const { return Indices.GetData() == nullptr; }
-		inline int32 GetSliceStart() const { UE_LEARNING_CHECK(IsSlice()); return SliceStart; }
-		inline int32 GetSliceNum() const { UE_LEARNING_CHECK(IsSlice()); return SliceNum; }
-		inline TLearningArrayView<1, const int32> GetIndices() const { UE_LEARNING_CHECK(!IsSlice()); return Indices; }
+		inline int32 GetSliceStart() const { check(IsSlice()); return SliceStart; }
+		inline int32 GetSliceNum() const { check(IsSlice()); return SliceNum; }
+		inline TLearningArrayView<1, const int32> GetIndices() const { check(!IsSlice()); return Indices; }
 
 	private:
 
@@ -285,7 +212,7 @@ namespace UE::Learning
 	*/
 	static inline void SerializeToBytes(int64& InOutOffset, TLearningArrayView<1, uint8> Bytes, const uint8 InValue)
 	{
-		UE_LEARNING_CHECK(InOutOffset + (int64)sizeof(uint8) <= Bytes.Num());
+		check(InOutOffset + (int64)sizeof(uint8) <= Bytes.Num());
 		FMemory::Memcpy(&Bytes[InOutOffset], (uint8*)&InValue, sizeof(uint8));
 		InOutOffset += sizeof(uint8);
 	}
@@ -295,7 +222,7 @@ namespace UE::Learning
 	*/
 	static inline void SerializeToBytes(int64& InOutOffset, TLearningArrayView<1, uint8> Bytes, const int32 InValue)
 	{
-		UE_LEARNING_CHECK(InOutOffset + (int64)sizeof(int32) <= Bytes.Num());
+		check(InOutOffset + (int64)sizeof(int32) <= Bytes.Num());
 		FMemory::Memcpy(&Bytes[InOutOffset], (uint8*)&InValue, sizeof(int32));
 		InOutOffset += sizeof(int32);
 	}
@@ -305,7 +232,7 @@ namespace UE::Learning
 	*/
 	static inline void SerializeToBytes(int64& InOutOffset, TLearningArrayView<1, uint8> Bytes, const int64 InValue)
 	{
-		UE_LEARNING_CHECK(InOutOffset + (int64)sizeof(int64) <= Bytes.Num());
+		check(InOutOffset + (int64)sizeof(int64) <= Bytes.Num());
 		FMemory::Memcpy(&Bytes[InOutOffset], (uint8*)&InValue, sizeof(int64));
 		InOutOffset += sizeof(int64);
 	}
@@ -315,7 +242,7 @@ namespace UE::Learning
 	*/
 	static inline void DeserializeFromBytes(int64& InOutOffset, TLearningArrayView<1, const uint8> Bytes, uint8& OutValue)
 	{
-		UE_LEARNING_CHECK(InOutOffset + (int64)sizeof(uint8) <= Bytes.Num());
+		check(InOutOffset + (int64)sizeof(uint8) <= Bytes.Num());
 		FMemory::Memcpy((uint8*)&OutValue, &Bytes[InOutOffset], sizeof(uint8));
 		InOutOffset += sizeof(uint8);
 	}
@@ -325,7 +252,7 @@ namespace UE::Learning
 	*/
 	static inline void DeserializeFromBytes(int64& InOutOffset, TLearningArrayView<1, const uint8> Bytes, int32& OutValue)
 	{
-		UE_LEARNING_CHECK(InOutOffset + (int64)sizeof(int32) <= Bytes.Num());
+		check(InOutOffset + (int64)sizeof(int32) <= Bytes.Num());
 		FMemory::Memcpy((uint8*)&OutValue, &Bytes[InOutOffset], sizeof(int32));
 		InOutOffset += sizeof(int32);
 	}
@@ -335,7 +262,7 @@ namespace UE::Learning
 	*/
 	static inline void DeserializeFromBytes(int64& InOutOffset, TLearningArrayView<1, const uint8> Bytes, int64& OutValue)
 	{
-		UE_LEARNING_CHECK(InOutOffset + (int64)sizeof(int64) <= Bytes.Num());
+		check(InOutOffset + (int64)sizeof(int64) <= Bytes.Num());
 		FMemory::Memcpy((uint8*)&OutValue, &Bytes[InOutOffset], sizeof(int64));
 		InOutOffset += sizeof(int64);
 	}
@@ -345,7 +272,7 @@ namespace UE::Learning
 	*/
 	static inline void DeserializeFromBytes(int64& InOutOffset, TLearningArrayView<1, const uint8> Bytes, uint64& OutValue)
 	{
-		UE_LEARNING_CHECK(InOutOffset + (int64)sizeof(uint64) <= Bytes.Num());
+		check(InOutOffset + (int64)sizeof(uint64) <= Bytes.Num());
 		FMemory::Memcpy((uint8*)&OutValue, &Bytes[InOutOffset], sizeof(uint64));
 		InOutOffset += sizeof(uint64);
 	}
@@ -364,13 +291,11 @@ namespace UE::Learning
 		template<uint8 InDimNum>
 		inline void CheckShapesEqual(const TLearningArrayShape<InDimNum>& Lhs, const TLearningArrayShape<InDimNum>& Rhs)
 		{
-#if UE_LEARNING_ARRAY_CHECK_SHAPE == 1
 			for (uint8 Idx = 0; Idx < InDimNum; Idx++)
 			{
-				UE_LEARNING_ARRAY_SHAPE_CHECKF(Lhs[Idx] == Rhs[Idx],
+				checkf(Lhs[Idx] == Rhs[Idx],
 					TEXT("Array Shapes don't match on dimension %i of %i (lhs: %lld, rhs: %lld)"), Idx + 1, InDimNum, Lhs[Idx], Rhs[Idx]);
 			}
-#endif
 		}
 
 		/**
@@ -379,14 +304,12 @@ namespace UE::Learning
 		template<uint8 InDimNum, typename InElementType>
 		inline void Check(const TLearningArrayView<InDimNum, InElementType> View)
 		{
-#if UE_LEARNING_ARRAY_CHECK_VALUES == 1
 			for (int64 Idx = 0; Idx < View.Num(); Idx++)
 			{
-				UE_LEARNING_ARRAY_VALUE_CHECKF(
+				checkf(
 					FMath::IsFinite(View.GetData()[Idx]),
 					TEXT("Invalid value %f found at flat array index %lld"), View.GetData()[Idx], Idx);
 			}
-#endif
 		}
 
 		/**
@@ -395,14 +318,12 @@ namespace UE::Learning
 		template<uint8 InDimNum, typename InElementType, typename Allocator>
 		inline void Check(const TLearningArray<InDimNum, InElementType, Allocator>& View)
 		{
-#if UE_LEARNING_ARRAY_CHECK_VALUES == 1
 			for (int64 Idx = 0; Idx < View.Num(); Idx++)
 			{
-				UE_LEARNING_ARRAY_VALUE_CHECKF(
+				checkf(
 					FMath::IsFinite(View.GetData()[Idx]),
 					TEXT("Invalid value %f found at flat array index %lld"), View.GetData()[Idx], Idx);
 			}
-#endif
 		}
 
 
@@ -412,14 +333,12 @@ namespace UE::Learning
 		template<uint8 InDimNum, typename InElementType>
 		inline void Check(const TLearningArrayView<InDimNum, const InElementType> View)
 		{
-#if UE_LEARNING_ARRAY_CHECK_VALUES == 1
 			for (int64 Idx = 0; Idx < View.Num(); Idx++)
 			{
-				UE_LEARNING_ARRAY_VALUE_CHECKF(
+				checkf(
 					FMath::IsFinite(View.GetData()[Idx]),
 					TEXT("Invalid value %f found at flat array index %lld"), View.GetData()[Idx], Idx);
 			}
-#endif
 		}
 
 		/**
@@ -428,7 +347,6 @@ namespace UE::Learning
 		template<uint8 InDimNum, typename InElementType>
 		inline void Check(const TLearningArrayView<InDimNum, const InElementType> View, const FIndexSet Indices)
 		{
-#if UE_LEARNING_ARRAY_CHECK_VALUES == 1
 			if constexpr (InDimNum > 1)
 			{
 				for (const int32 Idx : Indices)
@@ -440,12 +358,11 @@ namespace UE::Learning
 			{
 				for (const int32 Idx : Indices)
 				{
-					UE_LEARNING_ARRAY_VALUE_CHECKF(
+					checkf(
 						FMath::IsFinite(View.GetData()[Idx]) && View.GetData()[Idx],
 						TEXT("Invalid value %f found at flat array index %lld"), View.GetData()[Idx], Idx);
 				}
 			}
-#endif
 		}
 
 		/**
@@ -454,7 +371,6 @@ namespace UE::Learning
 		template<uint8 InDimNum, typename InElementType>
 		inline void Check(const TLearningArrayView<InDimNum, InElementType> View, const FIndexSet Indices)
 		{
-#if UE_LEARNING_ARRAY_CHECK_VALUES == 1
 			if constexpr (InDimNum > 1)
 			{
 				for (const int32 Idx : Indices)
@@ -466,12 +382,11 @@ namespace UE::Learning
 			{
 				for (const int32 Idx : Indices)
 				{
-					UE_LEARNING_ARRAY_VALUE_CHECKF(
+					checkf(
 						FMath::IsFinite(View.GetData()[Idx]) && View.GetData()[Idx],
 						TEXT("Invalid value %f found at flat array index %lld"), View.GetData()[Idx], Idx);
 				}
 			}
-#endif
 		}
 
 		/**
@@ -480,7 +395,6 @@ namespace UE::Learning
 		template<uint8 InDimNum, typename InElementType>
 		inline void Check(const TLearningArray<InDimNum, InElementType>& View, const FIndexSet Indices)
 		{
-#if UE_LEARNING_ARRAY_CHECK_VALUES == 1
 			if constexpr (InDimNum > 1)
 			{
 				for (const int32 Idx : Indices)
@@ -492,12 +406,11 @@ namespace UE::Learning
 			{
 				for (const int32 Idx : Indices)
 				{
-					UE_LEARNING_ARRAY_VALUE_CHECKF(
+					checkf(
 						FMath::IsFinite(View.GetData()[Idx]),
 						TEXT("Invalid value %f found at flat array index %lld"), View.GetData()[Idx], Idx);
 				}
 			}
-#endif
 		}
 
 		/**
@@ -1036,7 +949,7 @@ namespace UE::Learning
 				Learning::SerializeToBytes(InOutOffset, Bytes, (int32)Shape[ShapeIdx]);
 			}
 
-			UE_LEARNING_CHECK(InOutOffset + (int64)sizeof(InElementType) * InArray.Num() <= Bytes.Num());
+			check(InOutOffset + (int64)sizeof(InElementType) * InArray.Num() <= Bytes.Num());
 			FMemory::Memcpy(&Bytes[InOutOffset], (uint8*)InArray.GetData(), sizeof(InElementType) * InArray.Num());
 			InOutOffset += sizeof(InElementType) * InArray.Num();
 		}
@@ -1053,7 +966,7 @@ namespace UE::Learning
 				Learning::SerializeToBytes(InOutOffset, Bytes, (int32)InArray.Shape()[ShapeIdx]);
 			}
 
-			UE_LEARNING_CHECK(InOutOffset + (int64)sizeof(InElementType) * InArray.Num() <= Bytes.Num());
+			check(InOutOffset + (int64)sizeof(InElementType) * InArray.Num() <= Bytes.Num());
 			FMemory::Memcpy(&Bytes[InOutOffset], (uint8*)InArray.GetData(), sizeof(InElementType) * InArray.Num());
 			InOutOffset += sizeof(InElementType) * InArray.Num();
 		}
@@ -1066,7 +979,7 @@ namespace UE::Learning
 		{
 			int32 DimNum = INDEX_NONE;
 			Learning::DeserializeFromBytes(InOutOffset, Bytes, DimNum);
-			UE_LEARNING_CHECK(DimNum == InDimNum);
+			check(DimNum == InDimNum);
 
 			TLearningArrayShape<InDimNum> Shape;
 			for (uint8 ShapeIdx = 0; ShapeIdx < InDimNum; ShapeIdx++)
@@ -1078,7 +991,7 @@ namespace UE::Learning
 
 			OutArray.SetNumUninitialized(Shape.Total());
 
-			UE_LEARNING_CHECK(InOutOffset + (int64)sizeof(InElementType) * OutArray.Num() <= Bytes.Num());
+			check(InOutOffset + (int64)sizeof(InElementType) * OutArray.Num() <= Bytes.Num());
 			FMemory::Memcpy((uint8*)OutArray.GetData(), &Bytes[InOutOffset], sizeof(InElementType) * OutArray.Num());
 			InOutOffset += sizeof(InElementType) * OutArray.Num();
 		}
@@ -1091,7 +1004,7 @@ namespace UE::Learning
 		{
 			int32 DimNum = INDEX_NONE;
 			Learning::DeserializeFromBytes(InOutOffset, Bytes, DimNum);
-			UE_LEARNING_CHECK(DimNum == InDimNum);
+			check(DimNum == InDimNum);
 
 			TLearningArrayShape<InDimNum> Shape;
 			for (uint8 ShapeIdx = 0; ShapeIdx < InDimNum; ShapeIdx++)
@@ -1103,7 +1016,7 @@ namespace UE::Learning
 
 			OutArray.SetNumUninitialized(Shape);
 
-			UE_LEARNING_CHECK(InOutOffset + (int64)sizeof(InElementType) * OutArray.Num() <= Bytes.Num());
+			check(InOutOffset + (int64)sizeof(InElementType) * OutArray.Num() <= Bytes.Num());
 			FMemory::Memcpy((uint8*)OutArray.GetData(), &Bytes[InOutOffset],  sizeof(InElementType) * OutArray.Num());
 			InOutOffset += sizeof(InElementType) * OutArray.Num();
 		}
diff --git a/Engine/Plugins/Experimental/LearningAgents/Source/LearningAgents/Private/LearningAgentsActions.cpp b/Engine/Plugins/Experimental/LearningAgents/Source/LearningAgents/Private/LearningAgentsActions.cpp
--- a/Engine/Plugins/Experimental/LearningAgents/Source/LearningAgents/Private/LearningAgentsActions.cpp
+++ b/Engine/Plugins/Experimental/LearningAgents/Source/LearningAgents/Private/LearningAgentsActions.cpp
@@ -70,7 +70,7 @@ namespace UE::Learning::Agents::Action::Private
 		case Learning::Action::EType::Array: return TEXT("Array");
 		case Learning::Action::EType::Encoding: return TEXT("Encoding");
 		default:
-			UE_LEARNING_NOT_IMPLEMENTED();
+			checkNoEntry();
 			return TEXT("Unimplemented");
 		}
 	}
@@ -193,8 +193,8 @@ namespace UE::Learning::Agents::Action::Private
 		{
 			const Learning::Action::FSchemaAndParameters SchemaParameters = Schema.GetAnd(SchemaElement);
 			const Learning::Action::FObjectAndParameters ObjectParameters = Object.GetAnd(ObjectElement);
-			UE_LEARNING_CHECK(SchemaParameters.Elements.Num() == SchemaParameters.ElementNames.Num());
-			UE_LEARNING_CHECK(ObjectParameters.Elements.Num() == ObjectParameters.ElementNames.Num());
+			check(SchemaParameters.Elements.Num() == SchemaParameters.ElementNames.Num());
+			check(ObjectParameters.Elements.Num() == ObjectParameters.ElementNames.Num());
 
 			if (SchemaParameters.Elements.Num() != ObjectParameters.Elements.Num())
 			{
@@ -234,7 +234,7 @@ namespace UE::Learning::Agents::Action::Private
 		{
 			const Learning::Action::FSchemaOrExclusiveParameters SchemaParameters = Schema.GetOrExclusive(SchemaElement);
 			const Learning::Action::FObjectOrExclusiveParameters ObjectParameters = Object.GetOrExclusive(ObjectElement);
-			UE_LEARNING_CHECK(SchemaParameters.Elements.Num() == SchemaParameters.ElementNames.Num());
+			check(SchemaParameters.Elements.Num() == SchemaParameters.ElementNames.Num());
 
 			const int32 SchemaSubElementIdx = SchemaParameters.ElementNames.Find(ObjectParameters.ElementName);
 
@@ -335,7 +335,7 @@ namespace UE::Learning::Agents::Action::Private
 
 		default:
 		{
-			UE_LEARNING_NOT_IMPLEMENTED();
+			checkNoEntry();
 			return true;
 		}
 		}
@@ -490,8 +490,8 @@ namespace UE::Learning::Agents::Action::Private
 		{
 			const Learning::Action::FSchemaAndParameters SchemaParameters = Schema.GetAnd(SchemaElement);
 			const Learning::Action::FModifierAndParameters ModifierParameters = Modifier.GetAnd(ModifierElement);
-			UE_LEARNING_CHECK(SchemaParameters.Elements.Num() == SchemaParameters.ElementNames.Num());
-			UE_LEARNING_CHECK(ModifierParameters.Elements.Num() == ModifierParameters.ElementNames.Num());
+			check(SchemaParameters.Elements.Num() == SchemaParameters.ElementNames.Num());
+			check(ModifierParameters.Elements.Num() == ModifierParameters.ElementNames.Num());
 
 			if (SchemaParameters.Elements.Num() < ModifierParameters.Elements.Num())
 			{
@@ -531,7 +531,7 @@ namespace UE::Learning::Agents::Action::Private
 		{
 			const Learning::Action::FSchemaOrExclusiveParameters SchemaParameters = Schema.GetOrExclusive(SchemaElement);
 			const Learning::Action::FModifierOrExclusiveParameters ModifierParameters = Modifier.GetOrExclusive(ModifierElement);
-			UE_LEARNING_CHECK(SchemaParameters.Elements.Num() == SchemaParameters.ElementNames.Num());
+			check(SchemaParameters.Elements.Num() == SchemaParameters.ElementNames.Num());
 
 			// For exclusive action we must have one less than the schema size or it means all are masked.
 			if (ModifierParameters.MaskedElements.Num() > SchemaParameters.Elements.Num() - 1)
@@ -687,7 +687,7 @@ namespace UE::Learning::Agents::Action::Private
 
 		default:
 		{
-			UE_LEARNING_NOT_IMPLEMENTED();
+			checkNoEntry();
 			return false;
 		}
 		}
@@ -869,7 +869,7 @@ namespace UE::Learning::Agents::Action::Private
 		case ELearningAgentsActivationFunction::ReLU: return Learning::Action::EEncodingActivationFunction::ReLU;
 		case ELearningAgentsActivationFunction::ELU: return Learning::Action::EEncodingActivationFunction::ELU;
 		case ELearningAgentsActivationFunction::TanH: return Learning::Action::EEncodingActivationFunction::TanH;
-		default: UE_LEARNING_NOT_IMPLEMENTED(); return Learning::Action::EEncodingActivationFunction::ReLU;
+		default: checkNoEntry(); return Learning::Action::EEncodingActivationFunction::ReLU;
 		}
 	}
 }
diff --git a/Engine/Plugins/Experimental/LearningAgents/Source/LearningAgents/Private/LearningAgentsController.cpp b/Engine/Plugins/Experimental/LearningAgents/Source/LearningAgents/Private/LearningAgentsController.cpp
--- a/Engine/Plugins/Experimental/LearningAgents/Source/LearningAgents/Private/LearningAgentsController.cpp
+++ b/Engine/Plugins/Experimental/LearningAgents/Source/LearningAgents/Private/LearningAgentsController.cpp
@@ -109,7 +109,7 @@ void ULearningAgentsController::SetupController(ULearningAgentsManager*& InManag
 
 void ULearningAgentsController::EvaluateController()
 {
-	UE_LEARNING_TRACE_CPUPROFILER_EVENT_SCOPE(ULearningAgentsController::EvaluateController);
+	TRACE_CPUPROFILER_EVENT_SCOPE(ULearningAgentsController::EvaluateController);
 
 	if (!IsSetup())
 	{
@@ -180,7 +180,7 @@ ULearningAgentsInteractor* ULearningAgentsController::GetInteractor(const TSubcl
 
 void ULearningAgentsController::RunController()
 {
-	UE_LEARNING_TRACE_CPUPROFILER_EVENT_SCOPE(ULearningAgentsController::RunController);
+	TRACE_CPUPROFILER_EVENT_SCOPE(ULearningAgentsController::RunController);
 
 	if (!IsSetup())
 	{
diff --git a/Engine/Plugins/Experimental/LearningAgents/Source/LearningAgents/Private/LearningAgentsCritic.cpp b/Engine/Plugins/Experimental/LearningAgents/Source/LearningAgents/Private/LearningAgentsCritic.cpp
--- a/Engine/Plugins/Experimental/LearningAgents/Source/LearningAgents/Private/LearningAgentsCritic.cpp
+++ b/Engine/Plugins/Experimental/LearningAgents/Source/LearningAgents/Private/LearningAgentsCritic.cpp
@@ -186,8 +186,8 @@ void ULearningAgentsCritic::SetupCritic(
 				CriticSettings.HiddenLayerNum + 2, // Add 2 to account for input and output layers
 				UE::Learning::Agents::Critic::Private::GetBuilderActivationFunction(CriticSettings.ActivationFunction)));
 
-		UE_LEARNING_CHECK(CriticInputSize == ObservationEncodedVectorSize + MemoryStateSize);
-		UE_LEARNING_CHECK(CriticOutputSize == 1);
+		check(CriticInputSize == ObservationEncodedVectorSize + MemoryStateSize);
+		check(CriticOutputSize == 1);
 
 		CriticNetwork->NeuralNetworkData->Init(CriticInputSize, CriticOutputSize, CriticCompatibilityHash, FileData);
 		CriticNetwork->ForceMarkDirty();
@@ -266,7 +266,7 @@ ULearningAgentsNeuralNetwork* ULearningAgentsCritic::GetCriticNetworkAsset()
 
 void ULearningAgentsCritic::EvaluateCritic()
 {
-	UE_LEARNING_TRACE_CPUPROFILER_EVENT_SCOPE(ULearningAgentsCritic::EvaluateCritic);
+	TRACE_CPUPROFILER_EVENT_SCOPE(ULearningAgentsCritic::EvaluateCritic);
 
 	if (!IsSetup())
 	{
diff --git a/Engine/Plugins/Experimental/LearningAgents/Source/LearningAgents/Private/LearningAgentsInteractor.cpp b/Engine/Plugins/Experimental/LearningAgents/Source/LearningAgents/Private/LearningAgentsInteractor.cpp
--- a/Engine/Plugins/Experimental/LearningAgents/Source/LearningAgents/Private/LearningAgentsInteractor.cpp
+++ b/Engine/Plugins/Experimental/LearningAgents/Source/LearningAgents/Private/LearningAgentsInteractor.cpp
@@ -284,7 +284,7 @@ void ULearningAgentsInteractor::MakeAgentActionModifiers_Implementation(
 
 void ULearningAgentsInteractor::GatherObservations(const UE::Learning::FIndexSet AgentSet, bool bIncrementIteration)
 {
-	UE_LEARNING_TRACE_CPUPROFILER_EVENT_SCOPE(ULearningAgentsInteractor::GatherObservations);
+	TRACE_CPUPROFILER_EVENT_SCOPE(ULearningAgentsInteractor::GatherObservations);
 
 	if (!IsSetup())
 	{
@@ -337,7 +337,7 @@ void ULearningAgentsInteractor::GatherObservations(const UE::Learning::FIndexSet
 
 void ULearningAgentsInteractor::MakeActionModifiers(const UE::Learning::FIndexSet AgentSet, bool bIncrementIteration)
 {
-	UE_LEARNING_TRACE_CPUPROFILER_EVENT_SCOPE(ULearningAgentsInteractor::MakeActionModifiers);
+	TRACE_CPUPROFILER_EVENT_SCOPE(ULearningAgentsInteractor::MakeActionModifiers);
 
 	if (!IsSetup())
 	{
@@ -413,7 +413,7 @@ void ULearningAgentsInteractor::MakeActionModifiers(const UE::Learning::FIndexSe
 
 void ULearningAgentsInteractor::PerformActions(const UE::Learning::FIndexSet AgentSet)
 {
-	UE_LEARNING_TRACE_CPUPROFILER_EVENT_SCOPE(ULearningAgentsInteractor::PerformActions);
+	TRACE_CPUPROFILER_EVENT_SCOPE(ULearningAgentsInteractor::PerformActions);
 
 	if (!IsSetup())
 	{
diff --git a/Engine/Plugins/Experimental/LearningAgents/Source/LearningAgents/Private/LearningAgentsManager.cpp b/Engine/Plugins/Experimental/LearningAgents/Source/LearningAgents/Private/LearningAgentsManager.cpp
--- a/Engine/Plugins/Experimental/LearningAgents/Source/LearningAgents/Private/LearningAgentsManager.cpp
+++ b/Engine/Plugins/Experimental/LearningAgents/Source/LearningAgents/Private/LearningAgentsManager.cpp
@@ -176,7 +176,7 @@ void ULearningAgentsManager::RemoveAgent(const int32 AgentId)
 		return;
 	}
 
-	UE_LEARNING_CHECKF(RemovedCount == 1, TEXT("Somehow agent id was included multiple times in set."));
+	checkf(RemovedCount == 1, TEXT("Somehow agent id was included multiple times in set."));
 
 	// Remove Agent
 
@@ -217,7 +217,7 @@ void ULearningAgentsManager::RemoveAgents(const TArray<int32>& AgentIds)
 			continue;
 		}
 
-		UE_LEARNING_CHECKF(RemovedCount == 1, TEXT("Somehow agent id was included multiple times in set."));
+		checkf(RemovedCount == 1, TEXT("Somehow agent id was included multiple times in set."));
 
 		// Remove Agent
 
@@ -420,13 +420,13 @@ bool ULearningAgentsManager::HasAgent(const int32 AgentId) const
 
 const UObject* ULearningAgentsManager::GetAgent(const int32 AgentId) const
 {
-	UE_LEARNING_CHECK(HasAgent(AgentId));
+	check(HasAgent(AgentId));
 	return Agents[AgentId];
 }
 
 UObject* ULearningAgentsManager::GetAgent(const int32 AgentId)
 {
-	UE_LEARNING_CHECK(HasAgent(AgentId));
+	check(HasAgent(AgentId));
 	return Agents[AgentId];
 }
 
diff --git a/Engine/Plugins/Experimental/LearningAgents/Source/LearningAgents/Private/LearningAgentsManagerListener.cpp b/Engine/Plugins/Experimental/LearningAgents/Source/LearningAgents/Private/LearningAgentsManagerListener.cpp
--- a/Engine/Plugins/Experimental/LearningAgents/Source/LearningAgents/Private/LearningAgentsManagerListener.cpp
+++ b/Engine/Plugins/Experimental/LearningAgents/Source/LearningAgents/Private/LearningAgentsManagerListener.cpp
@@ -71,19 +71,19 @@ ULearningAgentsManager* ULearningAgentsManagerListener::GetAgentManager() const
 
 const UObject* ULearningAgentsManagerListener::GetAgent(const int32 AgentId) const
 {
-	UE_LEARNING_CHECKF(Manager, TEXT("Manager variable not set."));
+	checkf(Manager, TEXT("Manager variable not set."));
 	return Manager->GetAgent(AgentId);
 }
 
 UObject* ULearningAgentsManagerListener::GetAgent(const int32 AgentId)
 {
-	UE_LEARNING_CHECKF(Manager, TEXT("Manager variable not set."));
+	checkf(Manager, TEXT("Manager variable not set."));
 	return Manager->GetAgent(AgentId);
 }
 
 bool ULearningAgentsManagerListener::HasAgent(const int32 AgentId) const
 {
-	UE_LEARNING_CHECKF(Manager, TEXT("Manager variable not set."));
+	checkf(Manager, TEXT("Manager variable not set."));
 	return Manager->HasAgent(AgentId);
 }
 
diff --git a/Engine/Plugins/Experimental/LearningAgents/Source/LearningAgents/Private/LearningAgentsObservations.cpp b/Engine/Plugins/Experimental/LearningAgents/Source/LearningAgents/Private/LearningAgentsObservations.cpp
--- a/Engine/Plugins/Experimental/LearningAgents/Source/LearningAgents/Private/LearningAgentsObservations.cpp
+++ b/Engine/Plugins/Experimental/LearningAgents/Source/LearningAgents/Private/LearningAgentsObservations.cpp
@@ -63,7 +63,7 @@ namespace UE::Learning::Agents::Observation::Private
 		case  Learning::Observation::EType::Set: return TEXT("Set");
 		case  Learning::Observation::EType::Encoding: return TEXT("Encoding");
 		default:
-			UE_LEARNING_NOT_IMPLEMENTED();
+			checkNoEntry();
 			return TEXT("Unimplemented");
 		}
 	}
@@ -140,8 +140,8 @@ namespace UE::Learning::Agents::Observation::Private
 		{
 			const Learning::Observation::FSchemaAndParameters SchemaParameters = Schema.GetAnd(SchemaElement);
 			const Learning::Observation::FObjectAndParameters ObjectParameters = Object.GetAnd(ObjectElement);
-			UE_LEARNING_CHECK(SchemaParameters.Elements.Num() == SchemaParameters.ElementNames.Num());
-			UE_LEARNING_CHECK(ObjectParameters.Elements.Num() == ObjectParameters.ElementNames.Num());
+			check(SchemaParameters.Elements.Num() == SchemaParameters.ElementNames.Num());
+			check(ObjectParameters.Elements.Num() == ObjectParameters.ElementNames.Num());
 
 			if (SchemaParameters.Elements.Num() != ObjectParameters.Elements.Num())
 			{
@@ -181,7 +181,7 @@ namespace UE::Learning::Agents::Observation::Private
 		{
 			const Learning::Observation::FSchemaOrExclusiveParameters SchemaParameters = Schema.GetOrExclusive(SchemaElement);
 			const Learning::Observation::FObjectOrExclusiveParameters ObjectParameters = Object.GetOrExclusive(ObjectElement);
-			UE_LEARNING_CHECK(SchemaParameters.Elements.Num() == SchemaParameters.ElementNames.Num());
+			check(SchemaParameters.Elements.Num() == SchemaParameters.ElementNames.Num());
 
 			const int32 SchemaSubElementIdx = SchemaParameters.ElementNames.Find(ObjectParameters.ElementName);
 
@@ -311,7 +311,7 @@ namespace UE::Learning::Agents::Observation::Private
 
 		default:
 		{
-			UE_LEARNING_NOT_IMPLEMENTED();
+			checkNoEntry();
 			return true;
 		}
 		}
@@ -445,7 +445,7 @@ namespace UE::Learning::Agents::Observation::Private
 		case ELearningAgentsActivationFunction::ReLU: return Learning::Observation::EEncodingActivationFunction::ReLU;
 		case ELearningAgentsActivationFunction::ELU: return Learning::Observation::EEncodingActivationFunction::ELU;
 		case ELearningAgentsActivationFunction::TanH: return Learning::Observation::EEncodingActivationFunction::TanH;
-		default: UE_LEARNING_NOT_IMPLEMENTED(); return Learning::Observation::EEncodingActivationFunction::ReLU;
+		default: checkNoEntry(); return Learning::Observation::EEncodingActivationFunction::ReLU;
 		}
 	}
 }
diff --git a/Engine/Plugins/Experimental/LearningAgents/Source/LearningAgents/Private/LearningAgentsPolicy.cpp b/Engine/Plugins/Experimental/LearningAgents/Source/LearningAgents/Private/LearningAgentsPolicy.cpp
--- a/Engine/Plugins/Experimental/LearningAgents/Source/LearningAgents/Private/LearningAgentsPolicy.cpp
+++ b/Engine/Plugins/Experimental/LearningAgents/Source/LearningAgents/Private/LearningAgentsPolicy.cpp
@@ -206,8 +206,8 @@ void ULearningAgentsPolicy::SetupPolicy(
 			UE::Learning::Observation::FNetworkSettings(),
 			UE::Learning::Random::Int(Seed ^ 0x658868dd));
 
-		UE_LEARNING_CHECK(EncoderInputSize == ObservationVectorSize);
-		UE_LEARNING_CHECK(EncoderOutputSize == ObservationEncodedVectorSize);
+		check(EncoderInputSize == ObservationVectorSize);
+		check(EncoderOutputSize == ObservationEncodedVectorSize);
 
 		EncoderNetwork->NeuralNetworkData->Init(EncoderInputSize, EncoderOutputSize, ObservationCompatibilityHash, FileData);
 		EncoderNetwork->ForceMarkDirty();
@@ -295,8 +295,8 @@ void ULearningAgentsPolicy::SetupPolicy(
 					})
 			));
 
-		UE_LEARNING_CHECK(PolicyInputSize == ObservationEncodedVectorSize + MemoryStateSize);
-		UE_LEARNING_CHECK(PolicyOutputSize == ActionEncodedVectorSize + MemoryStateSize);
+		check(PolicyInputSize == ObservationEncodedVectorSize + MemoryStateSize);
+		check(PolicyOutputSize == ActionEncodedVectorSize + MemoryStateSize);
 
 		PolicyNetwork->NeuralNetworkData->Init(PolicyInputSize, PolicyOutputSize, PolicyCompatibilityHash, FileData);
 		PolicyNetwork->ForceMarkDirty();
@@ -353,8 +353,8 @@ void ULearningAgentsPolicy::SetupPolicy(
 			UE::Learning::Action::FNetworkSettings(),
 			UE::Learning::Random::Int(Seed ^ 0xfa88bb7f));
 
-		UE_LEARNING_CHECK(DecoderInputSize == ActionEncodedVectorSize);
-		UE_LEARNING_CHECK(DecoderOutputSize == ActionDistributionVectorSize);
+		check(DecoderInputSize == ActionEncodedVectorSize);
+		check(DecoderOutputSize == ActionDistributionVectorSize);
 
 		DecoderNetwork->NeuralNetworkData->Init(DecoderInputSize, DecoderOutputSize, ActionCompatibilityHash, FileData);
 		DecoderNetwork->ForceMarkDirty();
@@ -539,7 +539,7 @@ ULearningAgentsNeuralNetwork* ULearningAgentsPolicy::GetDecoderNetworkAsset()
 
 void ULearningAgentsPolicy::EncodeObservations()
 {
-	UE_LEARNING_TRACE_CPUPROFILER_EVENT_SCOPE(ULearningAgentsPolicy::EncodeObservations);
+	TRACE_CPUPROFILER_EVENT_SCOPE(ULearningAgentsPolicy::EncodeObservations);
 
 	if (!IsSetup())
 	{
@@ -595,7 +595,7 @@ void ULearningAgentsPolicy::EncodeObservations()
 
 void ULearningAgentsPolicy::EvaluatePolicy()
 {
-	UE_LEARNING_TRACE_CPUPROFILER_EVENT_SCOPE(ULearningAgentsPolicy::EvaluatePolicy);
+	TRACE_CPUPROFILER_EVENT_SCOPE(ULearningAgentsPolicy::EvaluatePolicy);
 
 	if (!IsSetup())
 	{
@@ -664,7 +664,7 @@ void ULearningAgentsPolicy::EvaluatePolicy()
 
 void ULearningAgentsPolicy::DecodeAndSampleActions(const float ActionNoiseScale)
 {
-	UE_LEARNING_TRACE_CPUPROFILER_EVENT_SCOPE(ULearningAgentsPolicy::DecodeAndSampleActions);
+	TRACE_CPUPROFILER_EVENT_SCOPE(ULearningAgentsPolicy::DecodeAndSampleActions);
 
 	if (!IsSetup())
 	{
@@ -741,7 +741,7 @@ void ULearningAgentsPolicy::DecodeAndSampleActions(const float ActionNoiseScale)
 
 void ULearningAgentsPolicy::RunInference(const float ActionNoiseScale)
 {
-	UE_LEARNING_TRACE_CPUPROFILER_EVENT_SCOPE(ULearningAgentsPolicy::RunInference);
+	TRACE_CPUPROFILER_EVENT_SCOPE(ULearningAgentsPolicy::RunInference);
 
 	if (!IsSetup())
 	{
diff --git a/Engine/Plugins/Experimental/LearningAgents/Source/LearningAgentsTraining/Private/LearningAgentsImitationTrainer.cpp b/Engine/Plugins/Experimental/LearningAgents/Source/LearningAgentsTraining/Private/LearningAgentsImitationTrainer.cpp
--- a/Engine/Plugins/Experimental/LearningAgents/Source/LearningAgentsTraining/Private/LearningAgentsImitationTrainer.cpp
+++ b/Engine/Plugins/Experimental/LearningAgents/Source/LearningAgentsTraining/Private/LearningAgentsImitationTrainer.cpp
@@ -249,8 +249,8 @@ void ULearningAgentsImitationTrainer::BeginTraining(
 		StepIdx += Record.StepNum;
 	}
 
-	UE_LEARNING_CHECK(EpisodeIdx == TotalEpisodeNum);
-	UE_LEARNING_CHECK(StepIdx == TotalStepNum);
+	check(EpisodeIdx == TotalEpisodeNum);
+	check(StepIdx == TotalStepNum);
 
 	const int32 ObservationSchemaId = 0;
 	const int32 ActionSchemaId = 0;
@@ -486,7 +486,7 @@ void ULearningAgentsImitationTrainer::EndTraining()
 
 void ULearningAgentsImitationTrainer::IterateTraining()
 {
-	UE_LEARNING_TRACE_CPUPROFILER_EVENT_SCOPE(ULearningAgentsImitationTrainer::IterateTraining);
+	TRACE_CPUPROFILER_EVENT_SCOPE(ULearningAgentsImitationTrainer::IterateTraining);
 
 	if (!IsSetup())
 	{
diff --git a/Engine/Plugins/Experimental/LearningAgents/Source/LearningAgentsTraining/Private/LearningAgentsPPOTrainer.cpp b/Engine/Plugins/Experimental/LearningAgents/Source/LearningAgentsTraining/Private/LearningAgentsPPOTrainer.cpp
--- a/Engine/Plugins/Experimental/LearningAgents/Source/LearningAgentsTraining/Private/LearningAgentsPPOTrainer.cpp
+++ b/Engine/Plugins/Experimental/LearningAgents/Source/LearningAgentsTraining/Private/LearningAgentsPPOTrainer.cpp
@@ -499,7 +499,7 @@ void ULearningAgentsPPOTrainer::EndTraining()
 
 void ULearningAgentsPPOTrainer::ProcessExperience(const bool bResetAgentsOnUpdate)
 {
-	UE_LEARNING_TRACE_CPUPROFILER_EVENT_SCOPE(ULearningAgentsPPOTrainer::ProcessExperience);
+	TRACE_CPUPROFILER_EVENT_SCOPE(ULearningAgentsPPOTrainer::ProcessExperience);
 
 	if (!IsSetup())
 	{
diff --git a/Engine/Plugins/Experimental/LearningAgents/Source/LearningAgentsTraining/Private/LearningAgentsRecorder.cpp b/Engine/Plugins/Experimental/LearningAgents/Source/LearningAgentsTraining/Private/LearningAgentsRecorder.cpp
--- a/Engine/Plugins/Experimental/LearningAgents/Source/LearningAgentsTraining/Private/LearningAgentsRecorder.cpp
+++ b/Engine/Plugins/Experimental/LearningAgents/Source/LearningAgentsTraining/Private/LearningAgentsRecorder.cpp
@@ -72,7 +72,7 @@ void ULearningAgentsRecorder::FAgentRecordBuffer::Empty()
 
 void ULearningAgentsRecorder::FAgentRecordBuffer::CopyToRecord(FLearningAgentsRecord& Record) const
 {
-	UE_LEARNING_CHECK(StepNum > 0);
+	check(StepNum > 0);
 
 	Record.StepNum = StepNum;
 	Record.ObservationDimNum = GetObservation(0).Num();
diff --git a/Engine/Plugins/Experimental/LearningAgents/Source/LearningAgentsTraining/Private/LearningAgentsRecording.cpp b/Engine/Plugins/Experimental/LearningAgents/Source/LearningAgentsTraining/Private/LearningAgentsRecording.cpp
--- a/Engine/Plugins/Experimental/LearningAgents/Source/LearningAgentsTraining/Private/LearningAgentsRecording.cpp
+++ b/Engine/Plugins/Experimental/LearningAgents/Source/LearningAgentsTraining/Private/LearningAgentsRecording.cpp
@@ -71,7 +71,7 @@ void ULearningAgentsRecording::LoadRecordingFromFile(const FFilePath& File)
 			UE::Learning::Array::DeserializeFromBytes32<2, float>(Offset, RecordingData, Records[RecordIdx].ActionData);
 		}
 
-		UE_LEARNING_CHECK(Offset == RecordingData.Num());
+		check(Offset == RecordingData.Num());
 
 		ForceMarkDirty();
 	}
@@ -120,7 +120,7 @@ void ULearningAgentsRecording::SaveRecordingToFile(const FFilePath& File) const
 		UE::Learning::Array::SerializeToBytes32<2, float>(Offset, RecordingData, { Records[RecordIdx].StepNum, Records[RecordIdx].ActionDimNum }, Records[RecordIdx].ActionData);
 	}
 
-	UE_LEARNING_CHECK(Offset == RecordingData.Num());
+	check(Offset == RecordingData.Num());
 
 	if (!FFileHelper::SaveArrayToFile(RecordingData, *File.FilePath))
 	{
diff --git a/Engine/Plugins/Experimental/LearningAgents/Source/LearningAgentsTraining/Private/LearningAgentsTrainingEnvironment.cpp b/Engine/Plugins/Experimental/LearningAgents/Source/LearningAgentsTraining/Private/LearningAgentsTrainingEnvironment.cpp
--- a/Engine/Plugins/Experimental/LearningAgents/Source/LearningAgentsTraining/Private/LearningAgentsTrainingEnvironment.cpp
+++ b/Engine/Plugins/Experimental/LearningAgents/Source/LearningAgentsTraining/Private/LearningAgentsTrainingEnvironment.cpp
@@ -198,7 +198,7 @@ void ULearningAgentsTrainingEnvironment::ResetAgentEpisodes_Implementation(const
 
 void ULearningAgentsTrainingEnvironment::GatherRewards()
 {
-	UE_LEARNING_TRACE_CPUPROFILER_EVENT_SCOPE(ULearningAgentsTrainingEnvironment::GatherRewards);
+	TRACE_CPUPROFILER_EVENT_SCOPE(ULearningAgentsTrainingEnvironment::GatherRewards);
 
 	if (!IsSetup())
 	{
@@ -242,7 +242,7 @@ void ULearningAgentsTrainingEnvironment::GatherRewards()
 
 void ULearningAgentsTrainingEnvironment::GatherCompletions()
 {
-	UE_LEARNING_TRACE_CPUPROFILER_EVENT_SCOPE(ULearningAgentsTrainingEnvironment::GatherCompletions);
+	TRACE_CPUPROFILER_EVENT_SCOPE(ULearningAgentsTrainingEnvironment::GatherCompletions);
 
 	if (!IsSetup())
 	{
diff --git a/Engine/Plugins/Experimental/LearningAgents/Source/LearningTraining/Private/LearningExperience.cpp b/Engine/Plugins/Experimental/LearningAgents/Source/LearningTraining/Private/LearningExperience.cpp
--- a/Engine/Plugins/Experimental/LearningAgents/Source/LearningTraining/Private/LearningExperience.cpp
+++ b/Engine/Plugins/Experimental/LearningAgents/Source/LearningTraining/Private/LearningExperience.cpp
@@ -66,7 +66,7 @@ namespace UE::Learning
 
 	int32 FEpisodeBuffer::AddObservations(const FName& Name, const int32 SchemaId, const int32 Size)
 	{
-		UE_LEARNING_CHECKF(!ObservationNames.Contains(Name), TEXT("Observation name collision!"));
+		checkf(!ObservationNames.Contains(Name), TEXT("Observation name collision!"));
 
 		int32 ObservationId = ObservationNames.Num();
 		ObservationNames.Add(Name);
@@ -84,9 +84,9 @@ namespace UE::Learning
 
 	void FEpisodeBuffer::PushObservations(const int32 ObservationId, const TLearningArrayView<2, const float> InObservations, const FIndexSet Instances)
 	{
-		UE_LEARNING_TRACE_CPUPROFILER_EVENT_SCOPE(Learning::FEpisodeBuffer::PushObservations);
+		TRACE_CPUPROFILER_EVENT_SCOPE(Learning::FEpisodeBuffer::PushObservations);
 
-		UE_LEARNING_CHECKF(ObservationId >= 0 && ObservationId < ObservationArrays.Num(), TEXT("Observation id invalid!"));
+		checkf(ObservationId >= 0 && ObservationId < ObservationArrays.Num(), TEXT("Observation id invalid!"));
 
 		for (const int32 InstanceIdx : Instances)
 		{
@@ -96,14 +96,14 @@ namespace UE::Learning
 
 	const TLearningArrayView<2, const float> FEpisodeBuffer::GetObservations(const int32 ObservationId, const int32 InstanceIdx) const
 	{
-		UE_LEARNING_CHECKF(ObservationId >= 0 && ObservationId < ObservationArrays.Num(), TEXT("Observation id invalid!"));
+		checkf(ObservationId >= 0 && ObservationId < ObservationArrays.Num(), TEXT("Observation id invalid!"));
 
 		return ObservationArrays[ObservationId][InstanceIdx].Slice(0, EpisodeStepNums[InstanceIdx]);
 	}
 
 	int32 FEpisodeBuffer::AddActions(const FName& Name, const int32 SchemaId, const int32 Size)
 	{
-		UE_LEARNING_CHECKF(!ActionNames.Contains(Name), TEXT("Action name collision!"));
+		checkf(!ActionNames.Contains(Name), TEXT("Action name collision!"));
 
 		int32 ActionId = ActionNames.Num();
 		ActionNames.Add(Name);
@@ -121,9 +121,9 @@ namespace UE::Learning
 
 	void FEpisodeBuffer::PushActions(const int32 ActionId, const TLearningArrayView<2, const float> InActions, const FIndexSet Instances)
 	{
-		UE_LEARNING_TRACE_CPUPROFILER_EVENT_SCOPE(Learning::FEpisodeBuffer::PushActions);
+		TRACE_CPUPROFILER_EVENT_SCOPE(Learning::FEpisodeBuffer::PushActions);
 
-		UE_LEARNING_CHECKF(ActionId >= 0 && ActionId < ActionArrays.Num(), TEXT("Action id invalid!"));
+		checkf(ActionId >= 0 && ActionId < ActionArrays.Num(), TEXT("Action id invalid!"));
 
 		for (const int32 InstanceIdx : Instances)
 		{
@@ -133,14 +133,14 @@ namespace UE::Learning
 
 	const TLearningArrayView<2, const float> FEpisodeBuffer::GetActions(const int32 ActionId, const int32 InstanceIdx) const
 	{
-		UE_LEARNING_CHECKF(ActionId >= 0 && ActionId < ActionArrays.Num(), TEXT("Action id invalid!"));
+		checkf(ActionId >= 0 && ActionId < ActionArrays.Num(), TEXT("Action id invalid!"));
 
 		return ActionArrays[ActionId][InstanceIdx].Slice(0, EpisodeStepNums[InstanceIdx]);
 	}
 
 	int32 FEpisodeBuffer::AddActionModifiers(const FName& Name, const int32 SchemaId, const int32 Size)
 	{
-		UE_LEARNING_CHECKF(!ActionModifierNames.Contains(Name), TEXT("Action Modifier name collision!"));
+		checkf(!ActionModifierNames.Contains(Name), TEXT("Action Modifier name collision!"));
 
 		int32 ActionModifierId = ActionModifierNames.Num();
 		ActionModifierNames.Add(Name);
@@ -158,9 +158,9 @@ namespace UE::Learning
 
 	void FEpisodeBuffer::PushActionModifiers(const int32 ActionModifierId, const TLearningArrayView<2, const float> InActionModifiers, const FIndexSet Instances)
 	{
-		UE_LEARNING_TRACE_CPUPROFILER_EVENT_SCOPE(Learning::FEpisodeBuffer::PushActionModifiers);
+		TRACE_CPUPROFILER_EVENT_SCOPE(Learning::FEpisodeBuffer::PushActionModifiers);
 
-		UE_LEARNING_CHECKF(ActionModifierId >= 0 && ActionModifierId < ActionModifierArrays.Num(), TEXT("Action Modifier id invalid!"));
+		checkf(ActionModifierId >= 0 && ActionModifierId < ActionModifierArrays.Num(), TEXT("Action Modifier id invalid!"));
 
 		for (const int32 InstanceIdx : Instances)
 		{
@@ -170,14 +170,14 @@ namespace UE::Learning
 
 	const TLearningArrayView<2, const float> FEpisodeBuffer::GetActionModifiers(const int32 ActionModifierId, const int32 InstanceIdx) const
 	{
-		UE_LEARNING_CHECKF(ActionModifierId >= 0 && ActionModifierId < ActionModifierArrays.Num(), TEXT("Action Modifier id invalid!"));
+		checkf(ActionModifierId >= 0 && ActionModifierId < ActionModifierArrays.Num(), TEXT("Action Modifier id invalid!"));
 
 		return ActionModifierArrays[ActionModifierId][InstanceIdx].Slice(0, EpisodeStepNums[InstanceIdx]);
 	}
 
 	int32 FEpisodeBuffer::AddMemoryStates(const FName& Name, const int32 Size)
 	{
-		UE_LEARNING_CHECKF(!MemoryStateNames.Contains(Name), TEXT("Memory State name collision!"));
+		checkf(!MemoryStateNames.Contains(Name), TEXT("Memory State name collision!"));
 
 		int32 MemoryStateId = MemoryStateNames.Num();
 		MemoryStateNames.Add(Name);
@@ -194,9 +194,9 @@ namespace UE::Learning
 
 	void FEpisodeBuffer::PushMemoryStates(const int32 MemoryStateId, const TLearningArrayView<2, const float> InMemoryStates, const FIndexSet Instances)
 	{
-		UE_LEARNING_TRACE_CPUPROFILER_EVENT_SCOPE(Learning::FEpisodeBuffer::PushMemoryStates);
+		TRACE_CPUPROFILER_EVENT_SCOPE(Learning::FEpisodeBuffer::PushMemoryStates);
 
-		UE_LEARNING_CHECKF(MemoryStateId >= 0 && MemoryStateId < MemoryStateArrays.Num(), TEXT("Memory state id invalid!"));
+		checkf(MemoryStateId >= 0 && MemoryStateId < MemoryStateArrays.Num(), TEXT("Memory state id invalid!"));
 
 		for (const int32 InstanceIdx : Instances)
 		{
@@ -206,14 +206,14 @@ namespace UE::Learning
 
 	const TLearningArrayView<2, const float> FEpisodeBuffer::GetMemoryStates(const int32 MemoryStateId, const int32 InstanceIdx) const
 	{
-		UE_LEARNING_CHECKF(MemoryStateId >= 0 && MemoryStateId < MemoryStateArrays.Num(), TEXT("Memory state id invalid!"));
+		checkf(MemoryStateId >= 0 && MemoryStateId < MemoryStateArrays.Num(), TEXT("Memory state id invalid!"));
 
 		return MemoryStateArrays[MemoryStateId][InstanceIdx].Slice(0, EpisodeStepNums[InstanceIdx]);
 	}
 
 	int32 FEpisodeBuffer::AddRewards(const FName& Name, const int32 Size)
 	{
-		UE_LEARNING_CHECKF(!RewardNames.Contains(Name), TEXT("Reward name collision!"));
+		checkf(!RewardNames.Contains(Name), TEXT("Reward name collision!"));
 
 		int32 RewardId = RewardNames.Num();
 		RewardNames.Add(Name);
@@ -230,9 +230,9 @@ namespace UE::Learning
 
 	void FEpisodeBuffer::PushRewards(const int32 RewardId, const TLearningArrayView<1, const float> InRewards, const FIndexSet Instances)
 	{
-		UE_LEARNING_TRACE_CPUPROFILER_EVENT_SCOPE(Learning::FEpisodeBuffer::PushRewards);
+		TRACE_CPUPROFILER_EVENT_SCOPE(Learning::FEpisodeBuffer::PushRewards);
 
-		UE_LEARNING_CHECKF(RewardId >= 0 && RewardId < RewardArrays.Num(), TEXT("Reward id invalid!"));
+		checkf(RewardId >= 0 && RewardId < RewardArrays.Num(), TEXT("Reward id invalid!"));
 
 		TLearningArrayView<2, const float> RewardsReshaped(InRewards.GetData(), { InRewards.Num(), 1 });
 		for (const int32 InstanceIdx : Instances)
@@ -243,9 +243,9 @@ namespace UE::Learning
 
 	void FEpisodeBuffer::PushRewards(const int32 RewardId, const TLearningArrayView<2, const float> InRewards, const FIndexSet Instances)
 	{
-		UE_LEARNING_TRACE_CPUPROFILER_EVENT_SCOPE(Learning::FEpisodeBuffer::PushRewards);
+		TRACE_CPUPROFILER_EVENT_SCOPE(Learning::FEpisodeBuffer::PushRewards);
 
-		UE_LEARNING_CHECKF(RewardId >= 0 && RewardId < RewardArrays.Num(), TEXT("Reward id invalid!"));
+		checkf(RewardId >= 0 && RewardId < RewardArrays.Num(), TEXT("Reward id invalid!"));
 
 		for (const int32 InstanceIdx : Instances)
 		{
@@ -255,18 +255,18 @@ namespace UE::Learning
 
 	const TLearningArrayView<2, const float> FEpisodeBuffer::GetRewards(const int32 RewardId, const int32 InstanceIdx) const
 	{
-		UE_LEARNING_CHECKF(RewardId >= 0 && RewardId < RewardArrays.Num(), TEXT("Reward id invalid!"));
+		checkf(RewardId >= 0 && RewardId < RewardArrays.Num(), TEXT("Reward id invalid!"));
 
 		return RewardArrays[RewardId][RewardId].Slice(0, EpisodeStepNums[InstanceIdx]);
 	}
 
 	void FEpisodeBuffer::IncrementEpisodeStepNums(const FIndexSet Instances)
 	{
-		UE_LEARNING_TRACE_CPUPROFILER_EVENT_SCOPE(Learning::FEpisodeBuffer::IncrementEpisodeStepNums);
+		TRACE_CPUPROFILER_EVENT_SCOPE(Learning::FEpisodeBuffer::IncrementEpisodeStepNums);
 
 		for (const int32 InstanceIdx : Instances)
 		{
-			UE_LEARNING_CHECKF(EpisodeStepNums[InstanceIdx] < MaxStepNum, TEXT("Episode Buffer full!"));
+			checkf(EpisodeStepNums[InstanceIdx] < MaxStepNum, TEXT("Episode Buffer full!"));
 			EpisodeStepNums[InstanceIdx]++;
 		}
 	}
@@ -386,19 +386,19 @@ namespace UE::Learning
 		const FIndexSet Instances,
 		const bool bAddTruncatedEpisodeWhenFull)
 	{
-		UE_LEARNING_TRACE_CPUPROFILER_EVENT_SCOPE(Learning::FReplayBuffer::AddEpisodes);
+		TRACE_CPUPROFILER_EVENT_SCOPE(Learning::FReplayBuffer::AddEpisodes);
 
-		UE_LEARNING_CHECKF(EpisodeBuffer.ObservationArrays.Num() == Observations.Num(), TEXT("Observation number mismatch!"));
-		UE_LEARNING_CHECKF(EpisodeBuffer.ObservationArrays.Num() == InEpisodeFinalObservations.Num(), TEXT("Final Observation number mismatch!"));
-		UE_LEARNING_CHECKF(EpisodeBuffer.ActionArrays.Num() == Actions.Num(), TEXT("Action number mismatch!"));
-		UE_LEARNING_CHECKF(EpisodeBuffer.ActionModifierArrays.Num() == ActionModifiers.Num(), TEXT("Action Modifier number mismatch!"));
-		UE_LEARNING_CHECKF(EpisodeBuffer.MemoryStateArrays.Num() == MemoryStates.Num(), TEXT("Memory State number mismatch!"));
-		UE_LEARNING_CHECKF(EpisodeBuffer.MemoryStateArrays.Num() == InEpisodeFinalMemoryStates.Num(), TEXT("Final Memory State number mismatch!"));
-		UE_LEARNING_CHECKF(EpisodeBuffer.RewardArrays.Num() == Rewards.Num(), TEXT("Reward number mismatch!"));
+		checkf(EpisodeBuffer.ObservationArrays.Num() == Observations.Num(), TEXT("Observation number mismatch!"));
+		checkf(EpisodeBuffer.ObservationArrays.Num() == InEpisodeFinalObservations.Num(), TEXT("Final Observation number mismatch!"));
+		checkf(EpisodeBuffer.ActionArrays.Num() == Actions.Num(), TEXT("Action number mismatch!"));
+		checkf(EpisodeBuffer.ActionModifierArrays.Num() == ActionModifiers.Num(), TEXT("Action Modifier number mismatch!"));
+		checkf(EpisodeBuffer.MemoryStateArrays.Num() == MemoryStates.Num(), TEXT("Memory State number mismatch!"));
+		checkf(EpisodeBuffer.MemoryStateArrays.Num() == InEpisodeFinalMemoryStates.Num(), TEXT("Final Memory State number mismatch!"));
+		checkf(EpisodeBuffer.RewardArrays.Num() == Rewards.Num(), TEXT("Reward number mismatch!"));
 
 		for (const int32 InstanceIdx : Instances)
 		{
-			UE_LEARNING_CHECKF(InEpisodeCompletionModes[InstanceIdx] != ECompletionMode::Running,
+			checkf(InEpisodeCompletionModes[InstanceIdx] != ECompletionMode::Running,
 				TEXT("Tried to add experience from episode that is still running..."));
 
 			const int32 EpisodeStepNum = EpisodeBuffer.GetEpisodeStepNums()[InstanceIdx];
@@ -466,7 +466,7 @@ namespace UE::Learning
 			if (bAddTruncatedEpisodeWhenFull && EpisodeNum < MaxEpisodeNum && StepNum < MaxStepNum)
 			{
 				const int32 PartialStepNum = MaxStepNum - StepNum;
-				UE_LEARNING_CHECK(PartialStepNum > 0 && PartialStepNum < EpisodeStepNum);
+				check(PartialStepNum > 0 && PartialStepNum < EpisodeStepNum);
 
 				// Copy the data into the replay buffer
 
@@ -794,7 +794,7 @@ namespace UE::Learning
 			const TFunctionRef<void(const FIndexSet Instances)> CompletionFunction,
 			const FIndexSet Instances)
 		{
-			UE_LEARNING_TRACE_CPUPROFILER_EVENT_SCOPE(Learning::Experience::GatherExperienceUntilReplayBufferFull);
+			TRACE_CPUPROFILER_EVENT_SCOPE(Learning::Experience::GatherExperienceUntilReplayBufferFull);
 
 			// Reset Everything
 			ReplayBuffer.Reset();
diff --git a/Engine/Plugins/Experimental/LearningAgents/Source/LearningTraining/Private/LearningExternalTrainer.cpp b/Engine/Plugins/Experimental/LearningAgents/Source/LearningTraining/Private/LearningExternalTrainer.cpp
--- a/Engine/Plugins/Experimental/LearningAgents/Source/LearningTraining/Private/LearningExternalTrainer.cpp
+++ b/Engine/Plugins/Experimental/LearningAgents/Source/LearningTraining/Private/LearningExternalTrainer.cpp
@@ -37,8 +37,8 @@ namespace UE::Learning
 		const ESubprocessFlags TrainingProcessFlags,
 		const ELogSetting LogSettings)
 	{
-		UE_LEARNING_CHECK(FPaths::FileExists(PythonExecutablePath));
-		UE_LEARNING_CHECK(FPaths::DirectoryExists(PythonContentPath));
+		check(FPaths::FileExists(PythonExecutablePath));
+		check(FPaths::DirectoryExists(PythonContentPath));
 
 		Timeout = InTimeout;
 		IntermediatePath = InIntermediatePath;
@@ -600,8 +600,8 @@ namespace UE::Learning
 	{
 		Timeout = InTimeout;
 
-		UE_LEARNING_CHECK(FPaths::FileExists(PythonExecutablePath));
-		UE_LEARNING_CHECK(FPaths::DirectoryExists(PythonContentPath));
+		check(FPaths::FileExists(PythonExecutablePath));
+		check(FPaths::DirectoryExists(PythonContentPath));
 
 		IFileManager& FileManager = IFileManager::Get();
 		const FString CommandLineArguments = FString::Printf(TEXT("\"%s\" Training -p \"%s\" -m \"%s\" \"%s\" Socket \"%s:%i\" \"%s\""),
diff --git a/Engine/Plugins/Experimental/LearningAgents/Source/LearningTraining/Private/LearningPPOTrainer.cpp b/Engine/Plugins/Experimental/LearningAgents/Source/LearningTraining/Private/LearningPPOTrainer.cpp
--- a/Engine/Plugins/Experimental/LearningAgents/Source/LearningTraining/Private/LearningPPOTrainer.cpp
+++ b/Engine/Plugins/Experimental/LearningAgents/Source/LearningTraining/Private/LearningPPOTrainer.cpp
@@ -68,7 +68,7 @@ int32 ULearningSocketPPOTrainerServerCommandlet::Main(const FString& Commandline
 		}
 		else
 		{
-			UE_LEARNING_NOT_IMPLEMENTED();
+			checkNoEntry();
 			return 1;
 		}
 	}
@@ -99,7 +99,7 @@ int32 ULearningSocketPPOTrainerServerCommandlet::Main(const FString& Commandline
 	}
 
 #else
-	UE_LEARNING_NOT_IMPLEMENTED();
+	checkNoEntry();
 #endif
 
 	return 0;
diff --git a/Engine/Plugins/Experimental/LearningAgents/Source/LearningTraining/Private/LearningSharedMemoryTraining.cpp b/Engine/Plugins/Experimental/LearningAgents/Source/LearningTraining/Private/LearningSharedMemoryTraining.cpp
--- a/Engine/Plugins/Experimental/LearningAgents/Source/LearningTraining/Private/LearningSharedMemoryTraining.cpp
+++ b/Engine/Plugins/Experimental/LearningAgents/Source/LearningTraining/Private/LearningSharedMemoryTraining.cpp
@@ -99,7 +99,7 @@ namespace UE::Learning::SharedMemoryTraining
 		{
 			FScopeNullableWriteLock ScopeLock(NetworkLock);
 
-			UE_LEARNING_CHECKF(Controls[(uint8)EControls::NetworkId] == NetworkId, TEXT("Received unexpected NetworkId!"));
+			checkf(Controls[(uint8)EControls::NetworkId] == NetworkId, TEXT("Received unexpected NetworkId!"));
 
 			if (NetworkData.Num() != OutNetwork.GetSnapshotByteNum())
 			{
diff --git a/Engine/Plugins/Experimental/LearningAgents/Source/LearningTraining/Private/LearningSocketTraining.cpp b/Engine/Plugins/Experimental/LearningAgents/Source/LearningTraining/Private/LearningSocketTraining.cpp
--- a/Engine/Plugins/Experimental/LearningAgents/Source/LearningTraining/Private/LearningSocketTraining.cpp
+++ b/Engine/Plugins/Experimental/LearningAgents/Source/LearningTraining/Private/LearningSocketTraining.cpp
@@ -80,7 +80,7 @@ namespace UE::Learning::SocketTraining
 		FRWLock* NetworkLock,
 		const ELogSetting LogSettings)
 	{
-		UE_LEARNING_CHECK(OutNetworkBuffer.Num() == OutNetwork.GetSnapshotByteNum());
+		check(OutNetworkBuffer.Num() == OutNetwork.GetSnapshotByteNum());
 
 		if (LogSettings != ELogSetting::Silent)
 		{
diff --git a/Engine/Plugins/Experimental/LearningAgents/Source/LearningTraining/Private/LearningTrainer.cpp b/Engine/Plugins/Experimental/LearningAgents/Source/LearningTraining/Private/LearningTrainer.cpp
--- a/Engine/Plugins/Experimental/LearningAgents/Source/LearningTraining/Private/LearningTrainer.cpp
+++ b/Engine/Plugins/Experimental/LearningAgents/Source/LearningTraining/Private/LearningTrainer.cpp
@@ -217,7 +217,7 @@ namespace UE::Learning::Trainer
 		}
 
 		default:
-			UE_LEARNING_NOT_IMPLEMENTED();
+			checkNoEntry();
 		}
 
 		return Object;
@@ -343,7 +343,7 @@ namespace UE::Learning::Trainer
 		}
 
 		default:
-			UE_LEARNING_NOT_IMPLEMENTED();
+			checkNoEntry();
 		}
 
 		return Object;
@@ -355,7 +355,7 @@ namespace UE::Learning::Trainer
 		{
 		case ETrainerDevice::GPU: return TEXT("GPU");
 		case ETrainerDevice::CPU: return TEXT("CPU");
-		default: UE_LEARNING_NOT_IMPLEMENTED(); return TEXT("Unknown");
+		default: checkNoEntry(); return TEXT("Unknown");
 		}
 	}
 
@@ -368,7 +368,7 @@ namespace UE::Learning::Trainer
 		case ETrainerResponse::Completed: return TEXT("Training completed");
 		case ETrainerResponse::Stopped: return TEXT("Training stopped");
 		case ETrainerResponse::Timeout: return TEXT("Communication timeout");
-		default: UE_LEARNING_NOT_IMPLEMENTED(); return TEXT("Unknown");
+		default: checkNoEntry(); return TEXT("Unknown");
 		}
 	}
 
@@ -379,21 +379,21 @@ namespace UE::Learning::Trainer
 
 	float DiscountFactorFromHalfLifeSteps(const int32 HalfLifeSteps)
 	{
-		UE_LEARNING_CHECKF(HalfLifeSteps >= 1, TEXT("Number of HalfLifeSteps should be at least 1 but got %i"), HalfLifeSteps);
+		checkf(HalfLifeSteps >= 1, TEXT("Number of HalfLifeSteps should be at least 1 but got %i"), HalfLifeSteps);
 
 		return FMath::Pow(0.5f, 1.0f / FMath::Max(HalfLifeSteps, 1));
 	}
 
 	FString GetPythonExecutablePath(const FString& IntermediateDir)
 	{
-		UE_LEARNING_CHECKF(PLATFORM_WINDOWS || PLATFORM_MAC || PLATFORM_LINUX, TEXT("Python only supported on Windows, Mac, and Linux."));
+		checkf(PLATFORM_WINDOWS || PLATFORM_MAC || PLATFORM_LINUX, TEXT("Python only supported on Windows, Mac, and Linux."));
 
 		return IntermediateDir / TEXT("PipInstall") / (PLATFORM_WINDOWS ? TEXT("Scripts/python.exe") : TEXT("bin/python3"));
 	}
 
 	FString GetSitePackagesPath(const FString& EngineDir)
 	{
-		UE_LEARNING_CHECKF(PLATFORM_WINDOWS || PLATFORM_MAC || PLATFORM_LINUX, TEXT("Python only supported on Windows, Mac, and Linux."));
+		checkf(PLATFORM_WINDOWS || PLATFORM_MAC || PLATFORM_LINUX, TEXT("Python only supported on Windows, Mac, and Linux."));
 
 		return EngineDir / TEXT("Plugins/Experimental/PythonFoundationPackages/Content/Python/Lib") / FPlatformMisc::GetUBTPlatform() / TEXT("site-packages");
 	}
