LearningAgents: Added FrameAttribute data structure, and added Range Offsets to FrameRangeSet data structure to make some algorithms easier to implement.

diff --git a/Engine/Plugins/Experimental/LearningAgents/Source/Learning/Private/Learning.ispc b/Engine/Plugins/Experimental/LearningAgents/Source/Learning/Private/Learning.ispc
--- a/Engine/Plugins/Experimental/LearningAgents/Source/Learning/Private/Learning.ispc
+++ b/Engine/Plugins/Experimental/LearningAgents/Source/Learning/Private/Learning.ispc
@@ -17,6 +17,19 @@ inline float LearningSigmoid(const float X)
 	return 1.0f / (1.0f + exp(-X));
 }
 
+inline float LearningSquare(float X)
+{
+	return X * X;
+}
+
+inline float LearningInvExpApprox(float X)
+{
+	uniform const float A = 1.00746054f; // 1 / 1! in Taylor series
+	uniform const float B = 0.45053901f; // 1 / 2! in Taylor series
+	uniform const float C = 0.25724632f; // 1 / 3! in Taylor series
+	return 1 / (1 + A * X + B * X * X + C * X * X * X);
+}
+
 /////////////////////////////////////////////////////////
 // 
 // Functions matching those in LearningRandom.cpp
@@ -678,4 +691,60 @@ export void LearningKMeansComputeClusteredIndex(
 		OutSampleMapping[SampleIdx] = SortedIdx;
 		OutInverseSampleMapping[SortedIdx] = SampleIdx;
 	}
+}
+
+export void LearningFilterGaussian(
+	uniform float Out[], 
+	uniform const float In[], 
+	uniform const int ValueNum,
+	uniform const float StdInFrames)
+{
+	uniform const int StdRange = (uniform int)round(StdInFrames * 3.0f);
+
+	foreach (ValueIdx = 0 ... ValueNum)
+	{
+		float Total = 0.0f;
+		Out[ValueIdx] = 0.0f;
+
+		const int RangeMin = max(ValueIdx - StdRange, 0);
+		const int RangeMax = min(ValueIdx + StdRange, ValueNum - 1);
+
+		for (int Offset = -StdRange; Offset <= +StdRange; Offset++)
+		{
+			if (ValueIdx + Offset >= 0 && ValueIdx + Offset < ValueNum)
+			{
+				const float Weight = LearningInvExpApprox(LearningSquare(Offset / StdInFrames));
+				Out[ValueIdx] += Weight * In[ValueIdx + Offset];
+				Total += Weight;
+			}
+		}
+
+		Out[ValueIdx] = Out[ValueIdx] / Total;
+	}
+}
+
+export void LearningFilterMajorityVote(
+	uniform float Out[],
+	uniform const float In[],
+	uniform const int ValueNum,
+	uniform const float FilterWidthInFrames)
+{
+	uniform const int FilterRange = (uniform int)round(FilterWidthInFrames);
+
+	foreach(ValueIdx = 0 ... ValueNum)
+	{
+		int Total = 0;
+		const int RangeMin = max(ValueIdx - FilterRange, 0);
+		const int RangeMax = min(ValueIdx + FilterRange, ValueNum - 1);
+
+		for (int Offset = -FilterRange; Offset <= +FilterRange; Offset++)
+		{
+			if (ValueIdx + Offset >= 0 && ValueIdx + Offset < ValueNum)
+			{
+				Total += In[ValueIdx + Offset] ? 1 : -1;
+			}
+		}
+
+		Out[ValueIdx] = Total > 0;
+	}
 }
\ No newline at end of file
diff --git a/Engine/Plugins/Experimental/LearningAgents/Source/Learning/Private/LearningFrameAttribute.cpp b/Engine/Plugins/Experimental/LearningAgents/Source/Learning/Private/LearningFrameAttribute.cpp
new file mode 100644
--- /dev/null
+++ b/Engine/Plugins/Experimental/LearningAgents/Source/Learning/Private/LearningFrameAttribute.cpp
@@ -0,0 +1,2752 @@
+// Copyright Epic Games, Inc. All Rights Reserved.
+
+#include "LearningFrameAttribute.h"
+
+#if UE_LEARNING_ISPC
+#include "Learning.ispc.generated.h"
+#endif
+
+namespace UE::Learning
+{
+	namespace FrameAttribute::Private
+	{
+		static inline bool FindMinimum(
+			int32& OutChannelIdx,
+			int32& OutFrameIdx,
+			float& OutMinimumValue,
+			const TLearningArrayView<2, const float> Values)
+		{
+			OutChannelIdx = INDEX_NONE;
+			OutFrameIdx = INDEX_NONE;
+			OutMinimumValue = UE_MAX_FLT;
+
+			const int32 ChannelNum = Values.Num<0>();
+			const int32 FrameNum = Values.Num<1>();
+
+			for (int32 ChannelIdx = 0; ChannelIdx < ChannelNum; ChannelIdx++)
+			{
+				for (int32 FrameIdx = 0; FrameIdx < FrameNum; FrameIdx++)
+				{
+					if (Values[ChannelIdx][FrameIdx] < OutMinimumValue)
+					{
+						OutChannelIdx = ChannelIdx;
+						OutFrameIdx = FrameIdx;
+						OutMinimumValue = Values[ChannelIdx][FrameIdx];
+					}
+				}
+			}
+
+			return OutChannelIdx != INDEX_NONE;
+		}
+		
+		static inline bool FindMaximum(
+			int32& OutChannelIdx,
+			int32& OutFrameIdx,
+			float& OutMaximumValue,
+			const TLearningArrayView<2, const float> Values)
+		{
+			OutChannelIdx = INDEX_NONE;
+			OutFrameIdx = INDEX_NONE;
+			OutMaximumValue = -UE_MAX_FLT;
+
+			const int32 ChannelNum = Values.Num<0>();
+			const int32 FrameNum = Values.Num<1>();
+
+			for (int32 ChannelIdx = 0; ChannelIdx < ChannelNum; ChannelIdx++)
+			{
+				for (int32 FrameIdx = 0; FrameIdx < FrameNum; FrameIdx++)
+				{
+					if (Values[ChannelIdx][FrameIdx] > OutMaximumValue)
+					{
+						OutChannelIdx = ChannelIdx;
+						OutFrameIdx = FrameIdx;
+						OutMaximumValue = Values[ChannelIdx][FrameIdx];
+					}
+				}
+			}
+
+			return OutChannelIdx != INDEX_NONE;
+		}
+
+		static inline void Add(TLearningArrayView<1, float> Out, const TLearningArrayView<1, const float> Lhs, const TLearningArrayView<1, const float> Rhs)
+		{
+			check(Out.Num() == Lhs.Num());
+			check(Out.Num() == Rhs.Num());
+			const int32 ValueNum = Out.Num();
+			for (int32 ValueIdx = 0; ValueIdx < ValueNum; ValueIdx++)
+			{
+				Out[ValueIdx] = Lhs[ValueIdx] + Rhs[ValueIdx];
+			}
+		}
+
+		static inline void Sub(TLearningArrayView<1, float> Out, const TLearningArrayView<1, const float> Lhs, const TLearningArrayView<1, const float> Rhs)
+		{
+			check(Out.Num() == Lhs.Num());
+			check(Out.Num() == Rhs.Num());
+			const int32 ValueNum = Out.Num();
+			for (int32 ValueIdx = 0; ValueIdx < ValueNum; ValueIdx++)
+			{
+				Out[ValueIdx] = Lhs[ValueIdx] - Rhs[ValueIdx];
+			}
+		}
+
+		static inline void Mul(TLearningArrayView<1, float> Out, const TLearningArrayView<1, const float> Lhs, const TLearningArrayView<1, const float> Rhs)
+		{
+			check(Out.Num() == Lhs.Num());
+			check(Out.Num() == Rhs.Num());
+			const int32 ValueNum = Out.Num();
+			for (int32 ValueIdx = 0; ValueIdx < ValueNum; ValueIdx++)
+			{
+				Out[ValueIdx] = Lhs[ValueIdx] * Rhs[ValueIdx];
+			}
+		}
+
+		static inline void Div(TLearningArrayView<1, float> Out, const TLearningArrayView<1, const float> Lhs, const TLearningArrayView<1, const float> Rhs)
+		{
+			check(Out.Num() == Lhs.Num());
+			check(Out.Num() == Rhs.Num());
+			const int32 ValueNum = Out.Num();
+			for (int32 ValueIdx = 0; ValueIdx < ValueNum; ValueIdx++)
+			{
+				Out[ValueIdx] = Lhs[ValueIdx] / Rhs[ValueIdx];
+			}
+		}
+
+		static inline void Dot(TLearningArrayView<1, float> InOut, const TLearningArrayView<1, const float> Lhs, const TLearningArrayView<1, const float> Rhs)
+		{
+			check(InOut.Num() == Lhs.Num());
+			check(InOut.Num() == Rhs.Num());
+			const int32 ValueNum = InOut.Num();
+			for (int32 ValueIdx = 0; ValueIdx < ValueNum; ValueIdx++)
+			{
+				InOut[ValueIdx] += Lhs[ValueIdx] * Rhs[ValueIdx];
+			}
+		}
+
+		static inline void Neg(TLearningArrayView<1, float> Out, const TLearningArrayView<1, const float> In)
+		{
+			check(Out.Num() == In.Num());
+			const int32 ValueNum = Out.Num();
+			for (int32 ValueIdx = 0; ValueIdx < ValueNum; ValueIdx++)
+			{
+				Out[ValueIdx] = -In[ValueIdx];
+			}
+		}
+
+		static inline void Inv(TLearningArrayView<1, float> Out, const TLearningArrayView<1, const float> In)
+		{
+			check(Out.Num() == In.Num());
+			const int32 ValueNum = Out.Num();
+			for (int32 ValueIdx = 0; ValueIdx < ValueNum; ValueIdx++)
+			{
+				Out[ValueIdx] = 1.0f / In[ValueIdx];
+			}
+		}
+
+		static inline void Abs(TLearningArrayView<1, float> Out, const TLearningArrayView<1, const float> In)
+		{
+			check(Out.Num() == In.Num());
+			const int32 ValueNum = Out.Num();
+			for (int32 ValueIdx = 0; ValueIdx < ValueNum; ValueIdx++)
+			{
+				Out[ValueIdx] = FMath::Abs(In[ValueIdx]);
+			}
+		}
+
+		static inline void Sqrt(TLearningArrayView<1, float> Out, const TLearningArrayView<1, const float> In)
+		{
+			check(Out.Num() == In.Num());
+			const int32 ValueNum = Out.Num();
+			for (int32 ValueIdx = 0; ValueIdx < ValueNum; ValueIdx++)
+			{
+				Out[ValueIdx] = FMath::Sqrt(In[ValueIdx]);
+			}
+		}
+
+		static inline void LengthSquared(TLearningArrayView<1, float> InOut, const TLearningArrayView<1, const float> In)
+		{
+			check(InOut.Num() == In.Num());
+			const int32 ValueNum = InOut.Num();
+			for (int32 ValueIdx = 0; ValueIdx < ValueNum; ValueIdx++)
+			{
+				InOut[ValueIdx] += In[ValueIdx] * In[ValueIdx];
+			}
+		}
+
+		static inline void SqrtInplace(TLearningArrayView<1, float> InOut)
+		{
+			const int32 ValueNum = InOut.Num();
+			for (int32 ValueIdx = 0; ValueIdx < ValueNum; ValueIdx++)
+			{
+				InOut[ValueIdx] = FMath::Sqrt(InOut[ValueIdx]);
+			}
+		}
+
+		static inline void Normalize(TLearningArrayView<2, float> Out, const TLearningArrayView<2, const float> In)
+		{
+			check(Out.Num<0>() == In.Num<0>());
+			check(Out.Num<1>() == In.Num<1>());
+
+			const int32 ChannelNum = In.Num<0>();
+			const int32 FrameNum = In.Num<1>();
+			for (int32 FrameIdx = 0; FrameIdx < FrameNum; FrameIdx++)
+			{
+				float LengthSquared = 0.0f;
+				for (int32 ChannelIdx = 0; ChannelIdx < ChannelNum; ChannelIdx++)
+				{
+					LengthSquared += In[ChannelIdx][FrameIdx] * In[ChannelIdx][FrameIdx];
+				}
+
+				const float Length = FMath::Sqrt(LengthSquared);
+
+				for (int32 ChannelIdx = 0; ChannelIdx < ChannelNum; ChannelIdx++)
+				{
+					Out[ChannelIdx][FrameIdx] = In[ChannelIdx][FrameIdx] / Length;
+				}
+			}
+		}
+
+		static inline void AddConstant(TLearningArrayView<1, float> Out, const TLearningArrayView<1, const float> Lhs, const float Rhs)
+		{
+			check(Out.Num() == Lhs.Num());
+			const int32 ValueNum = Out.Num();
+			for (int32 ValueIdx = 0; ValueIdx < ValueNum; ValueIdx++)
+			{
+				Out[ValueIdx] = Lhs[ValueIdx] + Rhs;
+			}
+		}
+
+		static inline void SubConstant(TLearningArrayView<1, float> Out, const TLearningArrayView<1, const float> Lhs, const float Rhs)
+		{
+			check(Out.Num() == Lhs.Num());
+			const int32 ValueNum = Out.Num();
+			for (int32 ValueIdx = 0; ValueIdx < ValueNum; ValueIdx++)
+			{
+				Out[ValueIdx] = Lhs[ValueIdx] - Rhs;
+			}
+		}
+
+		static inline void MulConstant(TLearningArrayView<1, float> Out, const TLearningArrayView<1, const float> Lhs, const float Rhs)
+		{
+			check(Out.Num() == Lhs.Num());
+			const int32 ValueNum = Out.Num();
+			for (int32 ValueIdx = 0; ValueIdx < ValueNum; ValueIdx++)
+			{
+				Out[ValueIdx] = Lhs[ValueIdx] * Rhs;
+			}
+		}
+
+		static inline void DivConstant(TLearningArrayView<1, float> Out, const TLearningArrayView<1, const float> Lhs, const float Rhs)
+		{
+			check(Out.Num() == Lhs.Num());
+			const int32 ValueNum = Out.Num();
+			for (int32 ValueIdx = 0; ValueIdx < ValueNum; ValueIdx++)
+			{
+				Out[ValueIdx] = Lhs[ValueIdx] / Rhs;
+			}
+		}
+
+		static inline void ConstantAdd(TLearningArrayView<1, float> Out, const float Lhs, const TLearningArrayView<1, const float> Rhs)
+		{
+			check(Out.Num() == Rhs.Num());
+			const int32 ValueNum = Out.Num();
+			for (int32 ValueIdx = 0; ValueIdx < ValueNum; ValueIdx++)
+			{
+				Out[ValueIdx] = Lhs + Rhs[ValueIdx];
+			}
+		}
+
+		static inline void ConstantSub(TLearningArrayView<1, float> Out, const float Lhs, const TLearningArrayView<1, const float> Rhs)
+		{
+			check(Out.Num() == Rhs.Num());
+			const int32 ValueNum = Out.Num();
+			for (int32 ValueIdx = 0; ValueIdx < ValueNum; ValueIdx++)
+			{
+				Out[ValueIdx] = Lhs - Rhs[ValueIdx];
+			}
+		}
+
+		static inline void ConstantMul(TLearningArrayView<1, float> Out, const float Lhs, const TLearningArrayView<1, const float> Rhs)
+		{
+			check(Out.Num() == Rhs.Num());
+			const int32 ValueNum = Out.Num();
+			for (int32 ValueIdx = 0; ValueIdx < ValueNum; ValueIdx++)
+			{
+				Out[ValueIdx] = Lhs * Rhs[ValueIdx];
+			}
+		}
+
+		static inline void ConstantDiv(TLearningArrayView<1, float> Out, const float Lhs, const TLearningArrayView<1, const float> Rhs)
+		{
+			check(Out.Num() == Rhs.Num());
+			const int32 ValueNum = Out.Num();
+			for (int32 ValueIdx = 0; ValueIdx < ValueNum; ValueIdx++)
+			{
+				Out[ValueIdx] = Lhs / Rhs[ValueIdx];
+			}
+		}
+
+		static inline void LogicalAnd(TLearningArrayView<1, float> Out, const TLearningArrayView<1, const float> Lhs, const TLearningArrayView<1, const float> Rhs)
+		{
+			check(Out.Num() == Lhs.Num());
+			check(Out.Num() == Rhs.Num());
+			const int32 ValueNum = Out.Num();
+			for (int32 ValueIdx = 0; ValueIdx < ValueNum; ValueIdx++)
+			{
+				Out[ValueIdx] = Lhs[ValueIdx] && Rhs[ValueIdx];
+			}
+		}
+
+		static inline void LogicalOr(TLearningArrayView<1, float> Out, const TLearningArrayView<1, const float> Lhs, const TLearningArrayView<1, const float> Rhs)
+		{
+			check(Out.Num() == Lhs.Num());
+			check(Out.Num() == Rhs.Num());
+			const int32 ValueNum = Out.Num();
+			for (int32 ValueIdx = 0; ValueIdx < ValueNum; ValueIdx++)
+			{
+				Out[ValueIdx] = Lhs[ValueIdx] || Rhs[ValueIdx];
+			}
+		}
+
+		static inline void LogicalNot(TLearningArrayView<1, float> Out, const TLearningArrayView<1, const float> In)
+		{
+			check(Out.Num() == In.Num());
+			const int32 ValueNum = Out.Num();
+			for (int32 ValueIdx = 0; ValueIdx < ValueNum; ValueIdx++)
+			{
+				Out[ValueIdx] = !In[ValueIdx];
+			}
+		}
+
+		static inline void Gt(TLearningArrayView<1, float> Out, const TLearningArrayView<1, const float> Lhs, const TLearningArrayView<1, const float> Rhs)
+		{
+			check(Out.Num() == Lhs.Num());
+			check(Out.Num() == Rhs.Num());
+			const int32 ValueNum = Out.Num();
+			for (int32 ValueIdx = 0; ValueIdx < ValueNum; ValueIdx++)
+			{
+				Out[ValueIdx] = Lhs[ValueIdx] > Rhs[ValueIdx];
+			}
+		}
+
+		static inline void Ge(TLearningArrayView<1, float> Out, const TLearningArrayView<1, const float> Lhs, const TLearningArrayView<1, const float> Rhs)
+		{
+			check(Out.Num() == Lhs.Num());
+			check(Out.Num() == Rhs.Num());
+			const int32 ValueNum = Out.Num();
+			for (int32 ValueIdx = 0; ValueIdx < ValueNum; ValueIdx++)
+			{
+				Out[ValueIdx] = Lhs[ValueIdx] >= Rhs[ValueIdx];
+			}
+		}
+
+		static inline void Lt(TLearningArrayView<1, float> Out, const TLearningArrayView<1, const float> Lhs, const TLearningArrayView<1, const float> Rhs)
+		{
+			check(Out.Num() == Lhs.Num());
+			check(Out.Num() == Rhs.Num());
+			const int32 ValueNum = Out.Num();
+			for (int32 ValueIdx = 0; ValueIdx < ValueNum; ValueIdx++)
+			{
+				Out[ValueIdx] = Lhs[ValueIdx] < Rhs[ValueIdx];
+			}
+		}
+
+		static inline void Le(TLearningArrayView<1, float> Out, const TLearningArrayView<1, const float> Lhs, const TLearningArrayView<1, const float> Rhs)
+		{
+			check(Out.Num() == Lhs.Num());
+			check(Out.Num() == Rhs.Num());
+			const int32 ValueNum = Out.Num();
+			for (int32 ValueIdx = 0; ValueIdx < ValueNum; ValueIdx++)
+			{
+				Out[ValueIdx] = Lhs[ValueIdx] <= Rhs[ValueIdx];
+			}
+		}
+
+		static inline void Eq(TLearningArrayView<1, float> Out, const TLearningArrayView<1, const float> Lhs, const TLearningArrayView<1, const float> Rhs)
+		{
+			check(Out.Num() == Lhs.Num());
+			check(Out.Num() == Rhs.Num());
+			const int32 ValueNum = Out.Num();
+			for (int32 ValueIdx = 0; ValueIdx < ValueNum; ValueIdx++)
+			{
+				Out[ValueIdx] = Lhs[ValueIdx] == Rhs[ValueIdx];
+			}
+		}
+
+		static inline void Neq(TLearningArrayView<1, float> Out, const TLearningArrayView<1, const float> Lhs, const TLearningArrayView<1, const float> Rhs)
+		{
+			check(Out.Num() == Lhs.Num());
+			check(Out.Num() == Rhs.Num());
+			const int32 ValueNum = Out.Num();
+			for (int32 ValueIdx = 0; ValueIdx < ValueNum; ValueIdx++)
+			{
+				Out[ValueIdx] = Lhs[ValueIdx] != Rhs[ValueIdx];
+			}
+		}
+
+		static inline void GtConstant(TLearningArrayView<1, float> Out, const TLearningArrayView<1, const float> Lhs, const float Rhs)
+		{
+			check(Out.Num() == Lhs.Num());
+			const int32 ValueNum = Out.Num();
+			for (int32 ValueIdx = 0; ValueIdx < ValueNum; ValueIdx++)
+			{
+				Out[ValueIdx] = Lhs[ValueIdx] > Rhs;
+			}
+		}
+
+		static inline void GeConstant(TLearningArrayView<1, float> Out, const TLearningArrayView<1, const float> Lhs, const float Rhs)
+		{
+			check(Out.Num() == Lhs.Num());
+			const int32 ValueNum = Out.Num();
+			for (int32 ValueIdx = 0; ValueIdx < ValueNum; ValueIdx++)
+			{
+				Out[ValueIdx] = Lhs[ValueIdx] >= Rhs;
+			}
+		}
+
+		static inline void LtConstant(TLearningArrayView<1, float> Out, const TLearningArrayView<1, const float> Lhs, const float Rhs)
+		{
+			check(Out.Num() == Lhs.Num());
+			const int32 ValueNum = Out.Num();
+			for (int32 ValueIdx = 0; ValueIdx < ValueNum; ValueIdx++)
+			{
+				Out[ValueIdx] = Lhs[ValueIdx] < Rhs;
+			}
+		}
+
+		static inline void LeConstant(TLearningArrayView<1, float> Out, const TLearningArrayView<1, const float> Lhs, const float Rhs)
+		{
+			check(Out.Num() == Lhs.Num());
+			const int32 ValueNum = Out.Num();
+			for (int32 ValueIdx = 0; ValueIdx < ValueNum; ValueIdx++)
+			{
+				Out[ValueIdx] = Lhs[ValueIdx] <= Rhs;
+			}
+		}
+
+		static inline void EqConstant(TLearningArrayView<1, float> Out, const TLearningArrayView<1, const float> Lhs, const float Rhs)
+		{
+			check(Out.Num() == Lhs.Num());
+			const int32 ValueNum = Out.Num();
+			for (int32 ValueIdx = 0; ValueIdx < ValueNum; ValueIdx++)
+			{
+				Out[ValueIdx] = Lhs[ValueIdx] == Rhs;
+			}
+		}
+
+		static inline void NeqConstant(TLearningArrayView<1, float> Out, const TLearningArrayView<1, const float> Lhs, const float Rhs)
+		{
+			check(Out.Num() == Lhs.Num());
+			const int32 ValueNum = Out.Num();
+			for (int32 ValueIdx = 0; ValueIdx < ValueNum; ValueIdx++)
+			{
+				Out[ValueIdx] = Lhs[ValueIdx] != Rhs;
+			}
+		}
+
+		static inline void ConstantGt(TLearningArrayView<1, float> Out, const float Lhs, const TLearningArrayView<1, const float> Rhs)
+		{
+			check(Out.Num() == Rhs.Num());
+			const int32 ValueNum = Out.Num();
+			for (int32 ValueIdx = 0; ValueIdx < ValueNum; ValueIdx++)
+			{
+				Out[ValueIdx] = Lhs > Rhs[ValueIdx];
+			}
+		}
+
+		static inline void ConstantGe(TLearningArrayView<1, float> Out, const float Lhs, const TLearningArrayView<1, const float> Rhs)
+		{
+			check(Out.Num() == Rhs.Num());
+			const int32 ValueNum = Out.Num();
+			for (int32 ValueIdx = 0; ValueIdx < ValueNum; ValueIdx++)
+			{
+				Out[ValueIdx] = Lhs >= Rhs[ValueIdx];
+			}
+		}
+
+		static inline void ConstantLt(TLearningArrayView<1, float> Out, const float Lhs, const TLearningArrayView<1, const float> Rhs)
+		{
+			check(Out.Num() == Rhs.Num());
+			const int32 ValueNum = Out.Num();
+			for (int32 ValueIdx = 0; ValueIdx < ValueNum; ValueIdx++)
+			{
+				Out[ValueIdx] = Lhs < Rhs[ValueIdx];
+			}
+		}
+
+		static inline void ConstantLe(TLearningArrayView<1, float> Out, const float Lhs, const TLearningArrayView<1, const float> Rhs)
+		{
+			check(Out.Num() == Rhs.Num());
+			const int32 ValueNum = Out.Num();
+			for (int32 ValueIdx = 0; ValueIdx < ValueNum; ValueIdx++)
+			{
+				Out[ValueIdx] = Lhs < Rhs[ValueIdx];
+			}
+		}
+
+		static inline void ConstantEq(TLearningArrayView<1, float> Out, const float Lhs, const TLearningArrayView<1, const float> Rhs)
+		{
+			check(Out.Num() == Rhs.Num());
+			const int32 ValueNum = Out.Num();
+			for (int32 ValueIdx = 0; ValueIdx < ValueNum; ValueIdx++)
+			{
+				Out[ValueIdx] = Lhs == Rhs[ValueIdx];
+			}
+		}
+
+		static inline void ConstantNeq(TLearningArrayView<1, float> Out, const float Lhs, const TLearningArrayView<1, const float> Rhs)
+		{
+			check(Out.Num() == Rhs.Num());
+			const int32 ValueNum = Out.Num();
+			for (int32 ValueIdx = 0; ValueIdx < ValueNum; ValueIdx++)
+			{
+				Out[ValueIdx] = Lhs != Rhs[ValueIdx];
+			}
+		}
+
+		static inline void FilterGaussian(TLearningArrayView<1, float> Out, const TLearningArrayView<1, const float> In, const float StdInFrames)
+		{
+			check(Out.Num() == In.Num());
+			check(StdInFrames >= 0.0f);
+
+#if UE_LEARNING_ISPC
+
+			ispc::LearningFilterGaussian(Out.GetData(), In.GetData(), In.Num(), StdInFrames);
+
+#else
+			const int32 StdRange = FMath::RoundToInt(StdInFrames * 3.0f);
+
+			const int32 ValueNum = Out.Num();
+			for (int32 ValueIdx = 0; ValueIdx < ValueNum; ValueIdx++)
+			{
+				float Total = 0.0f;
+				Out[ValueIdx] = 0.0f;
+
+				const int32 RangeMin = FMath::Max(ValueIdx - StdRange, 0);
+				const int32 RangeMax = FMath::Min(ValueIdx + StdRange, ValueNum - 1);
+
+				for (int32 Offset = -StdRange; Offset <= +StdRange; Offset++)
+				{
+					if (ValueIdx + Offset >= 0 && ValueIdx + Offset < ValueNum)
+					{
+						const float Weight = FMath::InvExpApprox(FMath::Square(Offset / StdInFrames));
+						Out[ValueIdx] += Weight * In[ValueIdx + Offset];
+						Total += Weight;
+					}
+				}
+
+				Out[ValueIdx] = Out[ValueIdx] / Total;
+			}
+#endif
+		}
+
+		static inline void FilterMajorityVote(TLearningArrayView<1, float> Out, const TLearningArrayView<1, const float> In, const float FilterWidthInFrames)
+		{
+			check(Out.Num() == In.Num());
+			check(FilterWidthInFrames >= 0.0f);
+
+#if UE_LEARNING_ISPC
+
+			ispc::LearningFilterMajorityVote(Out.GetData(), In.GetData(), In.Num(), FilterWidthInFrames);
+
+#else
+			const int32 FilterRange = FMath::RoundToInt(FilterWidthInFrames);
+
+			const int32 ValueNum = Out.Num();
+			for (int32 ValueIdx = 0; ValueIdx < ValueNum; ValueIdx++)
+			{
+				int32 Total = 0;
+
+				const int32 RangeMin = FMath::Max(ValueIdx - FilterRange, 0);
+				const int32 RangeMax = FMath::Min(ValueIdx + FilterRange, ValueNum - 1);
+
+				for (int32 Offset = -FilterRange; Offset <= +FilterRange; Offset++)
+				{
+					if (ValueIdx + Offset >= 0 && ValueIdx + Offset < ValueNum)
+					{
+						Total += In[ValueIdx + Offset] ? 1 : -1;
+					}
+				}
+
+				Out[ValueIdx] = Total > 0;
+			}
+#endif
+		}
+
+		static inline void QuatInv(TLearningArrayView<2, float> Out, const TLearningArrayView<2, const float> In)
+		{
+			check(Out.Num<0>() == 4);
+			check(In.Num<0>() == 4);
+			check(Out.Num<1>() == In.Num<1>());
+
+			const int32 ValueNum = Out.Num<1>();
+			for (int32 ValueIdx = 0; ValueIdx < ValueNum; ValueIdx++)
+			{
+				const FQuat4f InValue = FQuat4f(In[0][ValueIdx], In[1][ValueIdx], In[2][ValueIdx], In[3][ValueIdx]);
+				const FQuat4f OutValue = InValue.Inverse();
+				Out[0][ValueIdx] = OutValue.X;
+				Out[1][ValueIdx] = OutValue.Y;
+				Out[2][ValueIdx] = OutValue.Z;
+				Out[3][ValueIdx] = OutValue.W;
+			}
+		}
+
+		static inline void QuatAbs(TLearningArrayView<2, float> Out, const TLearningArrayView<2, const float> In)
+		{
+			check(Out.Num<0>() == 4);
+			check(In.Num<0>() == 4);
+			check(Out.Num<1>() == In.Num<1>());
+
+			const int32 ValueNum = Out.Num<1>();
+			for (int32 ValueIdx = 0; ValueIdx < ValueNum; ValueIdx++)
+			{
+				const FQuat4f InValue = FQuat4f(In[0][ValueIdx], In[1][ValueIdx], In[2][ValueIdx], In[3][ValueIdx]);
+				const FQuat4f OutValue = InValue.GetShortestArcWith(FQuat4f::Identity);
+				Out[0][ValueIdx] = OutValue.X;
+				Out[1][ValueIdx] = OutValue.Y;
+				Out[2][ValueIdx] = OutValue.Z;
+				Out[3][ValueIdx] = OutValue.W;
+			}
+		}
+
+		static inline void QuatMul(
+			TLearningArrayView<1, float> OutX,
+			TLearningArrayView<1, float> OutY,
+			TLearningArrayView<1, float> OutZ,
+			TLearningArrayView<1, float> OutW,
+			const TLearningArrayView<1, const float> LhsX,
+			const TLearningArrayView<1, const float> LhsY,
+			const TLearningArrayView<1, const float> LhsZ,
+			const TLearningArrayView<1, const float> LhsW,
+			const TLearningArrayView<1, const float> RhsX,
+			const TLearningArrayView<1, const float> RhsY,
+			const TLearningArrayView<1, const float> RhsZ,
+			const TLearningArrayView<1, const float> RhsW)
+		{
+			const int32 ValueNum = OutX.Num();
+			for (int32 ValueIdx = 0; ValueIdx < ValueNum; ValueIdx++)
+			{
+				const FQuat4f Lhs = FQuat4f(LhsX[ValueIdx], LhsY[ValueIdx], LhsZ[ValueIdx], LhsW[ValueIdx]);
+				const FQuat4f Rhs = FQuat4f(RhsX[ValueIdx], RhsY[ValueIdx], RhsZ[ValueIdx], RhsW[ValueIdx]);
+				const FQuat4f Out = Lhs * Rhs;
+				OutX[ValueIdx] = Out.X;
+				OutY[ValueIdx] = Out.Y;
+				OutZ[ValueIdx] = Out.Z;
+				OutW[ValueIdx] = Out.W;
+			}
+		}
+
+		static inline void QuatMulInv(
+			TLearningArrayView<1, float> OutX,
+			TLearningArrayView<1, float> OutY,
+			TLearningArrayView<1, float> OutZ,
+			TLearningArrayView<1, float> OutW,
+			const TLearningArrayView<1, const float> LhsX,
+			const TLearningArrayView<1, const float> LhsY,
+			const TLearningArrayView<1, const float> LhsZ,
+			const TLearningArrayView<1, const float> LhsW,
+			const TLearningArrayView<1, const float> RhsX,
+			const TLearningArrayView<1, const float> RhsY,
+			const TLearningArrayView<1, const float> RhsZ,
+			const TLearningArrayView<1, const float> RhsW)
+		{
+			const int32 ValueNum = OutX.Num();
+			for (int32 ValueIdx = 0; ValueIdx < ValueNum; ValueIdx++)
+			{
+				const FQuat4f Lhs = FQuat4f(LhsX[ValueIdx], LhsY[ValueIdx], LhsZ[ValueIdx], LhsW[ValueIdx]);
+				const FQuat4f Rhs = FQuat4f(RhsX[ValueIdx], RhsY[ValueIdx], RhsZ[ValueIdx], RhsW[ValueIdx]);
+				const FQuat4f Out = Lhs * Rhs.Inverse();
+				OutX[ValueIdx] = Out.X;
+				OutY[ValueIdx] = Out.Y;
+				OutZ[ValueIdx] = Out.Z;
+				OutW[ValueIdx] = Out.W;
+			}
+		}
+
+		static inline void QuatInvMul(
+			TLearningArrayView<1, float> OutX,
+			TLearningArrayView<1, float> OutY,
+			TLearningArrayView<1, float> OutZ,
+			TLearningArrayView<1, float> OutW,
+			const TLearningArrayView<1, const float> LhsX,
+			const TLearningArrayView<1, const float> LhsY,
+			const TLearningArrayView<1, const float> LhsZ,
+			const TLearningArrayView<1, const float> LhsW,
+			const TLearningArrayView<1, const float> RhsX,
+			const TLearningArrayView<1, const float> RhsY,
+			const TLearningArrayView<1, const float> RhsZ,
+			const TLearningArrayView<1, const float> RhsW)
+		{
+			const int32 ValueNum = OutX.Num();
+			for (int32 ValueIdx = 0; ValueIdx < ValueNum; ValueIdx++)
+			{
+				const FQuat4f Lhs = FQuat4f(LhsX[ValueIdx], LhsY[ValueIdx], LhsZ[ValueIdx], LhsW[ValueIdx]);
+				const FQuat4f Rhs = FQuat4f(RhsX[ValueIdx], RhsY[ValueIdx], RhsZ[ValueIdx], RhsW[ValueIdx]);
+				const FQuat4f Out = Lhs.Inverse() * Rhs;
+				OutX[ValueIdx] = Out.X;
+				OutY[ValueIdx] = Out.Y;
+				OutZ[ValueIdx] = Out.Z;
+				OutW[ValueIdx] = Out.W;
+			}
+		}
+
+		static inline void QuatRotate(
+			TLearningArrayView<1, float> OutX,
+			TLearningArrayView<1, float> OutY,
+			TLearningArrayView<1, float> OutZ,
+			const TLearningArrayView<1, const float> LhsX,
+			const TLearningArrayView<1, const float> LhsY,
+			const TLearningArrayView<1, const float> LhsZ,
+			const TLearningArrayView<1, const float> LhsW,
+			const TLearningArrayView<1, const float> RhsX,
+			const TLearningArrayView<1, const float> RhsY,
+			const TLearningArrayView<1, const float> RhsZ)
+		{
+			const int32 ValueNum = OutX.Num();
+			for (int32 ValueIdx = 0; ValueIdx < ValueNum; ValueIdx++)
+			{
+				const FQuat4f Lhs = FQuat4f(LhsX[ValueIdx], LhsY[ValueIdx], LhsZ[ValueIdx], LhsW[ValueIdx]);
+				const FVector3f Rhs = FVector3f(RhsX[ValueIdx], RhsY[ValueIdx], RhsZ[ValueIdx]);
+				const FVector3f Out = Lhs.RotateVector(Rhs);
+				OutX[ValueIdx] = Out.X;
+				OutY[ValueIdx] = Out.Y;
+				OutZ[ValueIdx] = Out.Z;
+			}
+		}
+
+		static inline void QuatUnrotate(
+			TLearningArrayView<1, float> OutX,
+			TLearningArrayView<1, float> OutY,
+			TLearningArrayView<1, float> OutZ,
+			const TLearningArrayView<1, const float> LhsX,
+			const TLearningArrayView<1, const float> LhsY,
+			const TLearningArrayView<1, const float> LhsZ,
+			const TLearningArrayView<1, const float> LhsW,
+			const TLearningArrayView<1, const float> RhsX,
+			const TLearningArrayView<1, const float> RhsY,
+			const TLearningArrayView<1, const float> RhsZ)
+		{
+			const int32 ValueNum = OutX.Num();
+			for (int32 ValueIdx = 0; ValueIdx < ValueNum; ValueIdx++)
+			{
+				const FQuat4f Lhs = FQuat4f(LhsX[ValueIdx], LhsY[ValueIdx], LhsZ[ValueIdx], LhsW[ValueIdx]);
+				const FVector3f Rhs = FVector3f(RhsX[ValueIdx], RhsY[ValueIdx], RhsZ[ValueIdx]);
+				const FVector3f Out = Lhs.UnrotateVector(Rhs);
+				OutX[ValueIdx] = Out.X;
+				OutY[ValueIdx] = Out.Y;
+				OutZ[ValueIdx] = Out.Z;
+			}
+		}
+
+		static inline void QuatBetween(
+			TLearningArrayView<1, float> OutX,
+			TLearningArrayView<1, float> OutY,
+			TLearningArrayView<1, float> OutZ,
+			TLearningArrayView<1, float> OutW,
+			const TLearningArrayView<1, const float> LhsX,
+			const TLearningArrayView<1, const float> LhsY,
+			const TLearningArrayView<1, const float> LhsZ,
+			const TLearningArrayView<1, const float> RhsX,
+			const TLearningArrayView<1, const float> RhsY,
+			const TLearningArrayView<1, const float> RhsZ)
+		{
+			const int32 ValueNum = OutX.Num();
+			for (int32 ValueIdx = 0; ValueIdx < ValueNum; ValueIdx++)
+			{
+				const FVector3f Lhs = FVector3f(LhsX[ValueIdx], LhsY[ValueIdx], LhsZ[ValueIdx]);
+				const FVector3f Rhs = FVector3f(RhsX[ValueIdx], RhsY[ValueIdx], RhsZ[ValueIdx]);
+				const FQuat4f Out = FQuat4f::FindBetween(Lhs, Rhs);
+				OutX[ValueIdx] = Out.X;
+				OutY[ValueIdx] = Out.Y;
+				OutZ[ValueIdx] = Out.Z;
+				OutW[ValueIdx] = Out.W;
+			}
+		}
+
+		static inline void QuatMulConstant(TLearningArrayView<2, float> Out, const TLearningArrayView<2, const float> Lhs, const FQuat4f Rhs)
+		{
+			check(Out.Num<0>() == 4);
+			check(Lhs.Num<0>() == 4);
+			check(Out.Num<1>() == Lhs.Num<1>());
+
+			const int32 ValueNum = Out.Num<1>();
+			for (int32 ValueIdx = 0; ValueIdx < ValueNum; ValueIdx++)
+			{
+				const FQuat4f LhsValue = FQuat4f(Lhs[0][ValueIdx], Lhs[1][ValueIdx], Lhs[2][ValueIdx], Lhs[3][ValueIdx]);
+				const FQuat4f OutValue = LhsValue * Rhs;
+				Out[0][ValueIdx] = OutValue.X;
+				Out[1][ValueIdx] = OutValue.Y;
+				Out[2][ValueIdx] = OutValue.Z;
+				Out[3][ValueIdx] = OutValue.W;
+			}
+		}
+
+		static inline void QuatInvMulConstant(TLearningArrayView<2, float> Out, const TLearningArrayView<2, const float> Lhs, const FQuat4f Rhs)
+		{
+			check(Out.Num<0>() == 4);
+			check(Lhs.Num<0>() == 4);
+			check(Out.Num<1>() == Lhs.Num<1>());
+
+			const int32 ValueNum = Out.Num<1>();
+			for (int32 ValueIdx = 0; ValueIdx < ValueNum; ValueIdx++)
+			{
+				const FQuat4f LhsValue = FQuat4f(Lhs[0][ValueIdx], Lhs[1][ValueIdx], Lhs[2][ValueIdx], Lhs[3][ValueIdx]);
+				const FQuat4f OutValue = LhsValue.Inverse() * Rhs;
+				Out[0][ValueIdx] = OutValue.X;
+				Out[1][ValueIdx] = OutValue.Y;
+				Out[2][ValueIdx] = OutValue.Z;
+				Out[3][ValueIdx] = OutValue.W;
+			}
+		}
+
+		static inline void QuatMulInvConstant(TLearningArrayView<2, float> Out, const TLearningArrayView<2, const float> Lhs, const FQuat4f Rhs)
+		{
+			check(Out.Num<0>() == 4);
+			check(Lhs.Num<0>() == 4);
+			check(Out.Num<1>() == Lhs.Num<1>());
+
+			const int32 ValueNum = Out.Num<1>();
+			for (int32 ValueIdx = 0; ValueIdx < ValueNum; ValueIdx++)
+			{
+				const FQuat4f LhsValue = FQuat4f(Lhs[0][ValueIdx], Lhs[1][ValueIdx], Lhs[2][ValueIdx], Lhs[3][ValueIdx]);
+				const FQuat4f OutValue = LhsValue * Rhs.Inverse();
+				Out[0][ValueIdx] = OutValue.X;
+				Out[1][ValueIdx] = OutValue.Y;
+				Out[2][ValueIdx] = OutValue.Z;
+				Out[3][ValueIdx] = OutValue.W;
+			}
+		}
+
+		static inline void QuatRotateConstant(TLearningArrayView<2, float> Out, const TLearningArrayView<2, const float> Lhs, const FVector3f Rhs)
+		{
+			check(Out.Num<0>() == 3);
+			check(Lhs.Num<0>() == 4);
+			check(Out.Num<1>() == Lhs.Num<1>());
+
+			const int32 ValueNum = Out.Num<1>();
+			for (int32 ValueIdx = 0; ValueIdx < ValueNum; ValueIdx++)
+			{
+				const FQuat4f LhsValue = FQuat4f(Lhs[0][ValueIdx], Lhs[1][ValueIdx], Lhs[2][ValueIdx], Lhs[3][ValueIdx]);
+				const FVector3f OutValue = LhsValue.RotateVector(Rhs);
+				Out[0][ValueIdx] = OutValue.X;
+				Out[1][ValueIdx] = OutValue.Y;
+				Out[2][ValueIdx] = OutValue.Z;
+			}
+		}
+
+		static inline void QuatUnrotateConstant(TLearningArrayView<2, float> Out, const TLearningArrayView<2, const float> Lhs, const FVector3f Rhs)
+		{
+			check(Out.Num<0>() == 3);
+			check(Lhs.Num<0>() == 4);
+			check(Out.Num<1>() == Lhs.Num<1>());
+
+			const int32 ValueNum = Out.Num<1>();
+			for (int32 ValueIdx = 0; ValueIdx < ValueNum; ValueIdx++)
+			{
+				const FQuat4f LhsValue = FQuat4f(Lhs[0][ValueIdx], Lhs[1][ValueIdx], Lhs[2][ValueIdx], Lhs[3][ValueIdx]);
+				const FVector3f OutValue = LhsValue.RotateVector(Rhs);
+				Out[0][ValueIdx] = OutValue.X;
+				Out[1][ValueIdx] = OutValue.Y;
+				Out[2][ValueIdx] = OutValue.Z;
+			}
+		}
+
+		static inline void QuatBetweenConstant(TLearningArrayView<2, float> Out, const TLearningArrayView<2, const float> Lhs, const FVector3f Rhs)
+		{
+			check(Out.Num<0>() == 4);
+			check(Lhs.Num<0>() == 3);
+			check(Out.Num<1>() == Lhs.Num<1>());
+
+			const int32 ValueNum = Out.Num<1>();
+			for (int32 ValueIdx = 0; ValueIdx < ValueNum; ValueIdx++)
+			{
+				const FVector3f LhsValue = FVector3f(Lhs[0][ValueIdx], Lhs[1][ValueIdx], Lhs[2][ValueIdx]);
+				const FQuat4f OutValue = FQuat4f::FindBetween(LhsValue, Rhs);
+				Out[0][ValueIdx] = OutValue.X;
+				Out[1][ValueIdx] = OutValue.Y;
+				Out[2][ValueIdx] = OutValue.Z;
+				Out[3][ValueIdx] = OutValue.W;
+			}
+		}
+
+		static inline void QuatConstantMul(TLearningArrayView<2, float> Out, const FQuat4f Lhs, const TLearningArrayView<2, const float> Rhs)
+		{
+			check(Out.Num<0>() == 4);
+			check(Rhs.Num<0>() == 4);
+			check(Out.Num<1>() == Rhs.Num<1>());
+
+			const int32 ValueNum = Out.Num<1>();
+			for (int32 ValueIdx = 0; ValueIdx < ValueNum; ValueIdx++)
+			{
+				const FQuat4f RhsValue = FQuat4f(Rhs[0][ValueIdx], Rhs[1][ValueIdx], Rhs[2][ValueIdx], Rhs[3][ValueIdx]);
+				const FQuat4f OutValue = Lhs * RhsValue;
+				Out[0][ValueIdx] = OutValue.X;
+				Out[1][ValueIdx] = OutValue.Y;
+				Out[2][ValueIdx] = OutValue.Z;
+				Out[3][ValueIdx] = OutValue.W;
+			}
+		}
+
+		static inline void QuatConstantInvMul(TLearningArrayView<2, float> Out, const FQuat4f Lhs, const TLearningArrayView<2, const float> Rhs)
+		{
+			check(Out.Num<0>() == 4);
+			check(Rhs.Num<0>() == 4);
+			check(Out.Num<1>() == Rhs.Num<1>());
+
+			const int32 ValueNum = Out.Num<1>();
+			for (int32 ValueIdx = 0; ValueIdx < ValueNum; ValueIdx++)
+			{
+				const FQuat4f RhsValue = FQuat4f(Rhs[0][ValueIdx], Rhs[1][ValueIdx], Rhs[2][ValueIdx], Rhs[3][ValueIdx]);
+				const FQuat4f OutValue = Lhs.Inverse() * RhsValue;
+				Out[0][ValueIdx] = OutValue.X;
+				Out[1][ValueIdx] = OutValue.Y;
+				Out[2][ValueIdx] = OutValue.Z;
+				Out[3][ValueIdx] = OutValue.W;
+			}
+		}
+
+		static inline void QuatConstantMulInv(TLearningArrayView<2, float> Out, const FQuat4f Lhs, const TLearningArrayView<2, const float> Rhs)
+		{
+			check(Out.Num<0>() == 4);
+			check(Rhs.Num<0>() == 4);
+			check(Out.Num<1>() == Rhs.Num<1>());
+
+			const int32 ValueNum = Out.Num<1>();
+			for (int32 ValueIdx = 0; ValueIdx < ValueNum; ValueIdx++)
+			{
+				const FQuat4f RhsValue = FQuat4f(Rhs[0][ValueIdx], Rhs[1][ValueIdx], Rhs[2][ValueIdx], Rhs[3][ValueIdx]);
+				const FQuat4f OutValue = Lhs * RhsValue.Inverse();
+				Out[0][ValueIdx] = OutValue.X;
+				Out[1][ValueIdx] = OutValue.Y;
+				Out[2][ValueIdx] = OutValue.Z;
+				Out[3][ValueIdx] = OutValue.W;
+			}
+		}
+
+		static inline void QuatConstantRotate(TLearningArrayView<2, float> Out, const FQuat4f Lhs, const TLearningArrayView<2, const float> Rhs)
+		{
+			check(Out.Num<0>() == 3);
+			check(Rhs.Num<0>() == 3);
+			check(Out.Num<1>() == Rhs.Num<1>());
+
+			const int32 ValueNum = Out.Num<1>();
+			for (int32 ValueIdx = 0; ValueIdx < ValueNum; ValueIdx++)
+			{
+				const FVector3f RhsValue = FVector3f(Rhs[0][ValueIdx], Rhs[1][ValueIdx], Rhs[2][ValueIdx]);
+				const FVector3f OutValue = Lhs.RotateVector(RhsValue);
+				Out[0][ValueIdx] = OutValue.X;
+				Out[1][ValueIdx] = OutValue.Y;
+				Out[2][ValueIdx] = OutValue.Z;
+			}
+		}
+
+		static inline void QuatConstantUnrotate(TLearningArrayView<2, float> Out, const FQuat4f Lhs, const TLearningArrayView<2, const float> Rhs)
+		{
+			check(Out.Num<0>() == 3);
+			check(Rhs.Num<0>() == 3);
+			check(Out.Num<1>() == Rhs.Num<1>());
+
+			const int32 ValueNum = Out.Num<1>();
+			for (int32 ValueIdx = 0; ValueIdx < ValueNum; ValueIdx++)
+			{
+				const FVector3f RhsValue = FVector3f(Rhs[0][ValueIdx], Rhs[1][ValueIdx], Rhs[2][ValueIdx]);
+				const FVector3f OutValue = Lhs.UnrotateVector(RhsValue);
+				Out[0][ValueIdx] = OutValue.X;
+				Out[1][ValueIdx] = OutValue.Y;
+				Out[2][ValueIdx] = OutValue.Z;
+			}
+		}
+
+		static inline void QuatConstantBetween(TLearningArrayView<2, float> Out, const FVector3f Lhs, const TLearningArrayView<2, const float> Rhs)
+		{
+			check(Out.Num<0>() == 4);
+			check(Rhs.Num<0>() == 3);
+			check(Out.Num<1>() == Rhs.Num<1>());
+
+			const int32 ValueNum = Out.Num<1>();
+			for (int32 ValueIdx = 0; ValueIdx < ValueNum; ValueIdx++)
+			{
+				const FVector3f RhsValue = FVector3f(Rhs[0][ValueIdx], Rhs[1][ValueIdx], Rhs[2][ValueIdx]);
+				const FQuat4f OutValue = FQuat4f::FindBetween(Lhs, RhsValue);
+				Out[0][ValueIdx] = OutValue.X;
+				Out[1][ValueIdx] = OutValue.Y;
+				Out[2][ValueIdx] = OutValue.Z;
+				Out[3][ValueIdx] = OutValue.W;
+			}
+		}
+
+	}
+
+	void FFrameAttribute::Check() const
+	{
+		FrameRangeSet.Check();
+		check(AttributeData.Num<1>() == FrameRangeSet.GetTotalFrameNum());
+	}
+
+	void FFrameAttribute::Empty() { FrameRangeSet.Empty(); AttributeData.Empty(); }
+	bool FFrameAttribute::IsEmpty() const { return FrameRangeSet.IsEmpty(); }
+
+	const FFrameRangeSet& FFrameAttribute::GetFrameRangeSet() const { return FrameRangeSet; }
+	int32 FFrameAttribute::GetTotalFrameNum() const { return AttributeData.Num<1>(); }
+	int32 FFrameAttribute::GetTotalRangeNum() const { return FrameRangeSet.GetTotalRangeNum(); }
+	int32 FFrameAttribute::GetChannelNum() const { return AttributeData.Num<0>(); }
+
+	TLearningArrayView<2, const float> FFrameAttribute::GetAttributeData() const { return AttributeData; }
+
+	TLearningArrayView<1, const float> FFrameAttribute::GetChannelAttributeData(const int32 ChannelIdx) const { return AttributeData[ChannelIdx]; }
+
+	const float& FFrameAttribute::GetChannelAttributeDataAtFrame(const int32 ChannelIdx, const int32 RangeFrameIdx) const { return AttributeData[ChannelIdx][RangeFrameIdx]; }
+
+	TLearningArrayView<1, const float> FFrameAttribute::GetChannelEntryRangeAttributeData(const int32 ChannelIdx, const int32 EntryIdx, const int32 RangeIdx) const
+	{
+		return AttributeData[ChannelIdx].Slice(FrameRangeSet.GetEntryRangeOffset(EntryIdx, RangeIdx), FrameRangeSet.GetEntryRangeLength(EntryIdx, RangeIdx));
+	}
+
+	TLearningArrayView<1, const float> FFrameAttribute::GetChannelRangeAttributeData(const int32 ChannelIdx, const int32 RangeOffset, const int32 RangeLength) const
+	{
+		return AttributeData[ChannelIdx].Slice(RangeOffset, RangeLength);
+	}
+
+	TLearningArrayView<2, float> FFrameAttribute::GetAttributeData() { return AttributeData; }
+
+	TLearningArrayView<1, float> FFrameAttribute::GetChannelAttributeData(const int32 ChannelIdx) { return AttributeData[ChannelIdx]; }
+
+	float& FFrameAttribute::GetChannelAttributeDataAtFrame(const int32 ChannelIdx, const int32 RangeFrameIdx) { return AttributeData[ChannelIdx][RangeFrameIdx]; }
+
+	TLearningArrayView<1, float> FFrameAttribute::GetChannelEntryRangeAttributeData(const int32 ChannelIdx, const int32 EntryIdx, const int32 RangeIdx)
+	{
+		return AttributeData[ChannelIdx].Slice(FrameRangeSet.GetEntryRangeOffset(EntryIdx, RangeIdx), FrameRangeSet.GetEntryRangeLength(EntryIdx, RangeIdx));
+	}
+
+	TLearningArrayView<1, float> FFrameAttribute::GetChannelRangeAttributeData(const int32 ChannelIdx, const int32 RangeOffset, const int32 RangeLength)
+	{
+		return AttributeData[ChannelIdx].Slice(RangeOffset, RangeLength);
+	}
+
+	namespace FrameAttribute
+	{
+		void Intersection(FFrameAttribute& OutFrameAttribute, const FFrameAttribute& FrameAttribute, const FFrameRangeSet& FrameRangeSet)
+		{
+			if (FrameRangeSet::Equal(FrameAttribute.FrameRangeSet, FrameRangeSet))
+			{
+				OutFrameAttribute = FrameAttribute;
+				return;
+			}
+
+			// Perform Intersection and get offsets
+			TLearningArray<1, int32> LhsOffsets;
+			TLearningArray<1, int32> RhsOffsets;
+			LhsOffsets.SetNumUninitialized({ FrameAttribute.FrameRangeSet.GetTotalRangeNum() + FrameRangeSet.GetTotalRangeNum() });
+			RhsOffsets.SetNumUninitialized({ FrameAttribute.FrameRangeSet.GetTotalRangeNum() + FrameRangeSet.GetTotalRangeNum() });
+
+			const int32 OutTotalRangeNum = UE::Learning::FrameRangeSet::IntersectionWithOffsets(
+				OutFrameAttribute.FrameRangeSet,
+				LhsOffsets,
+				RhsOffsets,
+				FrameAttribute.FrameRangeSet,
+				FrameRangeSet);
+
+			// Resize back to correct size
+			LhsOffsets.SetNumUninitialized({ OutTotalRangeNum });
+			RhsOffsets.SetNumUninitialized({ OutTotalRangeNum });
+
+			const int32 ChannelNum = FrameAttribute.GetChannelNum();
+			const int32 TotalFrameNum = OutFrameAttribute.FrameRangeSet.GetTotalFrameNum();
+
+			OutFrameAttribute.AttributeData.SetNumUninitialized({ ChannelNum, TotalFrameNum });
+
+			for (int32 ChannelIdx = 0; ChannelIdx < ChannelNum; ChannelIdx++)
+			{
+				for (int32 RangeIdx = 0; RangeIdx < OutTotalRangeNum; RangeIdx++)
+				{
+					const int32 OutOffset = OutFrameAttribute.FrameRangeSet.GetAllRangeOffsets()[RangeIdx];
+					const int32 LhsOffset = LhsOffsets[RangeIdx];
+					const int32 Length = OutFrameAttribute.FrameRangeSet.GetAllRangeLengths()[RangeIdx];
+
+					Array::Copy(
+						OutFrameAttribute.GetChannelRangeAttributeData(ChannelIdx, OutOffset, Length),
+						FrameAttribute.GetChannelRangeAttributeData(ChannelIdx, LhsOffset, Length));
+				}
+			}
+		}
+
+		void NonZeroFrameRangeSet(FFrameRangeSet& OutFrameRangeSet, const FFrameAttribute& FrameAttribute, const int32 ChannelIdx)
+		{
+			check(ChannelIdx >= 0 && ChannelIdx < FrameAttribute.GetChannelNum());
+
+			const int32 EntryNum = FrameAttribute.FrameRangeSet.GetEntryNum();
+
+			OutFrameRangeSet.Empty();
+
+			TArray<int32> AddedRangeStarts;
+			TArray<int32> AddedRangeLengths;
+
+			int32 RangeOffset = 0;
+			for (int32 EntryIdx = 0; EntryIdx < EntryNum; EntryIdx++)
+			{
+				const int32 RangeNum = FrameAttribute.FrameRangeSet.GetEntryRangeNum(EntryIdx);
+				const int32 Sequence = FrameAttribute.FrameRangeSet.GetEntrySequence(EntryIdx);
+
+				AddedRangeStarts.Reset();
+				AddedRangeLengths.Reset();
+
+				for (int32 RangeIdx = 0; RangeIdx < RangeNum; RangeIdx++)
+				{
+					const int32 FrameNum = FrameAttribute.FrameRangeSet.GetEntryRangeLength(EntryIdx, RangeIdx);
+					const int32 StartFrame = FrameAttribute.FrameRangeSet.GetEntryRangeStart(EntryIdx, RangeIdx);
+					const int32 FrameOffset = FrameAttribute.FrameRangeSet.GetEntryRangeOffset(EntryIdx, RangeIdx);
+
+					int32 StartFrameIndex = INDEX_NONE;
+
+					for (int32 FrameIdx = 0; FrameIdx < FrameNum; FrameIdx++)
+					{
+						const float FrameValue = FrameAttribute.GetChannelAttributeDataAtFrame(ChannelIdx, FrameOffset + FrameIdx);
+
+						if (StartFrameIndex == INDEX_NONE && FrameValue == 1.0f)
+						{
+							StartFrameIndex = FrameIdx;
+						}
+						else if (StartFrameIndex == INDEX_NONE && FrameValue == 0.0f) {}
+						else if (StartFrameIndex != INDEX_NONE && FrameValue == 1.0f) {}
+						else if (StartFrameIndex != INDEX_NONE && FrameValue == 0.0f)
+						{
+							check(FrameIdx - StartFrameIndex > 0);
+							AddedRangeStarts.Add(StartFrame + StartFrameIndex);
+							AddedRangeLengths.Add(StartFrame + FrameIdx - StartFrameIndex);
+							StartFrameIndex = INDEX_NONE;
+						}
+					}
+
+					if (StartFrameIndex != INDEX_NONE)
+					{
+						check(FrameNum - StartFrameIndex > 0);
+						AddedRangeStarts.Add(StartFrame + StartFrameIndex);
+						AddedRangeLengths.Add(StartFrame + FrameNum - StartFrameIndex);
+					}
+				}
+
+				OutFrameRangeSet.AddEntry(Sequence, AddedRangeStarts, AddedRangeLengths);
+			}
+		}
+
+		void ReduceOp(
+			const FFrameAttribute& In,
+			const ReduceOpFunction Op)
+		{
+			Op(In, In.FrameRangeSet.GetAllRangeOffsets(), In.FrameRangeSet.GetAllRangeLengths());
+		}
+
+		void NullaryOp(
+			FFrameAttribute& Out,
+			const int32 OutChannelNum,
+			const FFrameRangeSet& FrameRangeSet,
+			const NullaryOpFunction Op)
+		{
+			Out.FrameRangeSet = FrameRangeSet;
+			Out.AttributeData.SetNumUninitialized({ OutChannelNum, FrameRangeSet.GetTotalFrameNum() });
+
+			Op(Out, Out.FrameRangeSet.GetAllRangeOffsets(), Out.FrameRangeSet.GetAllRangeLengths());
+		}
+
+		void UnaryOp(
+			FFrameAttribute& Out,
+			const int32 OutChannelNum,
+			const FFrameAttribute& In,
+			const UnaryOpFunction Op)
+		{
+			Out.FrameRangeSet = In.FrameRangeSet;
+			Out.AttributeData.SetNumUninitialized({ OutChannelNum, In.FrameRangeSet.GetTotalFrameNum() });
+
+			Op(Out, In, Out.FrameRangeSet.GetAllRangeOffsets(), Out.FrameRangeSet.GetAllRangeLengths());
+		}
+
+		void BinaryOp(
+			FFrameAttribute& Out,
+			const int32 OutChannelNum,
+			const FFrameAttribute& Lhs,
+			const FFrameAttribute& Rhs,
+			const BinaryOpFunction Op)
+		{
+			// Fast Path for when FrameRangeSets are Equal
+
+			if (FrameRangeSet::Equal(Lhs.FrameRangeSet, Rhs.FrameRangeSet))
+			{
+				Out.FrameRangeSet = Lhs.FrameRangeSet;
+				Out.AttributeData.SetNumUninitialized({ OutChannelNum, Lhs.GetTotalFrameNum() });
+
+				Op(Out, Lhs, Rhs, 
+					Out.FrameRangeSet.GetAllRangeOffsets(),
+					Lhs.FrameRangeSet.GetAllRangeOffsets(),
+					Rhs.FrameRangeSet.GetAllRangeOffsets(),
+					Out.FrameRangeSet.GetAllRangeLengths());
+
+				return;
+			}
+
+			// Slow Path for when FrameRangeSets are not equal and we need to compute the intersection
+
+			TLearningArray<1, int32> LhsRangeOffsets;
+			TLearningArray<1, int32> RhsRangeOffsets;
+			LhsRangeOffsets.SetNumUninitialized({ Lhs.FrameRangeSet.GetTotalRangeNum() + Rhs.FrameRangeSet.GetTotalRangeNum() });
+			RhsRangeOffsets.SetNumUninitialized({ Lhs.FrameRangeSet.GetTotalRangeNum() + Rhs.FrameRangeSet.GetTotalRangeNum() });
+
+			const int32 OutTotalRangeNum = UE::Learning::FrameRangeSet::IntersectionWithOffsets(
+				Out.FrameRangeSet,
+				LhsRangeOffsets,
+				RhsRangeOffsets,
+				Lhs.FrameRangeSet,
+				Rhs.FrameRangeSet);
+
+			// Resize back to correct size
+			LhsRangeOffsets.SetNumUninitialized({ OutTotalRangeNum });
+			RhsRangeOffsets.SetNumUninitialized({ OutTotalRangeNum });
+
+			// Allocate Attribute Data
+			Out.AttributeData.SetNumUninitialized({ OutChannelNum, Out.FrameRangeSet.GetTotalFrameNum() });
+
+			Op(Out, Lhs, Rhs,
+				Out.FrameRangeSet.GetAllRangeOffsets(),
+				LhsRangeOffsets,
+				RhsRangeOffsets,
+				Out.FrameRangeSet.GetAllRangeLengths());
+		}
+
+		void NaryOp(
+			FFrameAttribute& Out,
+			const int32 OutChannelNum,
+			const TArrayView<const ConstFrameAttributePtr> Inputs,
+			const NaryOpFunction Op)
+		{
+			if (Inputs.Num() == 0)
+			{
+				Out.Empty();
+				Op(Out, {}, {}, {}, {});
+			}
+
+			// Check All Equal
+
+			const int32 InputNum = Inputs.Num();
+
+			bool bAllEqual = true;
+			for (int32 InputIdx = 1; InputIdx < InputNum; InputIdx++)
+			{
+				bAllEqual &= FrameRangeSet::Equal(Inputs[0]->FrameRangeSet, Inputs[InputIdx]->FrameRangeSet);
+			}
+
+			if (bAllEqual)
+			{
+				Out.FrameRangeSet = Inputs[0]->FrameRangeSet;
+				Out.AttributeData.SetNumUninitialized({ OutChannelNum, Inputs[0]->FrameRangeSet.GetTotalFrameNum() });
+
+				TArray<TLearningArrayView<1, const int32>, TInlineAllocator<8>> InputRangeOffsetsViews;
+				InputRangeOffsetsViews.SetNumUninitialized(InputNum);
+				for (int32 InputIdx = 0; InputIdx < InputNum; InputIdx++)
+				{
+					InputRangeOffsetsViews[InputIdx] = Inputs[InputIdx]->FrameRangeSet.GetAllRangeOffsets();
+				}
+
+				Op(Out, Inputs,
+					Out.FrameRangeSet.GetAllRangeOffsets(),
+					InputRangeOffsetsViews,
+					Out.FrameRangeSet.GetAllRangeLengths());
+			}
+			else
+			{
+				// Currently we don't support the case where things are not equal.
+
+				checkNoEntry();
+				Out.Empty();
+				return;
+			}
+		}
+
+		bool FindMinimum(
+			int32& OutChannelIdx,
+			int32& OutFrameIdx,
+			float& OutMinimumValue,
+			const FFrameAttribute& In)
+		{
+			return Private::FindMinimum(OutChannelIdx, OutFrameIdx, OutMinimumValue, In.GetAttributeData());
+		}
+
+		bool FindMaximum(
+			int32& OutChannelIdx,
+			int32& OutFrameIdx,
+			float& OutMaximumValue,
+			const FFrameAttribute& In)
+		{
+			return Private::FindMaximum(OutChannelIdx, OutFrameIdx, OutMaximumValue, In.GetAttributeData());
+		}
+
+		void Zeros(FFrameAttribute& Out, const FFrameRangeSet& FrameRangeSet, const int32 ChannelNum)
+		{
+			Out.FrameRangeSet = FrameRangeSet;
+			Out.AttributeData.SetNumUninitialized({ ChannelNum, FrameRangeSet.GetTotalFrameNum() });
+			Array::Zero(Out.AttributeData);
+		}
+
+		void Ones(FFrameAttribute& Out, const FFrameRangeSet& FrameRangeSet, const int32 ChannelNum)
+		{
+			Out.FrameRangeSet = FrameRangeSet;
+			Out.AttributeData.SetNumUninitialized({ ChannelNum, FrameRangeSet.GetTotalFrameNum() });
+			Array::Set(Out.AttributeData, 1.0f);
+		}
+
+		void Fill(FFrameAttribute& Out, const FFrameRangeSet& FrameRangeSet, const TLearningArrayView<1, const float> Values)
+		{
+			const int32 ChannelNum = Values.Num();
+
+			Out.FrameRangeSet = FrameRangeSet;
+			Out.AttributeData.SetNumUninitialized({ ChannelNum, FrameRangeSet.GetTotalFrameNum() });
+			for (int32 ChannelIdx = 0; ChannelIdx < ChannelNum; ChannelIdx++)
+			{
+				Array::Set(Out.AttributeData[ChannelIdx], Values[ChannelIdx]);
+			}
+		}
+
+		void Add(FFrameAttribute& Out, const FFrameAttribute& Lhs, const FFrameAttribute& Rhs)
+		{
+			check(Lhs.GetChannelNum() == Rhs.GetChannelNum());
+
+			BinaryOp(Out, Lhs.GetChannelNum(), Lhs, Rhs, [](
+				FFrameAttribute& Out,
+				const FFrameAttribute& Lhs,
+				const FFrameAttribute& Rhs,
+				const TLearningArrayView<1, const int32> OutRangeOffsets,
+				const TLearningArrayView<1, const int32> LhsRangeOffsets,
+				const TLearningArrayView<1, const int32> RhsRangeOffsets,
+				const TLearningArrayView<1, const int32> RangeLengths) {
+
+					const int32 ChannelNum = Out.GetChannelNum();
+					const int32 RangeNum = RangeLengths.Num();
+					for (int32 ChannelIdx = 0; ChannelIdx < ChannelNum; ChannelIdx++)
+					{
+						for (int32 RangeIdx = 0; RangeIdx < RangeNum; RangeIdx++)
+						{
+							Private::Add(
+								Out.GetChannelRangeAttributeData(ChannelIdx, OutRangeOffsets[RangeIdx], RangeLengths[RangeIdx]),
+								Lhs.GetChannelRangeAttributeData(ChannelIdx, LhsRangeOffsets[RangeIdx], RangeLengths[RangeIdx]),
+								Rhs.GetChannelRangeAttributeData(ChannelIdx, RhsRangeOffsets[RangeIdx], RangeLengths[RangeIdx]));
+						}
+					}
+				});
+		}
+
+		void Sub(FFrameAttribute& Out, const FFrameAttribute& Lhs, const FFrameAttribute& Rhs)
+		{
+			check(Lhs.GetChannelNum() == Rhs.GetChannelNum());
+
+			BinaryOp(Out, Lhs.GetChannelNum(), Lhs, Rhs, [](
+				FFrameAttribute& Out,
+				const FFrameAttribute& Lhs,
+				const FFrameAttribute& Rhs,
+				const TLearningArrayView<1, const int32> OutRangeOffsets,
+				const TLearningArrayView<1, const int32> LhsRangeOffsets,
+				const TLearningArrayView<1, const int32> RhsRangeOffsets,
+				const TLearningArrayView<1, const int32> RangeLengths) {
+
+					const int32 ChannelNum = Out.GetChannelNum();
+					const int32 RangeNum = RangeLengths.Num();
+					for (int32 ChannelIdx = 0; ChannelIdx < ChannelNum; ChannelIdx++)
+					{
+						for (int32 RangeIdx = 0; RangeIdx < RangeNum; RangeIdx++)
+						{
+							Private::Sub(
+								Out.GetChannelRangeAttributeData(ChannelIdx, OutRangeOffsets[RangeIdx], RangeLengths[RangeIdx]),
+								Lhs.GetChannelRangeAttributeData(ChannelIdx, LhsRangeOffsets[RangeIdx], RangeLengths[RangeIdx]),
+								Rhs.GetChannelRangeAttributeData(ChannelIdx, RhsRangeOffsets[RangeIdx], RangeLengths[RangeIdx]));
+						}
+					}
+				});
+		}
+
+		void Mul(FFrameAttribute& Out, const FFrameAttribute& Lhs, const FFrameAttribute& Rhs)
+		{
+			check(Lhs.GetChannelNum() == Rhs.GetChannelNum());
+
+			BinaryOp(Out, Lhs.GetChannelNum(), Lhs, Rhs, [](
+				FFrameAttribute& Out,
+				const FFrameAttribute& Lhs,
+				const FFrameAttribute& Rhs,
+				const TLearningArrayView<1, const int32> OutRangeOffsets,
+				const TLearningArrayView<1, const int32> LhsRangeOffsets,
+				const TLearningArrayView<1, const int32> RhsRangeOffsets,
+				const TLearningArrayView<1, const int32> RangeLengths) {
+
+					const int32 ChannelNum = Out.GetChannelNum();
+					const int32 RangeNum = RangeLengths.Num();
+					for (int32 ChannelIdx = 0; ChannelIdx < ChannelNum; ChannelIdx++)
+					{
+						for (int32 RangeIdx = 0; RangeIdx < RangeNum; RangeIdx++)
+						{
+							Private::Mul(
+								Out.GetChannelRangeAttributeData(ChannelIdx, OutRangeOffsets[RangeIdx], RangeLengths[RangeIdx]),
+								Lhs.GetChannelRangeAttributeData(ChannelIdx, LhsRangeOffsets[RangeIdx], RangeLengths[RangeIdx]),
+								Rhs.GetChannelRangeAttributeData(ChannelIdx, RhsRangeOffsets[RangeIdx], RangeLengths[RangeIdx]));
+						}
+					}
+				});
+		}
+
+		void Div(FFrameAttribute& Out, const FFrameAttribute& Lhs, const FFrameAttribute& Rhs)
+		{
+			check(Lhs.GetChannelNum() == Rhs.GetChannelNum());
+
+			BinaryOp(Out, Lhs.GetChannelNum(), Lhs, Rhs, [](
+				FFrameAttribute& Out,
+				const FFrameAttribute& Lhs,
+				const FFrameAttribute& Rhs,
+				const TLearningArrayView<1, const int32> OutRangeOffsets,
+				const TLearningArrayView<1, const int32> LhsRangeOffsets,
+				const TLearningArrayView<1, const int32> RhsRangeOffsets,
+				const TLearningArrayView<1, const int32> RangeLengths) {
+
+					const int32 ChannelNum = Out.GetChannelNum();
+					const int32 RangeNum = RangeLengths.Num();
+					for (int32 ChannelIdx = 0; ChannelIdx < ChannelNum; ChannelIdx++)
+					{
+						for (int32 RangeIdx = 0; RangeIdx < RangeNum; RangeIdx++)
+						{
+							Private::Div(
+								Out.GetChannelRangeAttributeData(ChannelIdx, OutRangeOffsets[RangeIdx], RangeLengths[RangeIdx]),
+								Lhs.GetChannelRangeAttributeData(ChannelIdx, LhsRangeOffsets[RangeIdx], RangeLengths[RangeIdx]),
+								Rhs.GetChannelRangeAttributeData(ChannelIdx, RhsRangeOffsets[RangeIdx], RangeLengths[RangeIdx]));
+						}
+					}
+				});
+		}
+
+		void Dot(FFrameAttribute& Out, const FFrameAttribute& Lhs, const FFrameAttribute& Rhs)
+		{
+			check(Lhs.GetChannelNum() == Rhs.GetChannelNum());
+
+			BinaryOp(Out, 1, Lhs, Rhs, [](
+				FFrameAttribute& Out,
+				const FFrameAttribute& Lhs,
+				const FFrameAttribute& Rhs,
+				const TLearningArrayView<1, const int32> OutRangeOffsets,
+				const TLearningArrayView<1, const int32> LhsRangeOffsets,
+				const TLearningArrayView<1, const int32> RhsRangeOffsets,
+				const TLearningArrayView<1, const int32> RangeLengths) {
+
+					Array::Zero(Out.GetAttributeData());
+
+					const int32 ChannelNum = Out.GetChannelNum();
+					const int32 RangeNum = RangeLengths.Num();
+
+					for (int32 ChannelIdx = 0; ChannelIdx < ChannelNum; ChannelIdx++)
+					{
+						for (int32 RangeIdx = 0; RangeIdx < RangeNum; RangeIdx++)
+						{
+							Private::Dot(
+								Out.GetChannelRangeAttributeData(0, OutRangeOffsets[RangeIdx], RangeLengths[RangeIdx]),
+								Lhs.GetChannelRangeAttributeData(ChannelIdx, LhsRangeOffsets[RangeIdx], RangeLengths[RangeIdx]),
+								Rhs.GetChannelRangeAttributeData(ChannelIdx, RhsRangeOffsets[RangeIdx], RangeLengths[RangeIdx]));
+						}
+					}
+				});
+		}
+
+		void Neg(FFrameAttribute& Out, const FFrameAttribute& In)
+		{
+			UnaryOp(Out, In.GetChannelNum(), In, [](
+				FFrameAttribute& Out,
+				const FFrameAttribute& In,
+				const TLearningArrayView<1, const int32> RangeOffsets,
+				const TLearningArrayView<1, const int32> RangeLengths) {
+
+					const int32 ChannelNum = Out.GetChannelNum();
+
+					for (int32 ChannelIdx = 0; ChannelIdx < ChannelNum; ChannelIdx++)
+					{
+						Private::Neg(
+							Out.GetChannelAttributeData(ChannelIdx),
+							In.GetChannelAttributeData(ChannelIdx));
+					}
+				});
+		}
+
+		void Inv(FFrameAttribute& Out, const FFrameAttribute& In)
+		{
+			UnaryOp(Out, In.GetChannelNum(), In, [](
+				FFrameAttribute& Out,
+				const FFrameAttribute& In,
+				const TLearningArrayView<1, const int32> RangeOffsets,
+				const TLearningArrayView<1, const int32> RangeLengths) {
+
+					const int32 ChannelNum = Out.GetChannelNum();
+
+					for (int32 ChannelIdx = 0; ChannelIdx < ChannelNum; ChannelIdx++)
+					{
+						Private::Inv(
+							Out.GetChannelAttributeData(ChannelIdx),
+							In.GetChannelAttributeData(ChannelIdx));
+					}
+				});
+		}
+
+		void Abs(FFrameAttribute& Out, const FFrameAttribute& In)
+		{
+			UnaryOp(Out, In.GetChannelNum(), In, [](
+				FFrameAttribute& Out,
+				const FFrameAttribute& In,
+				const TLearningArrayView<1, const int32> RangeOffsets,
+				const TLearningArrayView<1, const int32> RangeLengths) {
+
+					const int32 ChannelNum = Out.GetChannelNum();
+
+					for (int32 ChannelIdx = 0; ChannelIdx < ChannelNum; ChannelIdx++)
+					{
+						Private::Abs(
+							Out.GetChannelAttributeData(ChannelIdx),
+							In.GetChannelAttributeData(ChannelIdx));
+					}
+				});
+		}
+
+		void Length(FFrameAttribute& Out, const FFrameAttribute& In)
+		{
+			UnaryOp(Out, 1, In, [](
+				FFrameAttribute& Out,
+				const FFrameAttribute& In,
+				const TLearningArrayView<1, const int32> RangeOffsets,
+				const TLearningArrayView<1, const int32> RangeLengths) {
+
+					const int32 ChannelNum = Out.GetChannelNum();
+
+					Array::Zero(Out.GetAttributeData());
+
+					for (int32 ChannelIdx = 0; ChannelIdx < ChannelNum; ChannelIdx++)
+					{
+						Private::LengthSquared(
+							Out.GetChannelAttributeData(0),
+							In.GetChannelAttributeData(ChannelIdx));
+					}
+
+					Private::SqrtInplace(Out.GetChannelAttributeData(0));
+				});
+		}
+
+		void Normalize(FFrameAttribute& Out, const FFrameAttribute& In)
+		{
+			UnaryOp(Out, In.GetChannelNum(), In, [](
+				FFrameAttribute& Out,
+				const FFrameAttribute& In,
+				const TLearningArrayView<1, const int32> RangeOffsets,
+				const TLearningArrayView<1, const int32> RangeLengths) {
+
+					Private::Normalize(Out.GetAttributeData(), In.GetAttributeData());
+
+				});
+		}
+
+		void Index(FFrameAttribute& Out, const FFrameAttribute& In, const int32 ChannelIdx)
+		{
+			check(ChannelIdx >= 0);
+			check(ChannelIdx < In.GetChannelNum());
+
+			UnaryOp(Out, 1, In, [ChannelIdx](
+				FFrameAttribute& Out,
+				const FFrameAttribute& In,
+				const TLearningArrayView<1, const int32> RangeOffsets,
+				const TLearningArrayView<1, const int32> RangeLengths) {
+
+					Array::Copy(Out.GetChannelAttributeData(0), In.GetChannelAttributeData(ChannelIdx));
+
+				});
+		}
+
+		void AddConstant(FFrameAttribute& Out, const FFrameAttribute& Lhs, const TLearningArrayView<1, const float> Rhs)
+		{
+			UnaryOp(Out, Lhs.GetChannelNum(), Lhs, [Rhs](
+				FFrameAttribute& Out,
+				const FFrameAttribute& In,
+				const TLearningArrayView<1, const int32> RangeOffsets,
+				const TLearningArrayView<1, const int32> RangeLengths) {
+
+					const int32 ChannelNum = Out.GetChannelNum();
+
+					for (int32 ChannelIdx = 0; ChannelIdx < ChannelNum; ChannelIdx++)
+					{
+						Private::AddConstant(
+							Out.GetChannelAttributeData(ChannelIdx),
+							In.GetChannelAttributeData(ChannelIdx),
+							Rhs[ChannelIdx]);
+					}
+				});
+		}
+
+		void SubConstant(FFrameAttribute& Out, const FFrameAttribute& Lhs, const TLearningArrayView<1, const float> Rhs)
+		{
+			UnaryOp(Out, Lhs.GetChannelNum(), Lhs, [Rhs](
+				FFrameAttribute& Out,
+				const FFrameAttribute& In,
+				const TLearningArrayView<1, const int32> RangeOffsets,
+				const TLearningArrayView<1, const int32> RangeLengths) {
+
+					const int32 ChannelNum = Out.GetChannelNum();
+
+					for (int32 ChannelIdx = 0; ChannelIdx < ChannelNum; ChannelIdx++)
+					{
+						Private::SubConstant(
+							Out.GetChannelAttributeData(ChannelIdx),
+							In.GetChannelAttributeData(ChannelIdx),
+							Rhs[ChannelIdx]);
+					}
+				});
+		}
+
+		void MulConstant(FFrameAttribute& Out, const FFrameAttribute& Lhs, const TLearningArrayView<1, const float> Rhs)
+		{
+			UnaryOp(Out, Lhs.GetChannelNum(), Lhs, [Rhs](
+				FFrameAttribute& Out,
+				const FFrameAttribute& In,
+				const TLearningArrayView<1, const int32> RangeOffsets,
+				const TLearningArrayView<1, const int32> RangeLengths) {
+
+					const int32 ChannelNum = Out.GetChannelNum();
+
+					for (int32 ChannelIdx = 0; ChannelIdx < ChannelNum; ChannelIdx++)
+					{
+						Private::MulConstant(
+							Out.GetChannelAttributeData(ChannelIdx),
+							In.GetChannelAttributeData(ChannelIdx),
+							Rhs[ChannelIdx]);
+					}
+				});
+		}
+
+		void DivConstant(FFrameAttribute& Out, const FFrameAttribute& Lhs, const TLearningArrayView<1, const float> Rhs)
+		{
+			UnaryOp(Out, Lhs.GetChannelNum(), Lhs, [Rhs](
+				FFrameAttribute& Out,
+				const FFrameAttribute& In,
+				const TLearningArrayView<1, const int32> RangeOffsets,
+				const TLearningArrayView<1, const int32> RangeLengths) {
+
+					const int32 ChannelNum = Out.GetChannelNum();
+
+					for (int32 ChannelIdx = 0; ChannelIdx < ChannelNum; ChannelIdx++)
+					{
+						Private::DivConstant(
+							Out.GetChannelAttributeData(ChannelIdx),
+							In.GetChannelAttributeData(ChannelIdx),
+							Rhs[ChannelIdx]);
+					}
+				});
+		}
+
+		void ConstantAdd(FFrameAttribute& Out, const TLearningArrayView<1, const float> Lhs, const FFrameAttribute& Rhs)
+		{
+			UnaryOp(Out, Rhs.GetChannelNum(), Rhs, [Lhs](
+				FFrameAttribute& Out,
+				const FFrameAttribute& In,
+				const TLearningArrayView<1, const int32> RangeOffsets,
+				const TLearningArrayView<1, const int32> RangeLengths) {
+
+					const int32 ChannelNum = Out.GetChannelNum();
+
+					for (int32 ChannelIdx = 0; ChannelIdx < ChannelNum; ChannelIdx++)
+					{
+						Private::ConstantAdd(
+							Out.GetChannelAttributeData(ChannelIdx),
+							Lhs[ChannelIdx],
+							In.GetChannelAttributeData(ChannelIdx));
+					}
+				});
+		}
+
+		void ConstantSub(FFrameAttribute& Out, const TLearningArrayView<1, const float> Lhs, const FFrameAttribute& Rhs)
+		{
+			UnaryOp(Out, Rhs.GetChannelNum(), Rhs, [Lhs](
+				FFrameAttribute& Out,
+				const FFrameAttribute& In,
+				const TLearningArrayView<1, const int32> RangeOffsets,
+				const TLearningArrayView<1, const int32> RangeLengths) {
+
+					const int32 ChannelNum = Out.GetChannelNum();
+
+					for (int32 ChannelIdx = 0; ChannelIdx < ChannelNum; ChannelIdx++)
+					{
+						Private::ConstantSub(
+							Out.GetChannelAttributeData(ChannelIdx),
+							Lhs[ChannelIdx],
+							In.GetChannelAttributeData(ChannelIdx));
+					}
+				});
+		}
+
+		void ConstantMul(FFrameAttribute& Out, const TLearningArrayView<1, const float> Lhs, const FFrameAttribute& Rhs)
+		{
+			UnaryOp(Out, Rhs.GetChannelNum(), Rhs, [Lhs](
+				FFrameAttribute& Out,
+				const FFrameAttribute& In,
+				const TLearningArrayView<1, const int32> RangeOffsets,
+				const TLearningArrayView<1, const int32> RangeLengths) {
+
+					const int32 ChannelNum = Out.GetChannelNum();
+
+					for (int32 ChannelIdx = 0; ChannelIdx < ChannelNum; ChannelIdx++)
+					{
+						Private::ConstantMul(
+							Out.GetChannelAttributeData(ChannelIdx),
+							Lhs[ChannelIdx],
+							In.GetChannelAttributeData(ChannelIdx));
+					}
+				});
+		}
+
+		void ConstantDiv(FFrameAttribute& Out, const TLearningArrayView<1, const float> Lhs, const FFrameAttribute& Rhs)
+		{
+			UnaryOp(Out, Rhs.GetChannelNum(), Rhs, [Lhs](
+				FFrameAttribute& Out,
+				const FFrameAttribute& In,
+				const TLearningArrayView<1, const int32> RangeOffsets,
+				const TLearningArrayView<1, const int32> RangeLengths) {
+
+					const int32 ChannelNum = Out.GetChannelNum();
+
+					for (int32 ChannelIdx = 0; ChannelIdx < ChannelNum; ChannelIdx++)
+					{
+						Private::ConstantDiv(
+							Out.GetChannelAttributeData(ChannelIdx),
+							Lhs[ChannelIdx],
+							In.GetChannelAttributeData(ChannelIdx));
+					}
+				});
+		}
+
+		void Sum(FFrameAttribute& Out, const TArrayView<const ConstFrameAttributePtr> Inputs)
+		{
+			const int32 InputNum = Inputs.Num();
+
+			if (InputNum == 0)
+			{
+				Out.Empty();
+			}
+			else if (InputNum == 1)
+			{
+				Out = *Inputs[0];
+			}
+			else 
+			{
+				FFrameAttribute Tmp = *Inputs[0];
+
+				for (int32 InputIdx = 1; InputIdx < InputNum; InputIdx++)
+				{
+					Add(Out, Tmp, *Inputs[InputIdx]);
+					Tmp = Out;
+				}
+			}
+		}
+
+		void Prod(FFrameAttribute& Out, const TArrayView<const ConstFrameAttributePtr> Inputs)
+		{
+			const int32 InputNum = Inputs.Num();
+
+			if (InputNum == 0)
+			{
+				Out.Empty();
+			}
+			else if (InputNum == 1)
+			{
+				Out = *Inputs[0];
+			}
+			else
+			{
+				FFrameAttribute Tmp = *Inputs[0];
+
+				for (int32 InputIdx = 1; InputIdx < InputNum; InputIdx++)
+				{
+					Mul(Out, Tmp, *Inputs[InputIdx]);
+					Tmp = Out;
+				}
+			}
+		}
+
+		void LogicalAnd(FFrameAttribute& Out, const FFrameAttribute& Lhs, const FFrameAttribute& Rhs)
+		{
+			check(Lhs.GetChannelNum() == Rhs.GetChannelNum());
+
+			BinaryOp(Out, Lhs.GetChannelNum(), Lhs, Rhs, [](
+				FFrameAttribute& Out,
+				const FFrameAttribute& Lhs,
+				const FFrameAttribute& Rhs,
+				const TLearningArrayView<1, const int32> OutRangeOffsets,
+				const TLearningArrayView<1, const int32> LhsRangeOffsets,
+				const TLearningArrayView<1, const int32> RhsRangeOffsets,
+				const TLearningArrayView<1, const int32> RangeLengths) {
+
+					const int32 ChannelNum = Out.GetChannelNum();
+					const int32 RangeNum = RangeLengths.Num();
+					for (int32 ChannelIdx = 0; ChannelIdx < ChannelNum; ChannelIdx++)
+					{
+						for (int32 RangeIdx = 0; RangeIdx < RangeNum; RangeIdx++)
+						{
+							Private::LogicalAnd(
+								Out.GetChannelRangeAttributeData(ChannelIdx, OutRangeOffsets[RangeIdx], RangeLengths[RangeIdx]),
+								Lhs.GetChannelRangeAttributeData(ChannelIdx, LhsRangeOffsets[RangeIdx], RangeLengths[RangeIdx]),
+								Rhs.GetChannelRangeAttributeData(ChannelIdx, RhsRangeOffsets[RangeIdx], RangeLengths[RangeIdx]));
+						}
+					}
+				});
+		}
+
+		void LogicalOr(FFrameAttribute& Out, const FFrameAttribute& Lhs, const FFrameAttribute& Rhs)
+		{
+			check(Lhs.GetChannelNum() == Rhs.GetChannelNum());
+
+			BinaryOp(Out, Lhs.GetChannelNum(), Lhs, Rhs, [](
+				FFrameAttribute& Out,
+				const FFrameAttribute& Lhs,
+				const FFrameAttribute& Rhs,
+				const TLearningArrayView<1, const int32> OutRangeOffsets,
+				const TLearningArrayView<1, const int32> LhsRangeOffsets,
+				const TLearningArrayView<1, const int32> RhsRangeOffsets,
+				const TLearningArrayView<1, const int32> RangeLengths) {
+
+					const int32 ChannelNum = Out.GetChannelNum();
+					const int32 RangeNum = RangeLengths.Num();
+					for (int32 ChannelIdx = 0; ChannelIdx < ChannelNum; ChannelIdx++)
+					{
+						for (int32 RangeIdx = 0; RangeIdx < RangeNum; RangeIdx++)
+						{
+							Private::LogicalOr(
+								Out.GetChannelRangeAttributeData(ChannelIdx, OutRangeOffsets[RangeIdx], RangeLengths[RangeIdx]),
+								Lhs.GetChannelRangeAttributeData(ChannelIdx, LhsRangeOffsets[RangeIdx], RangeLengths[RangeIdx]),
+								Rhs.GetChannelRangeAttributeData(ChannelIdx, RhsRangeOffsets[RangeIdx], RangeLengths[RangeIdx]));
+						}
+					}
+				});
+		}
+
+		void LogicalNot(FFrameAttribute& Out, const FFrameAttribute& In)
+		{
+			UnaryOp(Out, In.GetChannelNum(), In, [](
+				FFrameAttribute& Out,
+				const FFrameAttribute& In,
+				const TLearningArrayView<1, const int32> RangeOffsets,
+				const TLearningArrayView<1, const int32> RangeLengths) {
+
+					const int32 ChannelNum = Out.GetChannelNum();
+
+					for (int32 ChannelIdx = 0; ChannelIdx < ChannelNum; ChannelIdx++)
+					{
+						Private::LogicalNot(
+							Out.GetChannelAttributeData(ChannelIdx),
+							In.GetChannelAttributeData(ChannelIdx));
+					}
+				});
+		}
+
+		void Gt(FFrameAttribute& Out, const FFrameAttribute& Lhs, const FFrameAttribute& Rhs)
+		{
+			check(Lhs.GetChannelNum() == Rhs.GetChannelNum());
+
+			BinaryOp(Out, Lhs.GetChannelNum(), Lhs, Rhs, [](
+				FFrameAttribute& Out,
+				const FFrameAttribute& Lhs,
+				const FFrameAttribute& Rhs,
+				const TLearningArrayView<1, const int32> OutRangeOffsets,
+				const TLearningArrayView<1, const int32> LhsRangeOffsets,
+				const TLearningArrayView<1, const int32> RhsRangeOffsets,
+				const TLearningArrayView<1, const int32> RangeLengths) {
+
+					const int32 ChannelNum = Out.GetChannelNum();
+					const int32 RangeNum = RangeLengths.Num();
+					for (int32 ChannelIdx = 0; ChannelIdx < ChannelNum; ChannelIdx++)
+					{
+						for (int32 RangeIdx = 0; RangeIdx < RangeNum; RangeIdx++)
+						{
+							Private::Gt(
+								Out.GetChannelRangeAttributeData(ChannelIdx, OutRangeOffsets[RangeIdx], RangeLengths[RangeIdx]),
+								Lhs.GetChannelRangeAttributeData(ChannelIdx, LhsRangeOffsets[RangeIdx], RangeLengths[RangeIdx]),
+								Rhs.GetChannelRangeAttributeData(ChannelIdx, RhsRangeOffsets[RangeIdx], RangeLengths[RangeIdx]));
+						}
+					}
+				});
+		}
+
+		void Ge(FFrameAttribute& Out, const FFrameAttribute& Lhs, const FFrameAttribute& Rhs)
+		{
+			check(Lhs.GetChannelNum() == Rhs.GetChannelNum());
+
+			BinaryOp(Out, Lhs.GetChannelNum(), Lhs, Rhs, [](
+				FFrameAttribute& Out,
+				const FFrameAttribute& Lhs,
+				const FFrameAttribute& Rhs,
+				const TLearningArrayView<1, const int32> OutRangeOffsets,
+				const TLearningArrayView<1, const int32> LhsRangeOffsets,
+				const TLearningArrayView<1, const int32> RhsRangeOffsets,
+				const TLearningArrayView<1, const int32> RangeLengths) {
+
+					const int32 ChannelNum = Out.GetChannelNum();
+					const int32 RangeNum = RangeLengths.Num();
+					for (int32 ChannelIdx = 0; ChannelIdx < ChannelNum; ChannelIdx++)
+					{
+						for (int32 RangeIdx = 0; RangeIdx < RangeNum; RangeIdx++)
+						{
+							Private::Ge(
+								Out.GetChannelRangeAttributeData(ChannelIdx, OutRangeOffsets[RangeIdx], RangeLengths[RangeIdx]),
+								Lhs.GetChannelRangeAttributeData(ChannelIdx, LhsRangeOffsets[RangeIdx], RangeLengths[RangeIdx]),
+								Rhs.GetChannelRangeAttributeData(ChannelIdx, RhsRangeOffsets[RangeIdx], RangeLengths[RangeIdx]));
+						}
+					}
+				});
+		}
+
+		void Lt(FFrameAttribute& Out, const FFrameAttribute& Lhs, const FFrameAttribute& Rhs)
+		{
+			check(Lhs.GetChannelNum() == Rhs.GetChannelNum());
+
+			BinaryOp(Out, Lhs.GetChannelNum(), Lhs, Rhs, [](
+				FFrameAttribute& Out,
+				const FFrameAttribute& Lhs,
+				const FFrameAttribute& Rhs,
+				const TLearningArrayView<1, const int32> OutRangeOffsets,
+				const TLearningArrayView<1, const int32> LhsRangeOffsets,
+				const TLearningArrayView<1, const int32> RhsRangeOffsets,
+				const TLearningArrayView<1, const int32> RangeLengths) {
+
+					const int32 ChannelNum = Out.GetChannelNum();
+					const int32 RangeNum = RangeLengths.Num();
+					for (int32 ChannelIdx = 0; ChannelIdx < ChannelNum; ChannelIdx++)
+					{
+						for (int32 RangeIdx = 0; RangeIdx < RangeNum; RangeIdx++)
+						{
+							Private::Lt(
+								Out.GetChannelRangeAttributeData(ChannelIdx, OutRangeOffsets[RangeIdx], RangeLengths[RangeIdx]),
+								Lhs.GetChannelRangeAttributeData(ChannelIdx, LhsRangeOffsets[RangeIdx], RangeLengths[RangeIdx]),
+								Rhs.GetChannelRangeAttributeData(ChannelIdx, RhsRangeOffsets[RangeIdx], RangeLengths[RangeIdx]));
+						}
+					}
+				});
+		}
+
+		void Le(FFrameAttribute& Out, const FFrameAttribute& Lhs, const FFrameAttribute& Rhs)
+		{
+			check(Lhs.GetChannelNum() == Rhs.GetChannelNum());
+
+			BinaryOp(Out, Lhs.GetChannelNum(), Lhs, Rhs, [](
+				FFrameAttribute& Out,
+				const FFrameAttribute& Lhs,
+				const FFrameAttribute& Rhs,
+				const TLearningArrayView<1, const int32> OutRangeOffsets,
+				const TLearningArrayView<1, const int32> LhsRangeOffsets,
+				const TLearningArrayView<1, const int32> RhsRangeOffsets,
+				const TLearningArrayView<1, const int32> RangeLengths) {
+
+					const int32 ChannelNum = Out.GetChannelNum();
+					const int32 RangeNum = RangeLengths.Num();
+					for (int32 ChannelIdx = 0; ChannelIdx < ChannelNum; ChannelIdx++)
+					{
+						for (int32 RangeIdx = 0; RangeIdx < RangeNum; RangeIdx++)
+						{
+							Private::Le(
+								Out.GetChannelRangeAttributeData(ChannelIdx, OutRangeOffsets[RangeIdx], RangeLengths[RangeIdx]),
+								Lhs.GetChannelRangeAttributeData(ChannelIdx, LhsRangeOffsets[RangeIdx], RangeLengths[RangeIdx]),
+								Rhs.GetChannelRangeAttributeData(ChannelIdx, RhsRangeOffsets[RangeIdx], RangeLengths[RangeIdx]));
+						}
+					}
+				});
+		}
+
+		void Eq(FFrameAttribute& Out, const FFrameAttribute& Lhs, const FFrameAttribute& Rhs)
+		{
+			check(Lhs.GetChannelNum() == Rhs.GetChannelNum());
+
+			BinaryOp(Out, Lhs.GetChannelNum(), Lhs, Rhs, [](
+				FFrameAttribute& Out,
+				const FFrameAttribute& Lhs,
+				const FFrameAttribute& Rhs,
+				const TLearningArrayView<1, const int32> OutRangeOffsets,
+				const TLearningArrayView<1, const int32> LhsRangeOffsets,
+				const TLearningArrayView<1, const int32> RhsRangeOffsets,
+				const TLearningArrayView<1, const int32> RangeLengths) {
+
+					const int32 ChannelNum = Out.GetChannelNum();
+					const int32 RangeNum = RangeLengths.Num();
+					for (int32 ChannelIdx = 0; ChannelIdx < ChannelNum; ChannelIdx++)
+					{
+						for (int32 RangeIdx = 0; RangeIdx < RangeNum; RangeIdx++)
+						{
+							Private::Eq(
+								Out.GetChannelRangeAttributeData(ChannelIdx, OutRangeOffsets[RangeIdx], RangeLengths[RangeIdx]),
+								Lhs.GetChannelRangeAttributeData(ChannelIdx, LhsRangeOffsets[RangeIdx], RangeLengths[RangeIdx]),
+								Rhs.GetChannelRangeAttributeData(ChannelIdx, RhsRangeOffsets[RangeIdx], RangeLengths[RangeIdx]));
+						}
+					}
+				});
+		}
+
+		void Neq(FFrameAttribute& Out, const FFrameAttribute& Lhs, const FFrameAttribute& Rhs)
+		{
+			check(Lhs.GetChannelNum() == Rhs.GetChannelNum());
+
+			BinaryOp(Out, Lhs.GetChannelNum(), Lhs, Rhs, [](
+				FFrameAttribute& Out,
+				const FFrameAttribute& Lhs,
+				const FFrameAttribute& Rhs,
+				const TLearningArrayView<1, const int32> OutRangeOffsets,
+				const TLearningArrayView<1, const int32> LhsRangeOffsets,
+				const TLearningArrayView<1, const int32> RhsRangeOffsets,
+				const TLearningArrayView<1, const int32> RangeLengths) {
+
+					const int32 ChannelNum = Out.GetChannelNum();
+					const int32 RangeNum = RangeLengths.Num();
+					for (int32 ChannelIdx = 0; ChannelIdx < ChannelNum; ChannelIdx++)
+					{
+						for (int32 RangeIdx = 0; RangeIdx < RangeNum; RangeIdx++)
+						{
+							Private::Neq(
+								Out.GetChannelRangeAttributeData(ChannelIdx, OutRangeOffsets[RangeIdx], RangeLengths[RangeIdx]),
+								Lhs.GetChannelRangeAttributeData(ChannelIdx, LhsRangeOffsets[RangeIdx], RangeLengths[RangeIdx]),
+								Rhs.GetChannelRangeAttributeData(ChannelIdx, RhsRangeOffsets[RangeIdx], RangeLengths[RangeIdx]));
+						}
+					}
+				});
+		}
+
+		void GtConstant(FFrameAttribute& Out, const FFrameAttribute& Lhs, const TLearningArrayView<1, const float> Rhs)
+		{
+			UnaryOp(Out, Lhs.GetChannelNum(), Lhs, [Rhs](
+				FFrameAttribute& Out,
+				const FFrameAttribute& In,
+				const TLearningArrayView<1, const int32> RangeOffsets,
+				const TLearningArrayView<1, const int32> RangeLengths) {
+
+					const int32 ChannelNum = Out.GetChannelNum();
+
+					for (int32 ChannelIdx = 0; ChannelIdx < ChannelNum; ChannelIdx++)
+					{
+						Private::GtConstant(
+							Out.GetChannelAttributeData(ChannelIdx),
+							In.GetChannelAttributeData(ChannelIdx),
+							Rhs[ChannelIdx]);
+					}
+				});
+		}
+
+		void GeConstant(FFrameAttribute& Out, const FFrameAttribute& Lhs, const TLearningArrayView<1, const float> Rhs)
+		{
+			UnaryOp(Out, Lhs.GetChannelNum(), Lhs, [Rhs](
+				FFrameAttribute& Out,
+				const FFrameAttribute& In,
+				const TLearningArrayView<1, const int32> RangeOffsets,
+				const TLearningArrayView<1, const int32> RangeLengths) {
+
+					const int32 ChannelNum = Out.GetChannelNum();
+
+					for (int32 ChannelIdx = 0; ChannelIdx < ChannelNum; ChannelIdx++)
+					{
+						Private::GeConstant(
+							Out.GetChannelAttributeData(ChannelIdx),
+							In.GetChannelAttributeData(ChannelIdx),
+							Rhs[ChannelIdx]);
+					}
+				});
+		}
+
+		void LtConstant(FFrameAttribute& Out, const FFrameAttribute& Lhs, const TLearningArrayView<1, const float> Rhs)
+		{
+			UnaryOp(Out, Lhs.GetChannelNum(), Lhs, [Rhs](
+				FFrameAttribute& Out,
+				const FFrameAttribute& In,
+				const TLearningArrayView<1, const int32> RangeOffsets,
+				const TLearningArrayView<1, const int32> RangeLengths) {
+
+					const int32 ChannelNum = Out.GetChannelNum();
+
+					for (int32 ChannelIdx = 0; ChannelIdx < ChannelNum; ChannelIdx++)
+					{
+						Private::LtConstant(
+							Out.GetChannelAttributeData(ChannelIdx),
+							In.GetChannelAttributeData(ChannelIdx),
+							Rhs[ChannelIdx]);
+					}
+				});
+		}
+
+		void LeConstant(FFrameAttribute& Out, const FFrameAttribute& Lhs, const TLearningArrayView<1, const float> Rhs)
+		{
+			UnaryOp(Out, Lhs.GetChannelNum(), Lhs, [Rhs](
+				FFrameAttribute& Out,
+				const FFrameAttribute& In,
+				const TLearningArrayView<1, const int32> RangeOffsets,
+				const TLearningArrayView<1, const int32> RangeLengths) {
+
+					const int32 ChannelNum = Out.GetChannelNum();
+
+					for (int32 ChannelIdx = 0; ChannelIdx < ChannelNum; ChannelIdx++)
+					{
+						Private::LeConstant(
+							Out.GetChannelAttributeData(ChannelIdx),
+							In.GetChannelAttributeData(ChannelIdx),
+							Rhs[ChannelIdx]);
+					}
+				});
+		}
+
+		void EqConstant(FFrameAttribute& Out, const FFrameAttribute& Lhs, const TLearningArrayView<1, const float> Rhs)
+		{
+			UnaryOp(Out, Lhs.GetChannelNum(), Lhs, [Rhs](
+				FFrameAttribute& Out,
+				const FFrameAttribute& In,
+				const TLearningArrayView<1, const int32> RangeOffsets,
+				const TLearningArrayView<1, const int32> RangeLengths) {
+
+					const int32 ChannelNum = Out.GetChannelNum();
+
+					for (int32 ChannelIdx = 0; ChannelIdx < ChannelNum; ChannelIdx++)
+					{
+						Private::EqConstant(
+							Out.GetChannelAttributeData(ChannelIdx),
+							In.GetChannelAttributeData(ChannelIdx),
+							Rhs[ChannelIdx]);
+					}
+				});
+		}
+
+		void NeqConstant(FFrameAttribute& Out, const FFrameAttribute& Lhs, const TLearningArrayView<1, const float> Rhs)
+		{
+			UnaryOp(Out, Lhs.GetChannelNum(), Lhs, [Rhs](
+				FFrameAttribute& Out,
+				const FFrameAttribute& In,
+				const TLearningArrayView<1, const int32> RangeOffsets,
+				const TLearningArrayView<1, const int32> RangeLengths) {
+
+					const int32 ChannelNum = Out.GetChannelNum();
+
+					for (int32 ChannelIdx = 0; ChannelIdx < ChannelNum; ChannelIdx++)
+					{
+						Private::NeqConstant(
+							Out.GetChannelAttributeData(ChannelIdx),
+							In.GetChannelAttributeData(ChannelIdx),
+							Rhs[ChannelIdx]);
+					}
+				});
+		}
+
+		void ConstantGt(FFrameAttribute& Out, const TLearningArrayView<1, const float> Lhs, const FFrameAttribute& Rhs)
+		{
+			UnaryOp(Out, Rhs.GetChannelNum(), Rhs, [Lhs](
+				FFrameAttribute& Out,
+				const FFrameAttribute& In,
+				const TLearningArrayView<1, const int32> RangeOffsets,
+				const TLearningArrayView<1, const int32> RangeLengths) {
+
+					const int32 ChannelNum = Out.GetChannelNum();
+
+					for (int32 ChannelIdx = 0; ChannelIdx < ChannelNum; ChannelIdx++)
+					{
+						Private::ConstantGt(
+							Out.GetChannelAttributeData(ChannelIdx),
+							Lhs[ChannelIdx],
+							In.GetChannelAttributeData(ChannelIdx));
+					}
+				});
+		}
+
+		void ConstantGe(FFrameAttribute& Out, const TLearningArrayView<1, const float> Lhs, const FFrameAttribute& Rhs)
+		{
+			UnaryOp(Out, Rhs.GetChannelNum(), Rhs, [Lhs](
+				FFrameAttribute& Out,
+				const FFrameAttribute& In,
+				const TLearningArrayView<1, const int32> RangeOffsets,
+				const TLearningArrayView<1, const int32> RangeLengths) {
+
+					const int32 ChannelNum = Out.GetChannelNum();
+
+					for (int32 ChannelIdx = 0; ChannelIdx < ChannelNum; ChannelIdx++)
+					{
+						Private::ConstantGe(
+							Out.GetChannelAttributeData(ChannelIdx),
+							Lhs[ChannelIdx],
+							In.GetChannelAttributeData(ChannelIdx));
+					}
+				});
+		}
+
+		void ConstantLt(FFrameAttribute& Out, const TLearningArrayView<1, const float> Lhs, const FFrameAttribute& Rhs)
+		{
+			UnaryOp(Out, Rhs.GetChannelNum(), Rhs, [Lhs](
+				FFrameAttribute& Out,
+				const FFrameAttribute& In,
+				const TLearningArrayView<1, const int32> RangeOffsets,
+				const TLearningArrayView<1, const int32> RangeLengths) {
+
+					const int32 ChannelNum = Out.GetChannelNum();
+
+					for (int32 ChannelIdx = 0; ChannelIdx < ChannelNum; ChannelIdx++)
+					{
+						Private::ConstantLt(
+							Out.GetChannelAttributeData(ChannelIdx),
+							Lhs[ChannelIdx],
+							In.GetChannelAttributeData(ChannelIdx));
+					}
+				});
+		}
+
+		void ConstantLe(FFrameAttribute& Out, const TLearningArrayView<1, const float> Lhs, const FFrameAttribute& Rhs)
+		{
+			UnaryOp(Out, Rhs.GetChannelNum(), Rhs, [Lhs](
+				FFrameAttribute& Out,
+				const FFrameAttribute& In,
+				const TLearningArrayView<1, const int32> RangeOffsets,
+				const TLearningArrayView<1, const int32> RangeLengths) {
+
+					const int32 ChannelNum = Out.GetChannelNum();
+
+					for (int32 ChannelIdx = 0; ChannelIdx < ChannelNum; ChannelIdx++)
+					{
+						Private::ConstantLe(
+							Out.GetChannelAttributeData(ChannelIdx),
+							Lhs[ChannelIdx],
+							In.GetChannelAttributeData(ChannelIdx));
+					}
+				});
+		}
+
+		void ConstantEq(FFrameAttribute& Out, const TLearningArrayView<1, const float> Lhs, const FFrameAttribute& Rhs)
+		{
+			UnaryOp(Out, Rhs.GetChannelNum(), Rhs, [Lhs](
+				FFrameAttribute& Out,
+				const FFrameAttribute& In,
+				const TLearningArrayView<1, const int32> RangeOffsets,
+				const TLearningArrayView<1, const int32> RangeLengths) {
+
+					const int32 ChannelNum = Out.GetChannelNum();
+
+					for (int32 ChannelIdx = 0; ChannelIdx < ChannelNum; ChannelIdx++)
+					{
+						Private::ConstantEq(
+							Out.GetChannelAttributeData(ChannelIdx),
+							Lhs[ChannelIdx],
+							In.GetChannelAttributeData(ChannelIdx));
+					}
+				});
+		}
+
+		void ConstantNeq(FFrameAttribute& Out, const TLearningArrayView<1, const float> Lhs, const FFrameAttribute& Rhs)
+		{
+			UnaryOp(Out, Rhs.GetChannelNum(), Rhs, [Lhs](
+				FFrameAttribute& Out,
+				const FFrameAttribute& In,
+				const TLearningArrayView<1, const int32> RangeOffsets,
+				const TLearningArrayView<1, const int32> RangeLengths) {
+
+					const int32 ChannelNum = Out.GetChannelNum();
+
+					for (int32 ChannelIdx = 0; ChannelIdx < ChannelNum; ChannelIdx++)
+					{
+						Private::ConstantNeq(
+							Out.GetChannelAttributeData(ChannelIdx),
+							Lhs[ChannelIdx],
+							In.GetChannelAttributeData(ChannelIdx));
+					}
+				});
+		}
+
+		void FilterGaussian(FFrameAttribute& Out, const FFrameAttribute& In, const float StdInFrames)
+		{
+			UnaryOp(Out, In.GetChannelNum(), In, [StdInFrames](
+				FFrameAttribute& Out,
+				const FFrameAttribute& In,
+				const TLearningArrayView<1, const int32> RangeOffsets,
+				const TLearningArrayView<1, const int32> RangeLengths) {
+
+					const int32 ChannelNum = Out.GetChannelNum();
+					const int32 RangeNum = RangeLengths.Num();
+					for (int32 ChannelIdx = 0; ChannelIdx < ChannelNum; ChannelIdx++)
+					{
+						for (int32 RangeIdx = 0; RangeIdx < RangeNum; RangeIdx++)
+						{
+							Private::FilterGaussian(
+								Out.GetChannelRangeAttributeData(ChannelIdx, RangeOffsets[RangeIdx], RangeLengths[RangeIdx]),
+								In.GetChannelRangeAttributeData(ChannelIdx, RangeOffsets[RangeIdx], RangeLengths[RangeIdx]),
+								StdInFrames);
+						}
+					}
+				});
+		}
+
+		void FilterMajorityVote(FFrameAttribute& Out, const FFrameAttribute& In, const float FilterWidthFrames)
+		{
+			UnaryOp(Out, In.GetChannelNum(), In, [FilterWidthFrames](
+				FFrameAttribute& Out,
+				const FFrameAttribute& In,
+				const TLearningArrayView<1, const int32> RangeOffsets,
+				const TLearningArrayView<1, const int32> RangeLengths) {
+
+					const int32 ChannelNum = Out.GetChannelNum();
+					const int32 RangeNum = RangeLengths.Num();
+					for (int32 ChannelIdx = 0; ChannelIdx < ChannelNum; ChannelIdx++)
+					{
+						for (int32 RangeIdx = 0; RangeIdx < RangeNum; RangeIdx++)
+						{
+							Private::FilterMajorityVote(
+								Out.GetChannelRangeAttributeData(ChannelIdx, RangeOffsets[RangeIdx], RangeLengths[RangeIdx]),
+								In.GetChannelRangeAttributeData(ChannelIdx, RangeOffsets[RangeIdx], RangeLengths[RangeIdx]),
+								FilterWidthFrames);
+						}
+					}
+				});
+		}
+
+		void QuatMul(FFrameAttribute& Out, const FFrameAttribute& Lhs, const FFrameAttribute& Rhs)
+		{
+			check(Lhs.GetChannelNum() == 4);
+			check(Rhs.GetChannelNum() == 4);
+
+			BinaryOp(Out, 4, Lhs, Rhs, [](
+				FFrameAttribute& Out,
+				const FFrameAttribute& Lhs,
+				const FFrameAttribute& Rhs,
+				const TLearningArrayView<1, const int32> OutRangeOffsets,
+				const TLearningArrayView<1, const int32> LhsRangeOffsets,
+				const TLearningArrayView<1, const int32> RhsRangeOffsets,
+				const TLearningArrayView<1, const int32> RangeLengths) {
+
+					const int32 RangeNum = RangeLengths.Num();
+
+					for (int32 RangeIdx = 0; RangeIdx < RangeNum; RangeIdx++)
+					{
+						Private::QuatMul(
+							Out.GetChannelRangeAttributeData(0, OutRangeOffsets[RangeIdx], RangeLengths[RangeIdx]),
+							Out.GetChannelRangeAttributeData(1, OutRangeOffsets[RangeIdx], RangeLengths[RangeIdx]),
+							Out.GetChannelRangeAttributeData(2, OutRangeOffsets[RangeIdx], RangeLengths[RangeIdx]),
+							Out.GetChannelRangeAttributeData(3, OutRangeOffsets[RangeIdx], RangeLengths[RangeIdx]),
+							Lhs.GetChannelRangeAttributeData(0, LhsRangeOffsets[RangeIdx], RangeLengths[RangeIdx]),
+							Lhs.GetChannelRangeAttributeData(1, LhsRangeOffsets[RangeIdx], RangeLengths[RangeIdx]),
+							Lhs.GetChannelRangeAttributeData(2, LhsRangeOffsets[RangeIdx], RangeLengths[RangeIdx]),
+							Lhs.GetChannelRangeAttributeData(3, LhsRangeOffsets[RangeIdx], RangeLengths[RangeIdx]),
+							Rhs.GetChannelRangeAttributeData(0, RhsRangeOffsets[RangeIdx], RangeLengths[RangeIdx]),
+							Rhs.GetChannelRangeAttributeData(1, RhsRangeOffsets[RangeIdx], RangeLengths[RangeIdx]),
+							Rhs.GetChannelRangeAttributeData(2, RhsRangeOffsets[RangeIdx], RangeLengths[RangeIdx]),
+							Rhs.GetChannelRangeAttributeData(3, RhsRangeOffsets[RangeIdx], RangeLengths[RangeIdx]));
+					}
+				});
+		}
+
+		void QuatDiv(FFrameAttribute& Out, const FFrameAttribute& Lhs, const FFrameAttribute& Rhs)
+		{
+			QuatMulInv(Out, Lhs, Rhs);
+		}
+
+		void QuatInv(FFrameAttribute& Out, const FFrameAttribute& In)
+		{
+			check(In.GetChannelNum() == 4);
+
+			UnaryOp(Out, 4, In, [](
+				FFrameAttribute& Out,
+				const FFrameAttribute& In,
+				const TLearningArrayView<1, const int32> RangeOffsets,
+				const TLearningArrayView<1, const int32> RangeLengths) {
+
+					Private::QuatInv(
+						Out.GetAttributeData(),
+						In.GetAttributeData());
+				});
+		}
+
+		void QuatAbs(FFrameAttribute& Out, const FFrameAttribute& In)
+		{
+			check(In.GetChannelNum() == 4);
+
+			UnaryOp(Out, 4, In, [](
+				FFrameAttribute& Out,
+				const FFrameAttribute& In,
+				const TLearningArrayView<1, const int32> RangeOffsets,
+				const TLearningArrayView<1, const int32> RangeLengths) {
+
+					Private::QuatAbs(
+						Out.GetAttributeData(),
+						In.GetAttributeData());
+				});
+		}
+
+		void QuatInvMul(FFrameAttribute& Out, const FFrameAttribute& Lhs, const FFrameAttribute& Rhs)
+		{
+			check(Lhs.GetChannelNum() == 4);
+			check(Rhs.GetChannelNum() == 4);
+
+			BinaryOp(Out, 4, Lhs, Rhs, [](
+				FFrameAttribute& Out,
+				const FFrameAttribute& Lhs,
+				const FFrameAttribute& Rhs,
+				const TLearningArrayView<1, const int32> OutRangeOffsets,
+				const TLearningArrayView<1, const int32> LhsRangeOffsets,
+				const TLearningArrayView<1, const int32> RhsRangeOffsets,
+				const TLearningArrayView<1, const int32> RangeLengths) {
+
+					const int32 RangeNum = RangeLengths.Num();
+
+					for (int32 RangeIdx = 0; RangeIdx < RangeNum; RangeIdx++)
+					{
+						Private::QuatInvMul(
+							Out.GetChannelRangeAttributeData(0, OutRangeOffsets[RangeIdx], RangeLengths[RangeIdx]),
+							Out.GetChannelRangeAttributeData(1, OutRangeOffsets[RangeIdx], RangeLengths[RangeIdx]),
+							Out.GetChannelRangeAttributeData(2, OutRangeOffsets[RangeIdx], RangeLengths[RangeIdx]),
+							Out.GetChannelRangeAttributeData(3, OutRangeOffsets[RangeIdx], RangeLengths[RangeIdx]),
+							Lhs.GetChannelRangeAttributeData(0, LhsRangeOffsets[RangeIdx], RangeLengths[RangeIdx]),
+							Lhs.GetChannelRangeAttributeData(1, LhsRangeOffsets[RangeIdx], RangeLengths[RangeIdx]),
+							Lhs.GetChannelRangeAttributeData(2, LhsRangeOffsets[RangeIdx], RangeLengths[RangeIdx]),
+							Lhs.GetChannelRangeAttributeData(3, LhsRangeOffsets[RangeIdx], RangeLengths[RangeIdx]),
+							Rhs.GetChannelRangeAttributeData(0, RhsRangeOffsets[RangeIdx], RangeLengths[RangeIdx]),
+							Rhs.GetChannelRangeAttributeData(1, RhsRangeOffsets[RangeIdx], RangeLengths[RangeIdx]),
+							Rhs.GetChannelRangeAttributeData(2, RhsRangeOffsets[RangeIdx], RangeLengths[RangeIdx]),
+							Rhs.GetChannelRangeAttributeData(3, RhsRangeOffsets[RangeIdx], RangeLengths[RangeIdx]));
+					}
+				});
+		}
+
+		void QuatMulInv(FFrameAttribute& Out, const FFrameAttribute& Lhs, const FFrameAttribute& Rhs)
+		{
+			check(Lhs.GetChannelNum() == 4);
+			check(Rhs.GetChannelNum() == 4);
+
+			BinaryOp(Out, 4, Lhs, Rhs, [](
+				FFrameAttribute& Out,
+				const FFrameAttribute& Lhs,
+				const FFrameAttribute& Rhs,
+				const TLearningArrayView<1, const int32> OutRangeOffsets,
+				const TLearningArrayView<1, const int32> LhsRangeOffsets,
+				const TLearningArrayView<1, const int32> RhsRangeOffsets,
+				const TLearningArrayView<1, const int32> RangeLengths) {
+
+					const int32 RangeNum = RangeLengths.Num();
+
+					for (int32 RangeIdx = 0; RangeIdx < RangeNum; RangeIdx++)
+					{
+						Private::QuatMulInv(
+							Out.GetChannelRangeAttributeData(0, OutRangeOffsets[RangeIdx], RangeLengths[RangeIdx]),
+							Out.GetChannelRangeAttributeData(1, OutRangeOffsets[RangeIdx], RangeLengths[RangeIdx]),
+							Out.GetChannelRangeAttributeData(2, OutRangeOffsets[RangeIdx], RangeLengths[RangeIdx]),
+							Out.GetChannelRangeAttributeData(3, OutRangeOffsets[RangeIdx], RangeLengths[RangeIdx]),
+							Lhs.GetChannelRangeAttributeData(0, LhsRangeOffsets[RangeIdx], RangeLengths[RangeIdx]),
+							Lhs.GetChannelRangeAttributeData(1, LhsRangeOffsets[RangeIdx], RangeLengths[RangeIdx]),
+							Lhs.GetChannelRangeAttributeData(2, LhsRangeOffsets[RangeIdx], RangeLengths[RangeIdx]),
+							Lhs.GetChannelRangeAttributeData(3, LhsRangeOffsets[RangeIdx], RangeLengths[RangeIdx]),
+							Rhs.GetChannelRangeAttributeData(0, RhsRangeOffsets[RangeIdx], RangeLengths[RangeIdx]),
+							Rhs.GetChannelRangeAttributeData(1, RhsRangeOffsets[RangeIdx], RangeLengths[RangeIdx]),
+							Rhs.GetChannelRangeAttributeData(2, RhsRangeOffsets[RangeIdx], RangeLengths[RangeIdx]),
+							Rhs.GetChannelRangeAttributeData(3, RhsRangeOffsets[RangeIdx], RangeLengths[RangeIdx]));
+					}
+				});
+		}
+
+		void QuatRotate(FFrameAttribute& Out, const FFrameAttribute& Lhs, const FFrameAttribute& Rhs)
+		{
+			check(Lhs.GetChannelNum() == 4);
+			check(Rhs.GetChannelNum() == 3);
+
+			BinaryOp(Out, 3, Lhs, Rhs, [](
+				FFrameAttribute& Out,
+				const FFrameAttribute& Lhs,
+				const FFrameAttribute& Rhs,
+				const TLearningArrayView<1, const int32> OutRangeOffsets,
+				const TLearningArrayView<1, const int32> LhsRangeOffsets,
+				const TLearningArrayView<1, const int32> RhsRangeOffsets,
+				const TLearningArrayView<1, const int32> RangeLengths) {
+
+					const int32 RangeNum = RangeLengths.Num();
+
+					for (int32 RangeIdx = 0; RangeIdx < RangeNum; RangeIdx++)
+					{
+						Private::QuatRotate(
+							Out.GetChannelRangeAttributeData(0, OutRangeOffsets[RangeIdx], RangeLengths[RangeIdx]),
+							Out.GetChannelRangeAttributeData(1, OutRangeOffsets[RangeIdx], RangeLengths[RangeIdx]),
+							Out.GetChannelRangeAttributeData(2, OutRangeOffsets[RangeIdx], RangeLengths[RangeIdx]),
+							Lhs.GetChannelRangeAttributeData(0, LhsRangeOffsets[RangeIdx], RangeLengths[RangeIdx]),
+							Lhs.GetChannelRangeAttributeData(1, LhsRangeOffsets[RangeIdx], RangeLengths[RangeIdx]),
+							Lhs.GetChannelRangeAttributeData(2, LhsRangeOffsets[RangeIdx], RangeLengths[RangeIdx]),
+							Lhs.GetChannelRangeAttributeData(3, LhsRangeOffsets[RangeIdx], RangeLengths[RangeIdx]),
+							Rhs.GetChannelRangeAttributeData(0, RhsRangeOffsets[RangeIdx], RangeLengths[RangeIdx]),
+							Rhs.GetChannelRangeAttributeData(1, RhsRangeOffsets[RangeIdx], RangeLengths[RangeIdx]),
+							Rhs.GetChannelRangeAttributeData(2, RhsRangeOffsets[RangeIdx], RangeLengths[RangeIdx]));
+					}
+				});
+		}
+
+		void QuatUnrotate(FFrameAttribute& Out, const FFrameAttribute& Lhs, const FFrameAttribute& Rhs)
+		{
+			check(Lhs.GetChannelNum() == 4);
+			check(Rhs.GetChannelNum() == 3);
+
+			BinaryOp(Out, 3, Lhs, Rhs, [](
+				FFrameAttribute& Out,
+				const FFrameAttribute& Lhs,
+				const FFrameAttribute& Rhs,
+				const TLearningArrayView<1, const int32> OutRangeOffsets,
+				const TLearningArrayView<1, const int32> LhsRangeOffsets,
+				const TLearningArrayView<1, const int32> RhsRangeOffsets,
+				const TLearningArrayView<1, const int32> RangeLengths) {
+
+					const int32 RangeNum = RangeLengths.Num();
+
+					for (int32 RangeIdx = 0; RangeIdx < RangeNum; RangeIdx++)
+					{
+						Private::QuatUnrotate(
+							Out.GetChannelRangeAttributeData(0, OutRangeOffsets[RangeIdx], RangeLengths[RangeIdx]),
+							Out.GetChannelRangeAttributeData(1, OutRangeOffsets[RangeIdx], RangeLengths[RangeIdx]),
+							Out.GetChannelRangeAttributeData(2, OutRangeOffsets[RangeIdx], RangeLengths[RangeIdx]),
+							Lhs.GetChannelRangeAttributeData(0, LhsRangeOffsets[RangeIdx], RangeLengths[RangeIdx]),
+							Lhs.GetChannelRangeAttributeData(1, LhsRangeOffsets[RangeIdx], RangeLengths[RangeIdx]),
+							Lhs.GetChannelRangeAttributeData(2, LhsRangeOffsets[RangeIdx], RangeLengths[RangeIdx]),
+							Lhs.GetChannelRangeAttributeData(3, LhsRangeOffsets[RangeIdx], RangeLengths[RangeIdx]),
+							Rhs.GetChannelRangeAttributeData(0, RhsRangeOffsets[RangeIdx], RangeLengths[RangeIdx]),
+							Rhs.GetChannelRangeAttributeData(1, RhsRangeOffsets[RangeIdx], RangeLengths[RangeIdx]),
+							Rhs.GetChannelRangeAttributeData(2, RhsRangeOffsets[RangeIdx], RangeLengths[RangeIdx]));
+					}
+				});
+		}
+
+		void QuatBetween(FFrameAttribute& Out, const FFrameAttribute& Lhs, const FFrameAttribute& Rhs)
+		{
+			check(Lhs.GetChannelNum() == 3);
+			check(Rhs.GetChannelNum() == 3);
+
+			BinaryOp(Out, 4, Lhs, Rhs, [](
+				FFrameAttribute& Out,
+				const FFrameAttribute& Lhs,
+				const FFrameAttribute& Rhs,
+				const TLearningArrayView<1, const int32> OutRangeOffsets,
+				const TLearningArrayView<1, const int32> LhsRangeOffsets,
+				const TLearningArrayView<1, const int32> RhsRangeOffsets,
+				const TLearningArrayView<1, const int32> RangeLengths) {
+
+					const int32 RangeNum = RangeLengths.Num();
+
+					for (int32 RangeIdx = 0; RangeIdx < RangeNum; RangeIdx++)
+					{
+						Private::QuatBetween(
+							Out.GetChannelRangeAttributeData(0, OutRangeOffsets[RangeIdx], RangeLengths[RangeIdx]),
+							Out.GetChannelRangeAttributeData(1, OutRangeOffsets[RangeIdx], RangeLengths[RangeIdx]),
+							Out.GetChannelRangeAttributeData(2, OutRangeOffsets[RangeIdx], RangeLengths[RangeIdx]),
+							Out.GetChannelRangeAttributeData(3, OutRangeOffsets[RangeIdx], RangeLengths[RangeIdx]),
+							Lhs.GetChannelRangeAttributeData(0, LhsRangeOffsets[RangeIdx], RangeLengths[RangeIdx]),
+							Lhs.GetChannelRangeAttributeData(1, LhsRangeOffsets[RangeIdx], RangeLengths[RangeIdx]),
+							Lhs.GetChannelRangeAttributeData(2, LhsRangeOffsets[RangeIdx], RangeLengths[RangeIdx]),
+							Rhs.GetChannelRangeAttributeData(0, RhsRangeOffsets[RangeIdx], RangeLengths[RangeIdx]),
+							Rhs.GetChannelRangeAttributeData(1, RhsRangeOffsets[RangeIdx], RangeLengths[RangeIdx]),
+							Rhs.GetChannelRangeAttributeData(2, RhsRangeOffsets[RangeIdx], RangeLengths[RangeIdx]));
+					}
+				});
+		}
+
+		void QuatMulConstant(FFrameAttribute& Out, const FFrameAttribute& Lhs, const FQuat4f Rhs)
+		{
+			check(Lhs.GetChannelNum() == 4);
+
+			UnaryOp(Out, 4, Lhs, [Rhs](
+				FFrameAttribute& Out,
+				const FFrameAttribute& Lhs,
+				const TLearningArrayView<1, const int32> RangeOffsets,
+				const TLearningArrayView<1, const int32> RangeLengths) {
+
+					Private::QuatMulConstant(
+						Out.GetAttributeData(),
+						Lhs.GetAttributeData(),
+						Rhs);
+				});
+		}
+
+		void QuatDivConstant(FFrameAttribute& Out, const FFrameAttribute& Lhs, const FQuat4f Rhs)
+		{
+			QuatMulInvConstant(Out, Lhs, Rhs);
+		}
+
+		void QuatInvMulConstant(FFrameAttribute& Out, const FFrameAttribute& Lhs, const FQuat4f Rhs)
+		{
+			check(Lhs.GetChannelNum() == 4);
+
+			UnaryOp(Out, 4, Lhs, [Rhs](
+				FFrameAttribute& Out,
+				const FFrameAttribute& Lhs,
+				const TLearningArrayView<1, const int32> RangeOffsets,
+				const TLearningArrayView<1, const int32> RangeLengths) {
+
+					Private::QuatInvMulConstant(
+						Out.GetAttributeData(),
+						Lhs.GetAttributeData(),
+						Rhs);
+				});
+		}
+
+		void QuatMulInvConstant(FFrameAttribute& Out, const FFrameAttribute& Lhs, const FQuat4f Rhs)
+		{
+			check(Lhs.GetChannelNum() == 4);
+
+			UnaryOp(Out, 4, Lhs, [Rhs](
+				FFrameAttribute& Out,
+				const FFrameAttribute& Lhs,
+				const TLearningArrayView<1, const int32> RangeOffsets,
+				const TLearningArrayView<1, const int32> RangeLengths) {
+
+					Private::QuatMulInvConstant(
+						Out.GetAttributeData(),
+						Lhs.GetAttributeData(),
+						Rhs);
+				});
+		}
+
+		void QuatRotateConstant(FFrameAttribute& Out, const FFrameAttribute& Lhs, const FVector3f Rhs)
+		{
+			check(Lhs.GetChannelNum() == 4);
+
+			UnaryOp(Out, 3, Lhs, [Rhs](
+				FFrameAttribute& Out,
+				const FFrameAttribute& Lhs,
+				const TLearningArrayView<1, const int32> RangeOffsets,
+				const TLearningArrayView<1, const int32> RangeLengths) {
+
+					Private::QuatRotateConstant(
+						Out.GetAttributeData(),
+						Lhs.GetAttributeData(),
+						Rhs);
+				});
+		}
+
+		void QuatUnrotateConstant(FFrameAttribute& Out, const FFrameAttribute& Lhs, const FVector3f Rhs)
+		{
+			check(Lhs.GetChannelNum() == 4);
+
+			UnaryOp(Out, 3, Lhs, [Rhs](
+				FFrameAttribute& Out,
+				const FFrameAttribute& Lhs,
+				const TLearningArrayView<1, const int32> RangeOffsets,
+				const TLearningArrayView<1, const int32> RangeLengths) {
+
+					Private::QuatUnrotateConstant(
+						Out.GetAttributeData(),
+						Lhs.GetAttributeData(),
+						Rhs);
+				});
+		}
+
+		void QuatBetweenConstant(FFrameAttribute& Out, const FFrameAttribute& Lhs, const FVector3f Rhs)
+		{
+			check(Lhs.GetChannelNum() == 3);
+
+			UnaryOp(Out, 4, Lhs, [Rhs](
+				FFrameAttribute& Out,
+				const FFrameAttribute& Lhs,
+				const TLearningArrayView<1, const int32> RangeOffsets,
+				const TLearningArrayView<1, const int32> RangeLengths) {
+
+					Private::QuatBetweenConstant(
+						Out.GetAttributeData(),
+						Lhs.GetAttributeData(),
+						Rhs);
+				});
+		}
+
+		void QuatConstantMul(FFrameAttribute& Out, const FQuat4f Lhs, const FFrameAttribute& Rhs)
+		{
+			check(Rhs.GetChannelNum() == 4);
+
+			UnaryOp(Out, 4, Rhs, [Lhs](
+				FFrameAttribute& Out,
+				const FFrameAttribute& Rhs,
+				const TLearningArrayView<1, const int32> RangeOffsets,
+				const TLearningArrayView<1, const int32> RangeLengths) {
+
+					Private::QuatConstantMul(
+						Out.GetAttributeData(),
+						Lhs,
+						Rhs.GetAttributeData());
+				});
+		}
+
+		void QuatConstantDiv(FFrameAttribute& Out, const FQuat4f Lhs, const FFrameAttribute& Rhs)
+		{
+			QuatConstantMulInv(Out, Lhs, Rhs);
+		}
+
+		void QuatConstantInvMul(FFrameAttribute& Out, const FQuat4f Lhs, const FFrameAttribute& Rhs)
+		{
+			check(Rhs.GetChannelNum() == 4);
+
+			UnaryOp(Out, 4, Rhs, [Lhs](
+				FFrameAttribute& Out,
+				const FFrameAttribute& Rhs,
+				const TLearningArrayView<1, const int32> RangeOffsets,
+				const TLearningArrayView<1, const int32> RangeLengths) {
+
+					Private::QuatConstantInvMul(
+						Out.GetAttributeData(),
+						Lhs,
+						Rhs.GetAttributeData());
+				});
+		}
+
+		void QuatConstantMulInv(FFrameAttribute& Out, const FQuat4f Lhs, const FFrameAttribute& Rhs)
+		{
+			check(Rhs.GetChannelNum() == 4);
+
+			UnaryOp(Out, 4, Rhs, [Lhs](
+				FFrameAttribute& Out,
+				const FFrameAttribute& Rhs,
+				const TLearningArrayView<1, const int32> RangeOffsets,
+				const TLearningArrayView<1, const int32> RangeLengths) {
+
+					Private::QuatConstantMulInv(
+						Out.GetAttributeData(),
+						Lhs,
+						Rhs.GetAttributeData());
+				});
+		}
+
+		void QuatConstantRotate(FFrameAttribute& Out, const FQuat4f Lhs, const FFrameAttribute& Rhs)
+		{
+			check(Rhs.GetChannelNum() == 3);
+
+			UnaryOp(Out, 3, Rhs, [Lhs](
+				FFrameAttribute& Out,
+				const FFrameAttribute& Rhs,
+				const TLearningArrayView<1, const int32> RangeOffsets,
+				const TLearningArrayView<1, const int32> RangeLengths) {
+
+					Private::QuatConstantRotate(
+						Out.GetAttributeData(),
+						Lhs,
+						Rhs.GetAttributeData());
+				});
+		}
+
+		void QuatConstantUnrotate(FFrameAttribute& Out, const FQuat4f Lhs, const FFrameAttribute& Rhs)
+		{
+			check(Rhs.GetChannelNum() == 3);
+
+			UnaryOp(Out, 3, Rhs, [Lhs](
+				FFrameAttribute& Out,
+				const FFrameAttribute& Rhs,
+				const TLearningArrayView<1, const int32> RangeOffsets,
+				const TLearningArrayView<1, const int32> RangeLengths) {
+
+					Private::QuatConstantUnrotate(
+						Out.GetAttributeData(),
+						Lhs,
+						Rhs.GetAttributeData());
+				});
+		}
+
+		void QuatConstantBetween(FFrameAttribute& Out, const FVector3f Lhs, const FFrameAttribute& Rhs)
+		{
+			check(Rhs.GetChannelNum() == 3);
+
+			UnaryOp(Out, 4, Rhs, [Lhs](
+				FFrameAttribute& Out,
+				const FFrameAttribute& Rhs,
+				const TLearningArrayView<1, const int32> RangeOffsets,
+				const TLearningArrayView<1, const int32> RangeLengths) {
+
+					Private::QuatConstantBetween(
+						Out.GetAttributeData(),
+						Lhs,
+						Rhs.GetAttributeData());
+				});
+		}
+
+	}
+}
\ No newline at end of file
diff --git a/Engine/Plugins/Experimental/LearningAgents/Source/Learning/Private/LearningFrameRangeSet.cpp b/Engine/Plugins/Experimental/LearningAgents/Source/Learning/Private/LearningFrameRangeSet.cpp
--- a/Engine/Plugins/Experimental/LearningAgents/Source/Learning/Private/LearningFrameRangeSet.cpp
+++ b/Engine/Plugins/Experimental/LearningAgents/Source/Learning/Private/LearningFrameRangeSet.cpp
@@ -35,19 +35,34 @@ namespace UE::Learning
 			}
 		}
 
-		static inline void FrameOffsetsCheck(
-			const TLearningArrayView<1, const int32> FrameOffsets,
+		static inline void OffsetsCheck(
+			const TLearningArrayView<1, const int32> Offsets,
 			const int32 MinimumOffset,
 			const int32 MaximumOffset)
 		{
-			const int32 FrameNum = FrameOffsets.Num();
+			const int32 FrameNum = Offsets.Num();
 			for (int32 FrameIdx = 0; FrameIdx < FrameNum - 1; FrameIdx++)
 			{
-				check(FrameOffsets[FrameIdx + 0] >= MinimumOffset);
-				check(FrameOffsets[FrameIdx + 1] >= MinimumOffset);
-				check(FrameOffsets[FrameIdx + 0] < MaximumOffset);
-				check(FrameOffsets[FrameIdx + 1] < MaximumOffset);
-				check(FrameOffsets[FrameIdx + 0] < FrameOffsets[FrameIdx + 1]);
+				check(Offsets[FrameIdx + 0] >= MinimumOffset);
+				check(Offsets[FrameIdx + 1] >= MinimumOffset);
+				check(Offsets[FrameIdx + 0] < MaximumOffset);
+				check(Offsets[FrameIdx + 1] < MaximumOffset);
+				check(Offsets[FrameIdx + 0] < Offsets[FrameIdx + 1]);
+			}
+		}
+
+		static inline void ComputeRangeOffsets(
+			TLearningArrayView<1, int32> Offsets,
+			const TLearningArrayView<1, const int32> Lengths,
+			const int32 InitialOffset = 0)
+		{
+			check(Offsets.Num() == Lengths.Num());
+			const int32 RangeNum = Lengths.Num();
+			int32 Offset = InitialOffset;
+			for (int32 RangeIdx = 0; RangeIdx < RangeNum; RangeIdx++)
+			{
+				Offsets[RangeIdx] = Offset;
+				Offset += Lengths[RangeIdx];
 			}
 		}
 
@@ -606,22 +621,6 @@ namespace UE::Learning
 			return OutIndex;
 		}
 
-		static inline void RangesFrameOffsets(
-			TLearningArrayView<1, int32> OutFrameOffsets,
-			const TLearningArrayView<1, const int32> RangeLengths,
-			const int32 EntryFrameOffset)
-		{
-			check(OutFrameOffsets.Num() == RangeLengths.Num());
-			
-			int32 Offset = EntryFrameOffset;
-			const int32 RangeNum = RangeLengths.Num();
-			for (int32 RangeIdx = 0; RangeIdx < RangeNum; RangeIdx++)
-			{
-				OutFrameOffsets[RangeIdx] = Offset;
-				Offset += RangeLengths[RangeIdx];
-			}
-		}
-
 		static inline int32 RangesTotalFrameNum(const TLearningArrayView<1, const int32> RangeLengths)
 		{
 			int32 Total = 0;
@@ -635,19 +634,19 @@ namespace UE::Learning
 			return Total;
 		}
 
-		static inline int32 RangesIntersectionWithFrameOffsets(
+		static inline int32 RangesIntersectionWithOffsets(
 			TLearningArrayView<1, int32> OutStarts,
 			TLearningArrayView<1, int32> OutLengths,
-			TLearningArrayView<1, int32> OutLhsFrameOffsets,
-			TLearningArrayView<1, int32> OutRhsFrameOffsets,
+			TLearningArrayView<1, int32> OutLhsOffsets,
+			TLearningArrayView<1, int32> OutRhsOffsets,
 			const TLearningArrayView<1, const int32> LhsStarts,
 			const TLearningArrayView<1, const int32> LhsLengths,
 			const TLearningArrayView<1, const int32> RhsStarts,
 			const TLearningArrayView<1, const int32> RhsLengths,
-			const int32 LhsEntryFrameOffset,
-			const int32 RhsEntryFrameOffset)
+			const int32 LhsEntryOffset,
+			const int32 RhsEntryOffset)
 		{
-			TRACE_CPUPROFILER_EVENT_SCOPE(Learning::FrameRangeSet::Private::RangesIntersectionWithFrameOffsets);
+			TRACE_CPUPROFILER_EVENT_SCOPE(Learning::FrameRangeSet::Private::RangesIntersectionWithOffsets);
 
 			RangesCheck(LhsStarts, LhsLengths);
 			RangesCheck(RhsStarts, RhsLengths);
@@ -656,10 +655,10 @@ namespace UE::Learning
 			{
 				Array::Copy(OutStarts.Slice(0, LhsStarts.Num()), LhsStarts);
 				Array::Copy(OutLengths.Slice(0, LhsStarts.Num()), LhsLengths);
-				RangesFrameOffsets(OutLhsFrameOffsets.Slice(0, LhsStarts.Num()), LhsLengths, LhsEntryFrameOffset);
-				RangesFrameOffsets(OutRhsFrameOffsets.Slice(0, RhsStarts.Num()), RhsLengths, RhsEntryFrameOffset);
-				FrameOffsetsCheck(OutLhsFrameOffsets.Slice(0, LhsStarts.Num()), LhsEntryFrameOffset, LhsEntryFrameOffset + RangesTotalFrameNum(LhsLengths));
-				FrameOffsetsCheck(OutRhsFrameOffsets.Slice(0, RhsStarts.Num()), RhsEntryFrameOffset, RhsEntryFrameOffset + RangesTotalFrameNum(RhsLengths));
+				ComputeRangeOffsets(OutLhsOffsets.Slice(0, LhsStarts.Num()), LhsLengths, LhsEntryOffset);
+				ComputeRangeOffsets(OutRhsOffsets.Slice(0, RhsStarts.Num()), RhsLengths, RhsEntryOffset);
+				OffsetsCheck(OutLhsOffsets.Slice(0, LhsStarts.Num()), LhsEntryOffset, LhsEntryOffset + RangesTotalFrameNum(LhsLengths));
+				OffsetsCheck(OutRhsOffsets.Slice(0, RhsStarts.Num()), RhsEntryOffset, RhsEntryOffset + RangesTotalFrameNum(RhsLengths));
 				return LhsLengths.Num();
 			}
 
@@ -677,9 +676,9 @@ namespace UE::Learning
 			int32 RhsIndex = 0;
 			int32 OutIndex = 0;
 
-			// Frame Offsets
-			int32 LhsFrameOffset = LhsEntryFrameOffset;
-			int32 RhsFrameOffset = RhsEntryFrameOffset;
+			// Offsets
+			int32 LhsOffset = LhsEntryOffset;
+			int32 RhsOffset = RhsEntryOffset;
 			int32 LhsActiveStart = INDEX_NONE;
 			int32 RhsActiveStart = INDEX_NONE;
 			int32 LhsActiveOffset = INDEX_NONE;
@@ -717,15 +716,15 @@ namespace UE::Learning
 
 						bOutActive = false;
 						OutLengths[OutIndex] = LhsT - OutStarts[OutIndex];
-						OutLhsFrameOffsets[OutIndex] = LhsActiveOffset + (OutStarts[OutIndex] - LhsActiveStart);
-						OutRhsFrameOffsets[OutIndex] = RhsActiveOffset + (OutStarts[OutIndex] - RhsActiveStart);
+						OutLhsOffsets[OutIndex] = LhsActiveOffset + (OutStarts[OutIndex] - LhsActiveStart);
+						OutRhsOffsets[OutIndex] = RhsActiveOffset + (OutStarts[OutIndex] - RhsActiveStart);
 						OutIndex++;
 					}
 
 					bLhsActive = bLhsActiveNext;
 					LhsActiveStart = bLhsActive ? LhsStarts[LhsIndex / 2] : INDEX_NONE;
-					LhsActiveOffset = bLhsActive ? LhsFrameOffset : INDEX_NONE;
-					LhsFrameOffset += bLhsActive ? 0 : LhsLengths[LhsIndex / 2];
+					LhsActiveOffset = bLhsActive ? LhsOffset : INDEX_NONE;
+					LhsOffset += bLhsActive ? 0 : LhsLengths[LhsIndex / 2];
 					LhsIndex++;
 				}
 				// Event from rhs coming first
@@ -749,15 +748,15 @@ namespace UE::Learning
 
 						bOutActive = false;
 						OutLengths[OutIndex] = RhsT - OutStarts[OutIndex];
-						OutLhsFrameOffsets[OutIndex] = LhsActiveOffset + (OutStarts[OutIndex] - LhsActiveStart);
-						OutRhsFrameOffsets[OutIndex] = RhsActiveOffset + (OutStarts[OutIndex] - RhsActiveStart);
+						OutLhsOffsets[OutIndex] = LhsActiveOffset + (OutStarts[OutIndex] - LhsActiveStart);
+						OutRhsOffsets[OutIndex] = RhsActiveOffset + (OutStarts[OutIndex] - RhsActiveStart);
 						OutIndex++;
 					}
 
 					bRhsActive = bRhsActiveNext;
 					RhsActiveStart = bRhsActive ? RhsStarts[RhsIndex / 2] : INDEX_NONE;
-					RhsActiveOffset = bRhsActive ? RhsFrameOffset : INDEX_NONE;
-					RhsFrameOffset += bRhsActive ? 0 : RhsLengths[RhsIndex / 2];
+					RhsActiveOffset = bRhsActive ? RhsOffset : INDEX_NONE;
+					RhsOffset += bRhsActive ? 0 : RhsLengths[RhsIndex / 2];
 					RhsIndex++;
 				}
 				// Event from lhs and rhs coming at same time
@@ -783,26 +782,26 @@ namespace UE::Learning
 
 						bOutActive = false;
 						OutLengths[OutIndex] = LhsT - OutStarts[OutIndex];
-						OutLhsFrameOffsets[OutIndex] = LhsActiveOffset + (OutStarts[OutIndex] - LhsActiveStart);
-						OutRhsFrameOffsets[OutIndex] = RhsActiveOffset + (OutStarts[OutIndex] - RhsActiveStart);
+						OutLhsOffsets[OutIndex] = LhsActiveOffset + (OutStarts[OutIndex] - LhsActiveStart);
+						OutRhsOffsets[OutIndex] = RhsActiveOffset + (OutStarts[OutIndex] - RhsActiveStart);
 						OutIndex++;
 					}
 
 					bLhsActive = bLhsActiveNext;
 					LhsActiveStart = bLhsActive ? LhsStarts[LhsIndex / 2] : INDEX_NONE;
-					LhsActiveOffset = bLhsActive ? LhsFrameOffset : INDEX_NONE;
+					LhsActiveOffset = bLhsActive ? LhsOffset : INDEX_NONE;
 					bRhsActive = bRhsActiveNext;
 					RhsActiveStart = bRhsActive ? RhsStarts[RhsIndex / 2] : INDEX_NONE;
-					RhsActiveOffset = bRhsActive ? RhsFrameOffset : INDEX_NONE;
-					LhsFrameOffset += bLhsActive ? 0 : LhsLengths[LhsIndex / 2];
-					RhsFrameOffset += bRhsActive ? 0 : RhsLengths[RhsIndex / 2];
+					RhsActiveOffset = bRhsActive ? RhsOffset : INDEX_NONE;
+					LhsOffset += bLhsActive ? 0 : LhsLengths[LhsIndex / 2];
+					RhsOffset += bRhsActive ? 0 : RhsLengths[RhsIndex / 2];
 					LhsIndex++; RhsIndex++;
 				}
 			}
 
 			RangesCheck(OutStarts.Slice(0, OutIndex), OutLengths.Slice(0, OutIndex));
-			FrameOffsetsCheck(OutLhsFrameOffsets.Slice(0, OutIndex), LhsEntryFrameOffset, LhsEntryFrameOffset + RangesTotalFrameNum(LhsLengths));
-			FrameOffsetsCheck(OutRhsFrameOffsets.Slice(0, OutIndex), RhsEntryFrameOffset, RhsEntryFrameOffset + RangesTotalFrameNum(RhsLengths));
+			OffsetsCheck(OutLhsOffsets.Slice(0, OutIndex), LhsEntryOffset, LhsEntryOffset + RangesTotalFrameNum(LhsLengths));
+			OffsetsCheck(OutRhsOffsets.Slice(0, OutIndex), RhsEntryOffset, RhsEntryOffset + RangesTotalFrameNum(RhsLengths));
 
 			// Return number of ranges added
 			return OutIndex;
@@ -949,6 +948,11 @@ namespace UE::Learning
 
 	void FFrameRangeSet::Check() const
 	{
+		check(EntrySequences.Num() == EntryRangeOffsets.Num());
+		check(EntrySequences.Num() == EntryRangeNums.Num());
+		check(RangeStarts.Num() == RangeLengths.Num());
+		check(RangeStarts.Num() == RangeOffsets.Num());
+
 		const int32 EntryNum = GetEntryNum();
 		for (int32 EntryIdx = 0; EntryIdx < EntryNum - 1; EntryIdx++)
 		{
@@ -971,6 +975,7 @@ namespace UE::Learning
 		EntryRangeNums.Empty();
 		RangeStarts.Empty();
 		RangeLengths.Empty();
+		RangeOffsets.Empty();
 	}
 
 	int32 FFrameRangeSet::GetEntryNum() const { return EntrySequences.Num(); }
@@ -981,42 +986,20 @@ namespace UE::Learning
 
 	int32 FFrameRangeSet::GetEntryTotalFrameNum(const int32 EntryIdx) const
 	{
-		int32 Total = 0;
-		const int32 RangeNum = GetEntryRangeNum(EntryIdx);
-		for (int32 RangeIdx = 0; RangeIdx < RangeNum; RangeIdx++)
-		{
-			Total += GetEntryRangeLength(EntryIdx, RangeIdx);
-		}
-		return Total;
-	}
-
-	int32 FFrameRangeSet::GetEntryFrameOffset(const int32 EntryIdx) const
-	{
-		int32 Offset = 0;
-		for (int32 PrevEntryIdx = 0; PrevEntryIdx < EntryIdx; PrevEntryIdx++)
-		{
-			Offset += GetEntryTotalFrameNum(PrevEntryIdx);
-		}
-
-		return Offset;
+		const int32 EntryRangeNum = GetEntryRangeNum(EntryIdx);
+		return 
+			EntryRangeNum == 0 ? 0 : 
+			GetEntryRangeOffset(EntryIdx, EntryRangeNum - 1) - 
+			GetEntryRangeOffset(EntryIdx, 0) + 
+			GetEntryRangeLength(EntryIdx, EntryRangeNum - 1);
 	}
 
 	TLearningArrayView<1, const int32> FFrameRangeSet::GetEntryRangeStarts(const int32 EntryIdx) const { return RangeStarts.Slice(EntryRangeOffsets[EntryIdx], EntryRangeNums[EntryIdx]); }
 	TLearningArrayView<1, const int32> FFrameRangeSet::GetEntryRangeLengths(const int32 EntryIdx) const { return RangeLengths.Slice(EntryRangeOffsets[EntryIdx], EntryRangeNums[EntryIdx]); }
+	TLearningArrayView<1, const int32> FFrameRangeSet::GetEntryRangeOffsets(const int32 EntryIdx) const { return RangeOffsets.Slice(EntryRangeOffsets[EntryIdx], EntryRangeNums[EntryIdx]); }
 	int32 FFrameRangeSet::GetEntryRangeStart(const int32 EntryIdx, const int32 RangeIdx) const { return RangeStarts[EntryRangeOffsets[EntryIdx] + RangeIdx]; }
 	int32 FFrameRangeSet::GetEntryRangeLength(const int32 EntryIdx, const int32 RangeIdx) const { return RangeLengths[EntryRangeOffsets[EntryIdx] + RangeIdx]; }
-	
-	int32 FFrameRangeSet::GetEntryRangeFrameOffset(const int32 EntryIdx, const int32 RangeIdx) const
-	{
-		int32 Offset = GetEntryFrameOffset(EntryIdx);
-
-		for (int32 PrevRangeIdx = 0; PrevRangeIdx < RangeIdx; PrevRangeIdx++)
-		{
-			Offset += GetEntryRangeLength(EntryIdx, PrevRangeIdx);
-		}
-
-		return Offset;
-	}
+	int32 FFrameRangeSet::GetEntryRangeOffset(const int32 EntryIdx, const int32 RangeIdx) const { return RangeOffsets[EntryRangeOffsets[EntryIdx] + RangeIdx]; }
 
 	float FFrameRangeSet::GetEntryRangeStartTime(const int32 EntryIdx, const int32 RangeIdx, const float FrameDeltaTime) const
 	{
@@ -1037,16 +1020,12 @@ namespace UE::Learning
 	
 	TLearningArrayView<1, const int32> FFrameRangeSet::GetAllRangeStarts() const { return RangeStarts; }
 	TLearningArrayView<1, const int32> FFrameRangeSet::GetAllRangeLengths() const { return RangeLengths; }
+	TLearningArrayView<1, const int32> FFrameRangeSet::GetAllRangeOffsets() const { return RangeOffsets; }
 
 	int32 FFrameRangeSet::GetTotalFrameNum() const
 	{
-		int32 Total = 0;
-		const int32 RangeNum = RangeLengths.Num();
-		for (int32 RangeIdx = 0; RangeIdx < RangeNum; RangeIdx++)
-		{
-			Total += RangeLengths[RangeIdx];
-		}
-		return Total;
+		const int32 RangeNum = GetTotalRangeNum();
+		return RangeNum == 0 ? 0 : RangeOffsets[RangeNum - 1] + RangeLengths[RangeNum - 1];
 	}
 
 	bool FFrameRangeSet::ContainsSequence(const int32 Sequence) const
@@ -1137,11 +1116,10 @@ namespace UE::Learning
 		return false;
 	}
 
-	bool FFrameRangeSet::FindFrameOffset(int32& OutEntryIdx, int32& OutRangeIdx, int32& OutRangeFrame, const int32 FrameOffset) const
+	bool FFrameRangeSet::FindOffset(int32& OutEntryIdx, int32& OutRangeIdx, int32& OutRangeFrame, const int32 Offset) const
 	{
 		const int32 EntryNum = GetEntryNum();
 
-		int32 Offset = 0;
 		for (int32 EntryIdx = 0; EntryIdx < EntryNum; EntryIdx++)
 		{
 			const int32 Sequence = GetEntrySequence(EntryIdx);
@@ -1151,16 +1129,15 @@ namespace UE::Learning
 			{
 				const int32 RangeStart = GetEntryRangeStart(EntryIdx, RangeIdx);
 				const int32 RangeLength = GetEntryRangeLength(EntryIdx, RangeIdx);
+				const int32 RangeOffset = GetEntryRangeOffset(EntryIdx, RangeIdx);
 
-				if (FrameOffset >= Offset && FrameOffset < Offset + RangeLength)
+				if (Offset >= RangeOffset && Offset < RangeOffset + RangeLength)
 				{
 					OutEntryIdx = EntryIdx;
 					OutRangeIdx = RangeIdx;
-					OutRangeFrame = RangeStart + FrameOffset - Offset;
+					OutRangeFrame = Offset - RangeOffset;
 					return true;
 				}
-
-				Offset += RangeLength;
 			}
 		}
 
@@ -1185,8 +1162,13 @@ namespace UE::Learning
 		const int32 AddRangeNum = InStarts.Num();
 		RangeStarts.SetNumUninitialized({ CurrRangeNum + AddRangeNum });
 		RangeLengths.SetNumUninitialized({ CurrRangeNum + AddRangeNum });
+		RangeOffsets.SetNumUninitialized({ CurrRangeNum + AddRangeNum });
 		Array::Copy(RangeStarts.Slice(CurrRangeNum, AddRangeNum), InStarts);
 		Array::Copy(RangeLengths.Slice(CurrRangeNum, AddRangeNum), InLengths);
+		for (int32 Idx = CurrRangeNum; Idx < CurrRangeNum + AddRangeNum; Idx++)
+		{
+			RangeOffsets[Idx] = Idx == 0 ? 0 : RangeOffsets[Idx - 1] + RangeLengths[Idx - 1];
+		}
 
 		const int32 CurrEntryNum = EntrySequences.Num();
 		EntrySequences.SetNumUninitialized({ CurrEntryNum + 1 });
@@ -1210,7 +1192,8 @@ namespace UE::Learning
 				Array::Equal(Lhs.EntryRangeOffsets, Rhs.EntryRangeOffsets) &&
 				Array::Equal(Lhs.EntryRangeNums, Rhs.EntryRangeNums) &&
 				Array::Equal(Lhs.RangeStarts, Rhs.RangeStarts) &&
-				Array::Equal(Lhs.RangeLengths, Rhs.RangeLengths);
+				Array::Equal(Lhs.RangeLengths, Rhs.RangeLengths) &&
+				Array::Equal(Lhs.RangeOffsets, Rhs.RangeOffsets);
 		}
 
 		void Intersection(FFrameSet& OutFrameSet, const FFrameSet& FrameSet, const FFrameRangeSet& FrameRangeSet)
@@ -1552,6 +1535,8 @@ namespace UE::Learning
 			Out.EntryRangeNums.SetNumUninitialized({ OutIndex });
 			Out.RangeStarts.SetNumUninitialized({ RangeIndex });
 			Out.RangeLengths.SetNumUninitialized({ RangeIndex });
+			Out.RangeOffsets.SetNumUninitialized({ RangeIndex });
+			Private::ComputeRangeOffsets(Out.RangeOffsets, Out.RangeLengths);
 			Out.Check();
 		}
 
@@ -1632,6 +1617,8 @@ namespace UE::Learning
 			Out.EntryRangeNums.SetNumUninitialized({ OutIndex });
 			Out.RangeStarts.SetNumUninitialized({ RangeIndex });
 			Out.RangeLengths.SetNumUninitialized({ RangeIndex });
+			Out.RangeOffsets.SetNumUninitialized({ RangeIndex });
+			Private::ComputeRangeOffsets(Out.RangeOffsets, Out.RangeLengths);
 			Out.Check();
 		}
 
@@ -1752,25 +1739,27 @@ namespace UE::Learning
 			Out.EntryRangeNums.SetNumUninitialized({ OutIndex });
 			Out.RangeStarts.SetNumUninitialized({ RangeIndex });
 			Out.RangeLengths.SetNumUninitialized({ RangeIndex });
+			Out.RangeOffsets.SetNumUninitialized({ RangeIndex });
+			Private::ComputeRangeOffsets(Out.RangeOffsets, Out.RangeLengths);
 			Out.Check();
 		}
 
-		int32 IntersectionWithFrameOffsets(
+		int32 IntersectionWithOffsets(
 			FFrameRangeSet& Out,
-			TLearningArrayView<1, int32> OutLhsFrameOffsets,
-			TLearningArrayView<1, int32> OutRhsFrameOffsets,
+			TLearningArrayView<1, int32> OutLhsOffsets,
+			TLearningArrayView<1, int32> OutRhsOffsets,
 			const FFrameRangeSet& Lhs,
 			const FFrameRangeSet& Rhs)
 		{
-			TRACE_CPUPROFILER_EVENT_SCOPE(Learning::FrameRangeSet::IntersectionWithFrameOffsets);
+			TRACE_CPUPROFILER_EVENT_SCOPE(Learning::FrameRangeSet::IntersectionWithOffsets);
 
 			Lhs.Check(); Rhs.Check();
 
 			if (Equal(Lhs, Rhs))
 			{
 				Out = Lhs;
-				AllRangeFrameOffsets(OutLhsFrameOffsets.Slice(0, Lhs.GetTotalRangeNum()), Lhs);
-				AllRangeFrameOffsets(OutRhsFrameOffsets.Slice(0, Rhs.GetTotalRangeNum()), Rhs);
+				Array::Copy(OutLhsOffsets.Slice(0, Lhs.GetTotalRangeNum()), Lhs.GetAllRangeOffsets());
+				Array::Copy(OutRhsOffsets.Slice(0, Rhs.GetTotalRangeNum()), Rhs.GetAllRangeOffsets());
 				return Out.GetTotalRangeNum();
 			}
 
@@ -1796,8 +1785,8 @@ namespace UE::Learning
 
 			// Frame Offsets
 			int32 OutIntersectionIdx = 0;
-			int32 LhsFrameOffset = 0;
-			int32 RhsFrameOffset = 0;
+			int32 LhsOffset = 0;
+			int32 RhsOffset = 0;
 
 			// While both sets have entries
 			while (LhsIndex < LhsEntryNum && RhsIndex < RhsEntryNum)
@@ -1805,13 +1794,13 @@ namespace UE::Learning
 				// If entry is in lhs but not rhs skip
 				if (Lhs.GetEntrySequence(LhsIndex) < Rhs.GetEntrySequence(RhsIndex))
 				{
-					LhsFrameOffset += Lhs.GetEntryTotalFrameNum(LhsIndex);
+					LhsOffset += Lhs.GetEntryTotalFrameNum(LhsIndex);
 					LhsIndex++;
 				}
 				// If entry is in rhs but not lhs skip
 				else if (Rhs.GetEntrySequence(RhsIndex) < Lhs.GetEntrySequence(LhsIndex))
 				{
-					RhsFrameOffset += Rhs.GetEntryTotalFrameNum(RhsIndex);
+					RhsOffset += Rhs.GetEntryTotalFrameNum(RhsIndex);
 					RhsIndex++;
 				}
 				// If entry is in both lhs and rhs
@@ -1821,20 +1810,20 @@ namespace UE::Learning
 					check(Rhs.GetEntryRangeNum(RhsIndex) > 0);
 
 					// Append intersection of subranges to output
-					const int32 RangeNum = FrameRangeSet::Private::RangesIntersectionWithFrameOffsets(
+					const int32 RangeNum = FrameRangeSet::Private::RangesIntersectionWithOffsets(
 						Out.RangeStarts.Slice(RangeIndex, Out.RangeStarts.Num() - RangeIndex),
 						Out.RangeLengths.Slice(RangeIndex, Out.RangeLengths.Num() - RangeIndex),
-						OutLhsFrameOffsets.Slice(RangeIndex, OutLhsFrameOffsets.Num() - RangeIndex),
-						OutRhsFrameOffsets.Slice(RangeIndex, OutRhsFrameOffsets.Num() - RangeIndex),
+						OutLhsOffsets.Slice(RangeIndex, OutLhsOffsets.Num() - RangeIndex),
+						OutRhsOffsets.Slice(RangeIndex, OutRhsOffsets.Num() - RangeIndex),
 						Lhs.GetEntryRangeStarts(LhsIndex),
 						Lhs.GetEntryRangeLengths(LhsIndex),
 						Rhs.GetEntryRangeStarts(RhsIndex),
 						Rhs.GetEntryRangeLengths(RhsIndex),
-						LhsFrameOffset,
-						RhsFrameOffset);
+						LhsOffset,
+						RhsOffset);
 
-					Private::FrameOffsetsCheck(OutLhsFrameOffsets.Slice(0, RangeIndex + RangeNum), 0, LhsFrameOffset + Lhs.GetEntryTotalFrameNum(LhsIndex));
-					Private::FrameOffsetsCheck(OutRhsFrameOffsets.Slice(0, RangeIndex + RangeNum), 0, RhsFrameOffset + Rhs.GetEntryTotalFrameNum(RhsIndex));
+					Private::OffsetsCheck(OutLhsOffsets.Slice(0, RangeIndex + RangeNum), 0, LhsOffset + Lhs.GetEntryTotalFrameNum(LhsIndex));
+					Private::OffsetsCheck(OutRhsOffsets.Slice(0, RangeIndex + RangeNum), 0, RhsOffset + Rhs.GetEntryTotalFrameNum(RhsIndex));
 
 					check(RangeNum <= Lhs.GetEntryRangeNum(LhsIndex) + Rhs.GetEntryRangeNum(RhsIndex));
 
@@ -1847,8 +1836,8 @@ namespace UE::Learning
 						OutIndex++;
 					}
 
-					LhsFrameOffset += Lhs.GetEntryTotalFrameNum(LhsIndex);
-					RhsFrameOffset += Rhs.GetEntryTotalFrameNum(RhsIndex);
+					LhsOffset += Lhs.GetEntryTotalFrameNum(LhsIndex);
+					RhsOffset += Rhs.GetEntryTotalFrameNum(RhsIndex);
 					LhsIndex++; RhsIndex++;
 				}
 			}
@@ -1859,10 +1848,12 @@ namespace UE::Learning
 			Out.EntryRangeNums.SetNumUninitialized({ OutIndex });
 			Out.RangeStarts.SetNumUninitialized({ RangeIndex });
 			Out.RangeLengths.SetNumUninitialized({ RangeIndex });
+			Out.RangeOffsets.SetNumUninitialized({ RangeIndex });
+			Private::ComputeRangeOffsets(Out.RangeOffsets, Out.RangeLengths);
 			Out.Check();
 
-			Private::FrameOffsetsCheck(OutLhsFrameOffsets.Slice(0, RangeIndex), 0, Lhs.GetTotalFrameNum());
-			Private::FrameOffsetsCheck(OutRhsFrameOffsets.Slice(0, RangeIndex), 0, Rhs.GetTotalFrameNum());
+			Private::OffsetsCheck(OutLhsOffsets.Slice(0, RangeIndex), 0, Lhs.GetTotalFrameNum());
+			Private::OffsetsCheck(OutRhsOffsets.Slice(0, RangeIndex), 0, Rhs.GetTotalFrameNum());
 
 			check(RangeIndex == Out.GetTotalRangeNum());
 			return RangeIndex;
@@ -1979,6 +1970,18 @@ namespace UE::Learning
 			Out.Check();
 		}
 
+		void MakeFromFrameSet(FFrameRangeSet& OutFrameRangeSet, const FFrameSet& FrameSet)
+		{
+			OutFrameRangeSet.EntrySequences = FrameSet.EntrySequences;
+			OutFrameRangeSet.EntryRangeNums = FrameSet.EntryFrameNums;
+			OutFrameRangeSet.EntryRangeOffsets = FrameSet.EntryFrameOffsets;
+			OutFrameRangeSet.RangeStarts = FrameSet.Frames;
+			OutFrameRangeSet.RangeLengths.SetNumUninitialized({ FrameSet.GetTotalFrameNum() });
+			OutFrameRangeSet.RangeOffsets.SetNumUninitialized({ FrameSet.GetTotalFrameNum() });
+			Array::Set(OutFrameRangeSet.RangeLengths, 1);
+			Private::ComputeRangeOffsets(OutFrameRangeSet.RangeOffsets, OutFrameRangeSet.RangeLengths);
+		}
+
 		void MakeFrameSetFromRangeStarts(FFrameSet& OutFrameSet, const FFrameRangeSet& FrameRangeSet)
 		{
 			OutFrameSet.EntrySequences = FrameRangeSet.EntrySequences;
@@ -2156,28 +2159,6 @@ namespace UE::Learning
 			check(RangeOffset == TotalRangeNum);
 		}
 
-		void AllRangeFrameOffsets(
-			TLearningArrayView<1, int32> OutRangeFrameOffsets,
-			const FFrameRangeSet& FrameRangeSet)
-		{
-			check(OutRangeFrameOffsets.Num() == FrameRangeSet.GetTotalRangeNum());
-
-			const int32 TotalRangeNum = FrameRangeSet.GetTotalRangeNum();
-			const int32 TotalFrameNum = FrameRangeSet.GetTotalFrameNum();
-			const int32 EntryNum = FrameRangeSet.GetEntryNum();
-
-			const int32 RangeNum = FrameRangeSet.GetTotalRangeNum();
-
-			int32 FrameOffset = 0;
-			for (int32 RangeIdx = 0; RangeIdx < RangeNum; RangeIdx++)
-			{
-				OutRangeFrameOffsets[RangeIdx] = FrameOffset;
-				FrameOffset += FrameRangeSet.RangeLengths[RangeIdx];
-			}
-
-			check(FrameOffset == TotalFrameNum);
-		}
-
 		void AllRangeStartTimes(
 			TLearningArrayView<1, float> OutRangeStartTimes,
 			const FFrameRangeSet& FrameRangeSet, 
@@ -2225,13 +2206,11 @@ namespace UE::Learning
 			const TFunctionRef<void(
 				const int32 TotalRangeIdx,
 				const int32 EntryIdx,
-				const int32 RangeIdx,
-				const int32 FrameOffset)> Body)
+				const int32 RangeIdx)> Body)
 		{
 			const int32 EntryNum = FrameRangeSet.GetEntryNum();
 
 			int32 TotalRangeIdx = 0;
-			int32 FrameOffset = 0;
 			for (int32 EntryIdx = 0; EntryIdx < EntryNum; EntryIdx++)
 			{
 				const int32 RangeNum = FrameRangeSet.GetEntryRangeNum(EntryIdx);
@@ -2241,11 +2220,9 @@ namespace UE::Learning
 					Body(
 						TotalRangeIdx,
 						EntryIdx,
-						RangeIdx,
-						FrameOffset);
+						RangeIdx);
 
 					TotalRangeIdx++;
-					FrameOffset += FrameRangeSet.GetEntryRangeLength(EntryIdx, RangeIdx);
 				}
 			}
 		}
@@ -2255,29 +2232,24 @@ namespace UE::Learning
 			const TFunctionRef<void(
 				const int32 TotalRangeIdx,
 				const int32 EntryIdx,
-				const int32 RangeIdx,
-				const int32 FrameOffset)> Body)
+				const int32 RangeIdx)> Body)
 		{
 			const int32 TotalRangeNum = FrameRangeSet.GetTotalRangeNum();
 
 			TLearningArray<1, int32> RangeEntries;
 			TLearningArray<1, int32> RangeIndices;
-			TLearningArray<1, int32> RangeFrameOffsets;
 
 			RangeEntries.SetNumUninitialized({ TotalRangeNum });
 			RangeIndices.SetNumUninitialized({ TotalRangeNum });
-			RangeFrameOffsets.SetNumUninitialized({ TotalRangeNum });
 
 			AllRangeEntries(RangeEntries, FrameRangeSet);
 			AllRangeIndices(RangeIndices, FrameRangeSet);
-			AllRangeFrameOffsets(RangeFrameOffsets, FrameRangeSet);
 
-			ParallelFor(TotalRangeNum, [&Body, &RangeEntries, &RangeIndices, &RangeFrameOffsets](const int32 TotalRangeIdx) {
+			ParallelFor(TotalRangeNum, [&Body, &RangeEntries, &RangeIndices](const int32 TotalRangeIdx) {
 				Body(
 					TotalRangeIdx,
 					RangeEntries[TotalRangeIdx],
-					RangeIndices[TotalRangeIdx],
-					RangeFrameOffsets[TotalRangeIdx]);
+					RangeIndices[TotalRangeIdx]);
 			});
 		}
 	}
diff --git a/Engine/Plugins/Experimental/LearningAgents/Source/Learning/Private/LearningFrameSet.cpp b/Engine/Plugins/Experimental/LearningAgents/Source/Learning/Private/LearningFrameSet.cpp
--- a/Engine/Plugins/Experimental/LearningAgents/Source/Learning/Private/LearningFrameSet.cpp
+++ b/Engine/Plugins/Experimental/LearningAgents/Source/Learning/Private/LearningFrameSet.cpp
@@ -210,6 +210,9 @@ namespace UE::Learning
 
 	void FFrameSet::Check() const
 	{
+		check(EntrySequences.Num() == EntryFrameOffsets.Num());
+		check(EntrySequences.Num() == EntryFrameNums.Num());
+
 		const int32 EntryNum = EntrySequences.Num();
 
 		for (int32 EntryIdx = 0; EntryIdx < EntryNum - 1; EntryIdx++)
@@ -242,7 +245,7 @@ namespace UE::Learning
 	TLearningArrayView<1, const int32> FFrameSet::GetEntryFrames(const int32 EntryIdx) const { return Frames.Slice(EntryFrameOffsets[EntryIdx], EntryFrameNums[EntryIdx]); }
 	int32 FFrameSet::GetEntryFrame(const int32 EntryIdx, const int32 FrameIdx) const { return Frames[EntryFrameOffsets[EntryIdx] + FrameIdx]; }
 	float FFrameSet::GetEntryFrameTime(const int32 EntryIdx, const int32 FrameIdx, const float FrameDeltaTime) const { return GetEntryFrame(EntryIdx, FrameIdx) * FrameDeltaTime; }
-	int32 FFrameSet::GetEntryFrameOffset(const int32 EntryIdx) const { return EntryFrameOffsets[EntryIdx]; }
+	int32 FFrameSet::GetEntryOffset(const int32 EntryIdx) const { return EntryFrameOffsets[EntryIdx]; }
 	int32 FFrameSet::GetTotalFrameNum() const { return Frames.Num(); }
 	
 	bool FFrameSet::ContainsSequence(const int32 Sequence) const
@@ -352,23 +355,20 @@ namespace UE::Learning
 		return false;
 	}
 
-	bool FFrameSet::FindFrameOffset(int32& OutEntryIdx, int32& OutFrameIdx, const int32 FrameOffset) const
+	bool FFrameSet::FindOffset(int32& OutEntryIdx, int32& OutFrameIdx, const int32 Offset) const
 	{
-		int32 Offset = 0;
-
 		const int32 EntryNum = GetEntryNum();
 		for (int32 EntryIdx = 0; EntryIdx < EntryNum; EntryIdx++)
 		{
 			const int32 FrameNum = GetEntryFrameNum(EntryIdx);
+			const int32 FrameOffset = GetEntryOffset(EntryIdx);
 
-			if (FrameOffset >= Offset && FrameOffset < FrameNum)
+			if (Offset >= FrameOffset && Offset < FrameOffset + FrameNum)
 			{
 				OutEntryIdx = EntryIdx;
-				OutFrameIdx = FrameOffset - Offset;
+				OutFrameIdx = Offset - FrameOffset;
 				return true;
 			}
-
-			Offset += FrameNum;
 		}
 
 		OutEntryIdx = INDEX_NONE;
@@ -802,7 +802,7 @@ namespace UE::Learning
 			for (int32 EntryIdx = 0; EntryIdx < EntryNum; EntryIdx++)
 			{
 				const int32 FrameNum = FrameSet.GetEntryFrameNum(EntryIdx);
-				const int32 FrameOffset = FrameSet.GetEntryFrameOffset(EntryIdx);
+				const int32 FrameOffset = FrameSet.GetEntryOffset(EntryIdx);
 				for (int32 FrameIdx = 0; FrameIdx < FrameNum; FrameIdx++)
 				{
 					Body(
diff --git a/Engine/Plugins/Experimental/LearningAgents/Source/Learning/Private/LearningRandom.cpp b/Engine/Plugins/Experimental/LearningAgents/Source/Learning/Private/LearningRandom.cpp
--- a/Engine/Plugins/Experimental/LearningAgents/Source/Learning/Private/LearningRandom.cpp
+++ b/Engine/Plugins/Experimental/LearningAgents/Source/Learning/Private/LearningRandom.cpp
@@ -576,7 +576,7 @@ namespace UE::Learning::Random
 	{
 		const int32 TotalFrameNum = FrameSet.GetTotalFrameNum();
 		const int32 FrameOffset = IntInRange(0x26810f69 ^ State, 0, TotalFrameNum - 1);
-		const bool bFound = FrameSet.FindFrameOffset(OutEntryIdx, OutFrameIdx, FrameOffset);
+		const bool bFound = FrameSet.FindOffset(OutEntryIdx, OutFrameIdx, FrameOffset);
 		check(bFound);
 	}
 
@@ -589,7 +589,7 @@ namespace UE::Learning::Random
 	{
 		const int32 TotalFrameNum = FrameRangeSet.GetTotalFrameNum();
 		const int32 FrameOffset = IntInRange(0x76ad14bb ^ State, 0, TotalFrameNum - 1);
-		const bool bFound = FrameRangeSet.FindFrameOffset(OutEntryIdx, OutRangeIdx, OutRangeFrame, FrameOffset);
+		const bool bFound = FrameRangeSet.FindOffset(OutEntryIdx, OutRangeIdx, OutRangeFrame, FrameOffset);
 		check(bFound);
 	}
 
diff --git a/Engine/Plugins/Experimental/LearningAgents/Source/Learning/Public/LearningFrameAttribute.h b/Engine/Plugins/Experimental/LearningAgents/Source/Learning/Public/LearningFrameAttribute.h
new file mode 100644
--- /dev/null
+++ b/Engine/Plugins/Experimental/LearningAgents/Source/Learning/Public/LearningFrameAttribute.h
@@ -0,0 +1,396 @@
+// Copyright Epic Games, Inc. All Rights Reserved.
+
+#pragma once
+
+#include "LearningArray.h"
+#include "LearningFrameRangeSet.h"
+
+#include "Math/Vector.h"
+#include "Math/Quat.h"
+
+namespace UE::Learning
+{
+
+	/**
+	 * A FFrameAttribute represents a attribute associated with every frame in a FFrameRangeSet. An attribute is made up of multiple "channels" 
+	 * such as the X, Y, Z components of a location.
+	 * 
+	 * The data for the attribute is stored in one large flat array of shape (ChannelNum, TotalFrameNum). This means the data is stored in a way
+	 * designed for SoA access by default. Some helper functions are provided for accessing slices of this data according to various range properties.
+	 * 
+	 * Also provided are many "operators" which allow for batched operations on attributes to create new attributes such as adding or subtracting 
+	 * attribute values. This operations are very efficient as they are computed across all frames in the attribute in batch and can be optimized 
+	 * using ISPC. If you perform a binary operation on two channels with different frames in the range sets then it will construct a new channel
+	 * which is the intersection of those two inputs.
+	 */
+	struct LEARNING_API FFrameAttribute
+	{
+	public:
+
+		/** Check if the FrameAttribute is well-formed. */
+		void Check() const;
+
+		/** True if the FrameAttribute is Empty, otherwise false */
+		bool IsEmpty() const;
+
+		/** Empties the FrameAttribute */
+		void Empty();
+
+		/** Gets the internal FrameRangeSet associated to this attribute */
+		const FFrameRangeSet& GetFrameRangeSet() const;
+
+		/** Gets the total number of frames for this attribute */
+		int32 GetTotalFrameNum() const;
+
+		/** Gets the total number of ranges for this attribute */
+		int32 GetTotalRangeNum() const;
+
+		/** Gets the number of channels in this attribute */
+		int32 GetChannelNum() const;
+
+		/** Gets a view of the complete large array of attribute data stored as (ChannelNum, TotalFrameNum) */
+		TLearningArrayView<2, const float> GetAttributeData() const;
+
+		/** Gets a view of the complete attribute data for a single channel */
+		TLearningArrayView<1, const float> GetChannelAttributeData(const int32 ChannelIdx) const;
+
+		/** Gets the attribute value for a given channel and frame index */
+		const float& GetChannelAttributeDataAtFrame(const int32 ChannelIdx, const int32 FrameIdx) const;
+
+		/** Gets the attribute data associated with a single channel, and range */
+		TLearningArrayView<1, const float> GetChannelEntryRangeAttributeData(const int32 ChannelIdx, const int32 EntryIdx, const int32 RangeIdx) const;
+
+		/** Gets the attribute data associated with range offset and length */
+		TLearningArrayView<1, const float> GetChannelRangeAttributeData(const int32 ChannelIdx, const int32 RangeOffset, const int32 RangeLength) const;
+
+		/** Gets a view of the complete large array of attribute data stored as (ChannelNum, TotalFrameNum) */
+		TLearningArrayView<2, float> GetAttributeData();
+
+		/** Gets a view of the complete attribute data for a single channel */
+		TLearningArrayView<1, float> GetChannelAttributeData(const int32 ChannelIdx);
+
+		/** Gets the attribute value for a given channel and frame index */
+		float& GetChannelAttributeDataAtFrame(const int32 ChannelIdx, const int32 FrameIdx);
+
+		/** Gets the attribute data associated with a single channel, and range */
+		TLearningArrayView<1, float> GetChannelEntryRangeAttributeData(const int32 ChannelIdx, const int32 EntryIdx, const int32 RangeIdx);
+
+		/** Gets the attribute data associated with range offset and length */
+		TLearningArrayView<1, float> GetChannelRangeAttributeData(const int32 ChannelIdx, const int32 RangeOffset, const int32 RangeLength);
+
+	public:
+
+		/** The internal associate FrameRangeSet */
+		FFrameRangeSet FrameRangeSet;
+
+		/** The large flat array of attribute data of shape (ChannelNum, TotalFrameNum) */
+		TLearningArray<2, float> AttributeData;
+	};
+
+	namespace FrameAttribute
+	{
+		/** Reduce op function. Takes as input a single frame attribute and a set of ranges and lengths for that attribute. */
+		using ReduceOpFunction = TFunctionRef<void(
+			const FFrameAttribute& In,
+			const TLearningArrayView<1, const int32> RangeOffsets,
+			const TLearningArrayView<1, const int32> RangeLengths)>;
+
+		/** Nullary op function. Produces as output a single frame attribute given a set of ranges and lengths for that attribute. */
+		using NullaryOpFunction = TFunctionRef<void(
+			FFrameAttribute& Out,
+			const TLearningArrayView<1, const int32> RangeOffsets,
+			const TLearningArrayView<1, const int32> RangeLengths)>;
+
+		/** Unary op function. Takes a single frame attribute as input and produces as output a single frame attribute. */
+		using UnaryOpFunction = TFunctionRef<void(
+			FFrameAttribute& Out,
+			const FFrameAttribute& In,
+			const TLearningArrayView<1, const int32> RangeOffsets,
+			const TLearningArrayView<1, const int32> RangeLengths)>;
+
+		/** Binary op function. Takes a two frame attributes as input and produces as output a single frame attribute. */
+		using BinaryOpFunction = TFunctionRef<void(
+			FFrameAttribute& Out,
+			const FFrameAttribute& Lhs,
+			const FFrameAttribute& Rhs,
+			const TLearningArrayView<1, const int32> OutRangeOffsets,
+			const TLearningArrayView<1, const int32> LhsRangeOffsets,
+			const TLearningArrayView<1, const int32> RhsRangeOffsets,
+			const TLearningArrayView<1, const int32> RangeLengths)>;
+
+		/** Convenience type to make the definition and usage of NaryOp easier */
+		using ConstFrameAttributePtr = const FFrameAttribute*;
+
+		/** Nary op function. Takes a multiple frame attributes as input and produces as output a single frame attribute. */
+		using NaryOpFunction = TFunctionRef<void(
+			FFrameAttribute& Out,
+			const TArrayView<const ConstFrameAttributePtr> In,
+			const TLearningArrayView<1, const int32> OutRangeOffsets,
+			const TArrayView<const TLearningArrayView<1, const int32>> InRangeOffsets,
+			const TLearningArrayView<1, const int32> RangeLengths)>;
+
+		/** Computes the intersection of a frame attribute and frame range set */
+		LEARNING_API void Intersection(FFrameAttribute& OutFrameAttribute, const FFrameAttribute& FrameAttribute, const FFrameRangeSet& FrameRangeSet);
+
+		/** Computes the frame range set where the given channel is non-zero in the frame attribute. */
+		LEARNING_API void NonZeroFrameRangeSet(FFrameRangeSet& OutFrameRangeSet, const FFrameAttribute& FrameAttribute, const int32 ChannelIdx);
+
+		/** Computes a reduction on a frame attribute */
+		LEARNING_API void ReduceOp(
+			const FFrameAttribute& In,
+			const ReduceOpFunction Op);
+
+		/** Creates a frame attribute from zero arguments and the given op */
+		LEARNING_API void NullaryOp(
+			FFrameAttribute& Out,
+			const int32 OutChannelNum,
+			const FFrameRangeSet& FrameRangeSet,
+			const NullaryOpFunction Op);
+
+		/** Creates a frame attribute from another and the given op */
+		LEARNING_API void UnaryOp(
+			FFrameAttribute& Out,
+			const int32 OutChannelNum,
+			const FFrameAttribute& In,
+			const UnaryOpFunction Op);
+
+		/** Creates a frame attribute from two others and the given op. Performs an intersection of the Lhs and Rhs if they do not match. */
+		LEARNING_API void BinaryOp(
+			FFrameAttribute& Out,
+			const int32 OutChannelNum,
+			const FFrameAttribute& Lhs,
+			const FFrameAttribute& Rhs,
+			const BinaryOpFunction Op);
+
+		/** Creates a frame attribute from multiple others and the given op. */
+		LEARNING_API void NaryOp(
+			FFrameAttribute& Out,
+			const int32 OutChannelNum,
+			const TArrayView<const ConstFrameAttributePtr> Inputs,
+			const NaryOpFunction Op);
+
+		/** Find the channel and frame with the smallest value */
+		LEARNING_API bool FindMinimum(
+			int32& OutChannelIdx,
+			int32& OutFrameIdx,
+			float& OutMinimumValue, 
+			const FFrameAttribute& In);
+
+		/** Find the channel and frame with the largest value */
+		LEARNING_API bool FindMaximum(
+			int32& OutChannelIdx,
+			int32& OutFrameIdx,
+			float& OutMaximumValue, 
+			const FFrameAttribute& In);
+
+		/** Create a frame attribute made up of zeros */
+		LEARNING_API void Zeros(FFrameAttribute& Out, const FFrameRangeSet& FrameRangeSet, const int32 ChannelNum);
+
+		/** Create a frame attribute made up of ones */
+		LEARNING_API void Ones(FFrameAttribute& Out, const FFrameRangeSet& FrameRangeSet, const int32 ChannelNum);
+
+		/** Create a frame attribute filled with the given value at each frame */
+		LEARNING_API void Fill(FFrameAttribute& Out, const FFrameRangeSet& FrameRangeSet, const TLearningArrayView<1, const float> Values);
+
+		/** Add two frame attributes. Channel numbers must match. */
+		LEARNING_API void Add(FFrameAttribute& Out, const FFrameAttribute& Lhs, const FFrameAttribute& Rhs);
+
+		/** Subtract two frame attributes. Channel numbers must match. */
+		LEARNING_API void Sub(FFrameAttribute& Out, const FFrameAttribute& Lhs, const FFrameAttribute& Rhs);
+
+		/** Multiply two frame attributes. Channel numbers must match. */
+		LEARNING_API void Mul(FFrameAttribute& Out, const FFrameAttribute& Lhs, const FFrameAttribute& Rhs);
+
+		/** Divide two frame attributes. Channel numbers must match. */
+		LEARNING_API void Div(FFrameAttribute& Out, const FFrameAttribute& Lhs, const FFrameAttribute& Rhs);
+
+		/** Compute the dot product of two frame attributes. Channel numbers must match. */
+		LEARNING_API void Dot(FFrameAttribute& Out, const FFrameAttribute& Lhs, const FFrameAttribute& Rhs);
+
+		/** Negate a frame attribute. */
+		LEARNING_API void Neg(FFrameAttribute& Out, const FFrameAttribute& In);
+
+		/** Invert a frame attribute (compute 1/x). */
+		LEARNING_API void Inv(FFrameAttribute& Out, const FFrameAttribute& In);
+
+		/** Compute the absolute value of a frame attribute. */
+		LEARNING_API void Abs(FFrameAttribute& Out, const FFrameAttribute& In);
+
+		/** Compute the length of a frame attribute over the channels dimension. */
+		LEARNING_API void Length(FFrameAttribute& Out, const FFrameAttribute& In);
+
+		/** Normalize a frame attribute over the channels dimension. */
+		LEARNING_API void Normalize(FFrameAttribute& Out, const FFrameAttribute& In);
+
+		/** Extract a single channel from frame attribute. */
+		LEARNING_API void Index(FFrameAttribute& Out, const FFrameAttribute& In, const int32 ChannelIdx);
+
+		/** Add a constant value to a frame attribute. Channel number must match the size of Rhs. */
+		LEARNING_API void AddConstant(FFrameAttribute& Out, const FFrameAttribute& Lhs, const TLearningArrayView<1, const float> Rhs);
+
+		/** Subtract a constant value to a frame attribute. Channel number must match the size of Rhs. */
+		LEARNING_API void SubConstant(FFrameAttribute& Out, const FFrameAttribute& Lhs, const TLearningArrayView<1, const float> Rhs);
+
+		/** Multiply a frame attribute by a constant value. Channel number must match the size of Rhs. */
+		LEARNING_API void MulConstant(FFrameAttribute& Out, const FFrameAttribute& Lhs, const TLearningArrayView<1, const float> Rhs);
+
+		/** Divide a frame attribute by a constant value. Channel number must match the size of Rhs. */
+		LEARNING_API void DivConstant(FFrameAttribute& Out, const FFrameAttribute& Lhs, const TLearningArrayView<1, const float> Rhs);
+
+		/** Add a constant value to a frame attribute. Channel number must match the size of Lhs. */
+		LEARNING_API void ConstantAdd(FFrameAttribute& Out, const TLearningArrayView<1, const float> Lhs, const FFrameAttribute& Rhs);
+
+		/** Subtract a constant value to a frame attribute. Channel number must match the size of Lhs. */
+		LEARNING_API void ConstantSub(FFrameAttribute& Out, const TLearningArrayView<1, const float> Lhs, const FFrameAttribute& Rhs);
+
+		/** Multiply a frame attribute by a constant value. Channel number must match the size of Lhs. */
+		LEARNING_API void ConstantMul(FFrameAttribute& Out, const TLearningArrayView<1, const float> Lhs, const FFrameAttribute& Rhs);
+
+		/** Divide a frame attribute by a constant value. Channel number must match the size of Lhs. */
+		LEARNING_API void ConstantDiv(FFrameAttribute& Out, const TLearningArrayView<1, const float> Lhs, const FFrameAttribute& Rhs);
+
+		/** Compute the sum of an array of frame attributes. Channel numbers must match. */
+		LEARNING_API void Sum(FFrameAttribute& Out, const TArrayView<const ConstFrameAttributePtr> Inputs);
+
+		/** Compute the product of an array of frame attributes. Channel numbers must match. */
+		LEARNING_API void Prod(FFrameAttribute& Out, const TArrayView<const ConstFrameAttributePtr> Inputs);
+
+		/** Computes the logical and of two frame attributes. Channel numbers must match. */
+		LEARNING_API void LogicalAnd(FFrameAttribute& Out, const FFrameAttribute& Lhs, const FFrameAttribute& Rhs);
+
+		/** Computes the logical or of two frame attributes. Channel numbers must match. */
+		LEARNING_API void LogicalOr(FFrameAttribute& Out, const FFrameAttribute& Lhs, const FFrameAttribute& Rhs);
+
+		/** Computes the logical not of a frame attribute. */
+		LEARNING_API void LogicalNot(FFrameAttribute& Out, const FFrameAttribute& In);
+
+		/** Computes if the values in one frame attribute are greater than another. Channel numbers must match. */
+		LEARNING_API void Gt(FFrameAttribute& Out, const FFrameAttribute& Lhs, const FFrameAttribute& Rhs);
+
+		/** Computes if the values in one frame attribute are greater than or equal to another. Channel numbers must match. */
+		LEARNING_API void Ge(FFrameAttribute& Out, const FFrameAttribute& Lhs, const FFrameAttribute& Rhs);
+
+		/** Computes if the values in one frame attribute are less than another. Channel numbers must match. */
+		LEARNING_API void Lt(FFrameAttribute& Out, const FFrameAttribute& Lhs, const FFrameAttribute& Rhs);
+
+		/** Computes if the values in one frame attribute are less than or equal to another. Channel numbers must match. */
+		LEARNING_API void Le(FFrameAttribute& Out, const FFrameAttribute& Lhs, const FFrameAttribute& Rhs);
+
+		/** Computes if the values in one frame attribute are equal to another. Channel numbers must match. */
+		LEARNING_API void Eq(FFrameAttribute& Out, const FFrameAttribute& Lhs, const FFrameAttribute& Rhs);
+
+		/** Computes if the values in one frame attribute are not equal to another. Channel numbers must match. */
+		LEARNING_API void Neq(FFrameAttribute& Out, const FFrameAttribute& Lhs, const FFrameAttribute& Rhs);
+
+		/** Computes if the values in a frame attribute are greater than a constant. Channel number must match the size of Rhs. */
+		LEARNING_API void GtConstant(FFrameAttribute& Out, const FFrameAttribute& Lhs, const TLearningArrayView<1, const float> Rhs);
+
+		/** Computes if the values in a frame attribute are greater than or equal to a constant. Channel number must match the size of Rhs. */
+		LEARNING_API void GeConstant(FFrameAttribute& Out, const FFrameAttribute& Lhs, const TLearningArrayView<1, const float> Rhs);
+
+		/** Computes if the values in a frame attribute are less than a constant. Channel number must match the size of Rhs. */
+		LEARNING_API void LtConstant(FFrameAttribute& Out, const FFrameAttribute& Lhs, const TLearningArrayView<1, const float> Rhs);
+
+		/** Computes if the values in a frame attribute are greater than or equal to a constant. Channel number must match the size of Rhs. */
+		LEARNING_API void LeConstant(FFrameAttribute& Out, const FFrameAttribute& Lhs, const TLearningArrayView<1, const float> Rhs);
+		
+		/** Computes if the values in a frame attribute are equal to a constant. Channel number must match the size of Rhs. */
+		LEARNING_API void EqConstant(FFrameAttribute& Out, const FFrameAttribute& Lhs, const TLearningArrayView<1, const float> Rhs);
+
+		/** Computes if the values in a frame attribute are not equal to a constant. Channel number must match the size of Rhs. */
+		LEARNING_API void NeqConstant(FFrameAttribute& Out, const FFrameAttribute& Lhs, const TLearningArrayView<1, const float> Rhs);
+
+		/** Computes if the values in a frame attribute are greater than a constant. Channel number must match the size of Lhs. */
+		LEARNING_API void ConstantGt(FFrameAttribute& Out, const TLearningArrayView<1, const float> Lhs, const FFrameAttribute& Rhs);
+		
+		/** Computes if the values in a frame attribute are greater than or equal to a constant. Channel number must match the size of Lhs. */
+		LEARNING_API void ConstantGe(FFrameAttribute& Out, const TLearningArrayView<1, const float> Lhs, const FFrameAttribute& Rhs);
+		
+		/** Computes if the values in a frame attribute are less than a constant. Channel number must match the size of Lhs. */
+		LEARNING_API void ConstantLt(FFrameAttribute& Out, const TLearningArrayView<1, const float> Lhs, const FFrameAttribute& Rhs);
+
+		/** Computes if the values in a frame attribute are greater than or equal to a constant. Channel number must match the size of Lhs. */
+		LEARNING_API void ConstantLe(FFrameAttribute& Out, const TLearningArrayView<1, const float> Lhs, const FFrameAttribute& Rhs);
+
+		/** Computes if the values in a frame attribute are equal to a constant. Channel number must match the size of Lhs. */
+		LEARNING_API void ConstantEq(FFrameAttribute& Out, const TLearningArrayView<1, const float> Lhs, const FFrameAttribute& Rhs);
+		
+		/** Computes if the values in a frame attribute are not equal to a constant. Channel number must match the size of Lhs. */
+		LEARNING_API void ConstantNeq(FFrameAttribute& Out, const TLearningArrayView<1, const float> Lhs, const FFrameAttribute& Rhs);
+
+		/** Applies a Gaussian smoothing filter to the ranges of values in the frame attribute. */
+		LEARNING_API void FilterGaussian(FFrameAttribute& Out, const FFrameAttribute& In, const float StdInFrames);
+
+		/** Applies a Majority Vote filter to the ranges of values in the frame attribute. */
+		LEARNING_API void FilterMajorityVote(FFrameAttribute& Out, const FFrameAttribute& In, const float FilterWidthFrames);
+
+		/** Computes the quaternion multiplication of two frame attributes. Both arguments are assumed to have 4 channels. */
+		LEARNING_API void QuatMul(FFrameAttribute& Out, const FFrameAttribute& Lhs, const FFrameAttribute& Rhs);
+
+		/** Computes the quaternion division of two frame attributes. Both arguments are assumed to have 4 channels. */
+		LEARNING_API void QuatDiv(FFrameAttribute& Out, const FFrameAttribute& Lhs, const FFrameAttribute& Rhs);
+
+		/** Compute the quaternion inverse of a frame attribute. Argument is assumed to have 4 channels. */
+		LEARNING_API void QuatInv(FFrameAttribute& Out, const FFrameAttribute& In);
+
+		/** Compute the quaternion closest to the identity quaternion for a frame attribute. Argument is assumed to have 4 channels. */
+		LEARNING_API void QuatAbs(FFrameAttribute& Out, const FFrameAttribute& In);
+
+		/** Computes the quaternion inverse multiplication of two frame attributes. Both arguments are assumed to have 4 channels. */
+		LEARNING_API void QuatInvMul(FFrameAttribute& Out, const FFrameAttribute& Lhs, const FFrameAttribute& Rhs);
+
+		/** Computes the quaternion multiplication inverse of two frame attributes. Both arguments are assumed to have 4 channels. */
+		LEARNING_API void QuatMulInv(FFrameAttribute& Out, const FFrameAttribute& Lhs, const FFrameAttribute& Rhs);
+
+		/** Computes the rotation of the Rhs by the Lhs quaternion. Lhs is assumed to have 4 channels, and Rhs is assumed to have 3. */
+		LEARNING_API void QuatRotate(FFrameAttribute& Out, const FFrameAttribute& Lhs, const FFrameAttribute& Rhs);
+
+		/** Computes the inverse rotation of the Rhs by the Lhs quaternion. Lhs is assumed to have 4 channels, and Rhs is assumed to have 3. */
+		LEARNING_API void QuatUnrotate(FFrameAttribute& Out, const FFrameAttribute& Lhs, const FFrameAttribute& Rhs);
+
+		/** Computes the shortest rotation between two vectors. Both arguments are assumed to have 3 channels. */
+		LEARNING_API void QuatBetween(FFrameAttribute& Out, const FFrameAttribute& Lhs, const FFrameAttribute& Rhs);
+
+		/** Computes the quaternion multiplication of a frame attribute by the given quaternion. Lhs is assumed to have 4 channels. */
+		LEARNING_API void QuatMulConstant(FFrameAttribute& Out, const FFrameAttribute& Lhs, const FQuat4f Rhs);
+
+		/** Computes the quaternion division of a frame attribute by the given quaternion. Lhs is assumed to have 4 channels. */
+		LEARNING_API void QuatDivConstant(FFrameAttribute& Out, const FFrameAttribute& Lhs, const FQuat4f Rhs);
+
+		/** Computes the quaternion inverse multiplication of a frame attribute by the given quaternion. Lhs is assumed to have 4 channels. */
+		LEARNING_API void QuatInvMulConstant(FFrameAttribute& Out, const FFrameAttribute& Lhs, const FQuat4f Rhs);
+
+		/** Computes the quaternion multiplication inverse of a frame attribute by the given quaternion. Lhs is assumed to have 4 channels. */
+		LEARNING_API void QuatMulInvConstant(FFrameAttribute& Out, const FFrameAttribute& Lhs, const FQuat4f Rhs);
+
+		/** Computes the quaternion rotation of the given vector by a frame attribute. Lhs is assumed to have 4 channels. */
+		LEARNING_API void QuatRotateConstant(FFrameAttribute& Out, const FFrameAttribute& Lhs, const FVector3f Rhs);
+
+		/** Computes the quaternion inverse rotation of the given vector by a frame attribute. Lhs is assumed to have 4 channels. */
+		LEARNING_API void QuatUnrotateConstant(FFrameAttribute& Out, const FFrameAttribute& Lhs, const FVector3f Rhs);
+
+		/** Computes the quaternion between the given frame attribute and vector. Lhs is assumed to have 3 channels. */
+		LEARNING_API void QuatBetweenConstant(FFrameAttribute& Out, const FFrameAttribute& Lhs, const FVector3f Rhs);
+
+		/** Computes the quaternion multiplication of a frame attribute by the given quaternion. Rhs is assumed to have 4 channels. */
+		LEARNING_API void QuatConstantMul(FFrameAttribute& Out, const FQuat4f Lhs, const FFrameAttribute& Rhs);
+
+		/** Computes the quaternion division of a frame attribute by the given quaternion. Rhs is assumed to have 4 channels. */
+		LEARNING_API void QuatConstantDiv(FFrameAttribute& Out, const FQuat4f Lhs, const FFrameAttribute& Rhs);
+		
+		/** Computes the quaternion inverse multiplication of a frame attribute by the given quaternion. Rhs is assumed to have 4 channels. */
+		LEARNING_API void QuatConstantInvMul(FFrameAttribute& Out, const FQuat4f Lhs, const FFrameAttribute& Rhs);
+		
+		/** Computes the quaternion multiplication inverse of a frame attribute by the given quaternion. Rhs is assumed to have 4 channels. */
+		LEARNING_API void QuatConstantMulInv(FFrameAttribute& Out, const FQuat4f Lhs, const FFrameAttribute& Rhs);
+		
+		/** Computes the quaternion rotation of the given frame attribute. Rhs is assumed to have 3 channels. */
+		LEARNING_API void QuatConstantRotate(FFrameAttribute& Out, const FQuat4f Lhs, const FFrameAttribute& Rhs);
+
+		/** Computes the quaternion inverse rotation of the given frame attribute. Rhs is assumed to have 3 channels. */
+		LEARNING_API void QuatConstantUnrotate(FFrameAttribute& Out, const FQuat4f Lhs, const FFrameAttribute& Rhs);
+
+		/** Computes the quaternion between the given frame attribute and vector. Rhs is assumed to have 3 channels. */
+		LEARNING_API void QuatConstantBetween(FFrameAttribute& Out, const FVector3f Lhs, const FFrameAttribute& Rhs);
+	}
+}
\ No newline at end of file
diff --git a/Engine/Plugins/Experimental/LearningAgents/Source/Learning/Public/LearningFrameRangeSet.h b/Engine/Plugins/Experimental/LearningAgents/Source/Learning/Public/LearningFrameRangeSet.h
--- a/Engine/Plugins/Experimental/LearningAgents/Source/Learning/Public/LearningFrameRangeSet.h
+++ b/Engine/Plugins/Experimental/LearningAgents/Source/Learning/Public/LearningFrameRangeSet.h
@@ -15,14 +15,12 @@ namespace UE::Learning
 	 * This data structure stores that information in a way that makes performing operations such as union, intersection, difference etc, efficient.
 	 * 
 	 * The way it works is by storing a sorted array of "entries", where each entry has a corresponding sequence, the number of ranges in the set for
-	 * that sequence, and an offset into two large arrays of range starts and stops. The sub-ranges of the large arrays of range starts and stops are 
-	 * also sorted.
+	 * that sequence, and an index into three large arrays of range starts, lengths, and offsets. The sub-ranges of the large arrays of range 
+	 * starts and lengths are also sorted, while the offsets arrays stores the offset that would be used if you had to take a slice of a large flat 
+	 * array of frame data.
 	 *
 	 * Having the entries sorted by sequence, and the ranges sorted, allows for efficient set operations via tape-merge algorithms.
 	 *
-	 * Also provided are some algorithms for getting "offsets", which correspond to the index associated with a frame if you were to
-	 * flatten this data-structure into one large array of frames by concatenating all ranges together into one large block.
-	 *
 	 * This data-structure is related to the FrameSet data-structure, which essentially acts the same way but stores individual frames instead of
 	 * frame ranges.
 	 */
@@ -66,23 +64,23 @@ namespace UE::Learning
 		/** Computes the total number of frames given by all ranges in an entry */
 		int32 GetEntryTotalFrameNum(const int32 EntryIdx) const;
 
-		/** Computes the flat offset of a particular entry if you were to flatten this data structure */
-		int32 GetEntryFrameOffset(const int32 EntryIdx) const;
-
 		/** Gets all the range starts associated with a given entry */
 		TLearningArrayView<1, const int32> GetEntryRangeStarts(const int32 EntryIdx) const;
 
 		/** Gets all the range lengths associated with a given entry */
 		TLearningArrayView<1, const int32> GetEntryRangeLengths(const int32 EntryIdx) const;
 		
+		/** Gets all the range offsets associated with a given entry */
+		TLearningArrayView<1, const int32> GetEntryRangeOffsets(const int32 EntryIdx) const;
+
 		/** Gets the range start associated with a given entry and range index */
 		int32 GetEntryRangeStart(const int32 EntryIdx, const int32 RangeIdx) const;
 
 		/** Gets the range length associated with a given entry and range index */
 		int32 GetEntryRangeLength(const int32 EntryIdx, const int32 RangeIdx) const;
 
-		/** Computes the flat offset associated with a given entry and range index */
-		int32 GetEntryRangeFrameOffset(const int32 EntryIdx, const int32 RangeIdx) const;
+		/** Computes the offset associated with a given entry and range index */
+		int32 GetEntryRangeOffset(const int32 EntryIdx, const int32 RangeIdx) const;
 
 		/** Gets the range start time associated with a given entry and range index */
 		float GetEntryRangeStartTime(const int32 EntryIdx, const int32 RangeIdx, const float FrameDeltaTime) const;
@@ -102,6 +100,9 @@ namespace UE::Learning
 		/** Gets an array of the range lengths for all ranges */
 		TLearningArrayView<1, const int32> GetAllRangeLengths() const;
 
+		/** Gets an array of the range offsets for all ranges */
+		TLearningArrayView<1, const int32> GetAllRangeOffsets() const;
+
 		/** Computes the total number of frames across all entries and ranges */
 		int32 GetTotalFrameNum() const;
 
@@ -123,15 +124,15 @@ namespace UE::Learning
 		/** Finds the entry index and range index associated with a given sequence and time in that sequence */
 		bool FindTime(int32& OutEntryIdx, int32& OutRangeIdx, float& OutRangeTime, const int32 Sequence, const float Time, const float FrameDeltaTime) const;
 
-		/** Finds entry index, frame index, and frame in that range, associated with some flat frame offset  */
-		bool FindFrameOffset(int32& OutEntryIdx, int32& OutRangeIdx, int32& OutRangeFrame, const int32 FrameOffset) const;
+		/** Finds entry index, frame index, and frame in that range, associated with some offset  */
+		bool FindOffset(int32& OutEntryIdx, int32& OutRangeIdx, int32& OutRangeFrame, const int32 Offset) const;
 
 	public:
 
 		/** Array of sequences associated with each entry */
 		TLearningArray<1, int32> EntrySequences;
 
-		/** Array of offsets into the RangeStarts and RangeLengths arrays associated with each entry */
+		/** Array of offsets into the RangeStarts, RangeLengths, RangeOffsets arrays associated with each entry */
 		TLearningArray<1, int32> EntryRangeOffsets;
 
 		/** Array of the number of ranges associated with each entry */
@@ -142,6 +143,9 @@ namespace UE::Learning
 
 		/** Large array of all range lengths for all entries, indexed using offsets from EntryRangeOffsets */
 		TLearningArray<1, int32> RangeLengths;
+
+		/** Large array of all range frame offsets for all entries, indexed using offsets from EntryRangeOffsets */
+		TLearningArray<1, int32> RangeOffsets;
 	};
 
 	namespace FrameRangeSet
@@ -166,13 +170,13 @@ namespace UE::Learning
 
 		/**
 		 * Computes the intersection of two frame range sets while recording the flat frame offsets in the lhs and rhs, associated with each range
-		 * that gets added to the output. The OutLhsFrameOffsets and OutRhsFrameOffsets outputs should be pre-allocated to at least the maximum of the
+		 * that gets added to the output. The OutLhsOffsets and OutRhsOffsets outputs should be pre-allocated to at least the maximum of the
 		 * number of ranges in Lhs and Rhs. Returns the number of ranges in the added to the output frame range set.
 		 */
-		LEARNING_API int32 IntersectionWithFrameOffsets(
+		LEARNING_API int32 IntersectionWithOffsets(
 			FFrameRangeSet& Out, 
-			TLearningArrayView<1, int32> OutLhsFrameOffsets,
-			TLearningArrayView<1, int32> OutRhsFrameOffsets,
+			TLearningArrayView<1, int32> OutLhsOffsets,
+			TLearningArrayView<1, int32> OutRhsOffsets,
 			const FFrameRangeSet& Lhs,
 			const FFrameRangeSet& Rhs);
 
@@ -185,6 +189,9 @@ namespace UE::Learning
 		/** Creates a new frame range set with both the start and the end of all ranges trimmed by the given number of frames */
 		LEARNING_API void TrimLeftAndRight(FFrameRangeSet& Out, const FFrameRangeSet& FrameRangeSet, const int32 TrimFrameNum);
 
+		/** Makes a Frame Range Set consisting of single-frame ranges from a frame set */
+		LEARNING_API void MakeFromFrameSet(FFrameRangeSet& OutFrameRangeSet, const FFrameSet& FrameSet);
+
 		/** Create a frame set from all of the range starts of a given frame range set */
 		LEARNING_API void MakeFrameSetFromRangeStarts(FFrameSet& OutFrameSet, const FFrameRangeSet& FrameRangeSet);
 
@@ -212,11 +219,6 @@ namespace UE::Learning
 			TLearningArrayView<1, int32> OutRangeSequences,
 			const FFrameRangeSet& FrameRangeSet);
 
-		/** Computes the flat frame offsets associated with every range in the set */
-		LEARNING_API void AllRangeFrameOffsets(
-			TLearningArrayView<1, int32> OutRangeFrameOffsets,
-			const FFrameRangeSet& FrameRangeSet);
-
 		/** Computes the start times associated with every range in the set */
 		LEARNING_API void AllRangeStartTimes(
 			TLearningArrayView<1, float> OutRangeStartTimes,
@@ -241,8 +243,7 @@ namespace UE::Learning
 			const TFunctionRef<void(
 				const int32 TotalRangeIdx,
 				const int32 EntryIdx,
-				const int32 RangeIdx,
-				const int32 FrameOffset)> Body);
+				const int32 RangeIdx)> Body);
 
 		/** Iterates over every range in the set in parallel and calls the provided callback */
 		LEARNING_API void ParallelForEachRange(
@@ -250,7 +251,6 @@ namespace UE::Learning
 			const TFunctionRef<void(
 				const int32 TotalRangeIdx,
 				const int32 EntryIdx,
-				const int32 RangeIdx,
-				const int32 FrameOffset)> Body);
+				const int32 RangeIdx)> Body);
 	}
 }
\ No newline at end of file
diff --git a/Engine/Plugins/Experimental/LearningAgents/Source/Learning/Public/LearningFrameSet.h b/Engine/Plugins/Experimental/LearningAgents/Source/Learning/Public/LearningFrameSet.h
--- a/Engine/Plugins/Experimental/LearningAgents/Source/Learning/Public/LearningFrameSet.h
+++ b/Engine/Plugins/Experimental/LearningAgents/Source/Learning/Public/LearningFrameSet.h
@@ -70,7 +70,7 @@ namespace UE::Learning
 		float GetEntryFrameTime(const int32 EntryIdx, const int32 FrameIdx, const float FrameDeltaTime) const;
 
 		/** Gets the flat offset associated with a particular entry if you were to flatten this data structure */
-		int32 GetEntryFrameOffset(const int32 EntryIdx) const;
+		int32 GetEntryOffset(const int32 EntryIdx) const;
 
 		/** Gets the total number of frames if you were to flatten this data structure */
 		int32 GetTotalFrameNum() const;
@@ -98,8 +98,8 @@ namespace UE::Learning
 			int32& OutEntryIdx, int32& OutFrameIdx, int32& OutFrameDifference, 
 			const int32 Sequence, const int32 Frame, const int32 RangeStart, const int32 RangeLength) const;
 
-		/** Finds entry index and frame index associated with some flat offset  */
-		bool FindFrameOffset(int32& OutEntryIdx, int32& OutFrameIdx, const int32 FrameOffset) const;
+		/** Finds entry index and frame index associated with some flat offset */
+		bool FindOffset(int32& OutEntryIdx, int32& OutFrameIdx, const int32 Offset) const;
 
 	public:
 
