[LearningAgents] New console flag to randomize seed of GymsManager with a randomly generated seed

-GymsManagerRandomizeNoSeed

diff --git a/Engine/Plugins/Experimental/LearningAgents/Source/LearningAgentsTraining/Private/LearningAgentsGymsManager.cpp b/Engine/Plugins/Experimental/LearningAgents/Source/LearningAgentsTraining/Private/LearningAgentsGymsManager.cpp
--- a/Engine/Plugins/Experimental/LearningAgents/Source/LearningAgentsTraining/Private/LearningAgentsGymsManager.cpp
+++ b/Engine/Plugins/Experimental/LearningAgents/Source/LearningAgentsTraining/Private/LearningAgentsGymsManager.cpp
@@ -8,6 +8,7 @@
 #include "NavigationSystem.h"
 #include "AI/NavigationSystemBase.h"
 #include "Misc/CommandLine.h"
+#include "Misc/DateTime.h"
 
 ALearningAgentsGymsManager::ALearningAgentsGymsManager()
 {
@@ -100,6 +101,11 @@ void ALearningAgentsGymsManager::InitializeRandomStream()
 		}
 		RandomSeed = FCString::Atoi(*RandomSeedInput);
 	}
+	if (FParse::Param(FCommandLine::Get(), TEXT("GymsManagerRandomizeNoSeed")))
+	{
+		RandomSeed = FDateTime::Now().GetTicks();
+		UE_LOG(LogLearning, Log, TEXT("Commandline flag -GymsManagerRandomizeNoSeed is active! Using randomly generated seed from timestamp: %d."), RandomSeed);
+	}
 	RandomStream = MakeShareable(new FRandomStream);
 	RandomStream->Initialize(RandomSeed);
 	UE_LOG(LogLearning, Log, TEXT("GymsManager initialized with seed %d!"), RandomSeed);
