LearningAgents: Added some more functions to FrameRangeSet structure.

diff --git a/Engine/Plugins/Experimental/LearningAgents/Source/Learning/Private/LearningFrameRangeSet.cpp b/Engine/Plugins/Experimental/LearningAgents/Source/Learning/Private/LearningFrameRangeSet.cpp
--- a/Engine/Plugins/Experimental/LearningAgents/Source/Learning/Private/LearningFrameRangeSet.cpp
+++ b/Engine/Plugins/Experimental/LearningAgents/Source/Learning/Private/LearningFrameRangeSet.cpp
@@ -183,6 +183,175 @@ namespace UE::Learning
 			return false;
 		}
 
+
+		static inline int32 FramesRangesUnion(
+			TLearningArrayView<1, int32> OutStarts,
+			TLearningArrayView<1, int32> OutLengths,
+			const TLearningArrayView<1, const int32> LhsFrames,
+			const TLearningArrayView<1, const int32> RhsStarts,
+			const TLearningArrayView<1, const int32> RhsLengths)
+		{
+			TRACE_CPUPROFILER_EVENT_SCOPE(Learning::FrameRangeSet::Private::FramesRangesUnion);
+
+			FramesCheck(LhsFrames);
+			RangesCheck(RhsStarts, RhsLengths);
+
+			if (LhsFrames.IsEmpty())
+			{
+				Array::Copy(OutStarts.Slice(0, RhsStarts.Num()), RhsStarts);
+				Array::Copy(OutLengths.Slice(0, RhsLengths.Num()), RhsLengths);
+				return RhsLengths.Num();
+			}
+
+			// Number of ranges in lhs and rhs
+			const int32 LhsNum = LhsFrames.Num();
+			const int32 RhsNum = RhsStarts.Num();
+
+			// Activation state of each list of ranges
+			bool bOutActive = false;
+			bool bLhsActive = false;
+			bool bRhsActive = false;
+
+			// Event index for each list of ranges
+			int32 LhsIndex = 0;
+			int32 RhsIndex = 0;
+			int32 OutIndex = 0;
+
+			// While both ranges have events to process
+			while (LhsIndex < LhsNum * 2 && RhsIndex < RhsNum * 2)
+			{
+				// Are the next lhs, and rhs events active or inactive
+				const bool bLhsActiveNext = LhsIndex % 2 == 0;
+				const bool bRhsActiveNext = RhsIndex % 2 == 0;
+
+				// Time of the next lhs, and rhs events
+				const int32 LhsT = bLhsActiveNext ? LhsFrames[LhsIndex / 2] : LhsFrames[LhsIndex / 2] + 1;
+				const int32 RhsT = bRhsActiveNext ? RhsStarts[RhsIndex / 2] : RhsStarts[RhsIndex / 2] + RhsLengths[RhsIndex / 2];
+
+				// Event from lhs is coming first
+				if (LhsT < RhsT)
+				{
+					// Activate output
+					if (!bOutActive && bLhsActiveNext)
+					{
+						bOutActive = true;
+						OutStarts[OutIndex] = LhsT;
+					}
+					// Deactivate output
+					else if (bOutActive && !bLhsActiveNext && !bRhsActive)
+					{
+
+						bOutActive = false;
+						OutLengths[OutIndex] = LhsT - OutStarts[OutIndex];
+						OutIndex++;
+					}
+
+					bLhsActive = bLhsActiveNext;
+					LhsIndex++;
+				}
+				// Event from rhs is coming first
+				else if (RhsT < LhsT)
+				{
+					// Activate output
+					if (!bOutActive && bRhsActiveNext)
+					{
+						bOutActive = true;
+						OutStarts[OutIndex] = RhsT;
+					}
+					// Deactivate output
+					else if (bOutActive && !bLhsActive && !bRhsActiveNext)
+					{
+						bOutActive = false;
+						OutLengths[OutIndex] = RhsT - OutStarts[OutIndex];
+						OutIndex++;
+					}
+
+					bRhsActive = bRhsActiveNext;
+					RhsIndex++;
+				}
+				// Event from lhs and rhs coming at same time
+				else
+				{
+					check(LhsT == RhsT);
+
+					// Activate output
+					if (!bOutActive && (bLhsActiveNext || bRhsActiveNext))
+					{
+						bOutActive = true;
+						OutStarts[OutIndex] = LhsT;
+					}
+					// Deactivate output
+					else if (bOutActive && !(bLhsActiveNext || bRhsActiveNext))
+					{
+						bOutActive = false;
+						OutLengths[OutIndex] = LhsT - OutStarts[OutIndex];
+						OutIndex++;
+					}
+
+					bLhsActive = bLhsActiveNext;
+					bRhsActive = bRhsActiveNext;
+					LhsIndex++; RhsIndex++;
+				}
+			}
+
+			// Process any remaining lhs events
+			while (LhsIndex < LhsNum * 2)
+			{
+				check(RhsIndex == RhsNum * 2);
+
+				const bool bLhsActiveNext = LhsIndex % 2 == 0;
+				const int32 LhsT = bLhsActiveNext ? LhsFrames[LhsIndex / 2] : LhsFrames[LhsIndex / 2] + 1;
+
+				// Activate output
+				if (!bOutActive && bLhsActiveNext)
+				{
+					bOutActive = true;
+					OutStarts[OutIndex] = LhsT;
+				}
+				// Deactivate output
+				else if (bOutActive && !bLhsActiveNext)
+				{
+					bOutActive = false;
+					OutLengths[OutIndex] = LhsT - OutStarts[OutIndex];
+					OutIndex++;
+				}
+
+				bLhsActive = bLhsActiveNext;
+				LhsIndex++;
+			}
+
+			// Process any remaining rhs events
+			while (RhsIndex < RhsNum * 2)
+			{
+				check(LhsIndex == LhsNum * 2);
+
+				const bool bRhsActiveNext = RhsIndex % 2 == 0;
+				const int32 RhsT = bRhsActiveNext ? RhsStarts[RhsIndex / 2] : RhsStarts[RhsIndex / 2] + RhsLengths[RhsIndex / 2];
+
+				// Activate output
+				if (!bOutActive && bRhsActiveNext)
+				{
+					bOutActive = true;
+					OutStarts[OutIndex] = RhsT;
+				}
+				// Deactivate output
+				else if (bOutActive && !bRhsActiveNext)
+				{
+					bOutActive = false;
+					OutLengths[OutIndex] = RhsT - OutStarts[OutIndex];
+					OutIndex++;
+				}
+
+				bRhsActive = bRhsActiveNext;
+				RhsIndex++;
+			}
+
+			RangesCheck(OutStarts.Slice(0, OutIndex), OutLengths.Slice(0, OutIndex));
+
+			// Return number of ranges added
+			return OutIndex;
+		}
+
 		static inline int32 FramesRangesIntersection(
 			TLearningArrayView<1, int32> OutFrames,
 			const TLearningArrayView<1, const int32> LhsFrames,
@@ -336,6 +505,139 @@ namespace UE::Learning
 		}
 
 
+		static inline int32 RangesFramesDifference(
+			TLearningArrayView<1, int32> OutStarts,
+			TLearningArrayView<1, int32> OutLengths,
+			const TLearningArrayView<1, const int32> LhsStarts,
+			const TLearningArrayView<1, const int32> LhsLengths,
+			const TLearningArrayView<1, const int32> RhsFrames)
+		{
+			TRACE_CPUPROFILER_EVENT_SCOPE(Learning::FrameRangeSet::Private::RangesFramesDifference);
+
+			RangesCheck(LhsStarts, LhsLengths);
+			FramesCheck(RhsFrames);
+
+			// Number of ranges in lhs and rhs
+			const int32 LhsNum = LhsStarts.Num();
+			const int32 RhsNum = RhsFrames.Num();
+
+			// Activation state of each list of ranges
+			bool bOutActive = false;
+			bool bLhsActive = false;
+			bool bRhsActive = false;
+
+			// Event index for each list of ranges
+			int32 LhsIndex = 0;
+			int32 RhsIndex = 0;
+			int32 OutIndex = 0;
+
+			// While both ranges have events to process
+			while (LhsIndex < LhsNum * 2 && RhsIndex < RhsNum * 2)
+			{
+				// Are the next lhs, and rhs events active or inactive
+				const bool bLhsActiveNext = LhsIndex % 2 == 0;
+				const bool bRhsActiveNext = RhsIndex % 2 == 0;
+
+				// Time of the next lhs, and rhs events
+				const int32 LhsT = bLhsActiveNext ? LhsStarts[LhsIndex / 2] : LhsStarts[LhsIndex / 2] + LhsLengths[LhsIndex / 2];
+				const int32 RhsT = bRhsActiveNext ? RhsFrames[RhsIndex / 2] : RhsFrames[RhsIndex / 2] + 1;
+
+				// Event coming from lhs first
+				if (LhsT < RhsT)
+				{
+					// Activate output
+					if (!bOutActive && !bRhsActive && bLhsActiveNext)
+					{
+						bOutActive = true;
+						OutStarts[OutIndex] = LhsT;
+					}
+					// Deactivate output
+					else if (bOutActive && !bLhsActiveNext)
+					{
+						bOutActive = false;
+						OutLengths[OutIndex] = LhsT - OutStarts[OutIndex];
+						OutIndex++;
+					}
+
+					bLhsActive = bLhsActiveNext;
+					LhsIndex++;
+				}
+				// Event coming from rhs first
+				else if (RhsT < LhsT)
+				{
+					// Activate output
+					if (!bOutActive && bLhsActive && !bRhsActiveNext)
+					{
+						bOutActive = true;
+						OutStarts[OutIndex] = RhsT;
+					}
+					// Deactivate output
+					else if (bOutActive && bRhsActiveNext)
+					{
+						bOutActive = false;
+						OutLengths[OutIndex] = RhsT - OutStarts[OutIndex];
+						OutIndex++;
+					}
+
+					bRhsActive = bRhsActiveNext;
+					RhsIndex++;
+				}
+				// Event from lhs and rhs coming at same time
+				else
+				{
+					check(LhsT == RhsT);
+
+					// Activate output
+					if (!bOutActive && bLhsActiveNext && !bRhsActiveNext)
+					{
+						bOutActive = true;
+						OutStarts[OutIndex] = LhsT;
+					}
+					// Deactivate output
+					else if (bOutActive && bRhsActiveNext)
+					{
+						bOutActive = false;
+						OutLengths[OutIndex] = LhsT - OutStarts[OutIndex];
+						OutIndex++;
+					}
+
+					bLhsActive = bLhsActiveNext;
+					bRhsActive = bRhsActiveNext;
+					LhsIndex++; RhsIndex++;
+				}
+			}
+
+			// Process any remaining lhs events
+			while (LhsIndex < LhsNum * 2)
+			{
+				const bool bLhsActiveNext = LhsIndex % 2 == 0;
+				const int32 LhsT = bLhsActiveNext ? LhsStarts[LhsIndex / 2] : LhsStarts[LhsIndex / 2] + LhsLengths[LhsIndex / 2];
+
+				// Activate output
+				if (!bOutActive && bLhsActiveNext)
+				{
+					bOutActive = true;
+					OutStarts[OutIndex] = LhsT;
+				}
+				// Deactivate output
+				else if (bOutActive && !bLhsActiveNext)
+				{
+					bOutActive = false;
+					OutLengths[OutIndex] = LhsT - OutStarts[OutIndex];
+					OutIndex++;
+				}
+
+				bLhsActive = bLhsActiveNext;
+				LhsIndex++;
+			}
+
+			RangesCheck(OutStarts.Slice(0, OutIndex), OutLengths.Slice(0, OutIndex));
+
+			// Return number of ranges added
+			return OutIndex;
+		}
+
+
 		static inline int32 RangesUnion(
 			TLearningArrayView<1, int32> OutStarts,
 			TLearningArrayView<1, int32> OutLengths,
@@ -1196,6 +1498,167 @@ namespace UE::Learning
 				Array::Equal(Lhs.RangeOffsets, Rhs.RangeOffsets);
 		}
 
+		void Union(FFrameRangeSet& OutFrameRangeSet, const FFrameSet& FrameSet, const FFrameRangeSet& FrameRangeSet)
+		{
+			TRACE_CPUPROFILER_EVENT_SCOPE(Learning::FrameRangeSet::Union);
+
+			FrameSet.Check(); FrameRangeSet.Check();
+
+			if (FrameSet.IsEmpty()) { OutFrameRangeSet = FrameRangeSet; return; }
+
+			// Allocate potential maximum number of entries and ranges we might to output
+
+			const int32 LhsEntryNum = FrameSet.GetEntryNum();
+			const int32 RhsEntryNum = FrameRangeSet.GetEntryNum();
+			const int32 LhsRangeNum = FrameSet.GetTotalFrameNum();
+			const int32 RhsRangeNum = FrameRangeSet.GetTotalRangeNum();
+
+			// Allocate potential maximum number of entries and ranges we might to output
+			OutFrameRangeSet.EntrySequences.SetNumUninitialized({ LhsEntryNum + RhsEntryNum });
+			OutFrameRangeSet.EntryRangeOffsets.SetNumUninitialized({ LhsEntryNum + RhsEntryNum });
+			OutFrameRangeSet.EntryRangeNums.SetNumUninitialized({ LhsEntryNum + RhsEntryNum });
+			OutFrameRangeSet.RangeStarts.SetNumUninitialized({ LhsRangeNum + RhsRangeNum });
+			OutFrameRangeSet.RangeLengths.SetNumUninitialized({ LhsRangeNum + RhsRangeNum });
+
+			// Entry index for each list of ranges
+			int32 OutIndex = 0;
+			int32 LhsIndex = 0;
+			int32 RhsIndex = 0;
+
+			// Output ranges index
+			int32 RangeIndex = 0;
+
+			// While both sets have entries
+			while (LhsIndex < LhsEntryNum && RhsIndex < RhsEntryNum)
+			{
+				// If entry from lhs is first
+				if (FrameSet.GetEntrySequence(LhsIndex) < FrameRangeSet.GetEntrySequence(RhsIndex))
+				{
+					check(FrameSet.GetEntryFrameNum(LhsIndex) > 0);
+
+					// Append subranges to output
+					OutFrameRangeSet.EntrySequences[OutIndex] = FrameSet.GetEntrySequence(LhsIndex);
+					OutFrameRangeSet.EntryRangeOffsets[OutIndex] = RangeIndex;
+					OutFrameRangeSet.EntryRangeNums[OutIndex] = FrameSet.GetEntryFrameNum(LhsIndex);
+
+					Array::Copy(
+						OutFrameRangeSet.RangeStarts.Slice(RangeIndex, FrameSet.GetEntryFrameNum(LhsIndex)),
+						FrameSet.GetEntryFrames(LhsIndex));
+
+					Array::Set(
+						OutFrameRangeSet.RangeLengths.Slice(RangeIndex, FrameSet.GetEntryFrameNum(LhsIndex)),
+						1);
+
+					RangeIndex += FrameSet.GetEntryFrameNum(LhsIndex);
+					OutIndex++;
+					LhsIndex++;
+				}
+				// If entry from rhs is first
+				else if (FrameRangeSet.GetEntrySequence(RhsIndex) < FrameSet.GetEntrySequence(LhsIndex))
+				{
+					check(FrameRangeSet.GetEntryRangeNum(RhsIndex) > 0);
+
+					// Append subranges to output
+					OutFrameRangeSet.EntrySequences[OutIndex] = FrameRangeSet.GetEntrySequence(RhsIndex);
+					OutFrameRangeSet.EntryRangeOffsets[OutIndex] = RangeIndex;
+					OutFrameRangeSet.EntryRangeNums[OutIndex] = FrameRangeSet.GetEntryRangeNum(RhsIndex);
+
+					Array::Copy(
+						OutFrameRangeSet.RangeStarts.Slice(RangeIndex, FrameRangeSet.GetEntryRangeNum(RhsIndex)),
+						FrameRangeSet.GetEntryRangeStarts(RhsIndex));
+
+					Array::Copy(
+						OutFrameRangeSet.RangeLengths.Slice(RangeIndex, FrameRangeSet.GetEntryRangeNum(RhsIndex)),
+						FrameRangeSet.GetEntryRangeLengths(RhsIndex));
+
+					RangeIndex += FrameRangeSet.GetEntryRangeNum(RhsIndex);
+					OutIndex++;
+					RhsIndex++;
+				}
+				// If both contain the same entry
+				else
+				{
+					check(FrameSet.GetEntryFrameNum(LhsIndex) > 0);
+					check(FrameRangeSet.GetEntryRangeNum(RhsIndex) > 0);
+
+					// Append union of subranges to output
+					const int32 RangeNum = FrameRangeSet::Private::FramesRangesUnion(
+						OutFrameRangeSet.RangeStarts.Slice(RangeIndex, OutFrameRangeSet.RangeStarts.Num() - RangeIndex),
+						OutFrameRangeSet.RangeLengths.Slice(RangeIndex, OutFrameRangeSet.RangeLengths.Num() - RangeIndex),
+						FrameSet.GetEntryFrames(LhsIndex),
+						FrameRangeSet.GetEntryRangeStarts(RhsIndex),
+						FrameRangeSet.GetEntryRangeLengths(RhsIndex));
+
+					check(RangeNum > 0);
+					check(RangeNum <= FrameSet.GetEntryFrameNum(LhsIndex) + FrameRangeSet.GetEntryRangeNum(RhsIndex));
+
+					OutFrameRangeSet.EntrySequences[OutIndex] = FrameSet.GetEntrySequence(LhsIndex);
+					OutFrameRangeSet.EntryRangeOffsets[OutIndex] = RangeIndex;
+					OutFrameRangeSet.EntryRangeNums[OutIndex] = RangeNum;
+					RangeIndex += RangeNum;
+					OutIndex++;
+					LhsIndex++;
+					RhsIndex++;
+				}
+			}
+
+			// Process any remaining lhs entries
+			while (LhsIndex < LhsEntryNum)
+			{
+				check(RhsIndex == RhsEntryNum);
+				check(FrameSet.GetEntryFrameNum(LhsIndex) > 0);
+
+				// Append subranges to output
+				OutFrameRangeSet.EntrySequences[OutIndex] = FrameSet.GetEntrySequence(LhsIndex);
+				OutFrameRangeSet.EntryRangeOffsets[OutIndex] = RangeIndex;
+				OutFrameRangeSet.EntryRangeNums[OutIndex] = FrameSet.GetEntryFrameNum(LhsIndex);
+
+				Array::Copy(
+					OutFrameRangeSet.RangeStarts.Slice(RangeIndex, FrameSet.GetEntryFrameNum(LhsIndex)),
+					FrameSet.GetEntryFrames(LhsIndex));
+
+				Array::Set(OutFrameRangeSet.RangeLengths.Slice(RangeIndex, FrameSet.GetEntryFrameNum(LhsIndex)), 1);
+
+				RangeIndex += FrameSet.GetEntryFrameNum(LhsIndex);
+				OutIndex++;
+				LhsIndex++;
+			}
+
+			// Process any remaining rhs entries
+			while (RhsIndex < RhsEntryNum)
+			{
+				check(LhsIndex == LhsEntryNum);
+				check(FrameRangeSet.GetEntryRangeNum(RhsIndex) > 0);
+
+				// Append subranges to output
+				OutFrameRangeSet.EntrySequences[OutIndex] = FrameRangeSet.GetEntrySequence(RhsIndex);
+				OutFrameRangeSet.EntryRangeOffsets[OutIndex] = RangeIndex;
+				OutFrameRangeSet.EntryRangeNums[OutIndex] = FrameRangeSet.GetEntryRangeNum(RhsIndex);
+
+				Array::Copy(
+					OutFrameRangeSet.RangeStarts.Slice(RangeIndex, FrameRangeSet.GetEntryRangeNum(RhsIndex)),
+					FrameRangeSet.GetEntryRangeStarts(RhsIndex));
+
+				Array::Copy(
+					OutFrameRangeSet.RangeLengths.Slice(RangeIndex, FrameRangeSet.GetEntryRangeNum(RhsIndex)),
+					FrameRangeSet.GetEntryRangeLengths(RhsIndex));
+
+				RangeIndex += FrameRangeSet.GetEntryRangeNum(RhsIndex);
+				OutIndex++;
+				RhsIndex++;
+			}
+
+			// Resize output to match what was added
+			OutFrameRangeSet.EntrySequences.SetNumUninitialized({ OutIndex });
+			OutFrameRangeSet.EntryRangeOffsets.SetNumUninitialized({ OutIndex });
+			OutFrameRangeSet.EntryRangeNums.SetNumUninitialized({ OutIndex });
+			OutFrameRangeSet.RangeStarts.SetNumUninitialized({ RangeIndex });
+			OutFrameRangeSet.RangeLengths.SetNumUninitialized({ RangeIndex });
+			OutFrameRangeSet.RangeOffsets.SetNumUninitialized({ RangeIndex });
+			Private::ComputeRangeOffsets(OutFrameRangeSet.RangeOffsets, OutFrameRangeSet.RangeLengths);
+			OutFrameRangeSet.Check();
+		}
+
 		void Intersection(FFrameSet& OutFrameSet, const FFrameSet& FrameSet, const FFrameRangeSet& FrameRangeSet)
 		{
 			TRACE_CPUPROFILER_EVENT_SCOPE(Learning::FrameRangeSet::Intersection);
@@ -1376,6 +1839,127 @@ namespace UE::Learning
 			OutFrameSet.Check();
 		}
 
+		void Difference(FFrameRangeSet& OutFrameRangeSet, const FFrameRangeSet& FrameRangeSet, const FFrameSet& FrameSet)
+		{
+			TRACE_CPUPROFILER_EVENT_SCOPE(Learning::FrameRangeSet::Difference);
+
+			FrameRangeSet.Check(); FrameSet.Check();
+
+			if (FrameSet.IsEmpty()) { OutFrameRangeSet = FrameRangeSet; return; }
+
+			const int32 LhsEntryNum = FrameRangeSet.GetEntryNum();
+			const int32 RhsEntryNum = FrameSet.GetEntryNum();
+			const int32 LhsRangeNum = FrameRangeSet.GetTotalRangeNum();
+			const int32 RhsRangeNum = FrameSet.GetTotalFrameNum();
+
+			// Allocate potential maximum number of entries and ranges we might to output
+			OutFrameRangeSet.EntrySequences.SetNumUninitialized({ LhsEntryNum });
+			OutFrameRangeSet.EntryRangeOffsets.SetNumUninitialized({ LhsEntryNum });
+			OutFrameRangeSet.EntryRangeNums.SetNumUninitialized({ LhsEntryNum });
+			OutFrameRangeSet.RangeStarts.SetNumUninitialized({ LhsRangeNum + RhsRangeNum });
+			OutFrameRangeSet.RangeLengths.SetNumUninitialized({ LhsRangeNum + RhsRangeNum });
+
+			// Anim index for each list of ranges
+			int32 OutIndex = 0;
+			int32 LhsIndex = 0;
+			int32 RhsIndex = 0;
+
+			// Output ranges index
+			int32 RangeIndex = 0;
+
+			// While both sets have entries
+			while (LhsIndex < LhsEntryNum && RhsIndex < RhsEntryNum)
+			{
+				// If entry from lhs is first
+				if (FrameRangeSet.GetEntrySequence(LhsIndex) < FrameSet.GetEntrySequence(RhsIndex))
+				{
+					check(FrameRangeSet.GetEntryRangeNum(LhsIndex) > 0);
+
+					// Append subranges to output
+					OutFrameRangeSet.EntrySequences[OutIndex] = FrameRangeSet.GetEntrySequence(LhsIndex);
+					OutFrameRangeSet.EntryRangeOffsets[OutIndex] = RangeIndex;
+					OutFrameRangeSet.EntryRangeNums[OutIndex] = FrameRangeSet.GetEntryRangeNum(LhsIndex);
+
+					Array::Copy(
+						OutFrameRangeSet.RangeStarts.Slice(RangeIndex, FrameRangeSet.GetEntryRangeNum(LhsIndex)),
+						FrameRangeSet.GetEntryRangeStarts(LhsIndex));
+
+					Array::Copy(
+						OutFrameRangeSet.RangeLengths.Slice(RangeIndex, FrameRangeSet.GetEntryRangeNum(LhsIndex)),
+						FrameRangeSet.GetEntryRangeLengths(LhsIndex));
+
+					RangeIndex += FrameRangeSet.GetEntryRangeNum(LhsIndex);
+					OutIndex++;
+					LhsIndex++;
+				}
+				// If entry is in rhs but not lhs skip
+				else if (FrameSet.GetEntrySequence(RhsIndex) < FrameRangeSet.GetEntrySequence(LhsIndex))
+				{
+					RhsIndex++;
+				}
+				// If entry is in both lhs and rhs
+				else
+				{
+					check(FrameRangeSet.GetEntryRangeNum(LhsIndex) > 0);
+					check(FrameSet.GetEntryFrameNum(RhsIndex) > 0);
+
+					// Append difference of subranges to output
+					const int32 RangeNum = FrameRangeSet::Private::RangesFramesDifference(
+						OutFrameRangeSet.RangeStarts.Slice(RangeIndex, OutFrameRangeSet.RangeStarts.Num() - RangeIndex),
+						OutFrameRangeSet.RangeLengths.Slice(RangeIndex, OutFrameRangeSet.RangeLengths.Num() - RangeIndex),
+						FrameRangeSet.GetEntryRangeStarts(LhsIndex),
+						FrameRangeSet.GetEntryRangeLengths(LhsIndex),
+						FrameSet.GetEntryFrames(RhsIndex));
+
+					check(RangeNum <= FrameRangeSet.GetEntryRangeNum(LhsIndex) + FrameSet.GetEntryFrameNum(RhsIndex));
+
+					if (RangeNum > 0)
+					{
+						OutFrameRangeSet.EntrySequences[OutIndex] = FrameRangeSet.GetEntrySequence(LhsIndex);
+						OutFrameRangeSet.EntryRangeOffsets[OutIndex] = RangeIndex;
+						OutFrameRangeSet.EntryRangeNums[OutIndex] = RangeNum;
+						RangeIndex += RangeNum;
+						OutIndex++;
+					}
+
+					LhsIndex++; RhsIndex++;
+				}
+			}
+
+			// Process any remaining lhs entries
+			while (LhsIndex < LhsEntryNum)
+			{
+				check(FrameRangeSet.GetEntryRangeNum(LhsIndex) > 0);
+
+				// Append subranges to output
+				OutFrameRangeSet.EntrySequences[OutIndex] = FrameRangeSet.GetEntrySequence(LhsIndex);
+				OutFrameRangeSet.EntryRangeOffsets[OutIndex] = RangeIndex;
+				OutFrameRangeSet.EntryRangeNums[OutIndex] = FrameRangeSet.GetEntryRangeNum(LhsIndex);
+
+				Array::Copy(
+					OutFrameRangeSet.RangeStarts.Slice(RangeIndex, FrameRangeSet.GetEntryRangeNum(LhsIndex)),
+					FrameRangeSet.GetEntryRangeStarts(LhsIndex));
+
+				Array::Copy(
+					OutFrameRangeSet.RangeLengths.Slice(RangeIndex, FrameRangeSet.GetEntryRangeNum(LhsIndex)),
+					FrameRangeSet.GetEntryRangeLengths(LhsIndex));
+
+				RangeIndex += FrameRangeSet.GetEntryRangeNum(LhsIndex);
+				OutIndex++;
+				LhsIndex++;
+			}
+
+			// Resize output to match what was added
+			OutFrameRangeSet.EntrySequences.SetNumUninitialized({ OutIndex });
+			OutFrameRangeSet.EntryRangeOffsets.SetNumUninitialized({ OutIndex });
+			OutFrameRangeSet.EntryRangeNums.SetNumUninitialized({ OutIndex });
+			OutFrameRangeSet.RangeStarts.SetNumUninitialized({ RangeIndex });
+			OutFrameRangeSet.RangeLengths.SetNumUninitialized({ RangeIndex });
+			OutFrameRangeSet.RangeOffsets.SetNumUninitialized({ RangeIndex });
+			Private::ComputeRangeOffsets(OutFrameRangeSet.RangeOffsets, OutFrameRangeSet.RangeLengths);
+			OutFrameRangeSet.Check();
+		}
+
 		void Union(FFrameRangeSet& Out, const FFrameRangeSet& Lhs, const FFrameRangeSet& Rhs)
 		{
 			TRACE_CPUPROFILER_EVENT_SCOPE(Learning::FrameRangeSet::Union);
diff --git a/Engine/Plugins/Experimental/LearningAgents/Source/Learning/Public/LearningFrameRangeSet.h b/Engine/Plugins/Experimental/LearningAgents/Source/Learning/Public/LearningFrameRangeSet.h
--- a/Engine/Plugins/Experimental/LearningAgents/Source/Learning/Public/LearningFrameRangeSet.h
+++ b/Engine/Plugins/Experimental/LearningAgents/Source/Learning/Public/LearningFrameRangeSet.h
@@ -153,12 +153,18 @@ namespace UE::Learning
 		/** Checks if two frame range sets are equal */
 		LEARNING_API bool Equal(const FFrameRangeSet& Lhs, const FFrameRangeSet& Rhs);
 
+		/** Computes the union of a frame set and frame range set */
+		LEARNING_API void Union(FFrameRangeSet& OutFrameRangeSet, const FFrameSet& FrameSet, const FFrameRangeSet& FrameRangeSet);
+
 		/** Computes the intersection of a frame set and frame range set */
 		LEARNING_API void Intersection(FFrameSet& OutFrameSet, const FFrameSet& FrameSet, const FFrameRangeSet& FrameRangeSet);
 
 		/** Computes the difference of a frame set and frame range set */
 		LEARNING_API void Difference(FFrameSet& OutFrameSet, const FFrameSet& FrameSet, const FFrameRangeSet& FrameRangeSet);
 
+		/** Computes the difference of a frame range set and a frame set */
+		LEARNING_API void Difference(FFrameRangeSet& OutFrameRangeSet, const FFrameRangeSet& FrameRangeSet, const FFrameSet& FrameSet);
+
 		/** Computes the union of two frame range sets */
 		LEARNING_API void Union(FFrameRangeSet& Out, const FFrameRangeSet& Lhs, const FFrameRangeSet& Rhs);
 
