LearningAgents: Added a few more functions for Frame Attributes and Frame Range Sets.

diff --git a/Engine/Plugins/Experimental/LearningAgents/Source/Learning/Private/LearningFrameAttribute.cpp b/Engine/Plugins/Experimental/LearningAgents/Source/Learning/Private/LearningFrameAttribute.cpp
--- a/Engine/Plugins/Experimental/LearningAgents/Source/Learning/Private/LearningFrameAttribute.cpp
+++ b/Engine/Plugins/Experimental/LearningAgents/Source/Learning/Private/LearningFrameAttribute.cpp
@@ -153,6 +153,26 @@ namespace UE::Learning
 			}
 		}
 
+		static inline void Log(TLearningArrayView<1, float> Out, const TLearningArrayView<1, const float> In)
+		{
+			check(Out.Num() == In.Num());
+			const int32 ValueNum = Out.Num();
+			for (int32 ValueIdx = 0; ValueIdx < ValueNum; ValueIdx++)
+			{
+				Out[ValueIdx] = FMath::Loge(In[ValueIdx]);
+			}
+		}
+
+		static inline void Exp(TLearningArrayView<1, float> Out, const TLearningArrayView<1, const float> In)
+		{
+			check(Out.Num() == In.Num());
+			const int32 ValueNum = Out.Num();
+			for (int32 ValueIdx = 0; ValueIdx < ValueNum; ValueIdx++)
+			{
+				Out[ValueIdx] = FMath::Exp(In[ValueIdx]);
+			}
+		}
+
 		static inline void Sqrt(TLearningArrayView<1, float> Out, const TLearningArrayView<1, const float> In)
 		{
 			check(Out.Num() == In.Num());
@@ -573,6 +593,36 @@ namespace UE::Learning
 #endif
 		}
 
+		static inline void MeanStd(float& OutMean, float& OutStd, const TLearningArrayView<1, const float> In)
+		{
+			const int32 Num = In.Num();
+
+			OutMean = 0.0f;
+			OutStd = 0.0f;
+			for (int32 Idx = 0; Idx < Num; Idx++)
+			{
+				OutStd += (((float)Idx / Num) / (Idx + 1)) * FMath::Square(In[Idx] - OutMean);
+				OutMean += (In[Idx] - OutMean) / (Idx + 1);
+			}
+			OutStd = FMath::Sqrt(OutStd);
+		}
+
+		static inline void LogMeanStd(float& OutMean, float& OutLogStd, const TLearningArrayView<1, const float> In)
+		{
+			const int32 Num = In.Num();
+
+			OutMean = 0.0f;
+			OutLogStd = 0.0f;
+			for (int32 Idx = 0; Idx < Num; Idx++)
+			{
+				const float Value = FMath::Loge(FMath::Max(In[Idx], UE_SMALL_NUMBER));
+				OutLogStd += (((float)Idx / Num) / (Idx + 1)) * FMath::Square(Value - OutMean);
+				OutMean += (Value - OutMean) / (Idx + 1);
+			}
+			OutMean = FMath::Exp(OutMean);
+			OutLogStd = FMath::Sqrt(OutLogStd);
+		}
+		
 		static inline void QuatInv(TLearningArrayView<2, float> Out, const TLearningArrayView<2, const float> In)
 		{
 			check(Out.Num<0>() == 4);
@@ -609,6 +659,41 @@ namespace UE::Learning
 			}
 		}
 
+		static inline void QuatToRotationVector(TLearningArrayView<2, float> Out, const TLearningArrayView<2, const float> In)
+		{
+			check(Out.Num<0>() == 3);
+			check(In.Num<0>() == 4);
+			check(Out.Num<1>() == In.Num<1>());
+
+			const int32 ValueNum = Out.Num<1>();
+			for (int32 ValueIdx = 0; ValueIdx < ValueNum; ValueIdx++)
+			{
+				const FQuat4f InValue = FQuat4f(In[0][ValueIdx], In[1][ValueIdx], In[2][ValueIdx], In[3][ValueIdx]);
+				const FVector3f OutValue = InValue.ToRotationVector();
+				Out[0][ValueIdx] = OutValue.X;
+				Out[1][ValueIdx] = OutValue.Y;
+				Out[2][ValueIdx] = OutValue.Z;
+			}
+		}
+
+		static inline void QuatFromRotationVector(TLearningArrayView<2, float> Out, const TLearningArrayView<2, const float> In)
+		{
+			check(Out.Num<0>() == 4);
+			check(In.Num<0>() == 3);
+			check(Out.Num<1>() == In.Num<1>());
+
+			const int32 ValueNum = Out.Num<1>();
+			for (int32 ValueIdx = 0; ValueIdx < ValueNum; ValueIdx++)
+			{
+				const FVector3f InValue = FVector3f(In[0][ValueIdx], In[1][ValueIdx], In[2][ValueIdx]);
+				const FQuat4f OutValue = FQuat4f::MakeFromRotationVector(InValue);
+				Out[0][ValueIdx] = OutValue.X;
+				Out[1][ValueIdx] = OutValue.Y;
+				Out[2][ValueIdx] = OutValue.Z;
+				Out[3][ValueIdx] = OutValue.W;
+			}
+		}
+
 		static inline void QuatMul(
 			TLearningArrayView<1, float> OutX,
 			TLearningArrayView<1, float> OutY,
@@ -975,6 +1060,101 @@ namespace UE::Learning
 			}
 		}
 
+		// We need these custom Vector4 functions because the defaults for the FVector4f class will treat it like a Vector3
+
+		static inline float Vector4Length(const FVector4f X)
+		{
+			return FMath::Sqrt(X.X * X.X + X.Y * X.Y + X.Z * X.Z + X.W * X.W);
+		}
+
+		static inline FVector4f Vector4Normalize(const FVector4f X)
+		{
+			return X / Vector4Length(X);
+		}
+
+		static inline FVector4f DominantEigenVector(
+			const FMatrix44f& A,
+			const FVector4f V0,
+			const int32 Iterations,
+			const float Epsilon)
+		{
+			// Initial Guess at Eigen Vector & Value
+			FVector4f V = V0;
+			float EV = (A.TransformFVector4(V) / V).X;
+
+			for (int32 Iteration = 0; Iteration < Iterations; Iteration++)
+			{
+				// Power Iteration
+				const FVector4f Av = A.TransformFVector4(V);
+
+				// Next Guess at Eigen Vector & Value
+				const FVector4f VNew = Vector4Normalize(Av);
+				const float EVNew = (A.TransformFVector4(VNew) / VNew).X;
+
+				// Break if converged
+				if (FMath::Abs(EV - EVNew) < Epsilon)
+				{
+					break;
+				}
+
+				// Update best guess
+				V = VNew;
+				EV = EVNew;
+			}
+
+			return V;
+		}
+
+		static inline void QuatMeanStd(FQuat4f& OutMean, FVector3f& OutStd, const TLearningArrayView<2, const float> In)
+		{
+			check(In.Num<0>() == 4);
+
+			const int32 Num = In.Num<1>();
+
+			FMatrix44f Accum;
+			FPlatformMemory::Memzero(&Accum, sizeof(FMatrix44f));
+
+			for (int64 Idx = 0; Idx < Num; Idx++)
+			{
+				const FQuat4f Q = FQuat4f(In[0][Idx], In[1][Idx], In[2][Idx], In[3][Idx]);
+
+				Accum.M[0][0] += ((Q.X * Q.X) - Accum.M[0][0]) / (Idx + 1);
+				Accum.M[0][1] += ((Q.X * Q.Y) - Accum.M[0][1]) / (Idx + 1);
+				Accum.M[0][2] += ((Q.X * Q.Z) - Accum.M[0][2]) / (Idx + 1);
+				Accum.M[0][3] += ((Q.X * Q.W) - Accum.M[0][3]) / (Idx + 1);
+
+				Accum.M[1][0] += ((Q.Y * Q.X) - Accum.M[1][0]) / (Idx + 1);
+				Accum.M[1][1] += ((Q.Y * Q.Y) - Accum.M[1][1]) / (Idx + 1);
+				Accum.M[1][2] += ((Q.Y * Q.Z) - Accum.M[1][2]) / (Idx + 1);
+				Accum.M[1][3] += ((Q.Y * Q.W) - Accum.M[1][3]) / (Idx + 1);
+
+				Accum.M[2][0] += ((Q.Z * Q.X) - Accum.M[2][0]) / (Idx + 1);
+				Accum.M[2][1] += ((Q.Z * Q.Y) - Accum.M[2][1]) / (Idx + 1);
+				Accum.M[2][2] += ((Q.Z * Q.Z) - Accum.M[2][2]) / (Idx + 1);
+				Accum.M[2][3] += ((Q.Z * Q.W) - Accum.M[2][3]) / (Idx + 1);
+
+				Accum.M[3][0] += ((Q.W * Q.X) - Accum.M[3][0]) / (Idx + 1);
+				Accum.M[3][1] += ((Q.W * Q.Y) - Accum.M[3][1]) / (Idx + 1);
+				Accum.M[3][2] += ((Q.W * Q.Z) - Accum.M[3][2]) / (Idx + 1);
+				Accum.M[3][3] += ((Q.W * Q.W) - Accum.M[3][3]) / (Idx + 1);
+			}
+
+			const FVector4f AverageQuat = DominantEigenVector(Accum, FVector4f(0.0f, 0.0f, 0.0f, 1.0f), 128, 0.0f);
+			OutMean = FQuat4f(AverageQuat.X, AverageQuat.Y, AverageQuat.Z, AverageQuat.W);
+			check(OutMean.IsNormalized());
+
+			OutStd = FVector3f::ZeroVector;
+			for (int64 Idx = 0; Idx < Num; Idx++)
+			{
+				const FQuat4f Q = FQuat4f(In[0][Idx], In[1][Idx], In[2][Idx], In[3][Idx]);
+				const FQuat4f Diff = (Q * OutMean.Inverse()).GetShortestArcWith(FQuat4f::Identity);
+				OutStd += FMath::Square(Diff.ToRotationVector()) / Num;
+			}
+			OutStd.X = FMath::Sqrt(OutStd.X);
+			OutStd.Y = FMath::Sqrt(OutStd.Y);
+			OutStd.Z = FMath::Sqrt(OutStd.Z);
+		}
+
 	}
 
 	void FFrameAttribute::Check() const
@@ -1504,6 +1684,44 @@ namespace UE::Learning
 				});
 		}
 
+		void Log(FFrameAttribute& Out, const FFrameAttribute& In)
+		{
+			UnaryOp(Out, In.GetChannelNum(), In, [](
+				FFrameAttribute& Out,
+				const FFrameAttribute& In,
+				const TLearningArrayView<1, const int32> RangeOffsets,
+				const TLearningArrayView<1, const int32> RangeLengths) {
+
+					const int32 ChannelNum = Out.GetChannelNum();
+
+					for (int32 ChannelIdx = 0; ChannelIdx < ChannelNum; ChannelIdx++)
+					{
+						Private::Log(
+							Out.GetChannelAttributeData(ChannelIdx),
+							In.GetChannelAttributeData(ChannelIdx));
+					}
+				});
+		}
+
+		void Exp(FFrameAttribute& Out, const FFrameAttribute& In)
+		{
+			UnaryOp(Out, In.GetChannelNum(), In, [](
+				FFrameAttribute& Out,
+				const FFrameAttribute& In,
+				const TLearningArrayView<1, const int32> RangeOffsets,
+				const TLearningArrayView<1, const int32> RangeLengths) {
+
+					const int32 ChannelNum = Out.GetChannelNum();
+
+					for (int32 ChannelIdx = 0; ChannelIdx < ChannelNum; ChannelIdx++)
+					{
+						Private::Exp(
+							Out.GetChannelAttributeData(ChannelIdx),
+							In.GetChannelAttributeData(ChannelIdx));
+					}
+				});
+		}
+
 		void Length(FFrameAttribute& Out, const FFrameAttribute& In)
 		{
 			UnaryOp(Out, 1, In, [](
@@ -2293,6 +2511,38 @@ namespace UE::Learning
 				});
 		}
 
+		void MeanStd(TLearningArrayView<1, float> OutMean, TLearningArrayView<1, float> OutStd, const FFrameAttribute& In)
+		{
+			ReduceOp(In, [&OutMean, &OutStd](
+				const FFrameAttribute& In,
+				const TLearningArrayView<1, const int32> RangeOffsets,
+				const TLearningArrayView<1, const int32> RangeLengths) {
+
+					const int32 ChannelNum = In.GetChannelNum();
+
+					for (int32 ChannelIdx = 0; ChannelIdx < ChannelNum; ChannelIdx++)
+					{
+						Private::MeanStd(OutMean[ChannelIdx], OutStd[ChannelIdx], In.GetChannelAttributeData(ChannelIdx));
+					}
+				});
+		}
+
+		void LogMeanStd(TLearningArrayView<1, float> OutMean, TLearningArrayView<1, float> OutLogStd, const FFrameAttribute& In)
+		{
+			ReduceOp(In, [&OutMean, &OutLogStd](
+				const FFrameAttribute& In,
+				const TLearningArrayView<1, const int32> RangeOffsets,
+				const TLearningArrayView<1, const int32> RangeLengths) {
+
+					const int32 ChannelNum = In.GetChannelNum();
+
+					for (int32 ChannelIdx = 0; ChannelIdx < ChannelNum; ChannelIdx++)
+					{
+						Private::LogMeanStd(OutMean[ChannelIdx], OutLogStd[ChannelIdx], In.GetChannelAttributeData(ChannelIdx));
+					}
+				});
+		}
+
 		void QuatMul(FFrameAttribute& Out, const FFrameAttribute& Lhs, const FFrameAttribute& Rhs)
 		{
 			check(Lhs.GetChannelNum() == 4);
@@ -2343,9 +2593,7 @@ namespace UE::Learning
 				const TLearningArrayView<1, const int32> RangeOffsets,
 				const TLearningArrayView<1, const int32> RangeLengths) {
 
-					Private::QuatInv(
-						Out.GetAttributeData(),
-						In.GetAttributeData());
+					Private::QuatInv(Out.GetAttributeData(), In.GetAttributeData());
 				});
 		}
 
@@ -2359,9 +2607,35 @@ namespace UE::Learning
 				const TLearningArrayView<1, const int32> RangeOffsets,
 				const TLearningArrayView<1, const int32> RangeLengths) {
 
-					Private::QuatAbs(
-						Out.GetAttributeData(),
-						In.GetAttributeData());
+					Private::QuatAbs(Out.GetAttributeData(), In.GetAttributeData());
+				});
+		}
+
+		void QuatToRotationVector(FFrameAttribute& Out, const FFrameAttribute& In)
+		{
+			check(In.GetChannelNum() == 4);
+
+			UnaryOp(Out, 3, In, [](
+				FFrameAttribute& Out,
+				const FFrameAttribute& In,
+				const TLearningArrayView<1, const int32> RangeOffsets,
+				const TLearningArrayView<1, const int32> RangeLengths) {
+				
+					Private::QuatToRotationVector(Out.GetAttributeData(), In.GetAttributeData());
+				});
+		}
+
+		void QuatFromRotationVector(FFrameAttribute& Out, const FFrameAttribute& In)
+		{
+			check(In.GetChannelNum() == 3);
+
+			UnaryOp(Out, 4, In, [](
+				FFrameAttribute& Out,
+				const FFrameAttribute& In,
+				const TLearningArrayView<1, const int32> RangeOffsets,
+				const TLearningArrayView<1, const int32> RangeLengths) {
+				
+					Private::QuatFromRotationVector(Out.GetAttributeData(),	In.GetAttributeData());
 				});
 		}
 
@@ -2748,5 +3022,18 @@ namespace UE::Learning
 				});
 		}
 
+		void QuatMeanStd(FQuat4f& OutMean, FVector3f& OutStd, const FFrameAttribute& In)
+		{
+			check(In.GetChannelNum() == 4);
+
+			ReduceOp(In, [&OutMean, &OutStd](
+				const FFrameAttribute& In,
+				const TLearningArrayView<1, const int32> RangeOffsets,
+				const TLearningArrayView<1, const int32> RangeLengths) {
+
+					Private::QuatMeanStd(OutMean, OutStd, In.GetAttributeData());
+				});
+		}
+
 	}
-}
\ No newline at end of file
+}
diff --git a/Engine/Plugins/Experimental/LearningAgents/Source/Learning/Private/LearningFrameRangeSet.cpp b/Engine/Plugins/Experimental/LearningAgents/Source/Learning/Private/LearningFrameRangeSet.cpp
--- a/Engine/Plugins/Experimental/LearningAgents/Source/Learning/Private/LearningFrameRangeSet.cpp
+++ b/Engine/Plugins/Experimental/LearningAgents/Source/Learning/Private/LearningFrameRangeSet.cpp
@@ -14,8 +14,6 @@ namespace UE::Learning
 			const int32 FrameNum = Frames.Num();
 			for (int32 FrameIdx = 0; FrameIdx < FrameNum - 1; FrameIdx++)
 			{
-				check(Frames[FrameIdx + 0] >= 0);
-				check(Frames[FrameIdx + 1] >= 0);
 				check(Frames[FrameIdx + 0] < Frames[FrameIdx + 1]);
 			}
 		}
@@ -203,6 +201,13 @@ namespace UE::Learning
 				return RhsLengths.Num();
 			}
 
+			if (RhsStarts.IsEmpty())
+			{
+				Array::Copy(OutStarts.Slice(0, LhsFrames.Num()), LhsFrames);
+				Array::Set(OutLengths.Slice(0, LhsFrames.Num()), 1);
+				return LhsFrames.Num();
+			}
+
 			// Number of ranges in lhs and rhs
 			const int32 LhsNum = LhsFrames.Num();
 			const int32 RhsNum = RhsStarts.Num();
@@ -363,6 +368,9 @@ namespace UE::Learning
 			FramesCheck(LhsFrames);
 			RangesCheck(RhsStarts, RhsLengths);
 
+			if (LhsFrames.IsEmpty()) { return 0; }
+			if (RhsStarts.IsEmpty()) { return 0; }
+
 			// Number of frames/ranges in lhs and rhs
 			const int32 LhsNum = LhsFrames.Num();
 			const int32 RhsNum = RhsStarts.Num();
@@ -439,6 +447,14 @@ namespace UE::Learning
 			FramesCheck(LhsFrames);
 			RangesCheck(RhsStarts, RhsLengths);
 
+			if (LhsFrames.IsEmpty()) { return 0; }
+
+			if (RhsStarts.IsEmpty())
+			{
+				Array::Copy(OutFrames.Slice(0, LhsFrames.Num()), LhsFrames);
+				return LhsFrames.Num();
+			}
+
 			// Number of frames/ranges in lhs and rhs
 			const int32 LhsNum = LhsFrames.Num();
 			const int32 RhsNum = RhsStarts.Num();
@@ -517,6 +533,15 @@ namespace UE::Learning
 			RangesCheck(LhsStarts, LhsLengths);
 			FramesCheck(RhsFrames);
 
+			if (LhsStarts.IsEmpty()) { return 0; }
+
+			if (RhsFrames.IsEmpty())
+			{
+				Array::Copy(OutStarts.Slice(0, LhsStarts.Num()), LhsStarts);
+				Array::Copy(OutLengths.Slice(0, LhsLengths.Num()), LhsLengths);
+				return LhsLengths.Num();
+			}
+
 			// Number of ranges in lhs and rhs
 			const int32 LhsNum = LhsStarts.Num();
 			const int32 RhsNum = RhsFrames.Num();
@@ -651,6 +676,20 @@ namespace UE::Learning
 			RangesCheck(LhsStarts, LhsLengths);
 			RangesCheck(RhsStarts, RhsLengths);
 
+			if (LhsStarts.IsEmpty())
+			{
+				Array::Copy(OutStarts.Slice(0, RhsStarts.Num()), RhsStarts);
+				Array::Copy(OutLengths.Slice(0, RhsLengths.Num()), RhsLengths);
+				return RhsLengths.Num();
+			}
+
+			if (RhsStarts.IsEmpty())
+			{
+				Array::Copy(OutStarts.Slice(0, LhsStarts.Num()), LhsStarts);
+				Array::Copy(OutLengths.Slice(0, LhsLengths.Num()), LhsLengths);
+				return LhsLengths.Num();
+			}
+
 			if (RangesEqual(LhsStarts, LhsLengths, RhsStarts, RhsLengths))
 			{
 				Array::Copy(OutStarts.Slice(0, LhsStarts.Num()), LhsStarts);
@@ -820,6 +859,9 @@ namespace UE::Learning
 			RangesCheck(LhsStarts, LhsLengths);
 			RangesCheck(RhsStarts, RhsLengths);
 
+			if (LhsStarts.IsEmpty()) { return 0; }
+			if (RhsStarts.IsEmpty()) { return 0; }
+
 			if (RangesEqual(LhsStarts, LhsLengths, RhsStarts, RhsLengths))
 			{
 				Array::Copy(OutStarts.Slice(0, LhsStarts.Num()), LhsStarts);
@@ -953,6 +995,9 @@ namespace UE::Learning
 			RangesCheck(LhsStarts, LhsLengths);
 			RangesCheck(RhsStarts, RhsLengths);
 
+			if (LhsStarts.IsEmpty()) { return 0; }
+			if (RhsStarts.IsEmpty()) { return 0; }
+
 			if (RangesEqual(LhsStarts, LhsLengths, RhsStarts, RhsLengths))
 			{
 				Array::Copy(OutStarts.Slice(0, LhsStarts.Num()), LhsStarts);
@@ -1122,6 +1167,15 @@ namespace UE::Learning
 			RangesCheck(LhsStarts, LhsLengths);
 			RangesCheck(RhsStarts, RhsLengths);
 
+			if (LhsStarts.IsEmpty()) { return 0; }
+
+			if (RhsStarts.IsEmpty())
+			{
+				Array::Copy(OutStarts.Slice(0, LhsStarts.Num()), LhsStarts);
+				Array::Copy(OutLengths.Slice(0, LhsLengths.Num()), LhsLengths);
+				return LhsLengths.Num();
+			}
+
 			if (RangesEqual(LhsStarts, LhsLengths, RhsStarts, RhsLengths))
 			{
 				return 0;
@@ -1665,6 +1719,9 @@ namespace UE::Learning
 
 			FrameSet.Check(); FrameRangeSet.Check();
 
+			if (FrameSet.IsEmpty()) { OutFrameSet.Empty(); return; }
+			if (FrameRangeSet.IsEmpty()) { OutFrameSet.Empty(); return; }
+
 			const int32 LhsEntryNum = FrameSet.GetEntryNum();
 			const int32 RhsEntryNum = FrameRangeSet.GetEntryNum();
 			const int32 LhsFrameNum = FrameSet.GetTotalFrameNum();
@@ -1737,6 +1794,9 @@ namespace UE::Learning
 
 			FrameSet.Check(); FrameRangeSet.Check();
 
+			if (FrameSet.IsEmpty()) { OutFrameSet.Empty(); return; }
+			if (FrameRangeSet.IsEmpty()) { OutFrameSet = FrameSet; return; }
+
 			const int32 LhsEntryNum = FrameSet.GetEntryNum();
 			const int32 RhsEntryNum = FrameRangeSet.GetEntryNum();
 			const int32 LhsFrameNum = FrameSet.GetTotalFrameNum();
@@ -1845,6 +1905,7 @@ namespace UE::Learning
 
 			FrameRangeSet.Check(); FrameSet.Check();
 
+			if (FrameRangeSet.IsEmpty()) { OutFrameRangeSet.Empty(); return; }
 			if (FrameSet.IsEmpty()) { OutFrameRangeSet = FrameRangeSet; return; }
 
 			const int32 LhsEntryNum = FrameRangeSet.GetEntryNum();
@@ -1966,6 +2027,8 @@ namespace UE::Learning
 
 			Lhs.Check(); Rhs.Check();
 
+			if (Lhs.IsEmpty()) { Out = Rhs; return; }
+			if (Rhs.IsEmpty()) { Out = Lhs; return; }
 			if (Equal(Lhs, Rhs)) { Out = Lhs; return; }
 
 			// Allocate potential maximum number of entries and ranges we might to output
@@ -2130,6 +2193,8 @@ namespace UE::Learning
 
 			Lhs.Check(); Rhs.Check();
 
+			if (Lhs.IsEmpty()) { Out.Empty(); return; }
+			if (Rhs.IsEmpty()) { Out.Empty(); return; }
 			if (Equal(Lhs, Rhs)) { Out = Lhs; return; }
 
 			const int32 LhsEntryNum = Lhs.GetEntryNum();
@@ -2212,6 +2277,8 @@ namespace UE::Learning
 
 			Lhs.Check(); Rhs.Check();
 
+			if (Lhs.IsEmpty()) { Out.Empty(); return; }
+			if (Rhs.IsEmpty()) { Out = Lhs; return; }
 			if (Equal(Lhs, Rhs)) { Out.Empty(); return; }
 
 			const int32 LhsEntryNum = Lhs.GetEntryNum();
@@ -2443,46 +2510,20 @@ namespace UE::Learning
 			return RangeIndex;
 		}
 
-		void TrimLeft(FFrameRangeSet& Out, const FFrameRangeSet& FrameRangeSet, const int32 TrimFrameNum)
+		void TrimStart(FFrameRangeSet& Out, const FFrameRangeSet& FrameRangeSet, const int32 TrimFrameNum)
 		{
-			check(TrimFrameNum >= 0);
-			FrameRangeSet.Check();
-
-			const int32 EntryNum = FrameRangeSet.GetEntryNum();
-
-			Out.Empty();
-
-			TArray<int32> RangeStartsAdded;
-			TArray<int32> RangeLengthsAdded;
-
-			int32 RangeOffset = 0;
-			for (int32 EntryIdx = 0; EntryIdx < EntryNum; EntryIdx++)
-			{
-				const int32 Sequence = FrameRangeSet.GetEntrySequence(EntryIdx);
-				const int32 RangeNum = FrameRangeSet.GetEntryRangeNum(EntryIdx);
-
-				RangeStartsAdded.Reset();
-				RangeLengthsAdded.Reset();
-				for (int32 RangeIdx = 0; RangeIdx < RangeNum; RangeIdx++)
-				{
-					const int32 RangeStart = FrameRangeSet.GetEntryRangeStart(EntryIdx, RangeIdx);
-					const int32 RangeLength = FrameRangeSet.GetEntryRangeLength(EntryIdx, RangeIdx);
-					if (RangeLength - TrimFrameNum > 0)
-					{
-						RangeStartsAdded.Add(RangeStart + TrimFrameNum);
-						RangeLengthsAdded.Add(RangeLength - TrimFrameNum);
-					}
-				}
-
-				Out.AddEntry(Sequence, RangeStartsAdded, RangeLengthsAdded);
-			}
+			Trim(Out, FrameRangeSet, TrimFrameNum, 0);
+		}
 
-			Out.Check();
+		void TrimEnd(FFrameRangeSet& Out, const FFrameRangeSet& FrameRangeSet, const int32 TrimFrameNum)
+		{
+			Trim(Out, FrameRangeSet, 0, TrimFrameNum);
 		}
 
-		void TrimRight(FFrameRangeSet& Out, const FFrameRangeSet& FrameRangeSet, const int32 TrimFrameNum)
+		void Trim(FFrameRangeSet& Out, const FFrameRangeSet& FrameRangeSet, const int32 TrimStartFrameNum, const int32 TrimEndFrameNum)
 		{
-			check(TrimFrameNum >= 0);
+			check(TrimStartFrameNum >= 0);
+			check(TrimEndFrameNum >= 0);
 			FrameRangeSet.Check();
 
 			const int32 EntryNum = FrameRangeSet.GetEntryNum();
@@ -2502,12 +2543,12 @@ namespace UE::Learning
 				RangeLengthsAdded.Reset();
 				for (int32 RangeIdx = 0; RangeIdx < RangeNum; RangeIdx++)
 				{
-					const int32 RangeStart = FrameRangeSet.GetEntryRangeStart(EntryIdx, RangeIdx);
-					const int32 RangeLength = FrameRangeSet.GetEntryRangeLength(EntryIdx, RangeIdx);
-					if (RangeLength - TrimFrameNum > 0)
+					const int32 RangeStart = FrameRangeSet.GetEntryRangeStart(EntryIdx, RangeIdx) + TrimStartFrameNum;
+					const int32 RangeLength = FrameRangeSet.GetEntryRangeLength(EntryIdx, RangeIdx) - TrimStartFrameNum - TrimEndFrameNum;
+					if (RangeLength > 0)
 					{
 						RangeStartsAdded.Add(RangeStart);
-						RangeLengthsAdded.Add(RangeLength - TrimFrameNum);
+						RangeLengthsAdded.Add(RangeLength);
 					}
 				}
 
@@ -2517,9 +2558,20 @@ namespace UE::Learning
 			Out.Check();
 		}
 
-		void TrimLeftAndRight(FFrameRangeSet& Out, const FFrameRangeSet& FrameRangeSet, const int32 TrimFrameNum)
+		void PadStart(FFrameRangeSet& Out, const FFrameRangeSet& FrameRangeSet, const int32 PadFrameNum)
+		{
+			Pad(Out, FrameRangeSet, PadFrameNum, 0);
+		}
+
+		void PadEnd(FFrameRangeSet& Out, const FFrameRangeSet& FrameRangeSet, const int32 PadFrameNum)
 		{
-			check(TrimFrameNum >= 0);
+			Pad(Out, FrameRangeSet, 0, PadFrameNum);
+		}
+
+		void Pad(FFrameRangeSet& Out, const FFrameRangeSet& FrameRangeSet, const int32 PadStartFrameNum, int32 PadEndFrameNum)
+		{
+			check(PadStartFrameNum >= 0);
+			check(PadEndFrameNum >= 0);
 			FrameRangeSet.Check();
 
 			const int32 EntryNum = FrameRangeSet.GetEntryNum();
@@ -2537,14 +2589,46 @@ namespace UE::Learning
 
 				RangeStartsAdded.Reset();
 				RangeLengthsAdded.Reset();
+
+				bool bRangeActive = false;
+				int32 RangeActiveStart = INDEX_NONE;
+				int32 RangeActiveLength = INDEX_NONE;
+
 				for (int32 RangeIdx = 0; RangeIdx < RangeNum; RangeIdx++)
 				{
-					const int32 RangeStart = FrameRangeSet.GetEntryRangeStart(EntryIdx, RangeIdx);
-					const int32 RangeLength = FrameRangeSet.GetEntryRangeLength(EntryIdx, RangeIdx);
-					if (RangeLength - 2 * TrimFrameNum > 0)
+					if (!bRangeActive)
+					{
+						bRangeActive = true;
+						RangeActiveStart = FrameRangeSet.GetEntryRangeStart(EntryIdx, RangeIdx) - PadStartFrameNum;
+						RangeActiveLength = FrameRangeSet.GetEntryRangeLength(EntryIdx, RangeIdx) + PadStartFrameNum + PadEndFrameNum;
+					}
+
+					if (RangeIdx < RangeNum - 1)
+					{
+						const int32 NextRangeStart = FrameRangeSet.GetEntryRangeStart(EntryIdx, RangeIdx + 1) - PadStartFrameNum;
+						const int32 NextRangeLength = FrameRangeSet.GetEntryRangeLength(EntryIdx, RangeIdx + 1) + PadStartFrameNum + PadEndFrameNum;
+
+						if (NextRangeStart < RangeActiveStart + RangeActiveLength)
+						{
+							RangeActiveLength = NextRangeStart + NextRangeLength - RangeActiveStart;
+							continue;
+						}
+						else
+						{
+							RangeStartsAdded.Add(RangeActiveStart);
+							RangeLengthsAdded.Add(RangeActiveLength);
+							bRangeActive = false;
+							RangeActiveStart = INDEX_NONE;
+							RangeActiveLength = INDEX_NONE;
+						}
+					}
+					else
 					{
-						RangeStartsAdded.Add(RangeStart + TrimFrameNum);
-						RangeLengthsAdded.Add(RangeLength - 2 * TrimFrameNum);
+						RangeStartsAdded.Add(RangeActiveStart);
+						RangeLengthsAdded.Add(RangeActiveLength);
+						bRangeActive = false;
+						RangeActiveStart = INDEX_NONE;
+						RangeActiveLength = INDEX_NONE;
 					}
 				}
 
diff --git a/Engine/Plugins/Experimental/LearningAgents/Source/Learning/Public/LearningFrameAttribute.h b/Engine/Plugins/Experimental/LearningAgents/Source/Learning/Public/LearningFrameAttribute.h
--- a/Engine/Plugins/Experimental/LearningAgents/Source/Learning/Public/LearningFrameAttribute.h
+++ b/Engine/Plugins/Experimental/LearningAgents/Source/Learning/Public/LearningFrameAttribute.h
@@ -216,6 +216,12 @@ namespace UE::Learning
 		/** Compute the absolute value of a frame attribute. */
 		LEARNING_API void Abs(FFrameAttribute& Out, const FFrameAttribute& In);
 
+		/** Compute the log of a frame attribute. */
+		LEARNING_API void Log(FFrameAttribute& Out, const FFrameAttribute& In);
+
+		/** Compute the exp of a frame attribute. */
+		LEARNING_API void Exp(FFrameAttribute& Out, const FFrameAttribute& In);
+
 		/** Compute the length of a frame attribute over the channels dimension. */
 		LEARNING_API void Length(FFrameAttribute& Out, const FFrameAttribute& In);
 
@@ -324,6 +330,12 @@ namespace UE::Learning
 		/** Applies a Majority Vote filter to the ranges of values in the frame attribute. */
 		LEARNING_API void FilterMajorityVote(FFrameAttribute& Out, const FFrameAttribute& In, const float FilterWidthFrames);
 
+		/** Computes the mean and std across all frames. OutMean and OutStd should be size of the number of channels */
+		LEARNING_API void MeanStd(TLearningArrayView<1, float> OutMean, TLearningArrayView<1, float> OutStd, const FFrameAttribute& In);
+
+		/** Computes the mean and log std across all frames in the log space. OutMean and OutLogStd should be size of the number of channels */
+		LEARNING_API void LogMeanStd(TLearningArrayView<1, float> OutMean, TLearningArrayView<1, float> OutLogStd, const FFrameAttribute& In);
+
 		/** Computes the quaternion multiplication of two frame attributes. Both arguments are assumed to have 4 channels. */
 		LEARNING_API void QuatMul(FFrameAttribute& Out, const FFrameAttribute& Lhs, const FFrameAttribute& Rhs);
 
@@ -336,6 +348,12 @@ namespace UE::Learning
 		/** Compute the quaternion closest to the identity quaternion for a frame attribute. Argument is assumed to have 4 channels. */
 		LEARNING_API void QuatAbs(FFrameAttribute& Out, const FFrameAttribute& In);
 
+		/** Compute a rotation vector from a quaternion. Argument is assumed to have 4 channels. */
+		LEARNING_API void QuatToRotationVector(FFrameAttribute& Out, const FFrameAttribute& In);
+
+		/** Compute a quaternion from a rotation vector. Argument is assumed to have 3 channels. */
+		LEARNING_API void QuatFromRotationVector(FFrameAttribute& Out, const FFrameAttribute& In);
+
 		/** Computes the quaternion inverse multiplication of two frame attributes. Both arguments are assumed to have 4 channels. */
 		LEARNING_API void QuatInvMul(FFrameAttribute& Out, const FFrameAttribute& Lhs, const FFrameAttribute& Rhs);
 
@@ -392,5 +410,8 @@ namespace UE::Learning
 
 		/** Computes the quaternion between the given frame attribute and vector. Rhs is assumed to have 3 channels. */
 		LEARNING_API void QuatConstantBetween(FFrameAttribute& Out, const FVector3f Lhs, const FFrameAttribute& Rhs);
+
+		/** Computes the quaternion mean and std across all frames. */
+		LEARNING_API void QuatMeanStd(FQuat4f& OutMean, FVector3f& OutStd, const FFrameAttribute& In);
 	}
 }
\ No newline at end of file
diff --git a/Engine/Plugins/Experimental/LearningAgents/Source/Learning/Public/LearningFrameRangeSet.h b/Engine/Plugins/Experimental/LearningAgents/Source/Learning/Public/LearningFrameRangeSet.h
--- a/Engine/Plugins/Experimental/LearningAgents/Source/Learning/Public/LearningFrameRangeSet.h
+++ b/Engine/Plugins/Experimental/LearningAgents/Source/Learning/Public/LearningFrameRangeSet.h
@@ -186,14 +186,23 @@ namespace UE::Learning
 			const FFrameRangeSet& Lhs,
 			const FFrameRangeSet& Rhs);
 
-		/** Creates a new frame range set with the left-hand-side (i.e. the start) of all ranges trimmed by the given number of frames */
-		LEARNING_API void TrimLeft(FFrameRangeSet& Out, const FFrameRangeSet& FrameRangeSet, const int32 TrimFrameNum);
+		/** Creates a new frame range set with the start of all ranges trimmed by the given number of frames */
+		LEARNING_API void TrimStart(FFrameRangeSet& Out, const FFrameRangeSet& FrameRangeSet, const int32 TrimFrameNum);
 
-		/** Creates a new frame range set with the right-hand-side (i.e. the end) of all ranges trimmed by the given number of frames */
-		LEARNING_API void TrimRight(FFrameRangeSet& Out, const FFrameRangeSet& FrameRangeSet, const int32 TrimFrameNum);
+		/** Creates a new frame range set with the end of all ranges trimmed by the given number of frames */
+		LEARNING_API void TrimEnd(FFrameRangeSet& Out, const FFrameRangeSet& FrameRangeSet, const int32 TrimFrameNum);
 
-		/** Creates a new frame range set with both the start and the end of all ranges trimmed by the given number of frames */
-		LEARNING_API void TrimLeftAndRight(FFrameRangeSet& Out, const FFrameRangeSet& FrameRangeSet, const int32 TrimFrameNum);
+		/** Creates a new frame range set with both the start and the end of all ranges trimmed by the given numbers of frames */
+		LEARNING_API void Trim(FFrameRangeSet& Out, const FFrameRangeSet& FrameRangeSet, const int32 TrimStartFrameNum, const int32 TrimEndFrameNum);
+
+		/** Creates a new frame range set with the start of all ranges padded by the given number of frames */
+		LEARNING_API void PadStart(FFrameRangeSet& Out, const FFrameRangeSet& FrameRangeSet, const int32 PadFrameNum);
+
+		/** Creates a new frame range set with the end of all ranges padded by the given number of frames */
+		LEARNING_API void PadEnd(FFrameRangeSet& Out, const FFrameRangeSet& FrameRangeSet, const int32 PadFrameNum);
+
+		/** Creates a new frame range set with both the start and the end of all ranges padded by the given numbers of frames */
+		LEARNING_API void Pad(FFrameRangeSet& Out, const FFrameRangeSet& FrameRangeSet, const int32 PadStartFrameNum, int32 PadEndFrameNum);
 
 		/** Makes a Frame Range Set consisting of single-frame ranges from a frame set */
 		LEARNING_API void MakeFromFrameSet(FFrameRangeSet& OutFrameRangeSet, const FFrameSet& FrameSet);
