[LearningAgents] Changes to support robust distributed training

+ Robust handling of connection dropping
+ Robust buffer filling before a training step with min_buffer_size and max_buffer_size settings
+ Robust handling of always use latest networks on UE side and discarding stale experiences (default behavior)

diff --git a/Engine/Plugins/Experimental/LearningAgents/Source/LearningAgentsTraining/Private/LearningAgentsImitationTrainer.cpp b/Engine/Plugins/Experimental/LearningAgents/Source/LearningAgentsTraining/Private/LearningAgentsImitationTrainer.cpp
--- a/Engine/Plugins/Experimental/LearningAgents/Source/LearningAgentsTraining/Private/LearningAgentsImitationTrainer.cpp
+++ b/Engine/Plugins/Experimental/LearningAgents/Source/LearningAgentsTraining/Private/LearningAgentsImitationTrainer.cpp
@@ -169,25 +169,6 @@ void ULearningAgentsImitationTrainer::BeginTraining(
 		return;
 	}
 
-	// Check Paths
-
-	const FString PythonExecutablePath = UE::Learning::Trainer::GetPythonExecutablePath(ImitationTrainerPathSettings.GetIntermediatePath());
-
-	if (!FPaths::FileExists(PythonExecutablePath))
-	{
-		UE_LOG(LogLearning, Error, TEXT("%s: Can't find Python executable \"%s\"."), *GetName(), *PythonExecutablePath);
-		return;
-	}
-	const FString PythonContentPath = UE::Learning::Trainer::GetPythonContentPath(ImitationTrainerPathSettings.GetEditorEnginePath());
-
-	if (!FPaths::DirectoryExists(PythonContentPath))
-	{
-		UE_LOG(LogLearning, Error, TEXT("%s: Can't find LearningAgents plugin Content \"%s\"."), *GetName(), *PythonContentPath);
-		return;
-	}
-
-	const FString IntermediatePath = UE::Learning::Trainer::GetIntermediatePath(ImitationTrainerPathSettings.GetIntermediatePath());
-
 	// Sizes
 
 	const int32 ObservationNum = Interactor->GetObservationVectorSize();
@@ -425,7 +406,6 @@ TSharedRef<FJsonObject> ULearningAgentsImitationTrainer::CreateTrainerConfig(con
 	TSharedRef<FJsonObject> ConfigObject = MakeShared<FJsonObject>();
 
 	// Add Training Config Entries
-	ConfigObject->SetStringField(TEXT("TaskName"), TEXT("Training"));
 	ConfigObject->SetStringField(TEXT("TrainerMethod"), TEXT("BehaviorCloning"));
 	ConfigObject->SetStringField(TEXT("TimeStamp"), *FDateTime::Now().ToFormattedString(TEXT("%Y-%m-%d_%H-%M-%S")));
 
diff --git a/Engine/Plugins/Experimental/LearningAgents/Source/LearningAgentsTraining/Private/LearningAgentsPPOTrainer.cpp b/Engine/Plugins/Experimental/LearningAgents/Source/LearningAgentsTraining/Private/LearningAgentsPPOTrainer.cpp
--- a/Engine/Plugins/Experimental/LearningAgents/Source/LearningAgentsTraining/Private/LearningAgentsPPOTrainer.cpp
+++ b/Engine/Plugins/Experimental/LearningAgents/Source/LearningAgentsTraining/Private/LearningAgentsPPOTrainer.cpp
@@ -450,7 +450,6 @@ TSharedRef<FJsonObject> ULearningAgentsPPOTrainer::CreateTrainerConfig(const FLe
 	TSharedRef<FJsonObject> ConfigObject = MakeShared<FJsonObject>();
 
 	// Add Training Task-specific Config Entries
-	ConfigObject->SetStringField(TEXT("TaskName"), TEXT("Training"));
 	ConfigObject->SetStringField(TEXT("TrainerMethod"), TEXT("PPO"));
 	ConfigObject->SetStringField(TEXT("TimeStamp"), *FDateTime::Now().ToFormattedString(TEXT("%Y-%m-%d_%H-%M-%S")));
 
@@ -617,58 +616,31 @@ void ULearningAgentsPPOTrainer::ProcessExperience(const bool bResetAgentsOnUpdat
 
 		ReplayBuffer->Reset();
 
-		// Get Updated Policy
-		Response = Trainer->ReceiveNetwork(PolicyNetworkId, *Policy->GetPolicyNetworkAsset()->NeuralNetworkData);
-		Policy->GetPolicyNetworkAsset()->ForceMarkDirty();
-
-		if (Response == UE::Learning::ETrainerResponse::Completed)
-		{
-			UE_LOG(LogLearning, Display, TEXT("%s: Trainer completed training."), *GetName());
-			DoneTraining();
-			return;
-		}
-		else if (Response != UE::Learning::ETrainerResponse::Success)
-		{
-			UE_LOG(LogLearning, Error, TEXT("%s: Error waiting for policy from trainer: %s. Check log for additional errors."), *GetName(), UE::Learning::Trainer::GetResponseString(Response));
-			bHasTrainingFailed = true;
-			EndTraining();
-			return;
-		}
-
-		// Get Updated Critic
-		Response = Trainer->ReceiveNetwork(CriticNetworkId, *Critic->GetCriticNetworkAsset()->NeuralNetworkData);
-		Critic->GetCriticNetworkAsset()->ForceMarkDirty();
-
-		if (Response != UE::Learning::ETrainerResponse::Success)
-		{
-			UE_LOG(LogLearning, Error, TEXT("%s: Error waiting for critic from trainer: %s. Check log for additional errors."), *GetName(), UE::Learning::Trainer::GetResponseString(Response));
-			bHasTrainingFailed = true;
-			EndTraining();
-			return;
-		}
-
-		// Get Updated Encoder
-		Response = Trainer->ReceiveNetwork(EncoderNetworkId, *Policy->GetEncoderNetworkAsset()->NeuralNetworkData);
-		Policy->GetEncoderNetworkAsset()->ForceMarkDirty();
-
-		if (Response != UE::Learning::ETrainerResponse::Success)
+		const TArray<int32> NetworkIds = {PolicyNetworkId, CriticNetworkId, EncoderNetworkId, DecoderNetworkId};
+		const TArray<TObjectPtr<ULearningNeuralNetworkData>> Networks = {Policy->GetPolicyNetworkAsset()->NeuralNetworkData, Critic->GetCriticNetworkAsset()->NeuralNetworkData, Policy->GetEncoderNetworkAsset()->NeuralNetworkData, Policy->GetDecoderNetworkAsset()->NeuralNetworkData};
+		TArray<UE::Learning::ETrainerResponse> NetworkResponses = Trainer->ReceiveNetworks(NetworkIds, Networks);
+	
+		TArray<ULearningAgentsNeuralNetwork*> NetworkAssets = { Policy->GetPolicyNetworkAsset(), Critic->GetCriticNetworkAsset(), Policy->GetEncoderNetworkAsset(), Policy->GetDecoderNetworkAsset() };
+		for (ULearningAgentsNeuralNetwork* NetworkAsset : NetworkAssets)
 		{
-			UE_LOG(LogLearning, Error, TEXT("%s: Error waiting for encoder from trainer: %s. Check log for additional errors."), *GetName(), UE::Learning::Trainer::GetResponseString(Response));
-			bHasTrainingFailed = true;
-			EndTraining();
-			return;
+			NetworkAsset->ForceMarkDirty();	
 		}
 
-		// Get Updated Decoder
-		Response = Trainer->ReceiveNetwork(DecoderNetworkId, *Policy->GetDecoderNetworkAsset()->NeuralNetworkData);
-		Policy->GetDecoderNetworkAsset()->ForceMarkDirty();
-
-		if (Response != UE::Learning::ETrainerResponse::Success)
+		for (int32 i = 0; i < NetworkIds.Num(); i++)
 		{
-			UE_LOG(LogLearning, Error, TEXT("%s: Error waiting for decoder from trainer: %s. Check log for additional errors."), *GetName(), UE::Learning::Trainer::GetResponseString(Response));
-			bHasTrainingFailed = true;
-			EndTraining();
-			return;
+			if (NetworkResponses[i] == UE::Learning::ETrainerResponse::Completed)
+			{
+				UE_LOG(LogLearning, Display, TEXT("%s: Trainer completed training."), *GetName());
+				DoneTraining();
+				return;
+			}
+			else if (NetworkResponses[i] != UE::Learning::ETrainerResponse::Success)
+			{
+				UE_LOG(LogLearning, Error, TEXT("Error receiving network (id=%d) from trainer: %s. Check log for errors."), NetworkIds[i], UE::Learning::Trainer::GetResponseString(NetworkResponses[i]));
+				bHasTrainingFailed = true;
+				EndTraining();
+				return;
+			}
 		}
 
 		if (bResetAgentsOnUpdate)
diff --git a/Engine/Plugins/Experimental/LearningAgents/Source/LearningTraining/Private/LearningExternalTrainer.cpp b/Engine/Plugins/Experimental/LearningAgents/Source/LearningTraining/Private/LearningExternalTrainer.cpp
--- a/Engine/Plugins/Experimental/LearningAgents/Source/LearningTraining/Private/LearningExternalTrainer.cpp
+++ b/Engine/Plugins/Experimental/LearningAgents/Source/LearningTraining/Private/LearningExternalTrainer.cpp
@@ -6,7 +6,6 @@
 #include "LearningNeuralNetwork.h"
 #include "LearningSharedMemoryTraining.h"
 #include "LearningSocketTraining.h"
-
 #include "GenericPlatform/GenericPlatformFile.h"
 #include "HAL/PlatformFileManager.h"
 
@@ -15,6 +14,8 @@
 #include "Misc/Guid.h"
 #include "Misc/Paths.h"
 
+#include "Algo/AllOf.h"
+
 #include "Dom/JsonObject.h"
 #include "Serialization/JsonSerializer.h"
 #include "Serialization/JsonWriter.h"
@@ -274,7 +275,7 @@ namespace UE::Learning
 
 		// Add intermediate path as a temp directory for the tensorboard logging on python side
 		TrainerConfigObject->SetStringField(TEXT("TempDirectory"), *FileManager.ConvertToAbsolutePathForExternalAppForRead(*IntermediatePath));
-
+		
 		TSharedRef<FJsonObject> SharedMemoryConfigObject = MakeShared<FJsonObject>();
 
 		TArray<TSharedPtr<FJsonValue>> NetworkGuidsArray;
@@ -410,6 +411,38 @@ namespace UE::Learning
 			LogSettings);
 	}
 
+	TArray<ETrainerResponse> FSharedMemoryTrainer::ReceiveNetworks(
+		const TArray<int32>& NetworkIds,
+		TArray<TObjectPtr<ULearningNeuralNetworkData>> Networks,
+		TArray<FRWLock*> NetworkLocks,
+		const ELogSetting LogSettings)
+	{
+		checkf(Controls.Region, TEXT("ReceiveNetworks: Controls Shared Memory Region is nullptr"));		
+		check(NetworkIds.Num() == Networks.Num());
+
+		TArray<ETrainerResponse> Responses;
+		Responses.Init(ETrainerResponse::Unexpected, NetworkIds.Num());
+
+		for (int32 i = 0; i < NetworkIds.Num(); ++i)
+		{
+			if (!ensureMsgf(NeuralNetworkSharedMemoryArrayViews.Num() >= NetworkIds[i], TEXT("Network %d has not been added. Call AddNetwork prior to ReceiveNetwork."), NetworkIds[i]))
+			{
+				return Responses;
+			}
+
+			Responses[i] = SharedMemoryTraining::RecvNetwork(
+				Controls.View,
+				NetworkIds[i],
+				*Networks[i],
+				TrainingProcess.IsValid() ? TrainingProcess->GetTrainingSubprocess() : nullptr,
+				NeuralNetworkSharedMemoryArrayViews[NetworkIds[i]].View,
+				Timeout,
+				NetworkLocks.IsValidIndex(i) ? NetworkLocks[i] : nullptr,
+				LogSettings);
+		}
+		return Responses;
+	}
+
 	ETrainerResponse FSharedMemoryTrainer::SendNetwork(
 		const int32 NetworkId,
 		const ULearningNeuralNetworkData& Network,
@@ -812,7 +845,124 @@ namespace UE::Learning
 		NetworkBuffers[NetworkId].SetNumUninitialized({Network.GetSnapshotByteNum()});
 		return NetworkId;
 	}
+	
+	TArray<ETrainerResponse> FSocketTrainer::ReceiveNetworks(
+		const TArray<int32>& NetworkIds,
+		TArray<TObjectPtr<ULearningNeuralNetworkData>> Networks,
+		TArray<FRWLock*> NetworkLocks,
+		const ELogSetting LogSettings)
+	{
+		check(NetworkIds.Num() == Networks.Num());
+
+		TArray<ETrainerResponse> Responses;
+		Responses.Init(ETrainerResponse::Unexpected, NetworkIds.Num());
+
+		if (!Socket)
+		{
+			UE_LOG(LogLearning, Error, TEXT("Training socket is nullptr"));
+			return Responses;
+		}
+
+		for (int32 i = 0; i < NetworkIds.Num(); i++)
+		{
+			if (!ensureMsgf(NetworkBuffers.Num() >= NetworkIds[i], TEXT("Network %d has not been added. Call AddNetwork prior to ReceiveNetwork."), NetworkIds[i]))
+			{
+				Responses[i] = ETrainerResponse::Unexpected;
+				continue;
+			}
+			
+			Responses[i] = SocketTraining::RecvNetwork(
+				*Socket,
+				NetworkIds[i],
+				NetworksVersion,
+				*Networks[i],
+				TrainingProcess.IsValid() ? TrainingProcess->GetTrainingSubprocess() : nullptr,
+				NetworkBuffers[NetworkIds[i]],
+				Timeout,
+				NetworkLocks.IsValidIndex(i) ? NetworkLocks[i] : nullptr,
+				LogSettings);
+			}
+
+		if (!Responses.Contains(ETrainerResponse::Completed))
+		{
+			bool bAllSuccess = Algo::AllOf(Responses, [](const ETrainerResponse& Response) { 
+				return Response == ETrainerResponse::Success; 
+			});
+			if (bAllSuccess)
+			{
+				return ReceiveQueuedNetworks(NetworkIds, Networks, NetworkLocks, LogSettings);
+			}
+		}
+		return Responses;
+	}
+
+	TArray<ETrainerResponse> FSocketTrainer::ReceiveQueuedNetworks(
+		const TArray<int32>& NetworkIds,
+		TArray<TObjectPtr<ULearningNeuralNetworkData>> Networks,
+		TArray<FRWLock*> NetworkLocks,
+		const ELogSetting LogSettings)
+	{
+		check(NetworkIds.Num() == Networks.Num());
+
+		TArray<ETrainerResponse> Responses;
+		Responses.Init(ETrainerResponse::Unexpected, NetworkIds.Num());
+
+		if (!Socket)
+		{
+			UE_LOG(LogLearning, Error, TEXT("Training socket is nullptr"));
+			return Responses;
+		}
+		
+		UE_LOG(LogLearning, Display, TEXT("Checking if there is any newer networks... Current version: %d"), NetworksVersion);
 
+		for (int32 i = 0; i < NetworkIds.Num(); i++)
+		{
+			if (!ensureMsgf(NetworkBuffers.Num() >= NetworkIds[i], TEXT("Network %d has not been added. Call AddNetwork prior to ReceiveNetwork."), NetworkIds[i]))
+			{
+				Responses[i] = ETrainerResponse::Unexpected;
+				continue;
+			}
+						
+			// Signal timeout is 0.0f to avoid blocking when checking queued networks
+			const float SignalTimeout = 0.0f;
+			ETrainerResponse Response = SocketTraining::RecvNetwork(
+				*Socket,
+				NetworkIds[i],
+				NetworksVersion,
+				*Networks[i],
+				TrainingProcess.IsValid() ? TrainingProcess->GetTrainingSubprocess() : nullptr,
+				NetworkBuffers[NetworkIds[i]],
+				Timeout,
+				NetworkLocks.IsValidIndex(i) ? NetworkLocks[i] : nullptr,
+				LogSettings,
+				SignalTimeout);
+			
+			// If it's the first network and we received a NetworkSignalTimeout, networks queue is consumed and can end recursion
+			if (i == 0 && Response == ETrainerResponse::NetworkSignalTimeout)
+			{
+				UE_LOG(LogLearning, Display, TEXT("Using the most updated networks version: %d"), NetworksVersion);
+				Responses.Init(ETrainerResponse::Success, NetworkIds.Num());
+				return Responses;
+			}
+
+			Responses[i] = Response;
+		}
+
+		if (!Responses.Contains(ETrainerResponse::Completed))
+		{
+			bool bAllSuccess = Algo::AllOf(Responses, [](const ETrainerResponse& Response) { 
+				return Response == ETrainerResponse::Success; 
+			});
+
+			// Recursively check if there's newer networks in the socket
+			if (bAllSuccess)
+			{
+				return ReceiveQueuedNetworks(NetworkIds, Networks, NetworkLocks, LogSettings);
+			}
+		}
+		return Responses;
+	}
+	
 	ETrainerResponse FSocketTrainer::ReceiveNetwork(
 		const int32 NetworkId,
 		ULearningNeuralNetworkData& OutNetwork,
@@ -833,6 +983,7 @@ namespace UE::Learning
 		return SocketTraining::RecvNetwork(
 			*Socket,
 			NetworkId,
+			NetworksVersion,
 			OutNetwork,
 			TrainingProcess.IsValid() ? TrainingProcess->GetTrainingSubprocess() : nullptr,
 			NetworkBuffers[NetworkId],
@@ -890,6 +1041,7 @@ namespace UE::Learning
 
 		return SocketTraining::SendExperience(
 			*Socket,
+			NetworksVersion,
 			ReplayBufferId,
 			ReplayBuffer,
 			TrainingProcess.IsValid() ? TrainingProcess->GetTrainingSubprocess() : nullptr,
diff --git a/Engine/Plugins/Experimental/LearningAgents/Source/LearningTraining/Private/LearningSocketTraining.cpp b/Engine/Plugins/Experimental/LearningAgents/Source/LearningTraining/Private/LearningSocketTraining.cpp
--- a/Engine/Plugins/Experimental/LearningAgents/Source/LearningTraining/Private/LearningSocketTraining.cpp
+++ b/Engine/Plugins/Experimental/LearningAgents/Source/LearningTraining/Private/LearningSocketTraining.cpp
@@ -73,26 +73,24 @@ namespace UE::Learning::SocketTraining
 	ETrainerResponse RecvNetwork(
 		FSocket& Socket,
 		const int32 NetworkId,
+		int32& OutNetworkVersion,
 		ULearningNeuralNetworkData& OutNetwork,
 		FSubprocess* Process,
 		TLearningArrayView<1, uint8> OutNetworkBuffer,
 		const float Timeout,
 		FRWLock* NetworkLock,
-		const ELogSetting LogSettings)
+		const ELogSetting LogSettings,
+		const float SignalTimeout)
 	{
 		check(OutNetworkBuffer.Num() == OutNetwork.GetSnapshotByteNum());
 
-		if (LogSettings != ELogSetting::Silent)
-		{
-			UE_LOG(LogLearning, Display, TEXT("Pulling Network..."));
-		}
-
 		uint8 Signal = (uint8)ESignal::Invalid;
 		ETrainerResponse Response = ETrainerResponse::Unexpected;
 
 		while (true)
 		{
-			Response = RecvWithTimeout(Socket, Process, &Signal, 1, Timeout);
+			Response = RecvWithTimeout(Socket, Process, &Signal, 1, SignalTimeout);
+			if (Response == ETrainerResponse::Timeout) { return ETrainerResponse::NetworkSignalTimeout; }
 			if (Response != ETrainerResponse::Success) { return Response; }
 
 			if (Signal == (uint8)ESignal::RecvComplete)
@@ -118,6 +116,9 @@ namespace UE::Learning::SocketTraining
 		if (Response != ETrainerResponse::Success) { return Response; }
 		checkf(Id == NetworkId, TEXT("Received unexpected network id!"))
 
+		Response = RecvWithTimeout(Socket, Process, (uint8*)&OutNetworkVersion, sizeof(int32), Timeout);
+		if (Response != ETrainerResponse::Success) { return Response; }
+
 		Response = RecvWithTimeout(Socket, Process, OutNetworkBuffer.GetData(), OutNetworkBuffer.Num(), Timeout);
 		if (Response != ETrainerResponse::Success) { return Response; }
 
@@ -287,6 +288,7 @@ namespace UE::Learning::SocketTraining
 
 	ETrainerResponse SendExperience(
 		FSocket& Socket,
+		const int32 NetworksVersion,
 		const int32 ReplayBufferId,
 		const FReplayBuffer& ReplayBuffer,
 		FSubprocess* Process,
@@ -302,6 +304,9 @@ namespace UE::Learning::SocketTraining
 		ETrainerResponse Response = SendWithTimeout(Socket, Process, &Signal, 1, Timeout);
 		if (Response != ETrainerResponse::Success) { return Response; }
 
+		Response = SendWithTimeout(Socket, Process, (const uint8*)&NetworksVersion, sizeof(int32), Timeout);
+		if (Response != ETrainerResponse::Success) { return Response; }
+
 		Response = SendWithTimeout(Socket, Process, (const uint8*)&ReplayBufferId, sizeof(int32), Timeout);
 		if (Response != ETrainerResponse::Success) { return Response; }
 
diff --git a/Engine/Plugins/Experimental/LearningAgents/Source/LearningTraining/Private/LearningSocketTraining.h b/Engine/Plugins/Experimental/LearningAgents/Source/LearningTraining/Private/LearningSocketTraining.h
--- a/Engine/Plugins/Experimental/LearningAgents/Source/LearningTraining/Private/LearningSocketTraining.h
+++ b/Engine/Plugins/Experimental/LearningAgents/Source/LearningTraining/Private/LearningSocketTraining.h
@@ -42,12 +42,14 @@ namespace UE::Learning
 		LEARNINGTRAINING_API ETrainerResponse RecvNetwork(
 			FSocket& Socket,
 			const int32 NetworkId,
+			int32& OutNetworkVersion,
 			ULearningNeuralNetworkData& OutNetwork,
 			FSubprocess* Process,
 			TLearningArrayView<1, uint8> OutNetworkBuffer,
 			const float Timeout = Trainer::DefaultTimeout,
 			FRWLock* NetworkLock = nullptr,
-			const ELogSetting LogSettings = Trainer::DefaultLogSettings);
+			const ELogSetting LogSettings = Trainer::DefaultLogSettings,
+			const float SignalTimeout = Trainer::DefaultTimeout);
 
 		LEARNINGTRAINING_API ETrainerResponse SendWithTimeout(
 			FSocket& Socket,
@@ -84,6 +86,7 @@ namespace UE::Learning
 
 		LEARNINGTRAINING_API ETrainerResponse SendExperience(
 			FSocket& Socket,
+			const int32 NetworksVersion,
 			const int32 ReplayBufferId,
 			const FReplayBuffer& ReplayBuffer,
 			FSubprocess* Process,
diff --git a/Engine/Plugins/Experimental/LearningAgents/Source/LearningTraining/Public/LearningExternalTrainer.h b/Engine/Plugins/Experimental/LearningAgents/Source/LearningTraining/Public/LearningExternalTrainer.h
--- a/Engine/Plugins/Experimental/LearningAgents/Source/LearningTraining/Public/LearningExternalTrainer.h
+++ b/Engine/Plugins/Experimental/LearningAgents/Source/LearningTraining/Public/LearningExternalTrainer.h
@@ -83,6 +83,22 @@ namespace UE::Learning
 			FRWLock* NetworkLock = nullptr,
 			const ELogSetting LogSettings = Trainer::DefaultLogSettings) = 0;
 
+		/**
+		* Wait for the trainer to push an array of updated networks.
+		*
+		* @param NetworkIds		Unique network ids
+		* @param OutNetworks		Networks to update
+		* @param Timeout		Timeout to wait in seconds
+		* @param NetworkLocks	Locks to use when updating networks
+		* @param LogSettings	The log verbosity level
+		* @returns				Trainer response
+		*/
+		virtual TArray<ETrainerResponse> ReceiveNetworks(
+			const TArray<int32>& NetworkIds,
+			TArray<TObjectPtr<ULearningNeuralNetworkData>> Networks,
+			TArray<FRWLock*> NetworkLocks = TArray<FRWLock*>(),
+			const ELogSetting LogSettings = Trainer::DefaultLogSettings) = 0;
+
 		/**
 		* Wait for the trainer to be ready and push the current policy network.
 		*
@@ -267,6 +283,12 @@ namespace UE::Learning
 			FRWLock* NetworkLock = nullptr,
 			const ELogSetting LogSettings = Trainer::DefaultLogSettings) override final;
 
+		virtual TArray<ETrainerResponse> ReceiveNetworks(
+			const TArray<int32>& NetworkIds,
+			TArray<TObjectPtr<ULearningNeuralNetworkData>> Networks,
+			TArray<FRWLock*> NetworkLocks = TArray<FRWLock*>(),
+			const ELogSetting LogSettings = Trainer::DefaultLogSettings) override final;
+
 		virtual ETrainerResponse SendNetwork(
 			const int32 NetworkId,
 			const ULearningNeuralNetworkData& Network,
@@ -411,6 +433,18 @@ namespace UE::Learning
 			FRWLock* NetworkLock = nullptr,
 			const ELogSetting LogSettings = Trainer::DefaultLogSettings) override final;
 
+		virtual TArray<ETrainerResponse> ReceiveNetworks(
+			const TArray<int32>& NetworkIds,
+			TArray<TObjectPtr<ULearningNeuralNetworkData>> Networks,
+			TArray<FRWLock*> NetworkLocks = TArray<FRWLock*>(),
+			const ELogSetting LogSettings = Trainer::DefaultLogSettings) override final;
+
+		virtual TArray<ETrainerResponse> ReceiveQueuedNetworks(
+			const TArray<int32>& NetworkIds,
+			TArray<TObjectPtr<ULearningNeuralNetworkData>> Networks,
+			TArray<FRWLock*> NetworkLocks = TArray<FRWLock*>(),
+			const ELogSetting LogSettings = Trainer::DefaultLogSettings);
+
 		virtual ETrainerResponse SendNetwork(
 			const int32 NetworkId,
 			const ULearningNeuralNetworkData& Network,
@@ -436,5 +470,7 @@ namespace UE::Learning
 		TSharedPtr<FSocketTrainerServerProcess> TrainingProcess = nullptr;
 
 		FSocket* Socket = nullptr;
+
+		int32 NetworksVersion = -1;
 	};
 }
diff --git a/Engine/Plugins/Experimental/LearningAgents/Source/LearningTraining/Public/LearningTrainer.h b/Engine/Plugins/Experimental/LearningAgents/Source/LearningTraining/Public/LearningTrainer.h
--- a/Engine/Plugins/Experimental/LearningAgents/Source/LearningTraining/Public/LearningTrainer.h
+++ b/Engine/Plugins/Experimental/LearningAgents/Source/LearningTraining/Public/LearningTrainer.h
@@ -52,6 +52,9 @@ namespace UE::Learning
 
 		// The communication timed-out
 		Timeout = 4,
+
+		// The communication timed-out for a network signal
+		NetworkSignalTimeout = 5,
 	};
 
 	/**
