LearningAgents: Updated MajorityVote filter width to integer parameter.

diff --git a/Engine/Plugins/Experimental/LearningAgents/Source/Learning/Private/Learning.ispc b/Engine/Plugins/Experimental/LearningAgents/Source/Learning/Private/Learning.ispc
--- a/Engine/Plugins/Experimental/LearningAgents/Source/Learning/Private/Learning.ispc
+++ b/Engine/Plugins/Experimental/LearningAgents/Source/Learning/Private/Learning.ispc
@@ -727,17 +727,15 @@ export void LearningFilterMajorityVote(
 	uniform float Out[],
 	uniform const float In[],
 	uniform const int ValueNum,
-	uniform const float FilterWidthInFrames)
+	uniform const int FilterWidthInFrames)
 {
-	uniform const int FilterRange = (uniform int)round(FilterWidthInFrames);
-
 	foreach(ValueIdx = 0 ... ValueNum)
 	{
 		int Total = 0;
-		const int RangeMin = max(ValueIdx - FilterRange, 0);
-		const int RangeMax = min(ValueIdx + FilterRange, ValueNum - 1);
+		const int RangeMin = max(ValueIdx - FilterWidthInFrames, 0);
+		const int RangeMax = min(ValueIdx + FilterWidthInFrames, ValueNum - 1);
 
-		for (int Offset = -FilterRange; Offset <= +FilterRange; Offset++)
+		for (int Offset = -FilterWidthInFrames; Offset <= +FilterWidthInFrames; Offset++)
 		{
 			if (ValueIdx + Offset >= 0 && ValueIdx + Offset < ValueNum)
 			{
diff --git a/Engine/Plugins/Experimental/LearningAgents/Source/Learning/Private/LearningFrameAttribute.cpp b/Engine/Plugins/Experimental/LearningAgents/Source/Learning/Private/LearningFrameAttribute.cpp
--- a/Engine/Plugins/Experimental/LearningAgents/Source/Learning/Private/LearningFrameAttribute.cpp
+++ b/Engine/Plugins/Experimental/LearningAgents/Source/Learning/Private/LearningFrameAttribute.cpp
@@ -560,7 +560,7 @@ namespace UE::Learning
 #endif
 		}
 
-		static inline void FilterMajorityVote(TLearningArrayView<1, float> Out, const TLearningArrayView<1, const float> In, const float FilterWidthInFrames)
+		static inline void FilterMajorityVote(TLearningArrayView<1, float> Out, const TLearningArrayView<1, const float> In, const int32 FilterWidthInFrames)
 		{
 			check(Out.Num() == In.Num());
 			check(FilterWidthInFrames >= 0.0f);
@@ -570,17 +570,15 @@ namespace UE::Learning
 			ispc::LearningFilterMajorityVote(Out.GetData(), In.GetData(), In.Num(), FilterWidthInFrames);
 
 #else
-			const int32 FilterRange = FMath::RoundToInt(FilterWidthInFrames);
-
 			const int32 ValueNum = Out.Num();
 			for (int32 ValueIdx = 0; ValueIdx < ValueNum; ValueIdx++)
 			{
 				int32 Total = 0;
 
-				const int32 RangeMin = FMath::Max(ValueIdx - FilterRange, 0);
-				const int32 RangeMax = FMath::Min(ValueIdx + FilterRange, ValueNum - 1);
+				const int32 RangeMin = FMath::Max(ValueIdx - FilterWidthInFrames, 0);
+				const int32 RangeMax = FMath::Min(ValueIdx + FilterWidthInFrames, ValueNum - 1);
 
-				for (int32 Offset = -FilterRange; Offset <= +FilterRange; Offset++)
+				for (int32 Offset = -FilterWidthInFrames; Offset <= +FilterWidthInFrames; Offset++)
 				{
 					if (ValueIdx + Offset >= 0 && ValueIdx + Offset < ValueNum)
 					{
@@ -2488,7 +2486,7 @@ namespace UE::Learning
 				});
 		}
 
-		void FilterMajorityVote(FFrameAttribute& Out, const FFrameAttribute& In, const float FilterWidthFrames)
+		void FilterMajorityVote(FFrameAttribute& Out, const FFrameAttribute& In, const int32 FilterWidthFrames)
 		{
 			UnaryOp(Out, In.GetChannelNum(), In, [FilterWidthFrames](
 				FFrameAttribute& Out,
diff --git a/Engine/Plugins/Experimental/LearningAgents/Source/Learning/Public/LearningFrameAttribute.h b/Engine/Plugins/Experimental/LearningAgents/Source/Learning/Public/LearningFrameAttribute.h
--- a/Engine/Plugins/Experimental/LearningAgents/Source/Learning/Public/LearningFrameAttribute.h
+++ b/Engine/Plugins/Experimental/LearningAgents/Source/Learning/Public/LearningFrameAttribute.h
@@ -328,7 +328,7 @@ namespace UE::Learning
 		LEARNING_API void FilterGaussian(FFrameAttribute& Out, const FFrameAttribute& In, const float StdInFrames);
 
 		/** Applies a Majority Vote filter to the ranges of values in the frame attribute. */
-		LEARNING_API void FilterMajorityVote(FFrameAttribute& Out, const FFrameAttribute& In, const float FilterWidthFrames);
+		LEARNING_API void FilterMajorityVote(FFrameAttribute& Out, const FFrameAttribute& In, const int32 FilterWidthFrames);
 
 		/** Computes the mean and std across all frames. OutMean and OutStd should be size of the number of channels */
 		LEARNING_API void MeanStd(TLearningArrayView<1, float> OutMean, TLearningArrayView<1, float> OutStd, const FFrameAttribute& In);
