Learning Agents: Exit on Training error if not in editor

#rnx

diff --git a/Engine/Plugins/Experimental/LearningAgents/Source/LearningAgentsTraining/Private/LearningAgentsImitationTrainer.cpp b/Engine/Plugins/Experimental/LearningAgents/Source/LearningAgentsTraining/Private/LearningAgentsImitationTrainer.cpp
--- a/Engine/Plugins/Experimental/LearningAgents/Source/LearningAgentsTraining/Private/LearningAgentsImitationTrainer.cpp
+++ b/Engine/Plugins/Experimental/LearningAgents/Source/LearningAgentsTraining/Private/LearningAgentsImitationTrainer.cpp
@@ -13,6 +13,7 @@
 #include "LearningAgentsPolicy.h"
 
 #include "Dom/JsonObject.h"
+#include "HAL/PlatformMisc.h"
 #include "Misc/Paths.h"
 
 TSharedRef<FJsonObject> FLearningAgentsImitationTrainerTrainingSettings::AsJsonConfig() const
@@ -539,6 +540,11 @@ void ULearningAgentsImitationTrainer::RunTraining(
 	if (bHasTrainingFailed)
 	{
 		UE_LOG(LogLearning, Error, TEXT("%s: Training has failed. Check log for errors."), *GetName());
+
+#if !WITH_EDITOR
+		FGenericPlatformMisc::RequestExitWithStatus(false, 99);
+#endif
+
 		return;
 	}
 
diff --git a/Engine/Plugins/Experimental/LearningAgents/Source/LearningAgentsTraining/Private/LearningAgentsPPOTrainer.cpp b/Engine/Plugins/Experimental/LearningAgents/Source/LearningAgentsTraining/Private/LearningAgentsPPOTrainer.cpp
--- a/Engine/Plugins/Experimental/LearningAgents/Source/LearningAgentsTraining/Private/LearningAgentsPPOTrainer.cpp
+++ b/Engine/Plugins/Experimental/LearningAgents/Source/LearningAgentsTraining/Private/LearningAgentsPPOTrainer.cpp
@@ -16,6 +16,7 @@
 #include "LearningExternalTrainer.h"
 
 #include "Dom/JsonObject.h"
+#include "HAL/PlatformMisc.h"
 
 ULearningAgentsPPOTrainer::ULearningAgentsPPOTrainer() : Super(FObjectInitializer::Get()) {}
 ULearningAgentsPPOTrainer::ULearningAgentsPPOTrainer(FVTableHelper& Helper) : Super(Helper) {}
@@ -680,6 +681,11 @@ void ULearningAgentsPPOTrainer::RunTraining(
 	if (bHasTrainingFailed)
 	{
 		UE_LOG(LogLearning, Error, TEXT("%s: Training has failed. Check log for errors."), *GetName());
+
+#if !WITH_EDITOR
+		FGenericPlatformMisc::RequestExitWithStatus(false, 99);
+#endif
+
 		return;
 	}
 
