Learning Agents: Fix bug with angle observation not using relative angle in calculation

diff --git a/Engine/Plugins/Experimental/LearningAgents/Source/LearningAgents/Private/LearningAgentsObservations.cpp b/Engine/Plugins/Experimental/LearningAgents/Source/LearningAgents/Private/LearningAgentsObservations.cpp
--- a/Engine/Plugins/Experimental/LearningAgents/Source/LearningAgents/Private/LearningAgentsObservations.cpp
+++ b/Engine/Plugins/Experimental/LearningAgents/Source/LearningAgents/Private/LearningAgentsObservations.cpp
@@ -2097,8 +2097,8 @@ FLearningAgentsObservationObjectElement ULearningAgentsObservations::MakeAngleOb
 	const FLinearColor VisualLoggerColor)
 {
 	const float LocalAngle = FMath::FindDeltaAngleDegrees(RelativeAngle, Angle);
-	const float EncodedX = FMath::Sin(FMath::DegreesToRadians(Angle));
-	const float EncodedY = FMath::Cos(FMath::DegreesToRadians(Angle));
+	const float EncodedX = FMath::Sin(FMath::DegreesToRadians(LocalAngle));
+	const float EncodedY = FMath::Cos(FMath::DegreesToRadians(LocalAngle));
 
 #if UE_LEARNING_AGENTS_ENABLE_VISUAL_LOG
 	if (bVisualLoggerEnabled && VisualLoggerListener)
