[LearningAgents] Fixes an issue with gyms spacing

- fixes an issue where when multi templates are used and ordered a certain way the spacing is off
- add GymsCountOverride cmdline arg

diff --git a/Engine/Plugins/Experimental/LearningAgents/Source/LearningAgentsTraining/Private/LearningAgentsGymsManager.cpp b/Engine/Plugins/Experimental/LearningAgents/Source/LearningAgentsTraining/Private/LearningAgentsGymsManager.cpp
--- a/Engine/Plugins/Experimental/LearningAgents/Source/LearningAgentsTraining/Private/LearningAgentsGymsManager.cpp
+++ b/Engine/Plugins/Experimental/LearningAgents/Source/LearningAgentsTraining/Private/LearningAgentsGymsManager.cpp
@@ -9,11 +9,27 @@
 #include "AI/NavigationSystemBase.h"
 #include "Misc/CommandLine.h"
 #include "Misc/DateTime.h"
+#include "Math/UnrealMathUtility.h"
 
 ALearningAgentsGymsManager::ALearningAgentsGymsManager()
 {
 	PrimaryActorTick.bCanEverTick = false;
 	PrimaryActorTick.bStartWithTickEnabled = false;
+
+	for (FSpawnGymInfo& GymTemplate : GymTemplates)
+	{
+		FString SpawnCountInput;
+		if (FParse::Value(FCommandLine::Get(), TEXT("GymsCountOverride"), SpawnCountInput))
+		{
+			if (SpawnCountInput.StartsWith(TEXT("=")))
+			{
+				SpawnCountInput = SpawnCountInput.RightChop(1);
+
+			}
+			GymTemplate.SpawnCount = FCString::Atoi(*SpawnCountInput);
+			UE_LOG(LogLearning, Log, TEXT("Commandline flag -GymsCountOverride is active! Value: %d Gyms."), GymTemplate.GetCount());
+		}
+	}
 }
 
 void ALearningAgentsGymsManager::SpawnGyms()
@@ -24,11 +40,13 @@ void ALearningAgentsGymsManager::SpawnGyms()
 
 	for (FSpawnGymInfo& GymTemplate : GymTemplates)
 	{
+		UE_LOG(LogLearning, Log, TEXT("Generating %d Gyms for the template %s."), GymTemplate.GetCount(), *GymTemplate.GymClass->GetName());
+
 		FActorSpawnParameters SpawnParameters;
 		SpawnParameters.Owner = this;
 
 		int32 CurrentSpawnCount = 0;
-		int32 GridSize = FMath::CeilToInt32(FMath::Sqrt(static_cast<float>(GymTemplate.SpawnCount)));
+		int32 GridSize = FMath::CeilToInt32(FMath::Sqrt(static_cast<float>(GymTemplate.GetCount())));
 		bool bFirstGym = true;
 
 		FVector MinBounds = FVector::ZeroVector;
@@ -39,7 +57,7 @@ void ALearningAgentsGymsManager::SpawnGyms()
 		{
 			for (int32 Col = 0; Col < GridSize; ++Col)
 			{
-				if (CurrentSpawnCount == GymTemplate.SpawnCount)
+				if (CurrentSpawnCount >= GymTemplate.GetCount())
 				{
 					break;
 				}
@@ -54,7 +72,10 @@ void ALearningAgentsGymsManager::SpawnGyms()
 				}
 				else
 				{
-					FVector SpawnLocation = FVector(Row * (GymsSpacing + (MaxBounds.X - MinBounds.X)) + TemplateSpawnLocation.X, Col * (GymsSpacing + (MaxBounds.Y - MinBounds.Y) + TemplateSpawnLocation.Y), TemplateSpawnLocation.Z);
+					FVector SpawnLocation = FVector(
+						Row * (GymsSpacing + (MaxBounds.X - MinBounds.X)) + TemplateSpawnLocation.X, 
+						Col * (GymsSpacing + (MaxBounds.Y - MinBounds.Y)) + TemplateSpawnLocation.Y, 
+						TemplateSpawnLocation.Z);
 					Gym = GetWorld()->SpawnActor<ALearningAgentsGymBase>(GymTemplate.GymClass, SpawnLocation, FRotator::ZeroRotator, SpawnParameters);
 				}
 
@@ -66,7 +87,10 @@ void ALearningAgentsGymsManager::SpawnGyms()
 				CurrentSpawnCount++;
 			}
 		}
-		TemplateSpawnLocation.X += (GymsSpacing + (MaxBounds.X - MinBounds.X)) * (GridSize - 1);
+		if (GridSize > 0)
+		{
+			TemplateSpawnLocation.X += (GymsSpacing + (MaxBounds.X - MinBounds.X)) * FMath::CeilToInt(static_cast<float>(GymTemplate.GetCount()) / GridSize);
+		}
 	}
 	
 	if (TObjectPtr<UNavigationSystemV1> NavigationSystem = FNavigationSystem::GetCurrent<UNavigationSystemV1>(GetWorld()))
@@ -84,7 +108,7 @@ int32 ALearningAgentsGymsManager::GetGymsCount() const
 	int32 Count = 0;
 	for (const FSpawnGymInfo& GymTemplate : GymTemplates)
 	{
-		Count += GymTemplate.SpawnCount;
+		Count += GymTemplate.GetCount();
 	}
 	return Count;
 }
diff --git a/Engine/Plugins/Experimental/LearningAgents/Source/LearningAgentsTraining/Public/LearningAgentsGymsManager.h b/Engine/Plugins/Experimental/LearningAgents/Source/LearningAgentsTraining/Public/LearningAgentsGymsManager.h
--- a/Engine/Plugins/Experimental/LearningAgents/Source/LearningAgentsTraining/Public/LearningAgentsGymsManager.h
+++ b/Engine/Plugins/Experimental/LearningAgents/Source/LearningAgentsTraining/Public/LearningAgentsGymsManager.h
@@ -22,8 +22,13 @@ struct FSpawnGymInfo
 	TSubclassOf<ALearningAgentsGymBase> GymClass;
 
 	/** The number of gym instances to spawn. */
-	UPROPERTY(EditAnywhere, BlueprintReadWrite, Category = "LearningAgents")
+	UPROPERTY(EditAnywhere, BlueprintReadWrite, Category = "LearningAgents", meta = (ClampMin = "0", UIMin = "0"))
 	int32 SpawnCount = 0;
+
+	int32 GetCount() const
+	{
+		return (SpawnCount < 0) ? 0 : SpawnCount;
+	}
 };
 
 
