[LearningAgents] Fixes an issue where cmdline gyms count isnt passed to total GetGymsCount

diff --git a/Engine/Plugins/Experimental/LearningAgents/Source/LearningAgentsTraining/Private/LearningAgentsGymsManager.cpp b/Engine/Plugins/Experimental/LearningAgents/Source/LearningAgentsTraining/Private/LearningAgentsGymsManager.cpp
--- a/Engine/Plugins/Experimental/LearningAgents/Source/LearningAgentsTraining/Private/LearningAgentsGymsManager.cpp
+++ b/Engine/Plugins/Experimental/LearningAgents/Source/LearningAgentsTraining/Private/LearningAgentsGymsManager.cpp
@@ -15,6 +15,13 @@ ALearningAgentsGymsManager::ALearningAgentsGymsManager()
 {
 	PrimaryActorTick.bCanEverTick = false;
 	PrimaryActorTick.bStartWithTickEnabled = false;
+}
+
+void ALearningAgentsGymsManager::SpawnGyms()
+{
+	InitializeRandomStream();
+
+	FVector TemplateSpawnLocation = GetActorLocation();
 
 	for (FSpawnGymInfo& GymTemplate : GymTemplates)
 	{
@@ -24,23 +31,18 @@ ALearningAgentsGymsManager::ALearningAgentsGymsManager()
 			if (SpawnCountInput.StartsWith(TEXT("=")))
 			{
 				SpawnCountInput = SpawnCountInput.RightChop(1);
-
 			}
-			GymTemplate.SpawnCount = FCString::Atoi(*SpawnCountInput);
-			UE_LOG(LogLearning, Log, TEXT("Commandline flag -GymsCountOverride is active! Value: %d Gyms."), GymTemplate.GetCount());
+			int32 SpawnCountOverride = FCString::Atoi(*SpawnCountInput);
+			if (SpawnCountOverride > 0)
+			{
+				GymTemplate.SpawnCount = SpawnCountOverride;
+				UE_LOG(LogLearning, Log, TEXT("Commandline flag -GymsCountOverride is active! Generating %d Gyms for the template %s."), GymTemplate.GetCount(), *GymTemplate.GymClass->GetName());
+			}
+			else
+			{
+				UE_LOG(LogLearning, Log, TEXT("Commandline flag -GymsCountOverride was passed an invalid input. Generating %d Gyms for the template %s."), GymTemplate.GetCount(), *GymTemplate.GymClass->GetName());
+			}
 		}
-	}
-}
-
-void ALearningAgentsGymsManager::SpawnGyms()
-{
-	InitializeRandomStream();
-
-	FVector TemplateSpawnLocation = GetActorLocation();
-
-	for (FSpawnGymInfo& GymTemplate : GymTemplates)
-	{
-		UE_LOG(LogLearning, Log, TEXT("Generating %d Gyms for the template %s."), GymTemplate.GetCount(), *GymTemplate.GymClass->GetName());
 
 		FActorSpawnParameters SpawnParameters;
 		SpawnParameters.Owner = this;
@@ -105,6 +107,20 @@ void ALearningAgentsGymsManager::SpawnGyms()
 
 int32 ALearningAgentsGymsManager::GetGymsCount() const
 {
+	FString InputString;
+	if (FParse::Value(FCommandLine::Get(), TEXT("GymsCountOverride"), InputString))
+	{
+		if (InputString.StartsWith(TEXT("=")))
+		{
+			InputString = InputString.RightChop(1);
+		}
+		int32 SpawnCountOverride = FCString::Atoi(*InputString);
+		if (SpawnCountOverride > 0)
+		{
+			return SpawnCountOverride * GymTemplates.Num();
+		}
+	}
+
 	int32 Count = 0;
 	for (const FSpawnGymInfo& GymTemplate : GymTemplates)
 	{
