LearningAgents: Added FindTotalRange function to FrameRangeSet

diff --git a/Engine/Plugins/Experimental/LearningAgents/Source/Learning/Private/LearningFrameRangeSet.cpp b/Engine/Plugins/Experimental/LearningAgents/Source/Learning/Private/LearningFrameRangeSet.cpp
--- a/Engine/Plugins/Experimental/LearningAgents/Source/Learning/Private/LearningFrameRangeSet.cpp
+++ b/Engine/Plugins/Experimental/LearningAgents/Source/Learning/Private/LearningFrameRangeSet.cpp
@@ -1472,13 +1472,34 @@ namespace UE::Learning
 		return false;
 	}
 
+	bool FFrameRangeSet::FindTotalRange(int32& OutEntryIdx, int32& OutRangeIdx, const int32 TotalRangeIdx) const
+	{
+		const int32 EntryNum = GetEntryNum();
+
+		for (int32 EntryIdx = 0; EntryIdx < EntryNum; EntryIdx++)
+		{
+			const int32 RangeOffset = EntryRangeOffsets[EntryIdx];
+			const int32 RangeNum = EntryRangeNums[EntryIdx];
+
+			if (TotalRangeIdx >= RangeOffset && TotalRangeIdx < RangeOffset + RangeNum)
+			{
+				OutEntryIdx = EntryIdx;
+				OutRangeIdx = TotalRangeIdx - RangeOffset;
+				return true;
+			}
+		}
+
+		OutEntryIdx = INDEX_NONE;
+		OutRangeIdx = INDEX_NONE;
+		return false;
+	}
+
 	bool FFrameRangeSet::FindOffset(int32& OutEntryIdx, int32& OutRangeIdx, int32& OutRangeFrame, const int32 Offset) const
 	{
 		const int32 EntryNum = GetEntryNum();
 
 		for (int32 EntryIdx = 0; EntryIdx < EntryNum; EntryIdx++)
 		{
-			const int32 Sequence = GetEntrySequence(EntryIdx);
 			const int32 RangeNum = GetEntryRangeNum(EntryIdx);
 
 			for (int32 RangeIdx = 0; RangeIdx < RangeNum; RangeIdx++)
diff --git a/Engine/Plugins/Experimental/LearningAgents/Source/Learning/Public/LearningFrameRangeSet.h b/Engine/Plugins/Experimental/LearningAgents/Source/Learning/Public/LearningFrameRangeSet.h
--- a/Engine/Plugins/Experimental/LearningAgents/Source/Learning/Public/LearningFrameRangeSet.h
+++ b/Engine/Plugins/Experimental/LearningAgents/Source/Learning/Public/LearningFrameRangeSet.h
@@ -121,6 +121,9 @@ namespace UE::Learning
 		/** Finds the entry index and range index associated with a given sequence and frame in that sequence */
 		bool Find(int32& OutEntryIdx, int32& OutRangeIdx, int32& OutRangeFrame, const int32 Sequence, const int32 Frame) const;
 
+		/** Finds the entry index and range index associated with a range index into all ranges */
+		bool FindTotalRange(int32& OutEntryIdx, int32& OutRangeIdx, const int32 TotalRangeIdx) const;
+
 		/** Finds the entry index and range index associated with a given sequence and time in that sequence */
 		bool FindTime(int32& OutEntryIdx, int32& OutRangeIdx, float& OutRangeTime, const int32 Sequence, const float Time, const float FrameDeltaTime) const;
 
