[LearningAgents] Trivial changing LearningAgentsEntitiesManagerComp UObject ref to AActor

diff --git a/Engine/Plugins/Experimental/LearningAgents/Source/LearningAgentsTraining/Private/LearningAgentsEntitiesManagerComponent.cpp b/Engine/Plugins/Experimental/LearningAgents/Source/LearningAgentsTraining/Private/LearningAgentsEntitiesManagerComponent.cpp
--- a/Engine/Plugins/Experimental/LearningAgents/Source/LearningAgentsTraining/Private/LearningAgentsEntitiesManagerComponent.cpp
+++ b/Engine/Plugins/Experimental/LearningAgents/Source/LearningAgentsTraining/Private/LearningAgentsEntitiesManagerComponent.cpp
@@ -76,7 +76,7 @@ void ULearningAgentsEntitiesManagerComponent::ResetLearningComponent()
 	}
 }
 
-TScriptInterface<ILearningAgentsEntityInterface> ULearningAgentsEntitiesManagerComponent::SpawnEntitiesAtRandomLocations(TSubclassOf<UObject> EntityClass, float EntitySpawnZOffset, int32 SpawnCount)
+TScriptInterface<ILearningAgentsEntityInterface> ULearningAgentsEntitiesManagerComponent::SpawnEntitiesAtRandomLocations(TSubclassOf<AActor> EntityClass, float EntitySpawnZOffset, int32 SpawnCount)
 {
 	if (SpawnCount <= 0)
 	{
@@ -99,7 +99,7 @@ TScriptInterface<ILearningAgentsEntityInterface> ULearningAgentsEntitiesManagerC
 	}
 }
 
-TScriptInterface<ILearningAgentsEntityInterface> ULearningAgentsEntitiesManagerComponent::SpawnEntityAtProjectedLocation(TSubclassOf<UObject> EntityClass, float EntitySpawnZOffset, const FTransform& InTransform)
+TScriptInterface<ILearningAgentsEntityInterface> ULearningAgentsEntitiesManagerComponent::SpawnEntityAtProjectedLocation(TSubclassOf<AActor> EntityClass, float EntitySpawnZOffset, const FTransform& InTransform)
 {
 	FTransform Transform = InTransform;
 	ProjectTransform(Transform);
@@ -116,7 +116,7 @@ TScriptInterface<ILearningAgentsEntityInterface> ULearningAgentsEntitiesManagerC
 	}
 }
 
-TArray<TScriptInterface<ILearningAgentsEntityInterface>> ULearningAgentsEntitiesManagerComponent::SpawnEntities(TSubclassOf<UObject> EntityClass, float EntitySpawnZOffset, int32 SpawnCount, const FTransform& InTransform)
+TArray<TScriptInterface<ILearningAgentsEntityInterface>> ULearningAgentsEntitiesManagerComponent::SpawnEntities(TSubclassOf<AActor> EntityClass, float EntitySpawnZOffset, int32 SpawnCount, const FTransform& InTransform)
 {
 	TArray<TScriptInterface<ILearningAgentsEntityInterface>> SpawnedEntities;
 	SpawnedEntities.SetNumUninitialized(SpawnCount);
diff --git a/Engine/Plugins/Experimental/LearningAgents/Source/LearningAgentsTraining/Public/LearningAgentsEntitiesManagerComponent.h b/Engine/Plugins/Experimental/LearningAgents/Source/LearningAgentsTraining/Public/LearningAgentsEntitiesManagerComponent.h
--- a/Engine/Plugins/Experimental/LearningAgents/Source/LearningAgentsTraining/Public/LearningAgentsEntitiesManagerComponent.h
+++ b/Engine/Plugins/Experimental/LearningAgents/Source/LearningAgentsTraining/Public/LearningAgentsEntitiesManagerComponent.h
@@ -20,7 +20,7 @@ struct FLearningAgentsEntityInfo
 
 	/** Specify the entity class to spawn. */
 	UPROPERTY(EditDefaultsOnly, BlueprintReadWrite, Category = "LearningAgents")
-	TSubclassOf<UObject> EntityClass;
+	TSubclassOf<AActor> EntityClass;
 
 	/** Specify the Z offset to apply to spawn locations. */
 	UPROPERTY(EditDefaultsOnly, BlueprintReadWrite, Category = "LearningAgents")
@@ -71,15 +71,15 @@ public:
 
 	/** Spawns pooled entities at random locations. */
 	UFUNCTION(BlueprintCallable, Category = "LearningAgents")
-	TScriptInterface<ILearningAgentsEntityInterface> SpawnEntitiesAtRandomLocations(TSubclassOf<UObject> EntityClass, float EntitySpawnZOffset, int32 SpawnCount);
+	TScriptInterface<ILearningAgentsEntityInterface> SpawnEntitiesAtRandomLocations(TSubclassOf<AActor> EntityClass, float EntitySpawnZOffset, int32 SpawnCount);
 
 	/** Spawns a single pooled entity with a specified transform projected in the gym. */
 	UFUNCTION(BlueprintCallable, Category = "LearningAgents")
-	TScriptInterface<ILearningAgentsEntityInterface> SpawnEntityAtProjectedLocation(TSubclassOf<UObject> EntityClass, float EntitySpawnZOffset, const FTransform& Transform);
+	TScriptInterface<ILearningAgentsEntityInterface> SpawnEntityAtProjectedLocation(TSubclassOf<AActor> EntityClass, float EntitySpawnZOffset, const FTransform& Transform);
 
 	/** Spawns multiple pooled entities. */
 	UFUNCTION(BlueprintCallable, Category = "LearningAgents")
-	TArray<TScriptInterface<ILearningAgentsEntityInterface>> SpawnEntities(TSubclassOf<UObject> EntityClass, float EntitySpawnZOffset, int32 SpawnCount, const FTransform& Transform);
+	TArray<TScriptInterface<ILearningAgentsEntityInterface>> SpawnEntities(TSubclassOf<AActor> EntityClass, float EntitySpawnZOffset, int32 SpawnCount, const FTransform& Transform);
 
 protected:
 	virtual void BeginPlay() override;
