Learning Agents: Fix incorrect schema in LearningObservation

diff --git a/Engine/Plugins/Experimental/LearningAgents/Source/Learning/Private/LearningObservation.cpp b/Engine/Plugins/Experimental/LearningAgents/Source/Learning/Private/LearningObservation.cpp
--- a/Engine/Plugins/Experimental/LearningAgents/Source/Learning/Private/LearningObservation.cpp
+++ b/Engine/Plugins/Experimental/LearningAgents/Source/Learning/Private/LearningObservation.cpp
@@ -828,7 +828,7 @@ namespace UE::Learning::Observation
 		case EType::And:
 		{
 			const FSchemaAndParameters ParametersA = SchemaA.GetAnd(SchemaElementA);
-			const FSchemaAndParameters ParametersB = SchemaA.GetAnd(SchemaElementB);
+			const FSchemaAndParameters ParametersB = SchemaB.GetAnd(SchemaElementB);
 
 			if (ParametersA.Elements.Num() != ParametersB.Elements.Num()) { return false; }
 
@@ -845,7 +845,7 @@ namespace UE::Learning::Observation
 		case EType::OrExclusive:
 		{
 			const FSchemaOrExclusiveParameters ParametersA = SchemaA.GetOrExclusive(SchemaElementA);
-			const FSchemaOrExclusiveParameters ParametersB = SchemaA.GetOrExclusive(SchemaElementB);
+			const FSchemaOrExclusiveParameters ParametersB = SchemaB.GetOrExclusive(SchemaElementB);
 
 			if (ParametersA.Elements.Num() != ParametersB.Elements.Num()) { return false; }
 
@@ -862,7 +862,7 @@ namespace UE::Learning::Observation
 		case EType::OrInclusive:
 		{
 			const FSchemaOrInclusiveParameters ParametersA = SchemaA.GetOrInclusive(SchemaElementA);
-			const FSchemaOrInclusiveParameters ParametersB = SchemaA.GetOrInclusive(SchemaElementB);
+			const FSchemaOrInclusiveParameters ParametersB = SchemaB.GetOrInclusive(SchemaElementB);
 
 			if (ParametersA.Elements.Num() != ParametersB.Elements.Num()) { return false; }
 
@@ -879,7 +879,7 @@ namespace UE::Learning::Observation
 		case EType::Array:
 		{
 			const FSchemaArrayParameters ParametersA = SchemaA.GetArray(SchemaElementA);
-			const FSchemaArrayParameters ParametersB = SchemaA.GetArray (SchemaElementB);
+			const FSchemaArrayParameters ParametersB = SchemaB.GetArray(SchemaElementB);
 
 			return (ParametersA.Num == ParametersB.Num) && AreSchemaObjectsCompatible(SchemaA, ParametersA.Element, SchemaB, ParametersB.Element);
 		}
@@ -887,7 +887,7 @@ namespace UE::Learning::Observation
 		case EType::Set:
 		{
 			const FSchemaSetParameters ParametersA = SchemaA.GetSet(SchemaElementA);
-			const FSchemaSetParameters ParametersB = SchemaA.GetSet(SchemaElementB);
+			const FSchemaSetParameters ParametersB = SchemaB.GetSet(SchemaElementB);
 
 			return (ParametersA.MaxNum == ParametersB.MaxNum) && AreSchemaObjectsCompatible(SchemaA, ParametersA.Element, SchemaB, ParametersB.Element);
 		}
