LearningAgents: Added hash of network contents

diff --git a/Engine/Plugins/Experimental/LearningAgents/Source/Learning/Private/LearningNeuralNetwork.cpp b/Engine/Plugins/Experimental/LearningAgents/Source/Learning/Private/LearningNeuralNetwork.cpp
--- a/Engine/Plugins/Experimental/LearningAgents/Source/Learning/Private/LearningNeuralNetwork.cpp
+++ b/Engine/Plugins/Experimental/LearningAgents/Source/Learning/Private/LearningNeuralNetwork.cpp
@@ -126,6 +126,11 @@ int32 ULearningNeuralNetworkData::GetCompatibilityHash() const
 	return CompatibilityHash;
 }
 
+int32 ULearningNeuralNetworkData::GetContentHash() const
+{
+	return ContentHash;
+}
+
 void ULearningNeuralNetworkData::UpdateNetwork()
 {
 	if (FileData.Num() > 0)
@@ -153,6 +158,10 @@ void ULearningNeuralNetworkData::UpdateNetwork()
 			UpdatedModel = RuntimeCPU->CreateModelCPU(ModelData);
 		}
 
+		// Compute the content hash
+
+		ContentHash = CityHash32((const char*)FileData.GetData(), FileData.Num());
+
 		// If we are not in the editor we can now clear the FileData and FileType since these will be
 		// using additional memory and we are not going to save this asset and so don't require them.
 
diff --git a/Engine/Plugins/Experimental/LearningAgents/Source/Learning/Public/LearningNeuralNetwork.h b/Engine/Plugins/Experimental/LearningAgents/Source/Learning/Public/LearningNeuralNetwork.h
--- a/Engine/Plugins/Experimental/LearningAgents/Source/Learning/Public/LearningNeuralNetwork.h
+++ b/Engine/Plugins/Experimental/LearningAgents/Source/Learning/Public/LearningNeuralNetwork.h
@@ -80,6 +80,9 @@ public:
 	// Gets the compatibility hash
 	int32 GetCompatibilityHash() const;
 
+	// Gets the content hash
+	int32 GetContentHash() const;
+
 private:
 
 	// Calls UpdateNetwork to update the internal representation.
@@ -100,6 +103,10 @@ private:
 	UPROPERTY(VisibleAnywhere, Category = "Network Properties")
 	int32 CompatibilityHash = 0;
 
+	// Content hash which can be used to see when the network data has changed
+	UPROPERTY(VisibleAnywhere, Category = "Network Properties")
+	int32 ContentHash = 0;
+
 	// File data used by NNE
 	UPROPERTY()
 	TArray<uint8> FileData;
