Learning Agents: Fix a bug with GetRewards

diff --git a/Engine/Plugins/Experimental/LearningAgents/Source/LearningTraining/Private/LearningExperience.cpp b/Engine/Plugins/Experimental/LearningAgents/Source/LearningTraining/Private/LearningExperience.cpp
--- a/Engine/Plugins/Experimental/LearningAgents/Source/LearningTraining/Private/LearningExperience.cpp
+++ b/Engine/Plugins/Experimental/LearningAgents/Source/LearningTraining/Private/LearningExperience.cpp
@@ -257,7 +257,7 @@ namespace UE::Learning
 	{
 		checkf(RewardId >= 0 && RewardId < RewardArrays.Num(), TEXT("Reward id invalid!"));
 
-		return RewardArrays[RewardId][RewardId].Slice(0, EpisodeStepNums[InstanceIdx]);
+		return RewardArrays[RewardId][InstanceIdx].Slice(0, EpisodeStepNums[InstanceIdx]);
 	}
 
 	void FEpisodeBuffer::IncrementEpisodeStepNums(const FIndexSet Instances)
