Learning Agents: Fix shared memory on Mac by supporting Guid Short format

diff --git a/Engine/Plugins/Experimental/LearningAgents/Source/LearningAgentsTraining/Private/LearningAgentsCommunicator.cpp b/Engine/Plugins/Experimental/LearningAgents/Source/LearningAgentsTraining/Private/LearningAgentsCommunicator.cpp
--- a/Engine/Plugins/Experimental/LearningAgents/Source/LearningAgentsTraining/Private/LearningAgentsCommunicator.cpp
+++ b/Engine/Plugins/Experimental/LearningAgents/Source/LearningAgentsTraining/Private/LearningAgentsCommunicator.cpp
@@ -29,12 +29,6 @@ FLearningAgentsSharedMemoryTrainerProcess ULearningAgentsCommunicatorLibrary::Sp
 		return TrainerProcess;
 	}
 
-	if (PLATFORM_MAC)
-	{
-		UE_LOG(LogLearning, Error, TEXT("SpawnSharedMemoryTrainingProcess: Shared Memory not supported on Mac. Switch to Socket Communicator instead."));
-		return TrainerProcess;
-	}
-
 	const FString PythonExecutablePath = UE::Learning::Trainer::GetPythonExecutablePath(TrainerProcessSettings.GetIntermediatePath());
 	if (!FPaths::FileExists(PythonExecutablePath))
 	{
@@ -77,12 +71,6 @@ FLearningAgentsCommunicator ULearningAgentsCommunicatorLibrary::MakeSharedMemory
 {
 	FLearningAgentsCommunicator Communicator;
 
-	if (PLATFORM_MAC)
-	{
-		UE_LOG(LogLearning, Error, TEXT("MakeSharedMemoryCommunicator: Shared Memory not supported on Mac. Switch to Socket Communicator instead."));
-		return Communicator;
-	}
-
 	// Task Name
 	FString CmdLineTaskName;
 	bool bTaskNameOverridden = FParse::Value(FCommandLine::Get(), TEXT("LearningAgentsTaskName="), CmdLineTaskName);
@@ -105,7 +93,7 @@ FLearningAgentsCommunicator ULearningAgentsCommunicatorLibrary::MakeSharedMemory
 	FGuid CmdLineControlsGuid;
 	if (FParse::Value(FCommandLine::Get(), TEXT("LearningAgentsControlsGuid="), CmdLineControlsGuidString) && FGuid::Parse(CmdLineControlsGuidString, CmdLineControlsGuid))
 	{
-		UE_LOG(LogLearning, Display, TEXT("Overriding Controls Guid with value from the cmdline: %s"), *CmdLineControlsGuid.ToString(EGuidFormats::DigitsWithHyphensInBraces));
+		UE_LOG(LogLearning, Display, TEXT("Overriding Controls Guid with value from the cmdline: %s"), *CmdLineControlsGuid.ToString(PLATFORM_MAC ? EGuidFormats::Short : EGuidFormats::DigitsWithHyphensInBraces));
 	}
 
 	// Config Path
diff --git a/Engine/Plugins/Experimental/LearningAgents/Source/LearningTraining/Private/LearningExternalTrainer.cpp b/Engine/Plugins/Experimental/LearningAgents/Source/LearningTraining/Private/LearningExternalTrainer.cpp
--- a/Engine/Plugins/Experimental/LearningAgents/Source/LearningTraining/Private/LearningExternalTrainer.cpp
+++ b/Engine/Plugins/Experimental/LearningAgents/Source/LearningTraining/Private/LearningExternalTrainer.cpp
@@ -81,7 +81,7 @@ namespace UE::Learning
 			*TrainerFileName,
 			LogSettings == ELogSetting::Normal ? TEXT("-l") : TEXT(""),
 			*FileManager.ConvertToAbsolutePathForExternalAppForRead(*(InIntermediatePath / TaskName + FString::FromInt(TaskId))),
-			*Controls.Guid.ToString(EGuidFormats::DigitsWithHyphensInBraces));
+			*Controls.Guid.ToString(PLATFORM_MAC ? EGuidFormats::Short : EGuidFormats::DigitsWithHyphensInBraces));
 
 		bool bLaunched = TrainingProcess.Launch(
 			FileManager.ConvertToAbsolutePathForExternalAppForRead(*PythonExecutablePath),
@@ -283,7 +283,7 @@ namespace UE::Learning
 		{
 			TSharedPtr<FJsonObject> JsonObject = MakeShared<FJsonObject>();
 			JsonObject->SetNumberField(TEXT("NetworkId"), Index);
-			JsonObject->SetStringField(TEXT("Guid"), *NeuralNetworkSharedMemoryArrayViews[Index].Guid.ToString(EGuidFormats::DigitsWithHyphensInBraces));
+			JsonObject->SetStringField(TEXT("Guid"), *NeuralNetworkSharedMemoryArrayViews[Index].Guid.ToString(PLATFORM_MAC ? EGuidFormats::Short : EGuidFormats::DigitsWithHyphensInBraces));
 
 			TSharedRef<FJsonValueObject> JsonValue = MakeShared<FJsonValueObject>(JsonObject);
 			NetworkGuidsArray.Add(JsonValue);
@@ -294,14 +294,14 @@ namespace UE::Learning
 		for (const FSharedMemoryExperienceContainer& SharedMemoryExperienceContainer : SharedMemoryExperienceContainers)
 		{
 			TSharedPtr<FJsonObject> JsonObject = MakeShared<FJsonObject>();
-			JsonObject->SetStringField(TEXT("EpisodeStartsGuid"), *SharedMemoryExperienceContainer.EpisodeStarts.Guid.ToString(EGuidFormats::DigitsWithHyphensInBraces));
-			JsonObject->SetStringField(TEXT("EpisodeLengthsGuid"), *SharedMemoryExperienceContainer.EpisodeLengths.Guid.ToString(EGuidFormats::DigitsWithHyphensInBraces));
-			JsonObject->SetStringField(TEXT("EpisodeCompletionModesGuid"), *SharedMemoryExperienceContainer.EpisodeCompletionModes.Guid.ToString(EGuidFormats::DigitsWithHyphensInBraces));
+			JsonObject->SetStringField(TEXT("EpisodeStartsGuid"), *SharedMemoryExperienceContainer.EpisodeStarts.Guid.ToString(PLATFORM_MAC ? EGuidFormats::Short : EGuidFormats::DigitsWithHyphensInBraces));
+			JsonObject->SetStringField(TEXT("EpisodeLengthsGuid"), *SharedMemoryExperienceContainer.EpisodeLengths.Guid.ToString(PLATFORM_MAC ? EGuidFormats::Short : EGuidFormats::DigitsWithHyphensInBraces));
+			JsonObject->SetStringField(TEXT("EpisodeCompletionModesGuid"), *SharedMemoryExperienceContainer.EpisodeCompletionModes.Guid.ToString(PLATFORM_MAC ? EGuidFormats::Short : EGuidFormats::DigitsWithHyphensInBraces));
 
 			TArray<TSharedPtr<FJsonValue>> EpisodeFinalObservationsGuidsArray;
 			for (const TSharedMemoryArrayView<2, float>& EpisodeFinalObservations : SharedMemoryExperienceContainer.EpisodeFinalObservations)
 			{
-				EpisodeFinalObservationsGuidsArray.Add(MakeShared<FJsonValueString>(*EpisodeFinalObservations.Guid.ToString(EGuidFormats::DigitsWithHyphensInBraces)));
+				EpisodeFinalObservationsGuidsArray.Add(MakeShared<FJsonValueString>(*EpisodeFinalObservations.Guid.ToString(PLATFORM_MAC ? EGuidFormats::Short : EGuidFormats::DigitsWithHyphensInBraces)));
 				
 			}
 			JsonObject->SetArrayField(TEXT("EpisodeFinalObservationsGuids"), EpisodeFinalObservationsGuidsArray);
@@ -309,7 +309,7 @@ namespace UE::Learning
 			TArray<TSharedPtr<FJsonValue>> EpisodeFinalMemoryStatesGuidsArray;
 			for (const TSharedMemoryArrayView<2, float>& EpisodeFinalMemoryStates : SharedMemoryExperienceContainer.EpisodeFinalMemoryStates)
 			{
-				EpisodeFinalMemoryStatesGuidsArray.Add(MakeShared<FJsonValueString>(*EpisodeFinalMemoryStates.Guid.ToString(EGuidFormats::DigitsWithHyphensInBraces)));
+				EpisodeFinalMemoryStatesGuidsArray.Add(MakeShared<FJsonValueString>(*EpisodeFinalMemoryStates.Guid.ToString(PLATFORM_MAC ? EGuidFormats::Short : EGuidFormats::DigitsWithHyphensInBraces)));
 
 			}
 			JsonObject->SetArrayField(TEXT("EpisodeFinalMemoryStatesGuids"), EpisodeFinalMemoryStatesGuidsArray);
@@ -317,7 +317,7 @@ namespace UE::Learning
 			TArray<TSharedPtr<FJsonValue>> ObservationsGuidsArray;
 			for (const TSharedMemoryArrayView<2, float>& Observations : SharedMemoryExperienceContainer.Observations)
 			{
-				ObservationsGuidsArray.Add(MakeShared<FJsonValueString>(*Observations.Guid.ToString(EGuidFormats::DigitsWithHyphensInBraces)));
+				ObservationsGuidsArray.Add(MakeShared<FJsonValueString>(*Observations.Guid.ToString(PLATFORM_MAC ? EGuidFormats::Short : EGuidFormats::DigitsWithHyphensInBraces)));
 
 			}
 			JsonObject->SetArrayField(TEXT("ObservationsGuids"), ObservationsGuidsArray);
@@ -325,7 +325,7 @@ namespace UE::Learning
 			TArray<TSharedPtr<FJsonValue>> ActionsGuidsArray;
 			for (const TSharedMemoryArrayView<2, float>& Actions : SharedMemoryExperienceContainer.Actions)
 			{
-				ActionsGuidsArray.Add(MakeShared<FJsonValueString>(*Actions.Guid.ToString(EGuidFormats::DigitsWithHyphensInBraces)));
+				ActionsGuidsArray.Add(MakeShared<FJsonValueString>(*Actions.Guid.ToString(PLATFORM_MAC ? EGuidFormats::Short : EGuidFormats::DigitsWithHyphensInBraces)));
 
 			}
 			JsonObject->SetArrayField(TEXT("ActionsGuids"), ActionsGuidsArray);
@@ -333,7 +333,7 @@ namespace UE::Learning
 			TArray<TSharedPtr<FJsonValue>> ActionModifiersGuidsArray;
 			for (const TSharedMemoryArrayView<2, float>& ActionModifiers : SharedMemoryExperienceContainer.ActionModifiers)
 			{
-				ActionModifiersGuidsArray.Add(MakeShared<FJsonValueString>(*ActionModifiers.Guid.ToString(EGuidFormats::DigitsWithHyphensInBraces)));
+				ActionModifiersGuidsArray.Add(MakeShared<FJsonValueString>(*ActionModifiers.Guid.ToString(PLATFORM_MAC ? EGuidFormats::Short : EGuidFormats::DigitsWithHyphensInBraces)));
 
 			}
 			JsonObject->SetArrayField(TEXT("ActionModifiersGuids"), ActionModifiersGuidsArray);
@@ -342,7 +342,7 @@ namespace UE::Learning
 			TArray<TSharedPtr<FJsonValue>> MemoryStatesGuidsArray;
 			for (const TSharedMemoryArrayView<2, float>& MemoryStates : SharedMemoryExperienceContainer.MemoryStates)
 			{
-				MemoryStatesGuidsArray.Add(MakeShared<FJsonValueString>(*MemoryStates.Guid.ToString(EGuidFormats::DigitsWithHyphensInBraces)));
+				MemoryStatesGuidsArray.Add(MakeShared<FJsonValueString>(*MemoryStates.Guid.ToString(PLATFORM_MAC ? EGuidFormats::Short : EGuidFormats::DigitsWithHyphensInBraces)));
 
 			}
 			JsonObject->SetArrayField(TEXT("MemoryStatesGuids"), MemoryStatesGuidsArray);
@@ -350,7 +350,7 @@ namespace UE::Learning
 			TArray<TSharedPtr<FJsonValue>> RewardsGuidsArray;
 			for (const TSharedMemoryArrayView<2, float>& Rewards : SharedMemoryExperienceContainer.Rewards)
 			{
-				RewardsGuidsArray.Add(MakeShared<FJsonValueString>(*Rewards.Guid.ToString(EGuidFormats::DigitsWithHyphensInBraces)));
+				RewardsGuidsArray.Add(MakeShared<FJsonValueString>(*Rewards.Guid.ToString(PLATFORM_MAC ? EGuidFormats::Short : EGuidFormats::DigitsWithHyphensInBraces)));
 
 			}
 			JsonObject->SetArrayField(TEXT("RewardsGuids"), RewardsGuidsArray);
@@ -364,7 +364,7 @@ namespace UE::Learning
 		FString SharedMemoryConfigString;
 		TSharedRef<TJsonWriter<>> SharedMemoryJsonWriter = TJsonWriterFactory<>::Create(&SharedMemoryConfigString, 0);
 		FJsonSerializer::Serialize(SharedMemoryConfigObject, SharedMemoryJsonWriter, true);
-		FFileHelper::SaveStringToFile(SharedMemoryConfigString, *(ConfigPath / FString::Printf(TEXT("shared-memory-%s.json"), *Controls.Guid.ToString(EGuidFormats::DigitsWithHyphensInBraces))));
+		FFileHelper::SaveStringToFile(SharedMemoryConfigString, *(ConfigPath / FString::Printf(TEXT("shared-memory-%s.json"), *Controls.Guid.ToString(PLATFORM_MAC ? EGuidFormats::Short : EGuidFormats::DigitsWithHyphensInBraces))));
 
 		FString DataConfigString;
 		TSharedRef<TJsonWriter<>> DataJsonWriter = TJsonWriterFactory<>::Create(&DataConfigString, 0);
diff --git a/Engine/Plugins/Experimental/LearningAgents/Source/LearningTraining/Public/LearningSharedMemory.h b/Engine/Plugins/Experimental/LearningAgents/Source/LearningTraining/Public/LearningSharedMemory.h
--- a/Engine/Plugins/Experimental/LearningAgents/Source/LearningTraining/Public/LearningSharedMemory.h
+++ b/Engine/Plugins/Experimental/LearningAgents/Source/LearningTraining/Public/LearningSharedMemory.h
@@ -48,7 +48,7 @@ namespace UE::Learning
 			}
 
 			FPlatformMemory::FSharedMemoryRegion* RegionPointer = FPlatformMemory::MapNamedSharedMemoryRegion(
-				Guid.ToString(EGuidFormats::DigitsWithHyphensInBraces),
+				Guid.ToString(PLATFORM_MAC ? EGuidFormats::Short : EGuidFormats::DigitsWithHyphensInBraces),
 				bCreate,
 				// I was getting odd issues allocating read-only memory so for now we always allocate as read-write
 				FPlatformMemory::ESharedMemoryAccess::Read | FPlatformMemory::ESharedMemoryAccess::Write,
