LearningAgents: Updated to create full set of Discrete and Named Discrete Observations and Actions

#virtualized

diff --git a/Engine/Plugins/Experimental/LearningAgents/Source/Learning/Private/LearningAction.cpp b/Engine/Plugins/Experimental/LearningAgents/Source/Learning/Private/LearningAction.cpp
--- a/Engine/Plugins/Experimental/LearningAgents/Source/Learning/Private/LearningAction.cpp
+++ b/Engine/Plugins/Experimental/LearningAgents/Source/Learning/Private/LearningAction.cpp
@@ -222,6 +222,58 @@ namespace UE::Learning::Action
 		return { Index, Generation };
 	}
 
+	FSchemaElement FSchema::CreateNamedDiscreteExclusive(const FSchemaNamedDiscreteExclusiveParameters Parameters, const FName Tag)
+	{
+		check(Parameters.PriorProbabilities.Num() == Parameters.ElementNames.Num());
+		check(Private::CheckPriorProbabilitiesExclusive(Parameters.PriorProbabilities));
+		check(!Private::ContainsDuplicates(Parameters.ElementNames));
+
+		FNamedDiscreteExclusiveData ElementData;
+		ElementData.Num = Parameters.ElementNames.Num();
+		ElementData.PriorProbabilitiesOffset = PriorProbabilities.Num();
+		ElementData.ElementsOffset = SubElementNames.Num();
+
+		PriorProbabilities.Append(Parameters.PriorProbabilities);
+		SubElementNames.Append(Parameters.ElementNames);
+		for (int32 Idx = 0; Idx < ElementData.Num; Idx++) { SubElementObjects.Add(FSchemaElement()); }
+
+		const int32 Index = Types.Add(EType::NamedDiscreteExclusive);
+		Tags.Add(Tag);
+		EncodedVectorSizes.Add(Parameters.ElementNames.Num());
+		ActionVectorSizes.Add(Parameters.ElementNames.Num());
+		ActionDistributionVectorSizes.Add(Parameters.ElementNames.Num());
+		ActionModifierVectorSizes.Add(1 + Parameters.ElementNames.Num());
+		TypeDataIndices.Add(NamedDiscreteExclusiveData.Add(ElementData));
+
+		return { Index, Generation };
+	}
+
+	FSchemaElement FSchema::CreateNamedDiscreteInclusive(const FSchemaNamedDiscreteInclusiveParameters Parameters, const FName Tag)
+	{
+		check(Parameters.PriorProbabilities.Num() == Parameters.ElementNames.Num());
+		check(Private::CheckPriorProbabilitiesInclusive(Parameters.PriorProbabilities));
+		check(!Private::ContainsDuplicates(Parameters.ElementNames));
+
+		FNamedDiscreteInclusiveData ElementData;
+		ElementData.Num = Parameters.ElementNames.Num();
+		ElementData.PriorProbabilitiesOffset = PriorProbabilities.Num();
+		ElementData.ElementsOffset = SubElementNames.Num();
+
+		PriorProbabilities.Append(Parameters.PriorProbabilities);
+		SubElementNames.Append(Parameters.ElementNames);
+		for (int32 Idx = 0; Idx < ElementData.Num; Idx++) { SubElementObjects.Add(FSchemaElement()); }
+
+		const int32 Index = Types.Add(EType::NamedDiscreteInclusive);
+		Tags.Add(Tag);
+		EncodedVectorSizes.Add(Parameters.ElementNames.Num());
+		ActionVectorSizes.Add(Parameters.ElementNames.Num());
+		ActionDistributionVectorSizes.Add(Parameters.ElementNames.Num());
+		ActionModifierVectorSizes.Add(1 + Parameters.ElementNames.Num());
+		TypeDataIndices.Add(NamedDiscreteInclusiveData.Add(ElementData));
+
+		return { Index, Generation };
+	}
+
 	FSchemaElement FSchema::CreateAnd(const FSchemaAndParameters Parameters, const FName Tag)
 	{
 		check(Parameters.Elements.Num() == Parameters.ElementNames.Num());
@@ -423,6 +475,28 @@ namespace UE::Learning::Action
 		return Parameters;
 	}
 
+	FSchemaNamedDiscreteExclusiveParameters FSchema::GetNamedDiscreteExclusive(const FSchemaElement Element) const
+	{
+		check(IsValid(Element) && GetType(Element) == EType::NamedDiscreteExclusive);
+		const FNamedDiscreteExclusiveData& ElementData = NamedDiscreteExclusiveData[TypeDataIndices[Element.Index]];
+
+		FSchemaNamedDiscreteExclusiveParameters Parameters;
+		Parameters.ElementNames = TArrayView<const FName>(SubElementNames.GetData() + ElementData.ElementsOffset, ElementData.Num);
+		Parameters.PriorProbabilities = TArrayView<const float>(PriorProbabilities.GetData() + ElementData.PriorProbabilitiesOffset, ElementData.Num);
+		return Parameters;
+	}
+
+	FSchemaNamedDiscreteInclusiveParameters FSchema::GetNamedDiscreteInclusive(const FSchemaElement Element) const
+	{
+		check(IsValid(Element) && GetType(Element) == EType::NamedDiscreteInclusive);
+		const FNamedDiscreteInclusiveData& ElementData = NamedDiscreteInclusiveData[TypeDataIndices[Element.Index]];
+
+		FSchemaNamedDiscreteInclusiveParameters Parameters;
+		Parameters.ElementNames = TArrayView<const FName>(SubElementNames.GetData() + ElementData.ElementsOffset, ElementData.Num);
+		Parameters.PriorProbabilities = TArrayView<const float>(PriorProbabilities.GetData() + ElementData.PriorProbabilitiesOffset, ElementData.Num);
+		return Parameters;
+	}
+
 	FSchemaAndParameters FSchema::GetAnd(const FSchemaElement Element) const
 	{
 		check(IsValid(Element) && GetType(Element) == EType::And);
@@ -499,6 +573,8 @@ namespace UE::Learning::Action
 		ContinuousData.Empty();
 		DiscreteExclusiveData.Empty();
 		DiscreteInclusiveData.Empty();
+		NamedDiscreteExclusiveData.Empty();
+		NamedDiscreteInclusiveData.Empty();
 		AndData.Empty();
 		OrExclusiveData.Empty();
 		OrInclusiveData.Empty();
@@ -529,6 +605,8 @@ namespace UE::Learning::Action
 		ContinuousData.Reset();
 		DiscreteExclusiveData.Reset();
 		DiscreteInclusiveData.Reset();
+		NamedDiscreteExclusiveData.Reset();
+		NamedDiscreteInclusiveData.Reset();		
 		AndData.Reset();
 		OrExclusiveData.Reset();
 		OrInclusiveData.Reset();
@@ -616,6 +694,46 @@ namespace UE::Learning::Action
 		return { Index, Generation };
 	}
 
+	FObjectElement FObject::CreateNamedDiscreteExclusive(const FObjectNamedDiscreteExclusiveParameters Parameters, const FName Tag)
+	{
+		const int32 Index = Types.Add(EType::NamedDiscreteExclusive);
+		Tags.Add(Tag);
+
+		ContinuousDataOffsets.Add(ContinuousValues.Num());
+		ContinuousDataNums.Add(0);
+
+		DiscreteDataOffsets.Add(DiscreteValues.Num());
+		DiscreteDataNums.Add(0);
+
+		ElementDataOffsets.Add(SubElementObjects.Num());
+		ElementDataNums.Add(1);
+
+		SubElementObjects.Add(FObjectElement());
+		SubElementNames.Add(Parameters.ElementName);
+
+		return { Index, Generation };
+	}
+
+	FObjectElement FObject::CreateNamedDiscreteInclusive(const FObjectNamedDiscreteInclusiveParameters Parameters, const FName Tag)
+	{
+		const int32 Index = Types.Add(EType::NamedDiscreteInclusive);
+		Tags.Add(Tag);
+
+		ContinuousDataOffsets.Add(ContinuousValues.Num());
+		ContinuousDataNums.Add(0);
+
+		DiscreteDataOffsets.Add(DiscreteValues.Num());
+		DiscreteDataNums.Add(0);
+
+		ElementDataOffsets.Add(SubElementObjects.Num());
+		ElementDataNums.Add(Parameters.ElementNames.Num());
+
+		for (int32 Idx = 0; Idx < Parameters.ElementNames.Num(); Idx++) { SubElementObjects.Add(FObjectElement()); }
+		SubElementNames.Append(Parameters.ElementNames);
+
+		return { Index, Generation };
+	}
+
 	FObjectElement FObject::CreateAnd(const FObjectAndParameters Parameters, const FName Tag)
 	{
 		check(Parameters.Elements.Num() == Parameters.ElementNames.Num());
@@ -777,6 +895,24 @@ namespace UE::Learning::Action
 		return Parameters;
 	}
 
+	FObjectNamedDiscreteExclusiveParameters FObject::GetNamedDiscreteExclusive(const FObjectElement Element) const
+	{
+		check(IsValid(Element) && GetType(Element) == EType::NamedDiscreteExclusive);
+
+		FObjectNamedDiscreteExclusiveParameters Parameters;
+		Parameters.ElementName = SubElementNames[ElementDataOffsets[Element.Index]];
+		return Parameters;
+	}
+
+	FObjectNamedDiscreteInclusiveParameters FObject::GetNamedDiscreteInclusive(const FObjectElement Element) const
+	{
+		check(IsValid(Element) && GetType(Element) == EType::NamedDiscreteInclusive);
+
+		FObjectNamedDiscreteInclusiveParameters Parameters;
+		Parameters.ElementNames = TArrayView<const FName>(SubElementNames.GetData() + ElementDataOffsets[Element.Index], ElementDataNums[Element.Index]);
+		return Parameters;
+	}
+
 	FObjectAndParameters FObject::GetAnd(const FObjectElement Element) const
 	{
 		check(IsValid(Element) && GetType(Element) == EType::And);
@@ -961,6 +1097,52 @@ namespace UE::Learning::Action
 		return { Index, Generation };
 	}
 
+	FModifierElement FModifier::CreateNamedDiscreteExclusive(const FModifierNamedDiscreteExclusiveParameters Parameters, const FName Tag)
+	{
+		const int32 Index = Types.Add(EType::NamedDiscreteExclusive);
+		Tags.Add(Tag);
+
+		ContinuousDataOffsets.Add(ContinuousMaskeds.Num());
+		ContinuousDataNums.Add(0);
+
+		DiscreteDataOffsets.Add(DiscreteValues.Num());
+		DiscreteDataNums.Add(0);
+
+		ElementDataOffsets.Add(SubElementModifiers.Num());
+		ElementDataNums.Add(Parameters.MaskedElementNames.Num());
+
+		MaskedDataOffsets.Add(MaskedElementNames.Num());
+		MaskedDataNums.Add(0);
+
+		SubElementNames.Append(Parameters.MaskedElementNames);
+		for (int32 Idx = 0; Idx < Parameters.MaskedElementNames.Num(); Idx++) { SubElementModifiers.Add(FModifierElement()); }
+
+		return { Index, Generation };
+	}
+
+	FModifierElement FModifier::CreateNamedDiscreteInclusive(const FModifierNamedDiscreteInclusiveParameters Parameters, const FName Tag)
+	{
+		const int32 Index = Types.Add(EType::NamedDiscreteInclusive);
+		Tags.Add(Tag);
+
+		ContinuousDataOffsets.Add(ContinuousMaskeds.Num());
+		ContinuousDataNums.Add(0);
+
+		DiscreteDataOffsets.Add(DiscreteValues.Num());
+		DiscreteDataNums.Add(0);
+
+		ElementDataOffsets.Add(SubElementModifiers.Num());
+		ElementDataNums.Add(Parameters.MaskedElementNames.Num());
+
+		MaskedDataOffsets.Add(MaskedElementNames.Num());
+		MaskedDataNums.Add(0);
+
+		SubElementNames.Append(Parameters.MaskedElementNames);
+		for (int32 Idx = 0; Idx < Parameters.MaskedElementNames.Num(); Idx++) { SubElementModifiers.Add(FModifierElement()); }
+
+		return { Index, Generation };
+	}
+
 	FModifierElement FModifier::CreateAnd(const FModifierAndParameters Parameters, const FName Tag)
 	{
 		check(Parameters.Elements.Num() == Parameters.ElementNames.Num());
@@ -1146,6 +1328,24 @@ namespace UE::Learning::Action
 		return Parameters;
 	}
 
+	FModifierNamedDiscreteExclusiveParameters FModifier::GetNamedDiscreteExclusive(const FModifierElement Element) const
+	{
+		check(IsValid(Element) && GetType(Element) == EType::NamedDiscreteExclusive);
+
+		FModifierNamedDiscreteExclusiveParameters Parameters;
+		Parameters.MaskedElementNames = TArrayView<const FName>(SubElementNames.GetData() + ElementDataOffsets[Element.Index], ElementDataNums[Element.Index]);
+		return Parameters;
+	}
+
+	FModifierNamedDiscreteInclusiveParameters FModifier::GetNamedDiscreteInclusive(const FModifierElement Element) const
+	{
+		check(IsValid(Element) && GetType(Element) == EType::NamedDiscreteInclusive);
+
+		FModifierNamedDiscreteInclusiveParameters Parameters;
+		Parameters.MaskedElementNames = TArrayView<const FName>(SubElementNames.GetData() + ElementDataOffsets[Element.Index], ElementDataNums[Element.Index]);
+		return Parameters;
+	}
+
 	FModifierAndParameters FModifier::GetAnd(const FModifierElement Element) const
 	{
 		check(IsValid(Element) && GetType(Element) == EType::And);
@@ -1286,6 +1486,24 @@ namespace UE::Learning::Action
 			return (int32)CityHash32((const char*)Hashes.GetData(), Hashes.Num() * Hashes.GetTypeSize());
 		}
 
+		static inline int32 HashElements(
+			const FSchema& Schema,
+			const TArrayView<const FName> SchemaElementNames,
+			const int32 Salt)
+		{
+			// Note: Here we xor all entries together. 
+			// This makes the hash in invariant to the ordering of names which is actually what we want 
+			// since this array is representing a set-like structure and it is fine to pass elements in a different order.
+
+			int32 Hash = 0x9de53147;
+			for (int32 SchemaElementIdx = 0; SchemaElementIdx < SchemaElementNames.Num(); SchemaElementIdx++)
+			{
+				Hash ^= HashFNameStable(SchemaElementNames[SchemaElementIdx]);
+			}
+
+			return Hash;
+		}
+
 		static inline int32 HashElements(
 			const FSchema& Schema,
 			const TArrayView<const FName> SchemaElementNames,
@@ -1327,6 +1545,18 @@ namespace UE::Learning::Action
 
 		case EType::DiscreteInclusive: return Private::HashCombine({ Hash, Private::HashInt(Schema.GetDiscreteInclusive(SchemaElement).Num) });
 
+		case EType::NamedDiscreteExclusive:
+		{
+			const FSchemaNamedDiscreteExclusiveParameters Parameters = Schema.GetNamedDiscreteExclusive(SchemaElement);
+			return Private::HashCombine({ Hash, Private::HashElements(Schema, Parameters.ElementNames, Salt) });
+		}
+
+		case EType::NamedDiscreteInclusive:
+		{
+			const FSchemaNamedDiscreteInclusiveParameters Parameters = Schema.GetNamedDiscreteInclusive(SchemaElement);
+			return Private::HashCombine({ Hash, Private::HashElements(Schema, Parameters.ElementNames, Salt) });
+		}
+
 		case EType::And:
 		{
 			const FSchemaAndParameters Parameters = Schema.GetAnd(SchemaElement);
@@ -1398,6 +1628,38 @@ namespace UE::Learning::Action
 
 		case EType::DiscreteInclusive: return SchemaA.GetDiscreteInclusive(SchemaElementA).Num == SchemaB.GetDiscreteInclusive(SchemaElementB).Num;
 
+		case EType::NamedDiscreteExclusive:
+		{
+			const FSchemaNamedDiscreteExclusiveParameters ParametersA = SchemaA.GetNamedDiscreteExclusive(SchemaElementA);
+			const FSchemaNamedDiscreteExclusiveParameters ParametersB = SchemaB.GetNamedDiscreteExclusive(SchemaElementB);
+
+			if (ParametersA.ElementNames.Num() != ParametersB.ElementNames.Num()) { return false; }
+
+			for (int32 SchemaElementAIdx = 0; SchemaElementAIdx < ParametersA.ElementNames.Num(); SchemaElementAIdx++)
+			{
+				const int32 SchemaElementBIdx = ParametersB.ElementNames.Find(ParametersA.ElementNames[SchemaElementAIdx]);
+				if (SchemaElementBIdx == INDEX_NONE) { return false; }
+			}
+
+			return true;
+		}
+
+		case EType::NamedDiscreteInclusive:
+		{
+			const FSchemaNamedDiscreteInclusiveParameters ParametersA = SchemaA.GetNamedDiscreteInclusive(SchemaElementA);
+			const FSchemaNamedDiscreteInclusiveParameters ParametersB = SchemaB.GetNamedDiscreteInclusive(SchemaElementB);
+
+			if (ParametersA.ElementNames.Num() != ParametersB.ElementNames.Num()) { return false; }
+
+			for (int32 SchemaElementAIdx = 0; SchemaElementAIdx < ParametersA.ElementNames.Num(); SchemaElementAIdx++)
+			{
+				const int32 SchemaElementBIdx = ParametersB.ElementNames.Find(ParametersA.ElementNames[SchemaElementAIdx]);
+				if (SchemaElementBIdx == INDEX_NONE) { return false; }
+			}
+
+			return true;
+		}
+
 		case EType::And:
 		{
 			const FSchemaAndParameters ParametersA = SchemaA.GetAnd(SchemaElementA);
@@ -1517,7 +1779,6 @@ namespace UE::Learning::Action
 				Parameters.Num,
 				Builder.MakeValuesCopy(LogPriorProbabilities),
 				Builder.MakeValuesOne(Parameters.Num));
-
 			break;
 		}
 
@@ -1532,12 +1793,50 @@ namespace UE::Learning::Action
 				LogPriorProbabilities[Idx] = Private::Logit(LogPriorProbabilities[Idx]);
 			}
 
-			OutElement = Builder.MakeDenormalize(Parameters.Num,
+			OutElement = Builder.MakeDenormalize(
+				Parameters.Num,
 				Builder.MakeValuesCopy(LogPriorProbabilities),
 				Builder.MakeValuesOne(Parameters.Num));
 			break;
 		}
 
+		case EType::NamedDiscreteExclusive:
+		{
+			const FSchemaNamedDiscreteExclusiveParameters Parameters = Schema.GetNamedDiscreteExclusive(SchemaElement);
+
+			TArray<float, TInlineAllocator<16>> LogPriorProbabilities;
+			LogPriorProbabilities.Append(Parameters.PriorProbabilities);
+			for (int32 Idx = 0; Idx < Parameters.ElementNames.Num(); Idx++)
+			{
+				// Clamp zero probabilities to the smallest (positive) float. This is approximately equal to a probability of 1:1e38
+				LogPriorProbabilities[Idx] = FMath::Loge(FMath::Max(LogPriorProbabilities[Idx], FLT_MIN));
+			}
+
+			OutElement = Builder.MakeDenormalize(
+				Parameters.ElementNames.Num(),
+				Builder.MakeValuesCopy(LogPriorProbabilities),
+				Builder.MakeValuesOne(Parameters.ElementNames.Num()));
+			break;
+		}
+
+		case EType::NamedDiscreteInclusive:
+		{
+			const FSchemaNamedDiscreteInclusiveParameters Parameters = Schema.GetNamedDiscreteInclusive(SchemaElement);
+
+			TArray<float, TInlineAllocator<16>> LogPriorProbabilities;
+			LogPriorProbabilities.Append(Parameters.PriorProbabilities);
+			for (int32 Idx = 0; Idx < Parameters.ElementNames.Num(); Idx++)
+			{
+				LogPriorProbabilities[Idx] = Private::Logit(LogPriorProbabilities[Idx]);
+			}
+
+			OutElement = Builder.MakeDenormalize(
+				Parameters.ElementNames.Num(),
+				Builder.MakeValuesCopy(LogPriorProbabilities),
+				Builder.MakeValuesOne(Parameters.ElementNames.Num()));
+			break;
+		}
+
 		case EType::And:
 		{
 			const FSchemaAndParameters Parameters = Schema.GetAnd(SchemaElement);
@@ -1821,6 +2120,77 @@ namespace UE::Learning::Action
 			break;
 		}
 
+		case EType::NamedDiscreteExclusive:
+		{
+			const int32 ValueNum = Schema.GetNamedDiscreteExclusive(SchemaElement).ElementNames.Num();
+			check(ValueNum == OutActionVector.Num());
+			check(ValueNum == ActionDistributionVector.Num());
+			check(1 + ValueNum == ActionModifierVector.Num());
+
+			if (ActionModifierVector[0])
+			{
+				TLearningArray<1, bool, TInlineAllocator<32>> Masked;
+				Masked.SetNumUninitialized({ ValueNum });
+				for (int32 ValueIdx = 0; ValueIdx < ValueNum; ValueIdx++)
+				{
+					Masked[ValueIdx] = ActionModifierVector[1 + ValueIdx] == 1.0f;
+				}
+				check(Private::CheckExclusiveMaskValid(Masked));
+
+				Random::SampleDistributionMultinoulliMasked(
+					OutActionVector,
+					InOutRandomState,
+					ActionDistributionVector,
+					Masked,
+					ActionNoiseScale);
+			}
+			else
+			{
+				Random::SampleDistributionMultinoulli(
+					OutActionVector,
+					InOutRandomState,
+					ActionDistributionVector,
+					ActionNoiseScale);
+			}
+
+			break;
+		}
+
+		case EType::NamedDiscreteInclusive:
+		{
+			const int32 ValueNum = Schema.GetNamedDiscreteInclusive(SchemaElement).ElementNames.Num();
+			check(ValueNum == OutActionVector.Num());
+			check(ValueNum == ActionDistributionVector.Num());
+			check(1 + ValueNum == ActionModifierVector.Num());
+
+			if (ActionModifierVector[0])
+			{
+				TLearningArray<1, bool, TInlineAllocator<32>> Masked;
+				Masked.SetNumUninitialized({ ValueNum });
+				for (int32 ValueIdx = 0; ValueIdx < ValueNum; ValueIdx++)
+				{
+					Masked[ValueIdx] = ActionModifierVector[1 + ValueIdx] == 1.0f;
+				}
+
+				Random::SampleDistributionBernoulliMasked(
+					OutActionVector,
+					InOutRandomState,
+					ActionDistributionVector,
+					Masked,
+					ActionNoiseScale);
+			}
+			else
+			{
+				Random::SampleDistributionBernoulli(
+					OutActionVector,
+					InOutRandomState,
+					ActionDistributionVector,
+					ActionNoiseScale);
+			}
+
+			break;
+		}
+
 		case EType::And:
 		{
 			const FSchemaAndParameters Parameters = Schema.GetAnd(SchemaElement);
@@ -2119,7 +2489,7 @@ namespace UE::Learning::Action
 
 		case EType::DiscreteExclusive:
 		{
-			int32 ActionValue = Object.GetDiscreteExclusive(ObjectElement).DiscreteIndex;
+			const int32 ActionValue = Object.GetDiscreteExclusive(ObjectElement).DiscreteIndex;
 			check(Schema.GetActionVectorSize(SchemaElement) > ActionValue && ActionValue >= 0);
 			check(Schema.GetActionVectorSize(SchemaElement) == OutActionVector.Num());
 
@@ -2131,7 +2501,7 @@ namespace UE::Learning::Action
 
 		case EType::DiscreteInclusive:
 		{
-			TArrayView<const int32> ActionValues = Object.GetDiscreteInclusive(ObjectElement).DiscreteIndices;
+			const TArrayView<const int32> ActionValues = Object.GetDiscreteInclusive(ObjectElement).DiscreteIndices;
 			check(Schema.GetActionVectorSize(SchemaElement) >= ActionValues.Num());
 			check(Schema.GetActionVectorSize(SchemaElement) == OutActionVector.Num());
 
@@ -2146,6 +2516,36 @@ namespace UE::Learning::Action
 			return;
 		}
 
+		case EType::NamedDiscreteExclusive:
+		{
+			const TArrayView<const FName> SchemaNames = Schema.GetNamedDiscreteExclusive(SchemaElement).ElementNames;
+			const FName ActionValue = Object.GetNamedDiscreteExclusive(ObjectElement).ElementName;
+			check(Schema.GetActionVectorSize(SchemaElement) == OutActionVector.Num());
+
+			// Set the single value in the action vector
+			const int32 ActionIndex = SchemaNames.Find(ActionValue);
+			check(ActionIndex != INDEX_NONE);
+			OutActionVector[ActionIndex] = 1.0f;
+			return;
+		}
+
+		case EType::NamedDiscreteInclusive:
+		{
+			const TArrayView<const FName> SchemaNames = Schema.GetNamedDiscreteInclusive(SchemaElement).ElementNames;
+			const TArrayView<const FName> ActionValues = Object.GetNamedDiscreteInclusive(ObjectElement).ElementNames;
+			check(Schema.GetActionVectorSize(SchemaElement) >= ActionValues.Num());
+			check(Schema.GetActionVectorSize(SchemaElement) == OutActionVector.Num());
+
+			// Set values in the action vector
+			for (int32 ActionValueIdx = 0; ActionValueIdx < ActionValues.Num(); ActionValueIdx++)
+			{
+				const int32 ActionIndex = SchemaNames.Find(ActionValues[ActionValueIdx]);
+				check(ActionIndex != INDEX_NONE);
+				OutActionVector[ActionIndex] = 1.0f;
+			}
+			return;
+		}
+
 		case EType::And:
 		{
 			// Check the number of sub-elements match
@@ -2393,6 +2793,49 @@ namespace UE::Learning::Action
 			return;
 		}
 
+		case EType::NamedDiscreteExclusive:
+		{
+			const TArrayView<const FName> SchemaNames = Schema.GetNamedDiscreteExclusive(SchemaElement).ElementNames;
+			check(ActionVectorSize == Schema.GetNamedDiscreteExclusive(SchemaElement).ElementNames.Num());
+
+			// Find Name
+			FName ExclusiveName = NAME_None;
+			for (int32 Idx = 0; Idx < ActionVectorSize; Idx++)
+			{
+				check(ActionVector[Idx] == 0.0f || ActionVector[Idx] == 1.0f);
+				if (ActionVector[Idx])
+				{
+					ExclusiveName = SchemaNames[Idx];
+					break;
+				}
+			}
+			check(ExclusiveName != NAME_None);
+
+			OutObjectElement = OutObject.CreateNamedDiscreteExclusive({ ExclusiveName }, SchemaElementTag);
+			return;
+		}
+
+		case EType::NamedDiscreteInclusive:
+		{
+			const TArrayView<const FName> SchemaNames = Schema.GetNamedDiscreteInclusive(SchemaElement).ElementNames;
+			check(ActionVectorSize == Schema.GetNamedDiscreteInclusive(SchemaElement).ElementNames.Num());
+
+			// Find Names
+			TArray<FName, TInlineAllocator<8>> InclusiveNames;
+			InclusiveNames.Reserve(ActionVectorSize);
+			for (int32 Idx = 0; Idx < ActionVectorSize; Idx++)
+			{
+				check(ActionVector[Idx] == 0.0f || ActionVector[Idx] == 1.0f);
+				if (ActionVector[Idx])
+				{
+					InclusiveNames.Add(SchemaNames[Idx]);
+				}
+			}
+
+			OutObjectElement = OutObject.CreateNamedDiscreteInclusive({ InclusiveNames }, SchemaElementTag);
+			return;
+		}
+
 		case EType::And:
 		{
 			const FSchemaAndParameters Parameters = Schema.GetAnd(SchemaElement);
@@ -2643,6 +3086,36 @@ namespace UE::Learning::Action
 			return;
 		}
 
+		case EType::NamedDiscreteExclusive:
+		{
+			const TArrayView<const FName> MaskNames = Modifier.GetNamedDiscreteExclusive(ModifierElement).MaskedElementNames;
+			check(Schema.GetNamedDiscreteExclusive(SchemaElement).ElementNames.Num() >= MaskNames.Num());
+
+			for (int32 MaskNameIdx = 0; MaskNameIdx < MaskNames.Num(); MaskNameIdx++)
+			{
+				const int32 MaskIdx = Schema.GetNamedDiscreteExclusive(SchemaElement).ElementNames.Find(MaskNames[MaskNameIdx]);
+				check(MaskIdx != INDEX_NONE);
+				OutActionModifierVector[1 + MaskIdx] = 1.0f;
+			}
+
+			return;
+		}
+
+		case EType::NamedDiscreteInclusive:
+		{
+			const TArrayView<const FName> MaskNames = Modifier.GetNamedDiscreteInclusive(ModifierElement).MaskedElementNames;
+			check(Schema.GetNamedDiscreteInclusive(SchemaElement).ElementNames.Num() >= MaskNames.Num());
+
+			for (int32 MaskNameIdx = 0; MaskNameIdx < MaskNames.Num(); MaskNameIdx++)
+			{
+				const int32 MaskIdx = Schema.GetNamedDiscreteInclusive(SchemaElement).ElementNames.Find(MaskNames[MaskNameIdx]);
+				check(MaskIdx != INDEX_NONE);
+				OutActionModifierVector[1 + MaskIdx] = 1.0f;
+			}
+
+			return;
+		}
+
 		case EType::And:
 		{
 			const FSchemaAndParameters SchemaParameters = Schema.GetAnd(SchemaElement);
@@ -2884,7 +3357,6 @@ namespace UE::Learning::Action
 		case EType::DiscreteExclusive:
 		{
 			const int32 ValueNum = Schema.GetDiscreteExclusive(SchemaElement).Num;
-
 			check(ActionModifierVectorSize == 1 + ValueNum);
 
 			// Find Indices
@@ -2906,7 +3378,6 @@ namespace UE::Learning::Action
 		case EType::DiscreteInclusive:
 		{
 			const int32 ValueNum = Schema.GetDiscreteInclusive(SchemaElement).Num;
-
 			check(ActionModifierVectorSize == 1 + ValueNum);
 
 			// Find Indices
@@ -2925,6 +3396,48 @@ namespace UE::Learning::Action
 			return;
 		}
 
+		case EType::NamedDiscreteExclusive:
+		{
+			const TArrayView<const FName> ElementNames = Schema.GetNamedDiscreteExclusive(SchemaElement).ElementNames;
+			check(ActionModifierVectorSize == 1 + ElementNames.Num());
+
+			// Find Names
+			TArray<FName, TInlineAllocator<8>> MaskedNames;
+			MaskedNames.Reserve(ElementNames.Num());
+			for (int32 Idx = 0; Idx < ElementNames.Num(); Idx++)
+			{
+				check(ActionModifierVector[1 + Idx] == 0.0f || ActionModifierVector[1 + Idx] == 1.0f);
+				if (ActionModifierVector[1 + Idx] == 1.0f)
+				{
+					MaskedNames.Add(ElementNames[Idx]);
+				}
+			}
+
+			OutModifierElement = OutModifier.CreateNamedDiscreteExclusive({ MaskedNames }, SchemaElementTag);
+			return;
+		}
+
+		case EType::NamedDiscreteInclusive:
+		{
+			const TArrayView<const FName> ElementNames = Schema.GetNamedDiscreteInclusive(SchemaElement).ElementNames;
+			check(ActionModifierVectorSize == 1 + ElementNames.Num());
+
+			// Find Names
+			TArray<FName, TInlineAllocator<8>> MaskedNames;
+			MaskedNames.Reserve(ElementNames.Num());
+			for (int32 Idx = 0; Idx < ElementNames.Num(); Idx++)
+			{
+				check(ActionModifierVector[1 + Idx] == 0.0f || ActionModifierVector[1 + Idx] == 1.0f);
+				if (ActionModifierVector[1 + Idx] == 1.0f)
+				{
+					MaskedNames.Add(ElementNames[Idx]);
+				}
+			}
+
+			OutModifierElement = OutModifier.CreateNamedDiscreteInclusive({ MaskedNames }, SchemaElementTag);
+			return;
+		}
+
 		case EType::And:
 		{
 			const FSchemaAndParameters Parameters = Schema.GetAnd(SchemaElement);
diff --git a/Engine/Plugins/Experimental/LearningAgents/Source/Learning/Private/LearningObservation.cpp b/Engine/Plugins/Experimental/LearningAgents/Source/Learning/Private/LearningObservation.cpp
--- a/Engine/Plugins/Experimental/LearningAgents/Source/Learning/Private/LearningObservation.cpp
+++ b/Engine/Plugins/Experimental/LearningAgents/Source/Learning/Private/LearningObservation.cpp
@@ -94,6 +94,74 @@ namespace UE::Learning::Observation
 		return { Index, Generation };
 	}
 
+	FSchemaElement FSchema::CreateDiscreteExclusive(const FSchemaDiscreteExclusiveParameters Parameters, const FName Tag)
+	{
+		FDiscreteExclusiveData ElementData;
+		ElementData.Num = Parameters.Num;
+
+		const int32 Index = Types.Add(EType::DiscreteExclusive);
+		Tags.Add(Tag);
+		ObservationVectorSizes.Add(Parameters.Num);
+		EncodedVectorSizes.Add(Parameters.Num);
+		TypeDataIndices.Add(DiscreteExclusiveData.Add(ElementData));
+
+		return { Index, Generation };
+	}
+
+	FSchemaElement FSchema::CreateDiscreteInclusive(const FSchemaDiscreteInclusiveParameters Parameters, const FName Tag)
+	{
+		FDiscreteInclusiveData ElementData;
+		ElementData.Num = Parameters.Num;
+
+		const int32 Index = Types.Add(EType::DiscreteInclusive);
+		Tags.Add(Tag);
+		ObservationVectorSizes.Add(Parameters.Num);
+		EncodedVectorSizes.Add(Parameters.Num);
+		TypeDataIndices.Add(DiscreteInclusiveData.Add(ElementData));
+
+		return { Index, Generation };
+	}
+
+	FSchemaElement FSchema::CreateNamedDiscreteExclusive(const FSchemaNamedDiscreteExclusiveParameters Parameters, const FName Tag)
+	{
+		check(!Private::ContainsDuplicates(Parameters.ElementNames));
+
+		FNamedDiscreteExclusiveData ElementData;
+		ElementData.Num = Parameters.ElementNames.Num();
+		ElementData.ElementsOffset = SubElementObjects.Num();
+
+		SubElementNames.Append(Parameters.ElementNames);
+		for (int32 Idx = 0; Idx < ElementData.Num; Idx++) { SubElementObjects.Add(FSchemaElement()); }
+
+		const int32 Index = Types.Add(EType::NamedDiscreteExclusive);
+		Tags.Add(Tag);
+		ObservationVectorSizes.Add(Parameters.ElementNames.Num());
+		EncodedVectorSizes.Add(Parameters.ElementNames.Num());
+		TypeDataIndices.Add(NamedDiscreteExclusiveData.Add(ElementData));
+
+		return { Index, Generation };
+	}
+
+	FSchemaElement FSchema::CreateNamedDiscreteInclusive(const FSchemaNamedDiscreteInclusiveParameters Parameters, const FName Tag)
+	{
+		check(!Private::ContainsDuplicates(Parameters.ElementNames));
+
+		FNamedDiscreteInclusiveData ElementData;
+		ElementData.Num = Parameters.ElementNames.Num();
+		ElementData.ElementsOffset = SubElementObjects.Num();
+
+		SubElementNames.Append(Parameters.ElementNames);
+		for (int32 Idx = 0; Idx < ElementData.Num; Idx++) { SubElementObjects.Add(FSchemaElement()); }
+
+		const int32 Index = Types.Add(EType::NamedDiscreteInclusive);
+		Tags.Add(Tag);
+		ObservationVectorSizes.Add(Parameters.ElementNames.Num());
+		EncodedVectorSizes.Add(Parameters.ElementNames.Num());
+		TypeDataIndices.Add(NamedDiscreteInclusiveData.Add(ElementData));
+
+		return { Index, Generation };
+	}
+
 	FSchemaElement FSchema::CreateAnd(const FSchemaAndParameters Parameters, const FName Tag)
 	{
 		check(Parameters.Elements.Num() == Parameters.ElementNames.Num());
@@ -270,6 +338,46 @@ namespace UE::Learning::Observation
 		return Parameters;
 	}
 
+	FSchemaDiscreteExclusiveParameters FSchema::GetDiscreteExclusive(const FSchemaElement Element) const
+	{
+		check(IsValid(Element) && GetType(Element) == EType::DiscreteExclusive);
+		const FDiscreteExclusiveData& ElementData = DiscreteExclusiveData[TypeDataIndices[Element.Index]];
+
+		FSchemaDiscreteExclusiveParameters Parameters;
+		Parameters.Num = ElementData.Num;
+		return Parameters;
+	}
+
+	FSchemaDiscreteInclusiveParameters FSchema::GetDiscreteInclusive(const FSchemaElement Element) const
+	{
+		check(IsValid(Element) && GetType(Element) == EType::DiscreteInclusive);
+		const FDiscreteInclusiveData& ElementData = DiscreteInclusiveData[TypeDataIndices[Element.Index]];
+
+		FSchemaDiscreteInclusiveParameters Parameters;
+		Parameters.Num = ElementData.Num;
+		return Parameters;
+	}
+
+	FSchemaNamedDiscreteExclusiveParameters FSchema::GetNamedDiscreteExclusive(const FSchemaElement Element) const
+	{
+		check(IsValid(Element) && GetType(Element) == EType::NamedDiscreteExclusive);
+		const FNamedDiscreteExclusiveData& ElementData = NamedDiscreteExclusiveData[TypeDataIndices[Element.Index]];
+
+		FSchemaNamedDiscreteExclusiveParameters Parameters;
+		Parameters.ElementNames = TArrayView<const FName>(SubElementNames.GetData() + ElementData.ElementsOffset, ElementData.Num);
+		return Parameters;
+	}
+
+	FSchemaNamedDiscreteInclusiveParameters FSchema::GetNamedDiscreteInclusive(const FSchemaElement Element) const
+	{
+		check(IsValid(Element) && GetType(Element) == EType::NamedDiscreteInclusive);
+		const FNamedDiscreteInclusiveData& ElementData = NamedDiscreteInclusiveData[TypeDataIndices[Element.Index]];
+
+		FSchemaNamedDiscreteInclusiveParameters Parameters;
+		Parameters.ElementNames = TArrayView<const FName>(SubElementNames.GetData() + ElementData.ElementsOffset, ElementData.Num);
+		return Parameters;
+	}
+
 	FSchemaAndParameters FSchema::GetAnd(const FSchemaElement Element) const
 	{
 		check(IsValid(Element) && GetType(Element) == EType::And);
@@ -359,6 +467,10 @@ namespace UE::Learning::Observation
 		TypeDataIndices.Empty();
 
 		ContinuousData.Empty();
+		DiscreteExclusiveData.Empty();
+		DiscreteInclusiveData.Empty();
+		NamedDiscreteExclusiveData.Empty();
+		NamedDiscreteInclusiveData.Empty();	
 		AndData.Empty();
 		OrExclusiveData.Empty();
 		OrInclusiveData.Empty();
@@ -386,6 +498,10 @@ namespace UE::Learning::Observation
 		TypeDataIndices.Reset();
 
 		ContinuousData.Reset();
+		DiscreteExclusiveData.Reset();
+		DiscreteInclusiveData.Reset();
+		NamedDiscreteExclusiveData.Reset();
+		NamedDiscreteInclusiveData.Reset();			
 		AndData.Reset();
 		OrExclusiveData.Reset();
 		OrInclusiveData.Reset();
@@ -407,6 +523,9 @@ namespace UE::Learning::Observation
 		ContinuousDataOffsets.Add(ContinuousValues.Num());
 		ContinuousDataNums.Add(0);
 
+		DiscreteDataOffsets.Add(DiscreteValues.Num());
+		DiscreteDataNums.Add(0);
+
 		SubElementDataOffsets.Add(SubElementObjects.Num());
 		SubElementDataNums.Add(0);
 
@@ -421,6 +540,9 @@ namespace UE::Learning::Observation
 		ContinuousDataOffsets.Add(ContinuousValues.Num());
 		ContinuousDataNums.Add(Parameters.Values.Num());
 
+		DiscreteDataOffsets.Add(DiscreteValues.Num());
+		DiscreteDataNums.Add(0);
+
 		SubElementDataOffsets.Add(SubElementObjects.Num());
 		SubElementDataNums.Add(0);
 
@@ -429,6 +551,86 @@ namespace UE::Learning::Observation
 		return { Index, Generation };
 	}
 
+	FObjectElement FObject::CreateDiscreteExclusive(const FObjectDiscreteExclusiveParameters Parameters, const FName Tag)
+	{
+		const int32 Index = Types.Add(EType::DiscreteExclusive);
+		Tags.Add(Tag);
+
+		ContinuousDataOffsets.Add(ContinuousValues.Num());
+		ContinuousDataNums.Add(0);
+
+		DiscreteDataOffsets.Add(DiscreteValues.Num());
+		DiscreteDataNums.Add(1);
+
+		SubElementDataOffsets.Add(SubElementObjects.Num());
+		SubElementDataNums.Add(0);
+
+		DiscreteValues.Add(Parameters.DiscreteIndex);
+
+		return { Index, Generation };
+	}
+
+	FObjectElement FObject::CreateDiscreteInclusive(const FObjectDiscreteInclusiveParameters Parameters, const FName Tag)
+	{
+		const int32 Index = Types.Add(EType::DiscreteInclusive);
+		Tags.Add(Tag);
+
+		ContinuousDataOffsets.Add(ContinuousValues.Num());
+		ContinuousDataNums.Add(0);
+
+		DiscreteDataOffsets.Add(DiscreteValues.Num());
+		DiscreteDataNums.Add(Parameters.DiscreteIndices.Num());
+
+		SubElementDataOffsets.Add(SubElementObjects.Num());
+		SubElementDataNums.Add(0);
+
+		DiscreteValues.Append(Parameters.DiscreteIndices);
+
+		return { Index, Generation };
+	}
+
+	FObjectElement FObject::CreateNamedDiscreteExclusive(const FObjectNamedDiscreteExclusiveParameters Parameters, const FName Tag)
+	{
+		const int32 Index = Types.Add(EType::NamedDiscreteExclusive);
+		Tags.Add(Tag);
+
+		ContinuousDataOffsets.Add(ContinuousValues.Num());
+		ContinuousDataNums.Add(0);
+
+		DiscreteDataOffsets.Add(DiscreteValues.Num());
+		DiscreteDataNums.Add(0);
+
+		SubElementDataOffsets.Add(SubElementObjects.Num());
+		SubElementDataNums.Add(1);
+
+		SubElementNames.Add(Parameters.ElementName);
+		SubElementObjects.Add(FObjectElement());
+
+		return { Index, Generation };
+	}
+
+	FObjectElement FObject::CreateNamedDiscreteInclusive(const FObjectNamedDiscreteInclusiveParameters Parameters, const FName Tag)
+	{
+		check(!Private::ContainsDuplicates(Parameters.ElementNames));
+
+		const int32 Index = Types.Add(EType::NamedDiscreteInclusive);
+		Tags.Add(Tag);
+
+		ContinuousDataOffsets.Add(ContinuousValues.Num());
+		ContinuousDataNums.Add(0);
+
+		DiscreteDataOffsets.Add(DiscreteValues.Num());
+		DiscreteDataNums.Add(0);
+
+		SubElementDataOffsets.Add(SubElementObjects.Num());
+		SubElementDataNums.Add(Parameters.ElementNames.Num());
+
+		SubElementNames.Append(Parameters.ElementNames);
+		for (int32 Idx = 0; Idx < Parameters.ElementNames.Num(); Idx++) { SubElementObjects.Add(FObjectElement()); }
+
+		return { Index, Generation };
+	}
+
 	FObjectElement FObject::CreateAnd(const FObjectAndParameters Parameters, const FName Tag)
 	{
 		check(Parameters.Elements.Num() == Parameters.ElementNames.Num());
@@ -441,6 +643,9 @@ namespace UE::Learning::Observation
 		ContinuousDataOffsets.Add(ContinuousValues.Num());
 		ContinuousDataNums.Add(0);
 
+		DiscreteDataOffsets.Add(DiscreteValues.Num());
+		DiscreteDataNums.Add(0);
+
 		SubElementDataOffsets.Add(SubElementObjects.Num());
 		SubElementDataNums.Add(Parameters.Elements.Num());
 
@@ -460,6 +665,9 @@ namespace UE::Learning::Observation
 		ContinuousDataOffsets.Add(ContinuousValues.Num());
 		ContinuousDataNums.Add(0);
 
+		DiscreteDataOffsets.Add(DiscreteValues.Num());
+		DiscreteDataNums.Add(0);
+
 		SubElementDataOffsets.Add(SubElementObjects.Num());
 		SubElementDataNums.Add(1);
 
@@ -481,6 +689,9 @@ namespace UE::Learning::Observation
 		ContinuousDataOffsets.Add(ContinuousValues.Num());
 		ContinuousDataNums.Add(0);
 
+		DiscreteDataOffsets.Add(DiscreteValues.Num());
+		DiscreteDataNums.Add(0);
+
 		SubElementDataOffsets.Add(SubElementObjects.Num());
 		SubElementDataNums.Add(Parameters.Elements.Num());
 
@@ -500,6 +711,9 @@ namespace UE::Learning::Observation
 		ContinuousDataOffsets.Add(ContinuousValues.Num());
 		ContinuousDataNums.Add(0);
 
+		DiscreteDataOffsets.Add(DiscreteValues.Num());
+		DiscreteDataNums.Add(0);
+
 		SubElementDataOffsets.Add(SubElementObjects.Num());
 		SubElementDataNums.Add(Parameters.Elements.Num());
 
@@ -522,6 +736,9 @@ namespace UE::Learning::Observation
 		ContinuousDataOffsets.Add(ContinuousValues.Num());
 		ContinuousDataNums.Add(0);
 
+		DiscreteDataOffsets.Add(DiscreteValues.Num());
+		DiscreteDataNums.Add(0);
+
 		SubElementDataOffsets.Add(SubElementObjects.Num());
 		SubElementDataNums.Add(Parameters.Elements.Num());
 
@@ -544,6 +761,9 @@ namespace UE::Learning::Observation
 		ContinuousDataOffsets.Add(ContinuousValues.Num());
 		ContinuousDataNums.Add(0);
 
+		DiscreteDataOffsets.Add(DiscreteValues.Num());
+		DiscreteDataNums.Add(0);
+
 		SubElementDataOffsets.Add(SubElementObjects.Num());
 		SubElementDataNums.Add(1);
 
@@ -579,6 +799,42 @@ namespace UE::Learning::Observation
 		return Parameters;
 	}
 
+	FObjectDiscreteExclusiveParameters FObject::GetDiscreteExclusive(const FObjectElement Element) const
+	{
+		check(IsValid(Element) && GetType(Element) == EType::DiscreteExclusive);
+
+		FObjectDiscreteExclusiveParameters Parameters;
+		Parameters.DiscreteIndex = DiscreteValues[DiscreteDataOffsets[Element.Index]];
+		return Parameters;
+	}
+
+	FObjectDiscreteInclusiveParameters FObject::GetDiscreteInclusive(const FObjectElement Element) const
+	{
+		check(IsValid(Element) && GetType(Element) == EType::DiscreteInclusive);
+
+		FObjectDiscreteInclusiveParameters Parameters;
+		Parameters.DiscreteIndices = TArrayView<const int32>(DiscreteValues.GetData() + DiscreteDataOffsets[Element.Index], DiscreteDataNums[Element.Index]);
+		return Parameters;
+	}
+
+	FObjectNamedDiscreteExclusiveParameters FObject::GetNamedDiscreteExclusive(const FObjectElement Element) const
+	{
+		check(IsValid(Element) && GetType(Element) == EType::NamedDiscreteExclusive);
+
+		FObjectNamedDiscreteExclusiveParameters Parameters;
+		Parameters.ElementName = SubElementNames[SubElementDataOffsets[Element.Index]];
+		return Parameters;
+	}
+
+	FObjectNamedDiscreteInclusiveParameters FObject::GetNamedDiscreteInclusive(const FObjectElement Element) const
+	{
+		check(IsValid(Element) && GetType(Element) == EType::NamedDiscreteInclusive);
+
+		FObjectNamedDiscreteInclusiveParameters Parameters;
+		Parameters.ElementNames = TArrayView<const FName>(SubElementNames.GetData() + SubElementDataOffsets[Element.Index], SubElementDataNums[Element.Index]);
+		return Parameters;
+	}
+
 	FObjectAndParameters FObject::GetAnd(const FObjectElement Element) const
 	{
 		check(IsValid(Element) && GetType(Element) == EType::And);
@@ -649,10 +905,14 @@ namespace UE::Learning::Observation
 		ContinuousDataOffsets.Empty();
 		ContinuousDataNums.Empty();
 
+		DiscreteDataOffsets.Reset();
+		DiscreteDataNums.Reset();
+
 		SubElementDataOffsets.Empty();
 		SubElementDataNums.Empty();
 
 		ContinuousValues.Empty();
+		DiscreteValues.Empty();
 		SubElementNames.Empty();
 		SubElementObjects.Empty();
 
@@ -672,10 +932,14 @@ namespace UE::Learning::Observation
 		ContinuousDataOffsets.Reset();
 		ContinuousDataNums.Reset();
 
+		DiscreteDataOffsets.Reset();
+		DiscreteDataNums.Reset();
+
 		SubElementDataOffsets.Reset();
 		SubElementDataNums.Reset();
 
 		ContinuousValues.Reset();
+		DiscreteValues.Reset();
 		SubElementNames.Reset();
 		SubElementObjects.Reset();
 
@@ -715,6 +979,24 @@ namespace UE::Learning::Observation
 			return (int32)CityHash32((const char*)Hashes.GetData(), Hashes.Num() * Hashes.GetTypeSize());
 		}
 
+		static inline int32 HashElements(
+			const FSchema& Schema,
+			const TArrayView<const FName> SchemaElementNames,
+			const int32 Salt)
+		{
+			// Note: Here we xor all entries together. 
+			// This makes the hash in invariant to the ordering of names which is actually what we want 
+			// since this array is representing a set-like structure and it is fine to pass elements in a different order.
+
+			int32 Hash = 0x5592716a;
+			for (int32 SchemaElementIdx = 0; SchemaElementIdx < SchemaElementNames.Num(); SchemaElementIdx++)
+			{
+				Hash ^= HashFNameStable(SchemaElementNames[SchemaElementIdx]);
+			}
+
+			return Hash;
+		}
+
 		static inline int32 HashElements(
 			const FSchema& Schema,
 			const TArrayView<const FName> SchemaElementNames,
@@ -752,6 +1034,21 @@ namespace UE::Learning::Observation
 
 		case EType::Continuous: return Private::HashCombine({ Hash, Private::HashInt(Schema.GetContinuous(SchemaElement).Num) });
 
+		case EType::DiscreteExclusive: return Private::HashCombine({ Hash, Private::HashInt(Schema.GetDiscreteExclusive(SchemaElement).Num) });
+		case EType::DiscreteInclusive: return Private::HashCombine({ Hash, Private::HashInt(Schema.GetDiscreteInclusive(SchemaElement).Num) });
+
+		case EType::NamedDiscreteExclusive:
+		{
+			const FSchemaNamedDiscreteExclusiveParameters Parameters = Schema.GetNamedDiscreteExclusive(SchemaElement);
+			return Private::HashCombine({ Hash, Private::HashElements(Schema, Parameters.ElementNames, Salt) });
+		}
+
+		case EType::NamedDiscreteInclusive:
+		{
+			const FSchemaNamedDiscreteInclusiveParameters Parameters = Schema.GetNamedDiscreteInclusive(SchemaElement);
+			return Private::HashCombine({ Hash, Private::HashElements(Schema, Parameters.ElementNames, Salt) });
+		}
+
 		case EType::And:
 		{
 			const FSchemaAndParameters Parameters = Schema.GetAnd(SchemaElement);
@@ -824,6 +1121,41 @@ namespace UE::Learning::Observation
 		case EType::Null: return true;
 
 		case EType::Continuous: return SchemaA.GetContinuous(SchemaElementA).Num == SchemaB.GetContinuous(SchemaElementB).Num;
+		
+		case EType::DiscreteExclusive: return SchemaA.GetDiscreteExclusive(SchemaElementA).Num == SchemaB.GetDiscreteExclusive(SchemaElementB).Num;
+		case EType::DiscreteInclusive: return SchemaA.GetDiscreteInclusive(SchemaElementA).Num == SchemaB.GetDiscreteInclusive(SchemaElementB).Num;
+
+		case EType::NamedDiscreteExclusive:
+		{
+			const FSchemaNamedDiscreteExclusiveParameters ParametersA = SchemaA.GetNamedDiscreteExclusive(SchemaElementA);
+			const FSchemaNamedDiscreteExclusiveParameters ParametersB = SchemaB.GetNamedDiscreteExclusive(SchemaElementB);
+
+			if (ParametersA.ElementNames.Num() != ParametersB.ElementNames.Num()) { return false; }
+
+			for (int32 SchemaElementAIdx = 0; SchemaElementAIdx < ParametersA.ElementNames.Num(); SchemaElementAIdx++)
+			{
+				const int32 SchemaElementBIdx = ParametersB.ElementNames.Find(ParametersA.ElementNames[SchemaElementAIdx]);
+				if (SchemaElementBIdx == INDEX_NONE) { return false; }
+			}
+
+			return true;
+		}
+
+		case EType::NamedDiscreteInclusive:
+		{
+			const FSchemaNamedDiscreteInclusiveParameters ParametersA = SchemaA.GetNamedDiscreteInclusive(SchemaElementA);
+			const FSchemaNamedDiscreteInclusiveParameters ParametersB = SchemaB.GetNamedDiscreteInclusive(SchemaElementB);
+
+			if (ParametersA.ElementNames.Num() != ParametersB.ElementNames.Num()) { return false; }
+
+			for (int32 SchemaElementAIdx = 0; SchemaElementAIdx < ParametersA.ElementNames.Num(); SchemaElementAIdx++)
+			{
+				const int32 SchemaElementBIdx = ParametersB.ElementNames.Find(ParametersA.ElementNames[SchemaElementAIdx]);
+				if (SchemaElementBIdx == INDEX_NONE) { return false; }
+			}
+
+			return true;
+		}
 
 		case EType::And:
 		{
@@ -935,6 +1267,50 @@ namespace UE::Learning::Observation
 			break;
 		}
 
+		case EType::DiscreteExclusive:
+		{
+			const int32 ValueNum = Schema.GetDiscreteExclusive(SchemaElement).Num;
+
+			OutElement = Builder.MakeDenormalize(
+				ValueNum,
+				Builder.MakeValuesZero(ValueNum),
+				Builder.MakeValuesOne(ValueNum));
+			break;
+		}
+
+		case EType::DiscreteInclusive:
+		{
+			const int32 ValueNum = Schema.GetDiscreteInclusive(SchemaElement).Num;
+
+			OutElement = Builder.MakeDenormalize(
+				ValueNum,
+				Builder.MakeValuesZero(ValueNum),
+				Builder.MakeValuesOne(ValueNum));
+			break;
+		}
+
+		case EType::NamedDiscreteExclusive:
+		{
+			const int32 ValueNum = Schema.GetNamedDiscreteExclusive(SchemaElement).ElementNames.Num();
+
+			OutElement = Builder.MakeDenormalize(
+				ValueNum,
+				Builder.MakeValuesZero(ValueNum),
+				Builder.MakeValuesOne(ValueNum));
+			break;
+		}
+
+		case EType::NamedDiscreteInclusive:
+		{
+			const int32 ValueNum = Schema.GetNamedDiscreteInclusive(SchemaElement).ElementNames.Num();
+
+			OutElement = Builder.MakeDenormalize(
+				ValueNum,
+				Builder.MakeValuesZero(ValueNum),
+				Builder.MakeValuesOne(ValueNum));
+			break;
+		}
+
 		case EType::And:
 		{
 			const FSchemaAndParameters Parameters = Schema.GetAnd(SchemaElement);
@@ -1205,6 +1581,65 @@ namespace UE::Learning::Observation
 			return;
 		}
 
+		case EType::DiscreteExclusive:
+		{
+			const int32 ObservationValue = Object.GetDiscreteExclusive(ObjectElement).DiscreteIndex;
+			check(Schema.GetObservationVectorSize(SchemaElement) > ObservationValue && ObservationValue >= 0);
+			check(Schema.GetObservationVectorSize(SchemaElement) == OutObservationVector.Num());
+
+			// Set the single value in the observation vector
+
+			OutObservationVector[ObservationValue] = 1.0f;
+			return;
+		}
+
+		case EType::DiscreteInclusive:
+		{
+			const TArrayView<const int32> ObservationValues = Object.GetDiscreteInclusive(ObjectElement).DiscreteIndices;
+			check(Schema.GetObservationVectorSize(SchemaElement) >= ObservationValues.Num());
+			check(Schema.GetObservationVectorSize(SchemaElement) == OutObservationVector.Num());
+
+			// Set values in the observation vector
+
+			for (int32 ObservationValueIdx = 0; ObservationValueIdx < ObservationValues.Num(); ObservationValueIdx++)
+			{
+				check(Schema.GetObservationVectorSize(SchemaElement) > ObservationValues[ObservationValueIdx] && ObservationValues[ObservationValueIdx] >= 0);
+				OutObservationVector[ObservationValues[ObservationValueIdx]] = 1.0f;
+			}
+
+			return;
+		}
+
+		case EType::NamedDiscreteExclusive:
+		{
+			const TArrayView<const FName> SchemaNames = Schema.GetNamedDiscreteExclusive(SchemaElement).ElementNames;
+			const FName ObservationValue = Object.GetNamedDiscreteExclusive(ObjectElement).ElementName;
+			check(Schema.GetObservationVectorSize(SchemaElement) == OutObservationVector.Num());
+
+			// Set the single value in the observation vector
+			const int32 ObservationIndex = SchemaNames.Find(ObservationValue);
+			check(ObservationIndex != INDEX_NONE);
+			OutObservationVector[ObservationIndex] = 1.0f;
+			return;
+		}
+
+		case EType::NamedDiscreteInclusive:
+		{
+			const TArrayView<const FName> SchemaNames = Schema.GetNamedDiscreteInclusive(SchemaElement).ElementNames;
+			const TArrayView<const FName> ObservationValues = Object.GetNamedDiscreteInclusive(ObjectElement).ElementNames;
+			check(Schema.GetObservationVectorSize(SchemaElement) >= ObservationValues.Num());
+			check(Schema.GetObservationVectorSize(SchemaElement) == OutObservationVector.Num());
+
+			// Set values in the observation vector
+			for (int32 ObservationValueIdx = 0; ObservationValueIdx < ObservationValues.Num(); ObservationValueIdx++)
+			{
+				const int32 ObservationIndex = SchemaNames.Find(ObservationValues[ObservationValueIdx]);
+				check(ObservationIndex != INDEX_NONE);
+				OutObservationVector[ObservationIndex] = 1.0f;
+			}
+			return;
+		}
+
 		case EType::And:
 		{
 			// Check the number of sub-elements match
@@ -1447,6 +1882,91 @@ namespace UE::Learning::Observation
 			OutObjectElement = OutObject.CreateContinuous({ MakeArrayView(ObservationValues.GetData(), ObservationValues.Num()) }, SchemaElementTag);
 			return;
 		}
+		
+		case EType::DiscreteExclusive:
+		{
+			check(ObservationVectorSize == Schema.GetDiscreteExclusive(SchemaElement).Num);
+
+			// Find Index
+			int32 ExclusiveIndex = INDEX_NONE;
+			for (int32 Idx = 0; Idx < ObservationVectorSize; Idx++)
+			{
+				check(ObservationVector[Idx] == 0.0f || ObservationVector[Idx] == 1.0f);
+				if (ObservationVector[Idx])
+				{
+					ExclusiveIndex = Idx;
+					break;
+				}
+			}
+			check(ExclusiveIndex != INDEX_NONE);
+
+			OutObjectElement = OutObject.CreateDiscreteExclusive({ ExclusiveIndex }, SchemaElementTag);
+			return;
+		}
+
+		case EType::DiscreteInclusive:
+		{
+			check(ObservationVectorSize == Schema.GetDiscreteInclusive(SchemaElement).Num);
+
+			// Find Indices
+			TArray<int32, TInlineAllocator<8>> InclusiveIndices;
+			InclusiveIndices.Reserve(ObservationVectorSize);
+			for (int32 Idx = 0; Idx < ObservationVectorSize; Idx++)
+			{
+				check(ObservationVector[Idx] == 0.0f || ObservationVector[Idx] == 1.0f);
+				if (ObservationVector[Idx])
+				{
+					InclusiveIndices.Add(Idx);
+				}
+			}
+
+			OutObjectElement = OutObject.CreateDiscreteInclusive({ InclusiveIndices }, SchemaElementTag);
+			return;
+		}
+
+
+		case EType::NamedDiscreteExclusive:
+		{
+			const TArrayView<const FName> SchemaNames = Schema.GetNamedDiscreteExclusive(SchemaElement).ElementNames;
+			check(ObservationVectorSize == Schema.GetNamedDiscreteExclusive(SchemaElement).ElementNames.Num());
+
+			// Find Name
+			FName ExclusiveName = NAME_None;
+			for (int32 Idx = 0; Idx < ObservationVectorSize; Idx++)
+			{
+				check(ObservationVector[Idx] == 0.0f || ObservationVector[Idx] == 1.0f);
+				if (ObservationVector[Idx])
+				{
+					ExclusiveName = SchemaNames[Idx];
+					break;
+				}
+			}
+			check(ExclusiveName != NAME_None);
+
+			OutObjectElement = OutObject.CreateNamedDiscreteExclusive({ ExclusiveName }, SchemaElementTag);
+			return;
+		}
+
+		case EType::NamedDiscreteInclusive:
+		{
+			const TArrayView<const FName> SchemaNames = Schema.GetNamedDiscreteInclusive(SchemaElement).ElementNames;
+			check(ObservationVectorSize == Schema.GetNamedDiscreteInclusive(SchemaElement).ElementNames.Num());
+
+			// Find Names
+			TArray<FName, TInlineAllocator<8>> InclusiveNames;
+			InclusiveNames.Reserve(ObservationVectorSize);
+			for (int32 Idx = 0; Idx < ObservationVectorSize; Idx++)
+			{
+				check(ObservationVector[Idx] == 0.0f || ObservationVector[Idx] == 1.0f);
+				if (ObservationVector[Idx])
+				{
+					InclusiveNames.Add(SchemaNames[Idx]);
+				}
+			}
+
+			OutObjectElement = OutObject.CreateNamedDiscreteInclusive({ InclusiveNames }, SchemaElementTag);
+			return;
+		}
 
 		case EType::And:
 		{
@@ -1685,6 +2205,26 @@ namespace UE::Learning::Observation
 			return;
 		}
 
+		case EType::DiscreteExclusive:
+		{
+			return;
+		}
+
+		case EType::DiscreteInclusive:
+		{
+			return;
+		}
+
+		case EType::NamedDiscreteExclusive:
+		{
+			return;
+		}
+
+		case EType::NamedDiscreteInclusive:
+		{
+			return;
+		}
+
 		case EType::And:
 		{
 			const FSchemaAndParameters Parameters = Schema.GetAnd(SchemaElement);
diff --git a/Engine/Plugins/Experimental/LearningAgents/Source/Learning/Public/LearningAction.h b/Engine/Plugins/Experimental/LearningAgents/Source/Learning/Public/LearningAction.h
--- a/Engine/Plugins/Experimental/LearningAgents/Source/Learning/Public/LearningAction.h
+++ b/Engine/Plugins/Experimental/LearningAgents/Source/Learning/Public/LearningAction.h
@@ -33,20 +33,26 @@ namespace UE::Learning::Action
 		// Set of inclusive discrete actions
 		DiscreteInclusive	= 3,
 
+		// Set of named exclusive discrete actions
+		NamedDiscreteExclusive = 4,
+
+		// Set of named inclusive discrete actions
+		NamedDiscreteInclusive = 5,
+
 		// Combination of multiple actions
-		And					= 4,
+		And					= 6,
 
 		// Exclusive choice from a set of actions
-		OrExclusive			= 5,
+		OrExclusive			= 7,
 
 		// Inclusive choice from a set of actions
-		OrInclusive			= 6,
+		OrInclusive			= 8,
 
 		// Fixed sized array of actions
-		Array				= 7,
+		Array				= 9,
 
 		// Encoding of another action
-		Encoding			= 8,
+		Encoding			= 10,
 	};
 
 	/**
@@ -84,6 +90,24 @@ namespace UE::Learning::Action
 		TArrayView<const float> PriorProbabilities;
 	};
 
+	struct FSchemaNamedDiscreteExclusiveParameters
+	{
+		// Names of the discrete actions.
+		TArrayView<const FName> ElementNames;
+
+		// Prior Probabilities of the discrete actions.
+		TArrayView<const float> PriorProbabilities;
+	};
+
+	struct FSchemaNamedDiscreteInclusiveParameters
+	{
+		// Names of the discrete actions.
+		TArrayView<const FName> ElementNames;
+
+		// Prior Probabilities of the discrete actions.
+		TArrayView<const float> PriorProbabilities;
+	};
+
 	struct FSchemaAndParameters
 	{
 		// Names of the sub-actions.
@@ -163,6 +187,8 @@ namespace UE::Learning::Action
 		FSchemaElement CreateContinuous(const FSchemaContinuousParameters Parameters, const FName Tag = NAME_None);
 		FSchemaElement CreateDiscreteExclusive(const FSchemaDiscreteExclusiveParameters Parameters, const FName Tag = NAME_None);
 		FSchemaElement CreateDiscreteInclusive(const FSchemaDiscreteInclusiveParameters Parameters, const FName Tag = NAME_None);
+		FSchemaElement CreateNamedDiscreteExclusive(const FSchemaNamedDiscreteExclusiveParameters Parameters, const FName Tag = NAME_None);
+		FSchemaElement CreateNamedDiscreteInclusive(const FSchemaNamedDiscreteInclusiveParameters Parameters, const FName Tag = NAME_None);
 		FSchemaElement CreateAnd(const FSchemaAndParameters Parameters, const FName Tag = NAME_None);
 		FSchemaElement CreateOrExclusive(const FSchemaOrExclusiveParameters Parameters, const FName Tag = NAME_None);
 		FSchemaElement CreateOrInclusive(const FSchemaOrInclusiveParameters Parameters, const FName Tag = NAME_None);
@@ -172,6 +198,8 @@ namespace UE::Learning::Action
 		FSchemaContinuousParameters GetContinuous(const FSchemaElement Element) const;
 		FSchemaDiscreteExclusiveParameters GetDiscreteExclusive(const FSchemaElement Element) const;
 		FSchemaDiscreteInclusiveParameters GetDiscreteInclusive(const FSchemaElement Element) const;
+		FSchemaNamedDiscreteExclusiveParameters GetNamedDiscreteExclusive(const FSchemaElement Element) const;
+		FSchemaNamedDiscreteInclusiveParameters GetNamedDiscreteInclusive(const FSchemaElement Element) const;
 		FSchemaAndParameters GetAnd(const FSchemaElement Element) const;
 		FSchemaOrExclusiveParameters GetOrExclusive(const FSchemaElement Element) const;
 		FSchemaOrInclusiveParameters GetOrInclusive(const FSchemaElement Element) const;
@@ -231,6 +259,20 @@ namespace UE::Learning::Action
 			int32 PriorProbabilitiesOffset = INDEX_NONE;
 		};
 
+		struct FNamedDiscreteExclusiveData
+		{
+			int32 Num = INDEX_NONE;
+			int32 PriorProbabilitiesOffset = INDEX_NONE;
+			int32 ElementsOffset = INDEX_NONE;
+		};
+
+		struct FNamedDiscreteInclusiveData
+		{
+			int32 Num = INDEX_NONE;
+			int32 PriorProbabilitiesOffset = INDEX_NONE;
+			int32 ElementsOffset = INDEX_NONE;
+		};
+
 		struct FAndData
 		{
 			int32 Num = INDEX_NONE;
@@ -280,6 +322,8 @@ namespace UE::Learning::Action
 		TArray<FContinuousData> ContinuousData;
 		TArray<FDiscreteExclusiveData> DiscreteExclusiveData;
 		TArray<FDiscreteInclusiveData> DiscreteInclusiveData;
+		TArray<FNamedDiscreteExclusiveData> NamedDiscreteExclusiveData;
+		TArray<FNamedDiscreteInclusiveData> NamedDiscreteInclusiveData;
 		TArray<FAndData> AndData;
 		TArray<FOrExclusiveData> OrExclusiveData;
 		TArray<FOrInclusiveData> OrInclusiveData;
@@ -320,6 +364,18 @@ namespace UE::Learning::Action
 		TArrayView<const int32> DiscreteIndices;
 	};
 
+	struct FObjectNamedDiscreteExclusiveParameters
+	{
+		// Named Exclusive discrete action name
+		FName ElementName = NAME_None;
+	};
+
+	struct FObjectNamedDiscreteInclusiveParameters
+	{
+		// Named Inclusive discrete action name
+		TArrayView<const FName> ElementNames;
+	};
+
 	struct FObjectAndParameters
 	{
 		// Names of the sub-actions.
@@ -372,6 +428,8 @@ namespace UE::Learning::Action
 		FObjectElement CreateContinuous(const FObjectContinuousParameters Parameters, const FName Tag = NAME_None);
 		FObjectElement CreateDiscreteExclusive(const FObjectDiscreteExclusiveParameters Parameters, const FName Tag = NAME_None);
 		FObjectElement CreateDiscreteInclusive(const FObjectDiscreteInclusiveParameters Parameters, const FName Tag = NAME_None);
+		FObjectElement CreateNamedDiscreteExclusive(const FObjectNamedDiscreteExclusiveParameters Parameters, const FName Tag = NAME_None);
+		FObjectElement CreateNamedDiscreteInclusive(const FObjectNamedDiscreteInclusiveParameters Parameters, const FName Tag = NAME_None);
 		FObjectElement CreateAnd(const FObjectAndParameters Parameters, const FName Tag = NAME_None);
 		FObjectElement CreateOrExclusive(const FObjectOrExclusiveParameters Parameters, const FName Tag = NAME_None);
 		FObjectElement CreateOrInclusive(const FObjectOrInclusiveParameters Parameters, const FName Tag = NAME_None);
@@ -381,6 +439,8 @@ namespace UE::Learning::Action
 		FObjectContinuousParameters GetContinuous(const FObjectElement Element) const;
 		FObjectDiscreteExclusiveParameters GetDiscreteExclusive(const FObjectElement Element) const;
 		FObjectDiscreteInclusiveParameters GetDiscreteInclusive(const FObjectElement Element) const;
+		FObjectNamedDiscreteExclusiveParameters GetNamedDiscreteExclusive(const FObjectElement Element) const;
+		FObjectNamedDiscreteInclusiveParameters GetNamedDiscreteInclusive(const FObjectElement Element) const;
 		FObjectAndParameters GetAnd(const FObjectElement Element) const;
 		FObjectOrExclusiveParameters GetOrExclusive(const FObjectElement Element) const;
 		FObjectOrInclusiveParameters GetOrInclusive(const FObjectElement Element) const;
@@ -458,6 +518,18 @@ namespace UE::Learning::Action
 		TArrayView<const int32> MaskedIndices;
 	};
 
+	struct FModifierNamedDiscreteExclusiveParameters
+	{
+		// Disallowed actions
+		TArrayView<const FName> MaskedElementNames;
+	};
+
+	struct FModifierNamedDiscreteInclusiveParameters
+	{
+		// Disallowed actions
+		TArrayView<const FName> MaskedElementNames;
+	};
+
 	struct FModifierAndParameters
 	{
 		// Names of the modified sub-actions.
@@ -516,6 +588,8 @@ namespace UE::Learning::Action
 		FModifierElement CreateContinuous(const FModifierContinuousParameters Parameters, const FName Tag = NAME_None);
 		FModifierElement CreateDiscreteExclusive(const FModifierDiscreteExclusiveParameters Parameters, const FName Tag = NAME_None);
 		FModifierElement CreateDiscreteInclusive(const FModifierDiscreteInclusiveParameters Parameters, const FName Tag = NAME_None);
+		FModifierElement CreateNamedDiscreteExclusive(const FModifierNamedDiscreteExclusiveParameters Parameters, const FName Tag = NAME_None);
+		FModifierElement CreateNamedDiscreteInclusive(const FModifierNamedDiscreteInclusiveParameters Parameters, const FName Tag = NAME_None);
 		FModifierElement CreateAnd(const FModifierAndParameters Parameters, const FName Tag = NAME_None);
 		FModifierElement CreateOrExclusive(const FModifierOrExclusiveParameters Parameters, const FName Tag = NAME_None);
 		FModifierElement CreateOrInclusive(const FModifierOrInclusiveParameters Parameters, const FName Tag = NAME_None);
@@ -525,6 +599,8 @@ namespace UE::Learning::Action
 		FModifierContinuousParameters GetContinuous(const FModifierElement Element) const;
 		FModifierDiscreteExclusiveParameters GetDiscreteExclusive(const FModifierElement Element) const;
 		FModifierDiscreteInclusiveParameters GetDiscreteInclusive(const FModifierElement Element) const;
+		FModifierNamedDiscreteExclusiveParameters GetNamedDiscreteExclusive(const FModifierElement Element) const;
+		FModifierNamedDiscreteInclusiveParameters GetNamedDiscreteInclusive(const FModifierElement Element) const;
 		FModifierAndParameters GetAnd(const FModifierElement Element) const;
 		FModifierOrExclusiveParameters GetOrExclusive(const FModifierElement Element) const;
 		FModifierOrInclusiveParameters GetOrInclusive(const FModifierElement Element) const;
diff --git a/Engine/Plugins/Experimental/LearningAgents/Source/Learning/Public/LearningObservation.h b/Engine/Plugins/Experimental/LearningAgents/Source/Learning/Public/LearningObservation.h
--- a/Engine/Plugins/Experimental/LearningAgents/Source/Learning/Public/LearningObservation.h
+++ b/Engine/Plugins/Experimental/LearningAgents/Source/Learning/Public/LearningObservation.h
@@ -27,23 +27,35 @@ namespace UE::Learning::Observation
 		// Vector of continuous float observations
 		Continuous	= 1,	
 
+		// Set of exclusive discrete actions
+		DiscreteExclusive = 2,
+
+		// Set of inclusive discrete actions
+		DiscreteInclusive = 3,
+
+		// Set of named exclusive discrete actions
+		NamedDiscreteExclusive = 4,
+
+		// Set of named inclusive discrete actions
+		NamedDiscreteInclusive = 5,
+
 		// Combination of multiple observations
-		And			= 2,
+		And			= 6,
 
 		// Exclusive choice from a set of observations
-		OrExclusive = 3,
+		OrExclusive = 7,
 
 		// Inclusive choice from a set of observations
-		OrInclusive = 4,
+		OrInclusive = 8,
 
 		// Fixed sized array of observations
-		Array		= 5,
+		Array		= 9,
 
 		// Variable sized, unordered set of observations
-		Set			= 6,
+		Set			= 10,
 
 		// Encoding of another observation 
-		Encoding	= 7,
+		Encoding	= 11,
 	};
 
 	/**
@@ -63,6 +75,30 @@ namespace UE::Learning::Observation
 		float Scale = 1.0f;
 	};
 
+	struct FSchemaDiscreteExclusiveParameters
+	{
+		// Number of values in the discrete observation
+		int32 Num = 0;
+	};
+
+	struct FSchemaDiscreteInclusiveParameters
+	{
+		// Number of values in the discrete observation
+		int32 Num = 0;
+	};
+
+	struct FSchemaNamedDiscreteExclusiveParameters
+	{
+		// Names of the discrete observations.
+		TArrayView<const FName> ElementNames;
+	};
+
+	struct FSchemaNamedDiscreteInclusiveParameters
+	{
+		// Names of the discrete observations.
+		TArrayView<const FName> ElementNames;
+	};
+
 	struct FSchemaAndParameters
 	{
 		// Names of the sub-observations.
@@ -164,6 +200,10 @@ namespace UE::Learning::Observation
 	{
 		FSchemaElement CreateNull(const FName Tag = NAME_None);
 		FSchemaElement CreateContinuous(const FSchemaContinuousParameters Parameters, const FName Tag = NAME_None);
+		FSchemaElement CreateDiscreteExclusive(const FSchemaDiscreteExclusiveParameters Parameters, const FName Tag = NAME_None);
+		FSchemaElement CreateDiscreteInclusive(const FSchemaDiscreteInclusiveParameters Parameters, const FName Tag = NAME_None);
+		FSchemaElement CreateNamedDiscreteExclusive(const FSchemaNamedDiscreteExclusiveParameters Parameters, const FName Tag = NAME_None);
+		FSchemaElement CreateNamedDiscreteInclusive(const FSchemaNamedDiscreteInclusiveParameters Parameters, const FName Tag = NAME_None);
 		FSchemaElement CreateAnd(const FSchemaAndParameters Parameters, const FName Tag = NAME_None);
 		FSchemaElement CreateOrExclusive(const FSchemaOrExclusiveParameters Parameters, const FName Tag = NAME_None);
 		FSchemaElement CreateOrInclusive(const FSchemaOrInclusiveParameters Parameters, const FName Tag = NAME_None);
@@ -172,6 +212,10 @@ namespace UE::Learning::Observation
 		FSchemaElement CreateEncoding(const FSchemaEncodingParameters Parameters, const FName Tag = NAME_None);
 
 		FSchemaContinuousParameters GetContinuous(const FSchemaElement Element) const;
+		FSchemaDiscreteExclusiveParameters GetDiscreteExclusive(const FSchemaElement Element) const;
+		FSchemaDiscreteInclusiveParameters GetDiscreteInclusive(const FSchemaElement Element) const;
+		FSchemaNamedDiscreteExclusiveParameters GetNamedDiscreteExclusive(const FSchemaElement Element) const;
+		FSchemaNamedDiscreteInclusiveParameters GetNamedDiscreteInclusive(const FSchemaElement Element) const;
 		FSchemaAndParameters GetAnd(const FSchemaElement Element) const;
 		FSchemaOrExclusiveParameters GetOrExclusive(const FSchemaElement Element) const;
 		FSchemaOrInclusiveParameters GetOrInclusive(const FSchemaElement Element) const;
@@ -214,6 +258,28 @@ namespace UE::Learning::Observation
 			float Scale = 0.0f;
 		};
 
+		struct FDiscreteExclusiveData
+		{
+			int32 Num = INDEX_NONE;
+		};
+
+		struct FDiscreteInclusiveData
+		{
+			int32 Num = INDEX_NONE;
+		};
+
+		struct FNamedDiscreteExclusiveData
+		{
+			int32 Num = INDEX_NONE;
+			int32 ElementsOffset = INDEX_NONE;
+		};
+
+		struct FNamedDiscreteInclusiveData
+		{
+			int32 Num = INDEX_NONE;
+			int32 ElementsOffset = INDEX_NONE;
+		};
+
 		struct FAndData
 		{
 			int32 Num = INDEX_NONE;
@@ -270,6 +336,10 @@ namespace UE::Learning::Observation
 		
 		/** These are indexed based on the Type using TypeDataIndices */
 		TArray<FContinuousData> ContinuousData;
+		TArray<FDiscreteExclusiveData> DiscreteExclusiveData;
+		TArray<FDiscreteInclusiveData> DiscreteInclusiveData;
+		TArray<FNamedDiscreteExclusiveData> NamedDiscreteExclusiveData;
+		TArray<FNamedDiscreteInclusiveData> NamedDiscreteInclusiveData;
 		TArray<FAndData> AndData;
 		TArray<FOrExclusiveData> OrExclusiveData;
 		TArray<FOrInclusiveData> OrInclusiveData;
@@ -296,6 +366,30 @@ namespace UE::Learning::Observation
 		TArrayView<const float> Values;
 	};
 
+	struct FObjectDiscreteExclusiveParameters
+	{
+		// Index of the chosen observation.
+		int32 DiscreteIndex;
+	};
+
+	struct FObjectDiscreteInclusiveParameters
+	{
+		// Indices of the chosen observations.
+		TArrayView<const int32> DiscreteIndices;
+	};
+
+	struct FObjectNamedDiscreteExclusiveParameters
+	{
+		// Name of the chosen observation.
+		FName ElementName;
+	};
+
+	struct FObjectNamedDiscreteInclusiveParameters
+	{
+		// Names of the chosen observations.
+		TArrayView<const FName> ElementNames;
+	};
+
 	struct FObjectAndParameters
 	{
 		// Names of the sub-observations.
@@ -352,6 +446,10 @@ namespace UE::Learning::Observation
 	{
 		FObjectElement CreateNull(const FName Tag = NAME_None);
 		FObjectElement CreateContinuous(const FObjectContinuousParameters Parameters, const FName Tag = NAME_None);
+		FObjectElement CreateDiscreteExclusive(const FObjectDiscreteExclusiveParameters Parameters, const FName Tag = NAME_None);
+		FObjectElement CreateDiscreteInclusive(const FObjectDiscreteInclusiveParameters Parameters, const FName Tag = NAME_None);
+		FObjectElement CreateNamedDiscreteExclusive(const FObjectNamedDiscreteExclusiveParameters Parameters, const FName Tag = NAME_None);
+		FObjectElement CreateNamedDiscreteInclusive(const FObjectNamedDiscreteInclusiveParameters Parameters, const FName Tag = NAME_None);
 		FObjectElement CreateAnd(const FObjectAndParameters Parameters, const FName Tag = NAME_None);
 		FObjectElement CreateOrExclusive(const FObjectOrExclusiveParameters Parameters, const FName Tag = NAME_None);
 		FObjectElement CreateOrInclusive(const FObjectOrInclusiveParameters Parameters, const FName Tag = NAME_None);
@@ -360,6 +458,10 @@ namespace UE::Learning::Observation
 		FObjectElement CreateEncoding(const FObjectEncodingParameters Parameters, const FName Tag = NAME_None);
 
 		FObjectContinuousParameters GetContinuous(const FObjectElement Element) const;
+		FObjectDiscreteExclusiveParameters GetDiscreteExclusive(const FObjectElement Element) const;
+		FObjectDiscreteInclusiveParameters GetDiscreteInclusive(const FObjectElement Element) const;
+		FObjectNamedDiscreteExclusiveParameters GetNamedDiscreteExclusive(const FObjectElement Element) const;
+		FObjectNamedDiscreteInclusiveParameters GetNamedDiscreteInclusive(const FObjectElement Element) const;
 		FObjectAndParameters GetAnd(const FObjectElement Element) const;
 		FObjectOrExclusiveParameters GetOrExclusive(const FObjectElement Element) const;
 		FObjectOrInclusiveParameters GetOrInclusive(const FObjectElement Element) const;
@@ -396,10 +498,13 @@ namespace UE::Learning::Observation
 		TArray<FName> Tags;
 		TArray<int32> ContinuousDataOffsets;
 		TArray<int32> ContinuousDataNums;
+		TArray<int32> DiscreteDataOffsets;
+		TArray<int32> DiscreteDataNums;
 		TArray<int32> SubElementDataOffsets;
 		TArray<int32> SubElementDataNums;
 
 		TArray<float> ContinuousValues;
+		TArray<int32> DiscreteValues;
 		TArray<FName> SubElementNames;
 		TArray<FObjectElement> SubElementObjects;
 	};
diff --git a/Engine/Plugins/Experimental/LearningAgents/Source/LearningAgents/Private/LearningAgentsActions.cpp b/Engine/Plugins/Experimental/LearningAgents/Source/LearningAgents/Private/LearningAgentsActions.cpp
--- a/Engine/Plugins/Experimental/LearningAgents/Source/LearningAgents/Private/LearningAgentsActions.cpp
+++ b/Engine/Plugins/Experimental/LearningAgents/Source/LearningAgents/Private/LearningAgentsActions.cpp
@@ -42,6 +42,13 @@ uint32 GetTypeHash(const FLearningAgentsActionModifierElement& Element)
 
 namespace UE::Learning::Agents::Action::Private
 {
+	// We declare this here instead of in LearningArray.h because FName is part of the UObject system and right now the 
+	// Learning core module does not have a dependency on the UObject part of UE which is a nice thing to preserve.	
+	static inline FString FormatNames(const TLearningArrayView<1, const FName> Array, const int32 MaxItemNum = 16)
+	{
+		return Array::Format<FName>(Array, [](const FName& Value) { return FString::Printf(TEXT("%s"), *Value.ToString()); }, MaxItemNum);
+	}
+
 	static inline bool ContainsDuplicates(const TArrayView<const int32> Indices)
 	{
 		TSet<int32, DefaultKeyFuncs<int32>, TInlineSetAllocator<32>> IndicesSet;
@@ -64,6 +71,8 @@ namespace UE::Learning::Agents::Action::Private
 		case Learning::Action::EType::Continuous: return TEXT("Continuous");
 		case Learning::Action::EType::DiscreteExclusive: return TEXT("DiscreteExclusive");
 		case Learning::Action::EType::DiscreteInclusive: return TEXT("DiscreteInclusive");
+		case Learning::Action::EType::NamedDiscreteExclusive: return TEXT("NamedDiscreteExclusive");
+		case Learning::Action::EType::NamedDiscreteInclusive: return TEXT("NamedDiscreteInclusive");
 		case Learning::Action::EType::And: return TEXT("Struct");
 		case Learning::Action::EType::OrExclusive: return TEXT("ExclusiveUnion");
 		case Learning::Action::EType::OrInclusive: return TEXT("InclusiveUnion");
@@ -189,6 +198,55 @@ namespace UE::Learning::Agents::Action::Private
 			return true;
 		}
 
+		case Learning::Action::EType::NamedDiscreteExclusive:
+		{
+			const Learning::Action::FSchemaNamedDiscreteExclusiveParameters SchemaParameters = Schema.GetNamedDiscreteExclusive(SchemaElement);
+			const Learning::Action::FObjectNamedDiscreteExclusiveParameters ObjectParameters = Object.GetNamedDiscreteExclusive(ObjectElement);
+			check(SchemaParameters.ElementNames.Num() == SchemaParameters.ElementNames.Num());
+
+			const int32 SchemaSubElementIdx = SchemaParameters.ElementNames.Find(ObjectParameters.ElementName);
+
+			if (SchemaSubElementIdx == INDEX_NONE)
+			{
+				UE_LOG(LogLearning, Error, TEXT("ValidateActionObjectMatchesSchema: Action '%s' Schema does not include '%s' action."),
+					*ActionSchemaElementTag.ToString(),
+					*ObjectParameters.ElementName.ToString());
+				return false;
+			}
+
+			return true;
+		}
+
+		case Learning::Action::EType::NamedDiscreteInclusive:
+		{
+			const Learning::Action::FSchemaNamedDiscreteInclusiveParameters SchemaParameters = Schema.GetNamedDiscreteInclusive(SchemaElement);
+			const Learning::Action::FObjectNamedDiscreteInclusiveParameters ObjectParameters = Object.GetNamedDiscreteInclusive(ObjectElement);
+
+			if (ObjectParameters.ElementNames.Num() > SchemaParameters.ElementNames.Num())
+			{
+				UE_LOG(LogLearning, Error, TEXT("ValidateActionObjectMatchesSchema: Action '%s' too many sub-actions provided. Expected at most '%i', got '%i'."),
+					*ActionSchemaElementTag.ToString(),
+					SchemaParameters.ElementNames.Num(),
+					ObjectParameters.ElementNames.Num());
+				return false;
+			}
+
+			for (int32 ObjectSubElementIdx = 0; ObjectSubElementIdx < ObjectParameters.ElementNames.Num(); ObjectSubElementIdx++)
+			{
+				const int32 SchemaSubElementIdx = SchemaParameters.ElementNames.Find(ObjectParameters.ElementNames[ObjectSubElementIdx]);
+
+				if (SchemaSubElementIdx == INDEX_NONE)
+				{
+					UE_LOG(LogLearning, Error, TEXT("ValidateActionObjectMatchesSchema: Action '%s' Schema does not include '%s' Action."),
+						*ActionSchemaElementTag.ToString(),
+						*ObjectParameters.ElementNames[ObjectSubElementIdx].ToString());
+					return false;
+				}
+			}
+
+			return true;
+		}
+
 		case Learning::Action::EType::And:
 		{
 			const Learning::Action::FSchemaAndParameters SchemaParameters = Schema.GetAnd(SchemaElement);
@@ -486,6 +544,54 @@ namespace UE::Learning::Agents::Action::Private
 			return true;
 		}
 
+		case Learning::Action::EType::NamedDiscreteExclusive:
+		{
+			const TArrayView<const FName> SchemaNames = Schema.GetNamedDiscreteExclusive(SchemaElement).ElementNames;
+			const TArrayView<const FName> ModifierMaskedNames = Modifier.GetNamedDiscreteExclusive(ModifierElement).MaskedElementNames;
+
+			// For exclusive action we must have one less than the schema size or it means all are masked.
+			if (ModifierMaskedNames.Num() > SchemaNames.Num() - 1)
+			{
+				UE_LOG(LogLearning, Error, TEXT("ValidateActionModifierMatchesSchema: Action Modifier '%s' too many indices provided. Expected at most '%i', got '%i'."),
+					*ActionSchemaElementTag.ToString(),
+					SchemaNames.Num() - 1,
+					ModifierMaskedNames.Num());
+				return false;
+			}
+
+			for (int32 MaskIdx = 0; MaskIdx < ModifierMaskedNames.Num(); MaskIdx++)
+			{
+				if (!SchemaNames.Contains(ModifierMaskedNames[MaskIdx]))
+				{
+					UE_LOG(LogLearning, Error, TEXT("ValidateActionModifierMatchesSchema: Action Modifier '%s' name '%s' not found in Schema."),
+						*ActionSchemaElementTag.ToString(),
+						*ModifierMaskedNames[MaskIdx].ToString());
+					return false;
+				}
+			}
+
+			return true;
+		}
+
+		case Learning::Action::EType::NamedDiscreteInclusive:
+		{
+			const TArrayView<const FName> SchemaNames = Schema.GetNamedDiscreteInclusive(SchemaElement).ElementNames;
+			const TArrayView<const FName> ModifierMaskedNames = Modifier.GetNamedDiscreteInclusive(ModifierElement).MaskedElementNames;
+
+			for (int32 MaskIdx = 0; MaskIdx < ModifierMaskedNames.Num(); MaskIdx++)
+			{
+				if (!SchemaNames.Contains(ModifierMaskedNames[MaskIdx]))
+				{
+					UE_LOG(LogLearning, Error, TEXT("ValidateActionModifierMatchesSchema: Action Modifier '%s' name '%s' not found in Schema."),
+						*ActionSchemaElementTag.ToString(),
+						*ModifierMaskedNames[MaskIdx].ToString());
+					return false;
+				}
+			}
+
+			return true;
+		}
+
 		case Learning::Action::EType::And:
 		{
 			const Learning::Action::FSchemaAndParameters SchemaParameters = Schema.GetAnd(SchemaElement);
@@ -740,6 +846,22 @@ namespace UE::Learning::Agents::Action::Private
 			return;
 		}
 
+		case UE::Learning::Action::EType::NamedDiscreteExclusive:
+		{
+			const UE::Learning::Action::FObjectNamedDiscreteExclusiveParameters Parameters = Object.GetNamedDiscreteExclusive(ObjectElement);
+
+			UE_LOG(LogLearning, Display, TEXT("%s%s \"%s\" (%s) %s"), *Indentation, *Prefix, *Tag.ToString(), GetActionTypeString(Type), *Parameters.ElementName.ToString());
+			return;
+		}
+
+		case UE::Learning::Action::EType::NamedDiscreteInclusive:
+		{
+			const UE::Learning::Action::FObjectNamedDiscreteInclusiveParameters Parameters = Object.GetNamedDiscreteInclusive(ObjectElement);
+
+			UE_LOG(LogLearning, Display, TEXT("%s%s \"%s\" (%s) %s"), *Indentation, *Prefix, *Tag.ToString(), GetActionTypeString(Type), *UE::Learning::Agents::Action::Private::FormatNames(Parameters.ElementNames));
+			return;
+		}
+
 		case UE::Learning::Action::EType::And:
 		{
 			const UE::Learning::Action::FObjectAndParameters Parameters = Object.GetAnd(ObjectElement);
@@ -1043,17 +1165,27 @@ FLearningAgentsActionSchemaElement ULearningAgentsActions::SpecifyNamedExclusive
 
 FLearningAgentsActionSchemaElement ULearningAgentsActions::SpecifyNamedExclusiveDiscreteActionFromArrayViews(ULearningAgentsActionSchema* Schema, const TArrayView<const FName> Names, const TArrayView<const float> PriorProbabilities, const FName Tag)
 {
-	const int32 NameNum = Names.Num();
+	if (!Schema)
+	{
+		UE_LOG(LogLearning, Error, TEXT("SpecifyNamedExclusiveDiscreteActionFromArrayViews: Schema is nullptr."));
+		return FLearningAgentsActionSchemaElement();
+	}
 
-	TArray<FLearningAgentsActionSchemaElement, TInlineAllocator<16>> Elements;
-	Elements.SetNumUninitialized(NameNum);
+	if (Names.Num() == 0)
+	{
+		UE_LOG(LogLearning, Error, TEXT("SpecifyNamedExclusiveDiscreteActionFromArrayViews: Specifying zero-sized Named Exclusive Discrete Action."));
+		return FLearningAgentsActionSchemaElement();
+	}
 
-	for (int32 NameIdx = 0; NameIdx < NameNum; NameIdx++)
+	TArray<float, TInlineAllocator<16>> NormalizedPriorProbabilities;
+	NormalizedPriorProbabilities.Init(1.0f / Names.Num(), Names.Num());
+	for (int32 Idx = 0; Idx < PriorProbabilities.Num(); Idx++)
 	{
-		Elements[NameIdx] = SpecifyNullAction(Schema);
+		NormalizedPriorProbabilities[Idx] = PriorProbabilities[Idx];
 	}
+	UE::Learning::Agents::Action::Private::NormalizeProbabilitiesExclusive(NormalizedPriorProbabilities);
 
-	return SpecifyExclusiveUnionActionFromArrayViews(Schema, Names, Elements, PriorProbabilities, Tag);
+	return { Schema->ActionSchema.CreateNamedDiscreteExclusive({ Names, MakeArrayView(NormalizedPriorProbabilities) }, Tag) };
 }
 
 FLearningAgentsActionSchemaElement ULearningAgentsActions::SpecifyInclusiveDiscreteAction(ULearningAgentsActionSchema* Schema, const int32 Size, const TArray<float>& PriorProbabilities, const FName Tag)
@@ -1136,17 +1268,26 @@ FLearningAgentsActionSchemaElement ULearningAgentsActions::SpecifyNamedInclusive
 
 FLearningAgentsActionSchemaElement ULearningAgentsActions::SpecifyNamedInclusiveDiscreteActionFromArrayViews(ULearningAgentsActionSchema* Schema, const TArrayView<const FName> Names, const TArrayView<const float> PriorProbabilities, const FName Tag)
 {
-	const int32 NameNum = Names.Num();
+	if (!Schema)
+	{
+		UE_LOG(LogLearning, Error, TEXT("SpecifyNamedInclusiveDiscreteActionFromArrayViews: Schema is nullptr."));
+		return FLearningAgentsActionSchemaElement();
+	}
 
-	TArray<FLearningAgentsActionSchemaElement, TInlineAllocator<16>> Elements;
-	Elements.SetNumUninitialized(NameNum);
+	if (Names.Num() == 0)
+	{
+		UE_LOG(LogLearning, Warning, TEXT("SpecifyNamedInclusiveDiscreteActionFromArrayViews: Specifying zero-sized Named Inclusive Discrete Action."));
+	}
 
-	for (int32 NameIdx = 0; NameIdx < NameNum; NameIdx++)
+	TArray<float, TInlineAllocator<16>> NormalizedPriorProbabilities;
+	NormalizedPriorProbabilities.Init(1.0f / Names.Num(), Names.Num());
+	for (int32 Idx = 0; Idx < PriorProbabilities.Num(); Idx++)
 	{
-		Elements[NameIdx] = SpecifyNullAction(Schema);
+		NormalizedPriorProbabilities[Idx] = PriorProbabilities[Idx];
 	}
+	UE::Learning::Agents::Action::Private::NormalizeProbabilitiesInclusive(NormalizedPriorProbabilities);
 
-	return SpecifyInclusiveUnionActionFromArrayViews(Schema, Names, Elements, PriorProbabilities, Tag);
+	return { Schema->ActionSchema.CreateNamedDiscreteInclusive({ Names, MakeArrayView(NormalizedPriorProbabilities) }, Tag) };
 }
 
 FLearningAgentsActionSchemaElement ULearningAgentsActions::SpecifyStructAction(ULearningAgentsActionSchema* Schema, const TMap<FName, FLearningAgentsActionSchemaElement>& Elements, const FName Tag)
@@ -1738,7 +1879,13 @@ FLearningAgentsActionObjectElement ULearningAgentsActions::MakeExclusiveDiscrete
 
 FLearningAgentsActionObjectElement ULearningAgentsActions::MakeNamedExclusiveDiscreteAction(ULearningAgentsActionObject* Object, const FName Name, const FName Tag)
 {
-	return MakeExclusiveUnionAction(Object, Name, MakeNullAction(Object), Tag);
+	if (!Object)
+	{
+		UE_LOG(LogLearning, Error, TEXT("MakeNamedExclusiveDiscreteAction: Object is nullptr."));
+		return FLearningAgentsActionObjectElement();
+	}
+
+	return { Object->ActionObject.CreateNamedDiscreteExclusive({ Name }, Tag) };
 }
 
 FLearningAgentsActionObjectElement ULearningAgentsActions::MakeInclusiveDiscreteAction(ULearningAgentsActionObject* Object, const TArray<int32>& Indices, const FName Tag)
@@ -1781,17 +1928,19 @@ FLearningAgentsActionObjectElement ULearningAgentsActions::MakeNamedInclusiveDis
 
 FLearningAgentsActionObjectElement ULearningAgentsActions::MakeNamedInclusiveDiscreteActionFromArrayView(ULearningAgentsActionObject* Object, const TArrayView<const FName> Names, const FName Tag)
 {
-	const int32 NameNum = Names.Num();
-
-	TArray<FLearningAgentsActionObjectElement, TInlineAllocator<32>> Elements;
-	Elements.SetNumUninitialized(NameNum);
+	if (!Object)
+	{
+		UE_LOG(LogLearning, Error, TEXT("MakeNamedInclusiveDiscreteActionFromArrayView: Object is nullptr."));
+		return FLearningAgentsActionObjectElement();
+	}
 
-	for (int32 NameIdx = 0; NameIdx < NameNum; NameIdx++)
+	if (UE::Learning::Agents::Action::Private::ContainsDuplicates(Names))
 	{
-		Elements[NameIdx] = MakeNullAction(Object);
+		UE_LOG(LogLearning, Error, TEXT("MakeNamedInclusiveDiscreteActionFromArrayView: Names contain duplicates."));
+		return FLearningAgentsActionObjectElement();
 	}
 
-	return MakeInclusiveUnionActionFromArrayViews(Object, Names, Elements, Tag);
+	return { Object->ActionObject.CreateNamedDiscreteInclusive({ Names }, Tag) };
 }
 
 FLearningAgentsActionObjectElement ULearningAgentsActions::MakeStructAction(ULearningAgentsActionObject* Object, const TMap<FName, FLearningAgentsActionObjectElement>& Elements, const FName Tag)
@@ -2222,6 +2371,12 @@ FLearningAgentsActionModifierElement ULearningAgentsActions::MakeExclusiveDiscre
 		return FLearningAgentsActionModifierElement();
 	}
 
+	if (UE::Learning::Agents::Action::Private::ContainsDuplicates(MaskedIndices))
+	{
+		UE_LOG(LogLearning, Error, TEXT("MakeExclusiveDiscreteActionModifierFromArrayView: MaskedIndices contain duplicates."));
+		return FLearningAgentsActionModifierElement();
+	}	
+
 	return { Modifier->ActionModifier.CreateDiscreteExclusive({ MaskedIndices }, Tag) };
 }
 
@@ -2232,7 +2387,19 @@ FLearningAgentsActionModifierElement ULearningAgentsActions::MakeNamedExclusiveD
 
 FLearningAgentsActionModifierElement ULearningAgentsActions::MakeNamedExclusiveDiscreteActionModifierFromArrayView(ULearningAgentsActionModifier* Modifier, const TArrayView<const FName> MaskedNames, const FName Tag)
 {
-	return MakeExclusiveUnionActionModifierFromArrayViews(Modifier, {}, {}, MaskedNames, Tag);
+	if (!Modifier)
+	{
+		UE_LOG(LogLearning, Error, TEXT("MakeNamedExclusiveDiscreteActionModifierFromArrayView: Modifier is nullptr."));
+		return FLearningAgentsActionModifierElement();
+	}
+
+	if (UE::Learning::Agents::Action::Private::ContainsDuplicates(MaskedNames))
+	{
+		UE_LOG(LogLearning, Error, TEXT("MakeNamedExclusiveDiscreteActionModifierFromArrayView: MaskedNames contain duplicates."));
+		return FLearningAgentsActionModifierElement();
+	}
+
+	return { Modifier->ActionModifier.CreateNamedDiscreteExclusive({ MaskedNames }, Tag) };
 }
 
 FLearningAgentsActionModifierElement ULearningAgentsActions::MakeInclusiveDiscreteActionModifier(ULearningAgentsActionModifier* Modifier, const TArray<int32>& MaskedIndices, const FName Tag)
@@ -2248,6 +2415,12 @@ FLearningAgentsActionModifierElement ULearningAgentsActions::MakeInclusiveDiscre
 		return FLearningAgentsActionModifierElement();
 	}
 
+	if (UE::Learning::Agents::Action::Private::ContainsDuplicates(MaskedIndices))
+	{
+		UE_LOG(LogLearning, Error, TEXT("MakeNamedExclusiveDiscreteActionModifierFromArrayView: MaskedIndices contain duplicates."));
+		return FLearningAgentsActionModifierElement();
+	}	
+
 	return { Modifier->ActionModifier.CreateDiscreteInclusive({ MaskedIndices }, Tag) };
 }
 
@@ -2258,7 +2431,19 @@ FLearningAgentsActionModifierElement ULearningAgentsActions::MakeNamedInclusiveD
 
 FLearningAgentsActionModifierElement ULearningAgentsActions::MakeNamedInclusiveDiscreteActionModifierFromArrayView(ULearningAgentsActionModifier* Modifier, const TArrayView<const FName> MaskedNames, const FName Tag)
 {
-	return MakeInclusiveUnionActionModifierFromArrayViews(Modifier, {}, {}, MaskedNames, Tag);
+	if (!Modifier)
+	{
+		UE_LOG(LogLearning, Error, TEXT("MakeNamedInclusiveDiscreteActionModifierFromArrayView: Modifier is nullptr."));
+		return FLearningAgentsActionModifierElement();
+	}
+	
+	if (UE::Learning::Agents::Action::Private::ContainsDuplicates(MaskedNames))
+	{
+		UE_LOG(LogLearning, Error, TEXT("MakeNamedExclusiveDiscreteActionModifierFromArrayView: MaskedNames contain duplicates."));
+		return FLearningAgentsActionModifierElement();
+	}	
+
+	return { Modifier->ActionModifier.CreateNamedDiscreteInclusive({ MaskedNames }, Tag) };
 }
 
 FLearningAgentsActionModifierElement ULearningAgentsActions::MakeStructActionModifier(ULearningAgentsActionModifier* Modifier, const TMap<FName, FLearningAgentsActionModifierElement>& Elements, const FName Tag)
@@ -2921,12 +3106,37 @@ bool ULearningAgentsActions::GetNamedExclusiveDiscreteAction(
 	const FVector VisualLoggerLocation,
 	const FLinearColor VisualLoggerColor)
 {
-	FLearningAgentsActionObjectElement OutElement;
-	if (!GetExclusiveUnionAction(OutName, OutElement, Object, Element, Tag))
+	if (!Object)
 	{
+		UE_LOG(LogLearning, Error, TEXT("GetNamedExclusiveDiscreteAction: Object is nullptr."));
+		OutName = NAME_None;
 		return false;
 	}
 
+	if (!Object->ActionObject.IsValid(Element.ObjectElement))
+	{
+		UE_LOG(LogLearning, Error, TEXT("GetNamedExclusiveDiscreteAction: Invalid Action Object."));
+		OutName = NAME_None;
+		return false;
+	}
+
+	if (Object->ActionObject.GetTag(Element.ObjectElement) != Tag)
+	{
+		UE_LOG(LogLearning, Warning, TEXT("GetNamedExclusiveDiscreteAction: Action tag does not match. Action is '%s' but asked for '%s'."), *Object->ActionObject.GetTag(Element.ObjectElement).ToString(), *Tag.ToString());
+	}
+
+	if (Object->ActionObject.GetType(Element.ObjectElement) != UE::Learning::Action::EType::NamedDiscreteExclusive)
+	{
+		UE_LOG(LogLearning, Error, TEXT("GetNamedExclusiveDiscreteAction: Action '%s' type does not match. Action is '%s' but asked for '%s'."),
+			*Object->ActionObject.GetTag(Element.ObjectElement).ToString(),
+			UE::Learning::Agents::Action::Private::GetActionTypeString(Object->ActionObject.GetType(Element.ObjectElement)),
+			UE::Learning::Agents::Action::Private::GetActionTypeString(UE::Learning::Action::EType::NamedDiscreteExclusive));
+		OutName = NAME_None;
+		return false;
+	}
+
+	OutName = Object->ActionObject.GetNamedDiscreteExclusive(Element.ObjectElement).ElementName;
+
 #if UE_LEARNING_AGENTS_ENABLE_VISUAL_LOG
 	if (bVisualLoggerEnabled && VisualLoggerListener)
 	{
@@ -3082,7 +3292,37 @@ bool ULearningAgentsActions::GetInclusiveDiscreteActionToArrayView(
 
 bool ULearningAgentsActions::GetNamedInclusiveDiscreteActionNum(int32& OutNum, const ULearningAgentsActionObject* Object, const FLearningAgentsActionObjectElement Element, const FName Tag)
 {
-	return GetInclusiveUnionActionNum(OutNum, Object, Element, Tag);
+	if (!Object)
+	{
+		UE_LOG(LogLearning, Error, TEXT("GetNamedInclusiveDiscreteActionNum: Object is nullptr."));
+		OutNum = 0;
+		return false;
+	}
+
+	if (!Object->ActionObject.IsValid(Element.ObjectElement))
+	{
+		UE_LOG(LogLearning, Error, TEXT("GetNamedInclusiveDiscreteActionNum: Invalid Action Object."));
+		OutNum = 0;
+		return false;
+	}
+
+	if (Object->ActionObject.GetTag(Element.ObjectElement) != Tag)
+	{
+		UE_LOG(LogLearning, Warning, TEXT("GetNamedInclusiveDiscreteActionNum: Action tag does not match. Action is '%s' but asked for '%s'."), *Object->ActionObject.GetTag(Element.ObjectElement).ToString(), *Tag.ToString());
+	}
+
+	if (Object->ActionObject.GetType(Element.ObjectElement) != UE::Learning::Action::EType::NamedDiscreteInclusive)
+	{
+		UE_LOG(LogLearning, Error, TEXT("GetNamedInclusiveDiscreteActionNum: Action '%s' type does not match. Action is '%s' but asked for '%s'."),
+			*Object->ActionObject.GetTag(Element.ObjectElement).ToString(),
+			UE::Learning::Agents::Action::Private::GetActionTypeString(Object->ActionObject.GetType(Element.ObjectElement)),
+			UE::Learning::Agents::Action::Private::GetActionTypeString(UE::Learning::Action::EType::NamedDiscreteInclusive));
+		OutNum = 0;
+		return false;
+	}
+
+	OutNum = Object->ActionObject.GetNamedDiscreteInclusive(Element.ObjectElement).ElementNames.Num();
+	return true;
 }
 
 bool ULearningAgentsActions::GetNamedInclusiveDiscreteAction(
@@ -3114,16 +3354,6 @@ bool ULearningAgentsActions::GetNamedInclusiveDiscreteAction(
 	return true;
 }
 
-namespace UE::Learning::Array
-{
-	// We declare this here instead of in LearningArray.h because FName is part of the UObject system and right now the 
-	// Learning core module does not have a dependency on the UObject part of UE which is a nice thing to preserve.	
-	static inline FString FormatNames(const TLearningArrayView<1, const FName> Array, const int32 MaxItemNum = 16)
-	{
-		return Format<FName>(Array, [](const FName& Value) { return FString::Printf(TEXT("%s"), *Value.ToString()); }, MaxItemNum);
-	}
-}
-
 bool ULearningAgentsActions::GetNamedInclusiveDiscreteActionToArrayView(
 	TArrayView<FName> OutNames,
 	const ULearningAgentsActionObject* Object,
@@ -3135,14 +3365,48 @@ bool ULearningAgentsActions::GetNamedInclusiveDiscreteActionToArrayView(
 	const FVector VisualLoggerLocation,
 	const FLinearColor VisualLoggerColor)
 {
-	TArray<FLearningAgentsActionObjectElement, TInlineAllocator<16>> OutElements;
-	OutElements.SetNumUninitialized(OutNames.Num());
+	if (!Object)
+	{
+		UE_LOG(LogLearning, Error, TEXT("GetNamedInclusiveDiscreteActionToArrayView: Object is nullptr."));
+		UE::Learning::Array::Set<1, FName>(OutNames, NAME_None);
+		return false;
+	}
 
-	if (!GetInclusiveUnionActionToArrayViews(OutNames, OutElements, Object, Element, Tag))
+	if (!Object->ActionObject.IsValid(Element.ObjectElement))
 	{
+		UE_LOG(LogLearning, Error, TEXT("GetNamedInclusiveDiscreteActionToArrayView: Invalid Action Object."));
+		UE::Learning::Array::Set<1, FName>(OutNames, NAME_None);
 		return false;
 	}
 
+	if (Object->ActionObject.GetTag(Element.ObjectElement) != Tag)
+	{
+		UE_LOG(LogLearning, Warning, TEXT("GetNamedInclusiveDiscreteActionToArrayView: Action tag does not match. Action is '%s' but asked for '%s'."), *Object->ActionObject.GetTag(Element.ObjectElement).ToString(), *Tag.ToString());
+	}
+
+	if (Object->ActionObject.GetType(Element.ObjectElement) != UE::Learning::Action::EType::NamedDiscreteInclusive)
+	{
+		UE_LOG(LogLearning, Error, TEXT("GetNamedInclusiveDiscreteActionToArrayView: Action '%s' type does not match. Action is '%s' but asked for '%s'."),
+			*Object->ActionObject.GetTag(Element.ObjectElement).ToString(),
+			UE::Learning::Agents::Action::Private::GetActionTypeString(Object->ActionObject.GetType(Element.ObjectElement)),
+			UE::Learning::Agents::Action::Private::GetActionTypeString(UE::Learning::Action::EType::NamedDiscreteInclusive));
+		UE::Learning::Array::Set<1, FName>(OutNames, NAME_None);
+		return false;
+	}
+
+	const TArrayView<const FName> Names = Object->ActionObject.GetNamedDiscreteInclusive(Element.ObjectElement).ElementNames;
+
+	if (Names.Num() != OutNames.Num())
+	{
+		UE_LOG(LogLearning, Error, TEXT("GetNamedInclusiveDiscreteActionToArrayView: Action '%s' size does not match. Action is '%i' elements but asked for '%i'."),
+			*Object->ActionObject.GetTag(Element.ObjectElement).ToString(),
+			Names.Num(), OutNames.Num());
+		UE::Learning::Array::Set<1, FName>(OutNames, NAME_None);
+		return false;
+	}
+
+	UE::Learning::Array::Copy<1, FName>(OutNames, Names);
+
 #if UE_LEARNING_AGENTS_ENABLE_VISUAL_LOG
 	if (bVisualLoggerEnabled && VisualLoggerListener)
 	{
@@ -3154,7 +3418,7 @@ bool ULearningAgentsActions::GetNamedInclusiveDiscreteActionToArrayView(
 			*VisualLoggerListener->GetName(),
 			*Tag.ToString(),
 			VisualLoggerAgentId,
-			*UE::Learning::Array::FormatNames(OutNames, 256));
+			*UE::Learning::Agents::Action::Private::FormatNames(OutNames, 256));
 	}
 #endif
 
diff --git a/Engine/Plugins/Experimental/LearningAgents/Source/LearningAgents/Private/LearningAgentsObservations.cpp b/Engine/Plugins/Experimental/LearningAgents/Source/LearningAgents/Private/LearningAgentsObservations.cpp
--- a/Engine/Plugins/Experimental/LearningAgents/Source/LearningAgents/Private/LearningAgentsObservations.cpp
+++ b/Engine/Plugins/Experimental/LearningAgents/Source/LearningAgents/Private/LearningAgentsObservations.cpp
@@ -36,6 +36,13 @@ uint32 GetTypeHash(const FLearningAgentsObservationObjectElement& Element)
 
 namespace UE::Learning::Agents::Observation::Private
 {
+	// We declare this here instead of in LearningArray.h because FName is part of the UObject system and right now the 
+	// Learning core module does not have a dependency on the UObject part of UE which is a nice thing to preserve.	
+	static inline FString FormatNames(const TLearningArrayView<1, const FName> Array, const int32 MaxItemNum = 16)
+	{
+		return Array::Format<FName>(Array, [](const FName& Value) { return FString::Printf(TEXT("%s"), *Value.ToString()); }, MaxItemNum);
+	}
+
 	static inline bool ContainsDuplicates(const TArrayView<const int32> Indices)
 	{
 		TSet<int32, DefaultKeyFuncs<int32>, TInlineSetAllocator<32>> IndicesSet;
@@ -56,6 +63,10 @@ namespace UE::Learning::Agents::Observation::Private
 		{
 		case  Learning::Observation::EType::Null: return TEXT("Null");
 		case  Learning::Observation::EType::Continuous: return TEXT("Continuous");
+		case  Learning::Observation::EType::DiscreteExclusive: return TEXT("DiscreteExclusive");
+		case  Learning::Observation::EType::DiscreteInclusive: return TEXT("DiscreteInclusive");
+		case  Learning::Observation::EType::NamedDiscreteExclusive: return TEXT("NamedDiscreteExclusive");
+		case  Learning::Observation::EType::NamedDiscreteInclusive: return TEXT("NamedDiscreteInclusive");
 		case  Learning::Observation::EType::And: return TEXT("Struct");
 		case  Learning::Observation::EType::OrExclusive: return TEXT("ExclusiveUnion");
 		case  Learning::Observation::EType::OrInclusive: return TEXT("InclusiveUnion");
@@ -136,6 +147,101 @@ namespace UE::Learning::Agents::Observation::Private
 			return true;
 		}
 
+		case Learning::Observation::EType::DiscreteExclusive:
+		{
+			const int32 SchemaElementSize = Schema.GetDiscreteExclusive(SchemaElement).Num;
+			const int32 ObjectElementIndex = Object.GetDiscreteExclusive(ObjectElement).DiscreteIndex;
+
+			if (ObjectElementIndex < 0 || ObjectElementIndex >= SchemaElementSize)
+			{
+				UE_LOG(LogLearning, Error, TEXT("ValidateObservationObjectMatchesSchema: Observation '%s' index out of range for Schema. Expected '<%i', got '%i'."),
+					*ObservationSchemaElementTag.ToString(),
+					SchemaElementSize,
+					ObjectElementIndex);
+				return false;
+			}
+
+			return true;
+		}
+
+		case Learning::Observation::EType::DiscreteInclusive:
+		{
+			const int32 SchemaElementSize = Schema.GetDiscreteInclusive(SchemaElement).Num;
+			const TArrayView<const int32> ObjectElementIndices = Object.GetDiscreteInclusive(ObjectElement).DiscreteIndices;
+
+			if (ObjectElementIndices.Num() > SchemaElementSize)
+			{
+				UE_LOG(LogLearning, Error, TEXT("ValidateObservationObjectMatchesSchema: Observation '%s' too many indices provided. Expected at most '%i', got '%i'."),
+					*ObservationSchemaElementTag.ToString(),
+					SchemaElementSize,
+					ObjectElementIndices.Num());
+				return false;
+			}
+
+			for (int32 SubElementIdx = 0; SubElementIdx < ObjectElementIndices.Num(); SubElementIdx++)
+			{
+				if (ObjectElementIndices[SubElementIdx] < 0 || ObjectElementIndices[SubElementIdx] >= SchemaElementSize)
+				{
+					UE_LOG(LogLearning, Error, TEXT("ValidateObservationObjectMatchesSchema: Observation '%s' index out of range for Schema. Expected '<%i', got '%i'."),
+						*ObservationSchemaElementTag.ToString(),
+						SchemaElementSize,
+						ObjectElementIndices[SubElementIdx]);
+					return false;
+				}
+			}
+
+			return true;
+		}
+
+		case Learning::Observation::EType::NamedDiscreteExclusive:
+		{
+			const Learning::Observation::FSchemaNamedDiscreteExclusiveParameters SchemaParameters = Schema.GetNamedDiscreteExclusive(SchemaElement);
+			const Learning::Observation::FObjectNamedDiscreteExclusiveParameters ObjectParameters = Object.GetNamedDiscreteExclusive(ObjectElement);
+			check(SchemaParameters.ElementNames.Num() == SchemaParameters.ElementNames.Num());
+
+			const int32 SchemaSubElementIdx = SchemaParameters.ElementNames.Find(ObjectParameters.ElementName);
+
+			if (SchemaSubElementIdx == INDEX_NONE)
+			{
+				UE_LOG(LogLearning, Error, TEXT("ValidateObservationObjectMatchesSchema: Observation '%s' Schema does not include '%s' observation."),
+					*ObservationSchemaElementTag.ToString(),
+					*ObjectParameters.ElementName.ToString());
+				return false;
+			}
+
+			return true;
+		}
+
+		case Learning::Observation::EType::NamedDiscreteInclusive:
+		{
+			const Learning::Observation::FSchemaNamedDiscreteInclusiveParameters SchemaParameters = Schema.GetNamedDiscreteInclusive(SchemaElement);
+			const Learning::Observation::FObjectNamedDiscreteInclusiveParameters ObjectParameters = Object.GetNamedDiscreteInclusive(ObjectElement);
+
+			if (ObjectParameters.ElementNames.Num() > SchemaParameters.ElementNames.Num())
+			{
+				UE_LOG(LogLearning, Error, TEXT("ValidateObservationObjectMatchesSchema: Observation '%s' too many sub-observations provided. Expected at most '%i', got '%i'."),
+					*ObservationSchemaElementTag.ToString(),
+					SchemaParameters.ElementNames.Num(),
+					ObjectParameters.ElementNames.Num());
+				return false;
+			}
+
+			for (int32 ObjectSubElementIdx = 0; ObjectSubElementIdx < ObjectParameters.ElementNames.Num(); ObjectSubElementIdx++)
+			{
+				const int32 SchemaSubElementIdx = SchemaParameters.ElementNames.Find(ObjectParameters.ElementNames[ObjectSubElementIdx]);
+
+				if (SchemaSubElementIdx == INDEX_NONE)
+				{
+					UE_LOG(LogLearning, Error, TEXT("ValidateObservationObjectMatchesSchema: Observation '%s' Schema does not include '%s' observation."),
+						*ObservationSchemaElementTag.ToString(),
+						*ObjectParameters.ElementNames[ObjectSubElementIdx].ToString());
+					return false;
+				}
+			}
+
+			return true;
+		}
+
 		case Learning::Observation::EType::And:
 		{
 			const Learning::Observation::FSchemaAndParameters SchemaParameters = Schema.GetAnd(SchemaElement);
@@ -348,6 +454,38 @@ namespace UE::Learning::Agents::Observation::Private
 			return;
 		}
 
+		case UE::Learning::Observation::EType::DiscreteExclusive:
+		{
+			const UE::Learning::Observation::FObjectDiscreteExclusiveParameters Parameters = Object.GetDiscreteExclusive(ObjectElement);
+
+			UE_LOG(LogLearning, Display, TEXT("%s%s \"%s\" (%s) %i"), *Indentation, *Prefix, *Tag.ToString(), GetObservationTypeString(Type), Parameters.DiscreteIndex);
+			return;
+		}
+
+		case UE::Learning::Observation::EType::DiscreteInclusive:
+		{
+			const UE::Learning::Observation::FObjectDiscreteInclusiveParameters Parameters = Object.GetDiscreteInclusive(ObjectElement);
+
+			UE_LOG(LogLearning, Display, TEXT("%s%s \"%s\" (%s) %s"), *Indentation, *Prefix, *Tag.ToString(), GetObservationTypeString(Type), *UE::Learning::Array::FormatInt32(Parameters.DiscreteIndices));
+			return;
+		}
+
+		case UE::Learning::Observation::EType::NamedDiscreteExclusive:
+		{
+			const UE::Learning::Observation::FObjectNamedDiscreteExclusiveParameters Parameters = Object.GetNamedDiscreteExclusive(ObjectElement);
+
+			UE_LOG(LogLearning, Display, TEXT("%s%s \"%s\" (%s) %s"), *Indentation, *Prefix, *Tag.ToString(), GetObservationTypeString(Type), *Parameters.ElementName.ToString());
+			return;
+		}
+
+		case UE::Learning::Observation::EType::NamedDiscreteInclusive:
+		{
+			const UE::Learning::Observation::FObjectNamedDiscreteInclusiveParameters Parameters = Object.GetNamedDiscreteInclusive(ObjectElement);
+
+			UE_LOG(LogLearning, Display, TEXT("%s%s \"%s\" (%s) %s"), *Indentation, *Prefix, *Tag.ToString(), GetObservationTypeString(Type), *UE::Learning::Agents::Observation::Private::FormatNames(Parameters.ElementNames));
+			return;
+		}
+
 		case UE::Learning::Observation::EType::And:
 		{
 			const UE::Learning::Observation::FObjectAndParameters Parameters = Object.GetAnd(ObjectElement);
@@ -531,12 +669,84 @@ FLearningAgentsObservationSchemaElement ULearningAgentsObservations::SpecifyCont
 
 FLearningAgentsObservationSchemaElement ULearningAgentsObservations::SpecifyExclusiveDiscreteObservation(ULearningAgentsObservationSchema* Schema, const int32 Size, const FName Tag)
 {
-	return SpecifyContinuousObservation(Schema, Size, 1.0f, Tag);
+	if (!Schema)
+	{
+		UE_LOG(LogLearning, Error, TEXT("SpecifyExclusiveDiscreteObservation: Schema is nullptr."));
+		return FLearningAgentsObservationSchemaElement();
+	}
+
+	if (Size < 1)
+	{
+		UE_LOG(LogLearning, Error, TEXT("SpecifyExclusiveDiscreteObservation: Invalid Exclusive Discrete Observation Size '%i'."), Size);
+		return FLearningAgentsObservationSchemaElement();
+	}
+
+	return { Schema->ObservationSchema.CreateDiscreteExclusive({ Size }, Tag) };
 }
 
 FLearningAgentsObservationSchemaElement ULearningAgentsObservations::SpecifyInclusiveDiscreteObservation(ULearningAgentsObservationSchema* Schema, const int32 Size, const FName Tag)
 {
-	return SpecifyContinuousObservation(Schema, Size, 1.0f, Tag);
+	if (!Schema)
+	{
+		UE_LOG(LogLearning, Error, TEXT("SpecifyInclusiveDiscreteObservation: Schema is nullptr."));
+		return FLearningAgentsObservationSchemaElement();
+	}
+
+	if (Size < 0)
+	{
+		UE_LOG(LogLearning, Error, TEXT("SpecifyInclusiveDiscreteObservation: Invalid Inclusive Discrete Observation Size '%i'."), Size);
+		return FLearningAgentsObservationSchemaElement();
+	}
+
+	if (Size == 0)
+	{
+		UE_LOG(LogLearning, Warning, TEXT("SpecifyInclusiveDiscreteObservation: Specifying zero-sized Inclusive Discrete Observation."));
+	}
+
+	return { Schema->ObservationSchema.CreateDiscreteInclusive({ Size }, Tag) };
+}
+
+FLearningAgentsObservationSchemaElement ULearningAgentsObservations::SpecifyNamedExclusiveDiscreteObservation(ULearningAgentsObservationSchema* Schema, const TArray<FName>& ElementNames, const FName Tag)
+{
+	return SpecifyNamedExclusiveDiscreteObservationFromArrayView(Schema, ElementNames, Tag);
+}
+
+FLearningAgentsObservationSchemaElement ULearningAgentsObservations::SpecifyNamedExclusiveDiscreteObservationFromArrayView(ULearningAgentsObservationSchema* Schema, const TArrayView<const FName> ElementNames, const FName Tag)
+{
+	if (!Schema)
+	{
+		UE_LOG(LogLearning, Error, TEXT("SpecifyNamedExclusiveDiscreteObservationFromArrayView: Schema is nullptr."));
+		return FLearningAgentsObservationSchemaElement();
+	}
+
+	if (ElementNames.Num() < 1)
+	{
+		UE_LOG(LogLearning, Error, TEXT("SpecifyNamedExclusiveDiscreteObservationFromArrayView: Invalid Named Exclusive Discrete Observation Size '%i'."), ElementNames.Num());
+		return FLearningAgentsObservationSchemaElement();
+	}
+
+	return { Schema->ObservationSchema.CreateNamedDiscreteExclusive({ ElementNames }, Tag) };
+}
+
+FLearningAgentsObservationSchemaElement ULearningAgentsObservations::SpecifyNamedInclusiveDiscreteObservation(ULearningAgentsObservationSchema* Schema, const TArray<FName>& ElementNames, const FName Tag)
+{
+	return SpecifyNamedInclusiveDiscreteObservationFromArrayView(Schema, ElementNames, Tag);
+}
+
+FLearningAgentsObservationSchemaElement ULearningAgentsObservations::SpecifyNamedInclusiveDiscreteObservationFromArrayView(ULearningAgentsObservationSchema* Schema, const TArrayView<const FName> ElementNames, const FName Tag)
+{
+	if (!Schema)
+	{
+		UE_LOG(LogLearning, Error, TEXT("SpecifyNamedInclusiveDiscreteObservationFromArrayView: Schema is nullptr."));
+		return FLearningAgentsObservationSchemaElement();
+	}
+
+	if (ElementNames.Num() == 0)
+	{
+		UE_LOG(LogLearning, Warning, TEXT("SpecifyNamedInclusiveDiscreteObservationFromArrayView: Specifying zero-sized Named Inclusive Discrete Observation."));
+	}
+
+	return { Schema->ObservationSchema.CreateNamedDiscreteInclusive({ ElementNames }, Tag) };
 }
 
 FLearningAgentsObservationSchemaElement ULearningAgentsObservations::SpecifyCountObservation(ULearningAgentsObservationSchema* Schema, const FName Tag)
@@ -1151,7 +1361,6 @@ FLearningAgentsObservationObjectElement ULearningAgentsObservations::MakeContinu
 FLearningAgentsObservationObjectElement ULearningAgentsObservations::MakeExclusiveDiscreteObservation(
 	ULearningAgentsObservationObject* Object, 
 	const int32 DiscreteIndex, 
-	const int32 Size, 
 	const FName Tag,
 	const bool bVisualLoggerEnabled,
 	ULearningAgentsManagerListener* VisualLoggerListener,
@@ -1159,16 +1368,12 @@ FLearningAgentsObservationObjectElement ULearningAgentsObservations::MakeExclusi
 	const FVector VisualLoggerLocation,
 	const FLinearColor VisualLoggerColor)
 {
-	if (DiscreteIndex < 0 || DiscreteIndex >= Size)
+	if (!Object)
 	{
-		UE_LOG(LogLearning, Error, TEXT("MakeExclusiveDiscreteObservation: Discrete index out of range: Got %i, expected <= %i."), DiscreteIndex, Size);
+		UE_LOG(LogLearning, Error, TEXT("MakeExclusiveDiscreteObservation: Object is nullptr."));
 		return FLearningAgentsObservationObjectElement();
 	}
 
-	TArray<float, TInlineAllocator<32>> Values;
-	Values.Init(0.0f, Size);
-	Values[DiscreteIndex] = 1.0f;
-
 #if UE_LEARNING_AGENTS_ENABLE_VISUAL_LOG
 	if (bVisualLoggerEnabled && VisualLoggerListener)
 	{
@@ -1176,23 +1381,20 @@ FLearningAgentsObservationObjectElement ULearningAgentsObservations::MakeExclusi
 
 		UE_LEARNING_AGENTS_VLOG_STRING(VisualLoggerObject, LogLearning, Display, VisualLoggerLocation,
 			VisualLoggerColor.ToFColor(true),
-			TEXT("Listener: %s\nTag: %s\nAgent Id: % 3i\nSize: [%i]\nIndex: [%i]\nEncoded: %s"),
+			TEXT("Listener: %s\nTag: %s\nAgent Id: % 3i\nIndex: [%i]"),
 			*VisualLoggerListener->GetName(),
 			*Tag.ToString(),
 			VisualLoggerAgentId,
-			Size,
-			DiscreteIndex,
-			*UE::Learning::Array::FormatFloat(Values, Size));
+			DiscreteIndex);
 	}
 #endif
 
-	return MakeContinuousObservationFromArrayView(Object, Values, Tag);
+	return { Object->ObservationObject.CreateDiscreteExclusive({ DiscreteIndex }, Tag) };
 }
 
 FLearningAgentsObservationObjectElement ULearningAgentsObservations::MakeInclusiveDiscreteObservation(
 	ULearningAgentsObservationObject* Object, 
 	const TArray<int32>& DiscreteIndices, 
-	const int32 Size, 
 	const FName Tag,
 	const bool bVisualLoggerEnabled,
 	ULearningAgentsManagerListener* VisualLoggerListener,
@@ -1200,13 +1402,12 @@ FLearningAgentsObservationObjectElement ULearningAgentsObservations::MakeInclusi
 	const FVector VisualLoggerLocation,
 	const FLinearColor VisualLoggerColor)
 {
-	return MakeInclusiveDiscreteObservationFromArrayView(Object, DiscreteIndices, Size, Tag, bVisualLoggerEnabled, VisualLoggerListener, VisualLoggerAgentId, VisualLoggerLocation, VisualLoggerColor);
+	return MakeInclusiveDiscreteObservationFromArrayView(Object, DiscreteIndices, Tag, bVisualLoggerEnabled, VisualLoggerListener, VisualLoggerAgentId, VisualLoggerLocation, VisualLoggerColor);
 }
 
 FLearningAgentsObservationObjectElement ULearningAgentsObservations::MakeInclusiveDiscreteObservationFromArrayView(
 	ULearningAgentsObservationObject* Object, 
 	const TArrayView<const int32> DiscreteIndices, 
-	const int32 Size, 
 	const FName Tag,
 	const bool bVisualLoggerEnabled,
 	ULearningAgentsManagerListener* VisualLoggerListener,
@@ -1214,24 +1415,103 @@ FLearningAgentsObservationObjectElement ULearningAgentsObservations::MakeInclusi
 	const FVector VisualLoggerLocation,
 	const FLinearColor VisualLoggerColor)
 {
+	if (!Object)
+	{
+		UE_LOG(LogLearning, Error, TEXT("MakeExclusiveDiscreteObservation: Object is nullptr."));
+		return FLearningAgentsObservationObjectElement();
+	}
+
 	if (UE::Learning::Agents::Observation::Private::ContainsDuplicates(DiscreteIndices))
 	{
 		UE_LOG(LogLearning, Error, TEXT("MakeInclusiveDiscreteObservationFromArrayView: Indices contain duplicates."));
 		return FLearningAgentsObservationObjectElement();
 	}
 
-	TArray<float, TInlineAllocator<32>> Values;
-	Values.Init(0.0f, Size);
+#if UE_LEARNING_AGENTS_ENABLE_VISUAL_LOG
+	if (bVisualLoggerEnabled && VisualLoggerListener)
+	{
+		const ULearningAgentsVisualLoggerObject* VisualLoggerObject = VisualLoggerListener->GetOrAddVisualLoggerObject(Tag);
+
+		UE_LEARNING_AGENTS_VLOG_STRING(VisualLoggerObject, LogLearning, Display, VisualLoggerLocation,
+			VisualLoggerColor.ToFColor(true),
+			TEXT("Listener: %s\nTag: %s\nAgent Id: % 3i\nIndices: %s"),
+			*VisualLoggerListener->GetName(),
+			*Tag.ToString(),
+			VisualLoggerAgentId,
+			*UE::Learning::Array::FormatInt32(DiscreteIndices));
+	}
+#endif
+
+	return { Object->ObservationObject.CreateDiscreteInclusive({ DiscreteIndices }, Tag) };
+}
+
+FLearningAgentsObservationObjectElement ULearningAgentsObservations::MakeNamedExclusiveDiscreteObservation(
+	ULearningAgentsObservationObject* Object,
+	const FName ElementName,
+	const FName Tag,
+	const bool bVisualLoggerEnabled,
+	ULearningAgentsManagerListener* VisualLoggerListener,
+	const int32 VisualLoggerAgentId,
+	const FVector VisualLoggerLocation,
+	const FLinearColor VisualLoggerColor)
+{
+	if (!Object)
+	{
+		UE_LOG(LogLearning, Error, TEXT("MakeNamedExclusiveDiscreteObservation: Object is nullptr."));
+		return FLearningAgentsObservationObjectElement();
+	}
+
+#if UE_LEARNING_AGENTS_ENABLE_VISUAL_LOG
+	if (bVisualLoggerEnabled && VisualLoggerListener)
+	{
+		const ULearningAgentsVisualLoggerObject* VisualLoggerObject = VisualLoggerListener->GetOrAddVisualLoggerObject(Tag);
+
+		UE_LEARNING_AGENTS_VLOG_STRING(VisualLoggerObject, LogLearning, Display, VisualLoggerLocation,
+			VisualLoggerColor.ToFColor(true),
+			TEXT("Listener: %s\nTag: %s\nAgent Id: % 3i\nName: [%s]"),
+			*VisualLoggerListener->GetName(),
+			*Tag.ToString(),
+			VisualLoggerAgentId,
+			*ElementName.ToString());
+	}
+#endif
+
+	return { Object->ObservationObject.CreateNamedDiscreteExclusive({ ElementName }, Tag) };
+}
 
-	for (int32 Idx = 0; Idx < DiscreteIndices.Num(); Idx++)
+FLearningAgentsObservationObjectElement ULearningAgentsObservations::MakeNamedInclusiveDiscreteObservation(
+	ULearningAgentsObservationObject* Object,
+	const TArray<FName>& ElementNames,
+	const FName Tag,
+	const bool bVisualLoggerEnabled,
+	ULearningAgentsManagerListener* VisualLoggerListener,
+	const int32 VisualLoggerAgentId,
+	const FVector VisualLoggerLocation,
+	const FLinearColor VisualLoggerColor)
+{
+	return MakeNamedInclusiveDiscreteObservationFromArrayView(Object, ElementNames, Tag, bVisualLoggerEnabled, VisualLoggerListener, VisualLoggerAgentId, VisualLoggerLocation, VisualLoggerColor);
+}
+
+FLearningAgentsObservationObjectElement ULearningAgentsObservations::MakeNamedInclusiveDiscreteObservationFromArrayView(
+	ULearningAgentsObservationObject* Object,
+	const TArrayView<const FName> ElementNames,
+	const FName Tag,
+	const bool bVisualLoggerEnabled,
+	ULearningAgentsManagerListener* VisualLoggerListener,
+	const int32 VisualLoggerAgentId,
+	const FVector VisualLoggerLocation,
+	const FLinearColor VisualLoggerColor)
+{
+	if (!Object)
 	{
-		if (DiscreteIndices[Idx] < 0 || DiscreteIndices[Idx] >= Size)
-		{
-			UE_LOG(LogLearning, Error, TEXT("MakeInclusiveDiscreteObservationFromArrayView: Discrete index out of range: Got %i, expected <= %i."), DiscreteIndices[Idx], Size);
-			return FLearningAgentsObservationObjectElement();
-		}
+		UE_LOG(LogLearning, Error, TEXT("MakeNamedInclusiveDiscreteObservationFromArrayView: Object is nullptr."));
+		return FLearningAgentsObservationObjectElement();
+	}
 
-		Values[DiscreteIndices[Idx]] = 1.0f;
+	if (UE::Learning::Agents::Observation::Private::ContainsDuplicates(ElementNames))
+	{
+		UE_LOG(LogLearning, Error, TEXT("MakeNamedInclusiveDiscreteObservationFromArrayView: Names contain duplicates."));
+		return FLearningAgentsObservationObjectElement();
 	}
 
 #if UE_LEARNING_AGENTS_ENABLE_VISUAL_LOG
@@ -1241,17 +1521,15 @@ FLearningAgentsObservationObjectElement ULearningAgentsObservations::MakeInclusi
 
 		UE_LEARNING_AGENTS_VLOG_STRING(VisualLoggerObject, LogLearning, Display, VisualLoggerLocation,
 			VisualLoggerColor.ToFColor(true),
-			TEXT("Listener: %s\nTag: %s\nAgent Id: % 3i\nSize: [%i]\nIndices: %s\nEncoded: %s"),
+			TEXT("Listener: %s\nTag: %s\nAgent Id: % 3i\nNames: %s"),
 			*VisualLoggerListener->GetName(),
 			*Tag.ToString(),
 			VisualLoggerAgentId,
-			Size,
-			*UE::Learning::Array::FormatInt32(DiscreteIndices, Size),
-			*UE::Learning::Array::FormatFloat(Values, Size));
+			*UE::Learning::Agents::Observation::Private::FormatNames(ElementNames, 256));
 	}
 #endif
 
-	return MakeContinuousObservationFromArrayView(Object, Values, Tag);
+	return { Object->ObservationObject.CreateNamedDiscreteInclusive({ ElementNames }, Tag) };
 }
 
 FLearningAgentsObservationObjectElement ULearningAgentsObservations::MakeCountObservation(
@@ -2581,57 +2859,71 @@ bool ULearningAgentsObservations::GetContinuousObservationToArrayView(TArrayView
 
 bool ULearningAgentsObservations::GetExclusiveDiscreteObservation(int32& OutIndex, const ULearningAgentsObservationObject* Object, const FLearningAgentsObservationObjectElement Element, const FName Tag)
 {
-	int32 DiscreteValueNum;
-	if (!GetContinuousObservationNum(DiscreteValueNum, Object, Element, Tag))
+	if (!Object)
 	{
-		OutIndex = INDEX_NONE;
+		UE_LOG(LogLearning, Error, TEXT("GetExclusiveDiscreteObservation: Object is nullptr."));
+		OutIndex = 0;
 		return false;
 	}
 
-	TArray<float, TInlineAllocator<32>> OneHot;
-	OneHot.SetNumUninitialized(DiscreteValueNum);
-	if (!GetContinuousObservationToArrayView(OneHot, Object, Element, Tag))
+	if (!Object->ObservationObject.IsValid(Element.ObjectElement))
 	{
-		OutIndex = INDEX_NONE;
+		UE_LOG(LogLearning, Error, TEXT("GetExclusiveDiscreteObservation: Invalid Observation Object."));
+		OutIndex = 0;
 		return false;
 	}
 
-	OutIndex = INDEX_NONE;
-	for (int32 Idx = 0; Idx < DiscreteValueNum; Idx++)
+	if (Object->ObservationObject.GetTag(Element.ObjectElement) != Tag)
 	{
-		if (OneHot[Idx])
-		{
-			OutIndex = Idx;
-			break;
-		}
+		UE_LOG(LogLearning, Warning, TEXT("GetExclusiveDiscreteObservation: Observation tag does not match. Observation is '%s' but asked for '%s'."), *Object->ObservationObject.GetTag(Element.ObjectElement).ToString(), *Tag.ToString());
 	}
 
-	return OutIndex != INDEX_NONE;
+	if (Object->ObservationObject.GetType(Element.ObjectElement) != UE::Learning::Observation::EType::DiscreteExclusive)
+	{
+		UE_LOG(LogLearning, Error, TEXT("GetExclusiveDiscreteObservation: Observation '%s' type does not match. Observation is '%s' but asked for '%s'."),
+			*Object->ObservationObject.GetTag(Element.ObjectElement).ToString(),
+			UE::Learning::Agents::Observation::Private::GetObservationTypeString(Object->ObservationObject.GetType(Element.ObjectElement)),
+			UE::Learning::Agents::Observation::Private::GetObservationTypeString(UE::Learning::Observation::EType::DiscreteExclusive));
+		OutIndex = 0;
+		return false;
+	}
+
+	OutIndex = Object->ObservationObject.GetDiscreteExclusive(Element.ObjectElement).DiscreteIndex;
+	return true;
 }
 
 bool ULearningAgentsObservations::GetInclusiveDiscreteObservationNum(int32& OutNum, const ULearningAgentsObservationObject* Object, const FLearningAgentsObservationObjectElement Element, const FName Tag)
 {
-	int32 DiscreteValueNum;
-	if (!GetContinuousObservationNum(DiscreteValueNum, Object, Element, Tag))
+	if (!Object)
 	{
+		UE_LOG(LogLearning, Error, TEXT("GetInclusiveDiscreteObservationNum: Object is nullptr."));
 		OutNum = 0;
 		return false;
 	}
 
-	TArray<float, TInlineAllocator<32>> OneHot;
-	OneHot.SetNumUninitialized(DiscreteValueNum);
-	if (!GetContinuousObservationToArrayView(OneHot, Object, Element, Tag))
+	if (!Object->ObservationObject.IsValid(Element.ObjectElement))
 	{
+		UE_LOG(LogLearning, Error, TEXT("GetInclusiveDiscreteObservationNum: Invalid Observation Object."));
 		OutNum = 0;
 		return false;
 	}
 
-	OutNum = 0;
-	for (int32 Idx = 0; Idx < DiscreteValueNum; Idx++)
+	if (Object->ObservationObject.GetTag(Element.ObjectElement) != Tag)
 	{
-		if (OneHot[Idx]) { OutNum++; }
+		UE_LOG(LogLearning, Warning, TEXT("GetInclusiveDiscreteObservationNum: Observation tag does not match. Observation is '%s' but asked for '%s'."), *Object->ObservationObject.GetTag(Element.ObjectElement).ToString(), *Tag.ToString());
 	}
 
+	if (Object->ObservationObject.GetType(Element.ObjectElement) != UE::Learning::Observation::EType::DiscreteInclusive)
+	{
+		UE_LOG(LogLearning, Error, TEXT("GetInclusiveDiscreteObservationNum: Observation '%s' type does not match. Observation is '%s' but asked for '%s'."),
+			*Object->ObservationObject.GetTag(Element.ObjectElement).ToString(),
+			UE::Learning::Agents::Observation::Private::GetObservationTypeString(Object->ObservationObject.GetType(Element.ObjectElement)),
+			UE::Learning::Agents::Observation::Private::GetObservationTypeString(UE::Learning::Observation::EType::DiscreteInclusive));
+		OutNum = 0;
+		return false;
+	}
+
+	OutNum = Object->ObservationObject.GetDiscreteInclusive(Element.ObjectElement).DiscreteIndices.Num();
 	return true;
 }
 
@@ -2653,43 +2945,173 @@ bool ULearningAgentsObservations::GetInclusiveDiscreteObservationToArrayView(TAr
 	if (!Object)
 	{
 		UE_LOG(LogLearning, Error, TEXT("GetInclusiveDiscreteObservationToArrayView: Object is nullptr."));
+		UE::Learning::Array::Zero<1, int32>(OutIndices);
 		return false;
 	}
 
-	int32 DiscreteValueNum;
-	if (!GetContinuousObservationNum(DiscreteValueNum, Object, Element, Tag))
+	if (!Object->ObservationObject.IsValid(Element.ObjectElement))
 	{
+		UE_LOG(LogLearning, Error, TEXT("GetInclusiveDiscreteObservationToArrayView: Invalid Observation Object."));
+		UE::Learning::Array::Zero<1, int32>(OutIndices);
 		return false;
 	}
 
-	TArray<float, TInlineAllocator<32>> OneHot;
-	OneHot.SetNumUninitialized(DiscreteValueNum);
-	if (!GetContinuousObservationToArrayView(OneHot, Object, Element, Tag))
+	if (Object->ObservationObject.GetTag(Element.ObjectElement) != Tag)
 	{
-		return false;
+		UE_LOG(LogLearning, Warning, TEXT("GetInclusiveDiscreteObservationToArrayView: Observation tag does not match. Observation is '%s' but asked for '%s'."), *Object->ObservationObject.GetTag(Element.ObjectElement).ToString(), *Tag.ToString());
 	}
 
-	int32 IndicesNum = 0;
-	for (int32 Idx = 0; Idx < DiscreteValueNum; Idx++)
+	if (Object->ObservationObject.GetType(Element.ObjectElement) != UE::Learning::Observation::EType::DiscreteInclusive)
 	{
-		if (OneHot[Idx]) { IndicesNum++; }
+		UE_LOG(LogLearning, Error, TEXT("GetInclusiveDiscreteObservationToArrayView: Observation '%s' type does not match. Observation is '%s' but asked for '%s'."),
+			*Object->ObservationObject.GetTag(Element.ObjectElement).ToString(),
+			UE::Learning::Agents::Observation::Private::GetObservationTypeString(Object->ObservationObject.GetType(Element.ObjectElement)),
+			UE::Learning::Agents::Observation::Private::GetObservationTypeString(UE::Learning::Observation::EType::DiscreteInclusive));
+		UE::Learning::Array::Zero<1, int32>(OutIndices);
+		return false;
 	}
 
-	if (IndicesNum != OutIndices.Num())
+	const TArrayView<const int32> Indices = Object->ObservationObject.GetDiscreteInclusive(Element.ObjectElement).DiscreteIndices;
+
+	if (Indices.Num() != OutIndices.Num())
 	{
-		UE_LOG(LogLearning, Error, TEXT("GetInclusiveDiscreteObservationToArrayView: Observation '%s' size does not match. Observation is '%i' indices but asked for '%i'."),
+		UE_LOG(LogLearning, Error, TEXT("GetInclusiveDiscreteObservationToArrayView: Observation '%s' size does not match. Observation is '%i' values but asked for '%i'."),
 			*Object->ObservationObject.GetTag(Element.ObjectElement).ToString(),
-			IndicesNum, OutIndices.Num());
+			Indices.Num(), OutIndices.Num());
 		UE::Learning::Array::Zero<1, int32>(OutIndices);
 		return false;
 	}
 
-	int32 Offset = 0;
-	for (int32 Idx = 0; Idx < DiscreteValueNum; Idx++)
+	UE::Learning::Array::Copy<1, int32>(OutIndices, Indices);
+	return true;
+}
+
+bool ULearningAgentsObservations::GetNamedExclusiveDiscreteObservation(FName& OutName, const ULearningAgentsObservationObject* Object, const FLearningAgentsObservationObjectElement Element, const FName Tag)
+{
+	if (!Object)
+	{
+		UE_LOG(LogLearning, Error, TEXT("GetNamedExclusiveDiscreteObservation: Object is nullptr."));
+		OutName = NAME_None;
+		return false;
+	}
+
+	if (!Object->ObservationObject.IsValid(Element.ObjectElement))
+	{
+		UE_LOG(LogLearning, Error, TEXT("GetNamedExclusiveDiscreteObservation: Invalid Observation Object."));
+		OutName = NAME_None;
+		return false;
+	}
+
+	if (Object->ObservationObject.GetTag(Element.ObjectElement) != Tag)
 	{
-		if (OneHot[Idx]) { OutIndices[Offset] = Idx; Offset++; }
+		UE_LOG(LogLearning, Warning, TEXT("GetNamedExclusiveDiscreteObservation: Observation tag does not match. Observation is '%s' but asked for '%s'."), *Object->ObservationObject.GetTag(Element.ObjectElement).ToString(), *Tag.ToString());
+	}
+
+	if (Object->ObservationObject.GetType(Element.ObjectElement) != UE::Learning::Observation::EType::NamedDiscreteExclusive)
+	{
+		UE_LOG(LogLearning, Error, TEXT("GetNamedExclusiveDiscreteObservation: Observation '%s' type does not match. Observation is '%s' but asked for '%s'."),
+			*Object->ObservationObject.GetTag(Element.ObjectElement).ToString(),
+			UE::Learning::Agents::Observation::Private::GetObservationTypeString(Object->ObservationObject.GetType(Element.ObjectElement)),
+			UE::Learning::Agents::Observation::Private::GetObservationTypeString(UE::Learning::Observation::EType::NamedDiscreteExclusive));
+		OutName = NAME_None;
+		return false;
+	}
+
+	OutName = Object->ObservationObject.GetNamedDiscreteExclusive(Element.ObjectElement).ElementName;
+	return true;
+}
+
+bool ULearningAgentsObservations::GetNamedInclusiveDiscreteObservationNum(int32& OutNum, const ULearningAgentsObservationObject* Object, const FLearningAgentsObservationObjectElement Element, const FName Tag)
+{
+	if (!Object)
+	{
+		UE_LOG(LogLearning, Error, TEXT("GetNamedInclusiveDiscreteObservationNum: Object is nullptr."));
+		OutNum = 0;
+		return false;
+	}
+
+	if (!Object->ObservationObject.IsValid(Element.ObjectElement))
+	{
+		UE_LOG(LogLearning, Error, TEXT("GetNamedInclusiveDiscreteObservationNum: Invalid Observation Object."));
+		OutNum = 0;
+		return false;
+	}
+
+	if (Object->ObservationObject.GetTag(Element.ObjectElement) != Tag)
+	{
+		UE_LOG(LogLearning, Warning, TEXT("GetNamedInclusiveDiscreteObservationNum: Observation tag does not match. Observation is '%s' but asked for '%s'."), *Object->ObservationObject.GetTag(Element.ObjectElement).ToString(), *Tag.ToString());
+	}
+
+	if (Object->ObservationObject.GetType(Element.ObjectElement) != UE::Learning::Observation::EType::NamedDiscreteInclusive)
+	{
+		UE_LOG(LogLearning, Error, TEXT("GetNamedInclusiveDiscreteObservationNum: Observation '%s' type does not match. Observation is '%s' but asked for '%s'."),
+			*Object->ObservationObject.GetTag(Element.ObjectElement).ToString(),
+			UE::Learning::Agents::Observation::Private::GetObservationTypeString(Object->ObservationObject.GetType(Element.ObjectElement)),
+			UE::Learning::Agents::Observation::Private::GetObservationTypeString(UE::Learning::Observation::EType::NamedDiscreteInclusive));
+		OutNum = 0;
+		return false;
+	}
+
+	OutNum = Object->ObservationObject.GetNamedDiscreteInclusive(Element.ObjectElement).ElementNames.Num();
+	return true;
+}
+
+bool ULearningAgentsObservations::GetNamedInclusiveDiscreteObservation(TArray<FName>& OutNames, const ULearningAgentsObservationObject* Object, const FLearningAgentsObservationObjectElement Element, const FName Tag)
+{
+	int32 OutNamesNum;
+	if (!GetNamedInclusiveDiscreteObservationNum(OutNamesNum, Object, Element, Tag))
+	{
+		OutNames.Empty();
+		return false;
+	}
+
+	OutNames.SetNumUninitialized(OutNamesNum);
+	return GetNamedInclusiveDiscreteObservationToArrayView(OutNames, Object, Element, Tag);
+}
+
+bool ULearningAgentsObservations::GetNamedInclusiveDiscreteObservationToArrayView(TArrayView<FName> OutNames, const ULearningAgentsObservationObject* Object, const FLearningAgentsObservationObjectElement Element, const FName Tag)
+{
+	if (!Object)
+	{
+		UE_LOG(LogLearning, Error, TEXT("GetNamedInclusiveDiscreteObservationToArrayView: Object is nullptr."));
+		UE::Learning::Array::Set<1, FName>(OutNames, NAME_None);
+		return false;
+	}
+
+	if (!Object->ObservationObject.IsValid(Element.ObjectElement))
+	{
+		UE_LOG(LogLearning, Error, TEXT("GetNamedInclusiveDiscreteObservationToArrayView: Invalid Observation Object."));
+		UE::Learning::Array::Set<1, FName>(OutNames, NAME_None);
+		return false;
+	}
+
+	if (Object->ObservationObject.GetTag(Element.ObjectElement) != Tag)
+	{
+		UE_LOG(LogLearning, Warning, TEXT("GetNamedInclusiveDiscreteObservationToArrayView: Observation tag does not match. Observation is '%s' but asked for '%s'."), *Object->ObservationObject.GetTag(Element.ObjectElement).ToString(), *Tag.ToString());
+	}
+
+	if (Object->ObservationObject.GetType(Element.ObjectElement) != UE::Learning::Observation::EType::NamedDiscreteInclusive)
+	{
+		UE_LOG(LogLearning, Error, TEXT("GetNamedInclusiveDiscreteObservationToArrayView: Observation '%s' type does not match. Observation is '%s' but asked for '%s'."),
+			*Object->ObservationObject.GetTag(Element.ObjectElement).ToString(),
+			UE::Learning::Agents::Observation::Private::GetObservationTypeString(Object->ObservationObject.GetType(Element.ObjectElement)),
+			UE::Learning::Agents::Observation::Private::GetObservationTypeString(UE::Learning::Observation::EType::NamedDiscreteInclusive));
+		UE::Learning::Array::Set<1, FName>(OutNames, NAME_None);
+		return false;
+	}
+
+	const TArrayView<const FName> Names = Object->ObservationObject.GetNamedDiscreteInclusive(Element.ObjectElement).ElementNames;
+
+	if (Names.Num() != OutNames.Num())
+	{
+		UE_LOG(LogLearning, Error, TEXT("GetNamedInclusiveDiscreteObservationToArrayView: Observation '%s' size does not match. Observation is '%i' values but asked for '%i'."),
+			*Object->ObservationObject.GetTag(Element.ObjectElement).ToString(),
+			Names.Num(), OutNames.Num());
+		UE::Learning::Array::Set<1, FName>(OutNames, NAME_None);
+		return false;
 	}
 
+	UE::Learning::Array::Copy<1, FName>(OutNames, Names);
 	return true;
 }
 
diff --git a/Engine/Plugins/Experimental/LearningAgents/Source/LearningAgents/Public/LearningAgentsObservations.h b/Engine/Plugins/Experimental/LearningAgents/Source/LearningAgents/Public/LearningAgentsObservations.h
--- a/Engine/Plugins/Experimental/LearningAgents/Source/LearningAgents/Public/LearningAgentsObservations.h
+++ b/Engine/Plugins/Experimental/LearningAgents/Source/LearningAgents/Public/LearningAgentsObservations.h
@@ -199,6 +199,48 @@ public:
 	UFUNCTION(BlueprintPure, Category = "LearningAgents", meta = (AdvancedDisplay = 2))
 	static FLearningAgentsObservationSchemaElement SpecifyInclusiveDiscreteObservation(ULearningAgentsObservationSchema* Schema, const int32 Size, const FName Tag = TEXT("InclusiveDiscreteObservation"));
 
+	/**
+	 * Specifies a new named exclusive discrete observation. This represents a discrete observation which is an exclusive selection from multiple choices.
+	 *
+	 * @param Schema The Observation Schema
+	 * @param ElementNames The names of the discrete options in the observation.
+	 * @param Tag The tag of this new observation. Used during observation object validation and debugging.
+	 * @return The newly created observation schema element.
+	 */
+	UFUNCTION(BlueprintPure, Category = "LearningAgents", meta = (AdvancedDisplay = 2))
+	static FLearningAgentsObservationSchemaElement SpecifyNamedExclusiveDiscreteObservation(ULearningAgentsObservationSchema* Schema, const TArray<FName>& ElementNames, const FName Tag = TEXT("NamedExclusiveDiscreteObservation"));
+
+	/**
+	 * Specifies a new named exclusive discrete observation. This represents a discrete observation which is an exclusive selection from multiple choices.
+	 *
+	 * @param Schema The Observation Schema
+	 * @param ElementNames The names of the discrete options in the observation.
+	 * @param Tag The tag of this new observation. Used during observation object validation and debugging.
+	 * @return The newly created observation schema element.
+	 */
+	static FLearningAgentsObservationSchemaElement SpecifyNamedExclusiveDiscreteObservationFromArrayView(ULearningAgentsObservationSchema* Schema, const TArrayView<const FName> ElementNames, const FName Tag = TEXT("NamedExclusiveDiscreteObservation"));
+
+	/**
+	 * Specifies a new named inclusive discrete observation. This represents a discrete observation which is an inclusive selection from multiple choices.
+	 *
+	 * @param Schema The Observation Schema
+	 * @param ElementNames The names of the discrete options in the observation.
+	 * @param Tag The tag of this new observation. Used during observation object validation and debugging.
+	 * @return The newly created observation schema element.
+	 */
+	UFUNCTION(BlueprintPure, Category = "LearningAgents", meta = (AdvancedDisplay = 2))
+	static FLearningAgentsObservationSchemaElement SpecifyNamedInclusiveDiscreteObservation(ULearningAgentsObservationSchema* Schema, const TArray<FName>& ElementNames, const FName Tag = TEXT("NamedInclusiveDiscreteObservation"));
+
+	/**
+	 * Specifies a new named inclusive discrete observation. This represents a discrete observation which is an inclusive selection from multiple choices.
+	 *
+	 * @param Schema The Observation Schema
+	 * @param ElementNames The names of the discrete options in the observation.
+	 * @param Tag The tag of this new observation. Used during observation object validation and debugging.
+	 * @return The newly created observation schema element.
+	 */
+	static FLearningAgentsObservationSchemaElement SpecifyNamedInclusiveDiscreteObservationFromArrayView(ULearningAgentsObservationSchema* Schema, const TArrayView<const FName> ElementNames, const FName Tag = TEXT("NamedInclusiveDiscreteObservation"));
+
 	/**
 	 * Specifies a new count observation. This represents a count of something such as the size of, or index into, an array.
 	 *
@@ -675,8 +717,7 @@ public:
 	 * Make a new exclusive discrete observation.
 	 *
 	 * @param Object The Observation Object
-	 * @param DiscreteIndex The index of the discrete observation. Values must be smaller than the given Size.
-	 * @param Size The size of the discrete observation. Must be equal to the size given during Specify.
+	 * @param DiscreteIndex The index of the discrete observation.
 	 * @param Tag The tag of the corresponding observation. Must match the tag given during Specify.
 	 * @param bVisualLoggerEnabled When true, debug data will be sent to the visual logger.
 	 * @param VisualLoggerListener The listener object which is making this observation. This must be set to use logging.
@@ -685,11 +726,10 @@ public:
 	 * @param VisualLoggerColor The color for the visual logger display.
 	 * @return The newly created observation object element.
 	 */
-	UFUNCTION(BlueprintPure, Category = "LearningAgents", meta = (AdvancedDisplay = 3, DefaultToSelf = "VisualLoggerListener"))
+	UFUNCTION(BlueprintPure, Category = "LearningAgents", meta = (AdvancedDisplay = 2, DefaultToSelf = "VisualLoggerListener"))
 	static FLearningAgentsObservationObjectElement MakeExclusiveDiscreteObservation(
 		ULearningAgentsObservationObject* Object, 
 		const int32 DiscreteIndex, 
-		const int32 Size,
 		const FName Tag = TEXT("ExclusiveDiscreteObservation"),
 		const bool bVisualLoggerEnabled = false,
 		ULearningAgentsManagerListener* VisualLoggerListener = nullptr,
@@ -701,8 +741,7 @@ public:
 	 * Make a new inclusive discrete observation.
 	 *
 	 * @param Object The Observation Object
-	 * @param DiscreteIndices The indices of the discrete observations. All values must be smaller than the given Size.
-	 * @param Size The size of the discrete observation. Must be equal to the size given during Specify.
+	 * @param DiscreteIndices The indices of the discrete observations.
 	 * @param Tag The tag of the corresponding observation. Must match the tag given during Specify.
 	 * @param bVisualLoggerEnabled When true, debug data will be sent to the visual logger.
 	 * @param VisualLoggerListener The listener object which is making this observation. This must be set to use logging.
@@ -711,11 +750,10 @@ public:
 	 * @param VisualLoggerColor The color for the visual logger display.
 	 * @return The newly created observation object element.
 	 */
-	UFUNCTION(BlueprintPure, Category = "LearningAgents", meta = (AdvancedDisplay = 3, DefaultToSelf = "VisualLoggerListener"))
+	UFUNCTION(BlueprintPure, Category = "LearningAgents", meta = (AdvancedDisplay = 2, DefaultToSelf = "VisualLoggerListener"))
 	static FLearningAgentsObservationObjectElement MakeInclusiveDiscreteObservation(
 		ULearningAgentsObservationObject* Object, 
 		const TArray<int32>& DiscreteIndices, 
-		const int32 Size, 
 		const FName Tag = TEXT("InclusiveDiscreteObservation"),
 		const bool bVisualLoggerEnabled = false,
 		ULearningAgentsManagerListener* VisualLoggerListener = nullptr,
@@ -727,8 +765,7 @@ public:
 	 * Make a new inclusive discrete observation.
 	 *
 	 * @param Object The Observation Object
-	 * @param DiscreteIndices The indices of the discrete observations. All values must be smaller than the given Size.
-	 * @param Size The size of the discrete observation. Must be equal to the size given during Specify.
+	 * @param DiscreteIndices The indices of the discrete observations.
 	 * @param Tag The tag of the corresponding observation. Must match the tag given during Specify.
 	 * @param bVisualLoggerEnabled When true, debug data will be sent to the visual logger.
 	 * @param VisualLoggerListener The listener object which is making this observation. This must be set to use logging.
@@ -740,7 +777,6 @@ public:
 	static FLearningAgentsObservationObjectElement MakeInclusiveDiscreteObservationFromArrayView(
 		ULearningAgentsObservationObject* Object, 
 		const TArrayView<const int32> DiscreteIndices, 
-		const int32 Size, 
 		const FName Tag = TEXT("InclusiveDiscreteObservation"),
 		const bool bVisualLoggerEnabled = false,
 		ULearningAgentsManagerListener* VisualLoggerListener = nullptr,
@@ -748,6 +784,77 @@ public:
 		const FVector VisualLoggerLocation = FVector::ZeroVector,
 		const FLinearColor VisualLoggerColor = FLinearColor::Red);
 
+	/**
+	 * Make a new named exclusive discrete observation.
+	 *
+	 * @param Object The Observation Object
+	 * @param ElementName The name of the discrete observation.
+	 * @param Tag The tag of the corresponding observation. Must match the tag given during Specify.
+	 * @param bVisualLoggerEnabled When true, debug data will be sent to the visual logger.
+	 * @param VisualLoggerListener The listener object which is making this observation. This must be set to use logging.
+	 * @param VisualLoggerAgentId The agent id associated with this observation.
+	 * @param VisualLoggerLocation A location for the visual logger information in the world.
+	 * @param VisualLoggerColor The color for the visual logger display.
+	 * @return The newly created observation object element.
+	 */
+	UFUNCTION(BlueprintPure, Category = "LearningAgents", meta = (AdvancedDisplay = 2, DefaultToSelf = "VisualLoggerListener"))
+	static FLearningAgentsObservationObjectElement MakeNamedExclusiveDiscreteObservation(
+		ULearningAgentsObservationObject* Object,
+		const FName ElementName,
+		const FName Tag = TEXT("NamedExclusiveDiscreteObservation"),
+		const bool bVisualLoggerEnabled = false,
+		ULearningAgentsManagerListener* VisualLoggerListener = nullptr,
+		const int32 VisualLoggerAgentId = -1,
+		const FVector VisualLoggerLocation = FVector::ZeroVector,
+		const FLinearColor VisualLoggerColor = FLinearColor::Red);
+
+	/**
+	 * Make a new named inclusive discrete observation.
+	 *
+	 * @param Object The Observation Object
+	 * @param ElementNames The names of the discrete observations.
+	 * @param Tag The tag of the corresponding observation. Must match the tag given during Specify.
+	 * @param bVisualLoggerEnabled When true, debug data will be sent to the visual logger.
+	 * @param VisualLoggerListener The listener object which is making this observation. This must be set to use logging.
+	 * @param VisualLoggerAgentId The agent id associated with this observation.
+	 * @param VisualLoggerLocation A location for the visual logger information in the world.
+	 * @param VisualLoggerColor The color for the visual logger display.
+	 * @return The newly created observation object element.
+	 */
+	UFUNCTION(BlueprintPure, Category = "LearningAgents", meta = (AdvancedDisplay = 2, DefaultToSelf = "VisualLoggerListener"))
+	static FLearningAgentsObservationObjectElement MakeNamedInclusiveDiscreteObservation(
+		ULearningAgentsObservationObject* Object,
+		const TArray<FName>& ElementNames,
+		const FName Tag = TEXT("NamedInclusiveDiscreteObservation"),
+		const bool bVisualLoggerEnabled = false,
+		ULearningAgentsManagerListener* VisualLoggerListener = nullptr,
+		const int32 VisualLoggerAgentId = -1,
+		const FVector VisualLoggerLocation = FVector::ZeroVector,
+		const FLinearColor VisualLoggerColor = FLinearColor::Red);
+
+	/**
+	 * Make a new named inclusive discrete observation.
+	 *
+	 * @param Object The Observation Object
+	 * @param ElementNames The names of the discrete observations. 
+	 * @param Tag The tag of the corresponding observation. Must match the tag given during Specify.
+	 * @param bVisualLoggerEnabled When true, debug data will be sent to the visual logger.
+	 * @param VisualLoggerListener The listener object which is making this observation. This must be set to use logging.
+	 * @param VisualLoggerAgentId The agent id associated with this observation.
+	 * @param VisualLoggerLocation A location for the visual logger information in the world.
+	 * @param VisualLoggerColor The color for the visual logger display.
+	 * @return The newly created observation object element.
+	 */
+	static FLearningAgentsObservationObjectElement MakeNamedInclusiveDiscreteObservationFromArrayView(
+		ULearningAgentsObservationObject* Object,
+		const TArrayView<const FName> ElementNames,
+		const FName Tag = TEXT("NamedInclusiveDiscreteObservation"),
+		const bool bVisualLoggerEnabled = false,
+		ULearningAgentsManagerListener* VisualLoggerListener = nullptr,
+		const int32 VisualLoggerAgentId = -1,
+		const FVector VisualLoggerLocation = FVector::ZeroVector,
+		const FLinearColor VisualLoggerColor = FLinearColor::Red);
+
 	/**
 	 * Make a new count observation.
 	 *
@@ -1535,6 +1642,16 @@ public:
 	static bool GetInclusiveDiscreteObservation(TArray<int32>& OutIndices, const ULearningAgentsObservationObject* Object, const FLearningAgentsObservationObjectElement Element, const FName Tag = TEXT("InclusiveDiscreteObservation"));
 	static bool GetInclusiveDiscreteObservationToArrayView(TArrayView<int32> OutIndices, const ULearningAgentsObservationObject* Object, const FLearningAgentsObservationObjectElement Element, const FName Tag = TEXT("InclusiveDiscreteObservation"));
 
+	UFUNCTION(BlueprintPure = false, Category = "LearningAgents", meta = (AdvancedDisplay = 3, ReturnDisplayName = "Success"))
+	static bool GetNamedExclusiveDiscreteObservation(FName& OutName, const ULearningAgentsObservationObject* Object, const FLearningAgentsObservationObjectElement Element, const FName Tag = TEXT("ExclusiveDiscreteObservation"));
+
+	UFUNCTION(BlueprintPure = false, Category = "LearningAgents", meta = (AdvancedDisplay = 3, ReturnDisplayName = "Success"))
+	static bool GetNamedInclusiveDiscreteObservationNum(int32& OutNum, const ULearningAgentsObservationObject* Object, const FLearningAgentsObservationObjectElement Element, const FName Tag = TEXT("InclusiveDiscreteObservation"));
+
+	UFUNCTION(BlueprintPure = false, Category = "LearningAgents", meta = (AdvancedDisplay = 3, ReturnDisplayName = "Success"))
+	static bool GetNamedInclusiveDiscreteObservation(TArray<FName>& OutNames, const ULearningAgentsObservationObject* Object, const FLearningAgentsObservationObjectElement Element, const FName Tag = TEXT("InclusiveDiscreteObservation"));
+	static bool GetNamedInclusiveDiscreteObservationToArrayView(TArrayView<FName> OutNames, const ULearningAgentsObservationObject* Object, const FLearningAgentsObservationObjectElement Element, const FName Tag = TEXT("InclusiveDiscreteObservation"));
+
 	UFUNCTION(BlueprintPure = false, Category = "LearningAgents", meta = (AdvancedDisplay = 3, ReturnDisplayName = "Success"))
 	static bool GetCountObservation(int32& OutNum, const ULearningAgentsObservationObject* Object, const FLearningAgentsObservationObjectElement Element, const int32 MaxNum, const FName Tag = TEXT("CountObservation"));
 
diff --git a/Engine/Plugins/Experimental/LearningAgents/Source/LearningTraining/Private/LearningTrainer.cpp b/Engine/Plugins/Experimental/LearningAgents/Source/LearningTraining/Private/LearningTrainer.cpp
--- a/Engine/Plugins/Experimental/LearningAgents/Source/LearningTraining/Private/LearningTrainer.cpp
+++ b/Engine/Plugins/Experimental/LearningAgents/Source/LearningTraining/Private/LearningTrainer.cpp
@@ -125,6 +125,56 @@ namespace UE::Learning::Trainer
 			break;
 		}
 
+		case Observation::EType::DiscreteExclusive:
+		{
+			const Observation::FSchemaDiscreteExclusiveParameters Parameters = ObservationSchema.GetDiscreteExclusive(ObservationSchemaElement);
+
+			Object->SetStringField(TEXT("Type"), TEXT("DiscreteExclusive"));
+			Object->SetNumberField(TEXT("Num"), Parameters.Num);
+			break;
+		}
+
+		case Observation::EType::DiscreteInclusive:
+		{
+			const Observation::FSchemaDiscreteInclusiveParameters Parameters = ObservationSchema.GetDiscreteInclusive(ObservationSchemaElement);
+
+			Object->SetStringField(TEXT("Type"), TEXT("DiscreteInclusive"));
+			Object->SetNumberField(TEXT("Num"), Parameters.Num);
+			break;
+		}
+
+		case Observation::EType::NamedDiscreteExclusive:
+		{
+			const Observation::FSchemaNamedDiscreteExclusiveParameters Parameters = ObservationSchema.GetNamedDiscreteExclusive(ObservationSchemaElement);
+
+			Object->SetStringField(TEXT("Type"), TEXT("NamedDiscreteExclusive"));
+
+			TArray<TSharedPtr<FJsonValue>> ElementNames;
+			for (int32 ElementIdx = 0; ElementIdx < Parameters.ElementNames.Num(); ElementIdx++)
+			{
+				ElementNames.Add(MakeShared<FJsonValueString>(Parameters.ElementNames[ElementIdx].ToString()));
+			}
+
+			Object->SetArrayField(TEXT("ElementNames"), ElementNames);
+			break;
+		}
+
+		case Observation::EType::NamedDiscreteInclusive:
+		{
+			const Observation::FSchemaNamedDiscreteInclusiveParameters Parameters = ObservationSchema.GetNamedDiscreteInclusive(ObservationSchemaElement);
+
+			Object->SetStringField(TEXT("Type"), TEXT("NamedDiscreteInclusive"));
+
+			TArray<TSharedPtr<FJsonValue>> ElementNames;
+			for (int32 ElementIdx = 0; ElementIdx < Parameters.ElementNames.Num(); ElementIdx++)
+			{
+				ElementNames.Add(MakeShared<FJsonValueString>(Parameters.ElementNames[ElementIdx].ToString()));
+			}
+
+			Object->SetArrayField(TEXT("ElementNames"), ElementNames);
+			break;
+		}
+
 		case Observation::EType::And:
 		{
 			const Observation::FSchemaAndParameters Parameters = ObservationSchema.GetAnd(ObservationSchemaElement);
@@ -268,6 +318,38 @@ namespace UE::Learning::Trainer
 			break;
 		}
 
+		case Action::EType::NamedDiscreteExclusive:
+		{
+			const Action::FSchemaNamedDiscreteExclusiveParameters Parameters = ActionSchema.GetNamedDiscreteExclusive(ActionSchemaElement);
+
+			Object->SetStringField(TEXT("Type"), TEXT("NamedDiscreteExclusive"));
+
+			TArray<TSharedPtr<FJsonValue>> ElementNames;
+			for (int32 ElementIdx = 0; ElementIdx < Parameters.ElementNames.Num(); ElementIdx++)
+			{
+				ElementNames.Add(MakeShared<FJsonValueString>(Parameters.ElementNames[ElementIdx].ToString()));
+			}
+
+			Object->SetArrayField(TEXT("ElementNames"), ElementNames);
+			break;
+		}
+
+		case Action::EType::NamedDiscreteInclusive:
+		{
+			const Action::FSchemaNamedDiscreteInclusiveParameters Parameters = ActionSchema.GetNamedDiscreteInclusive(ActionSchemaElement);
+
+			Object->SetStringField(TEXT("Type"), TEXT("NamedDiscreteInclusive"));
+
+			TArray<TSharedPtr<FJsonValue>> ElementNames;
+			for (int32 ElementIdx = 0; ElementIdx < Parameters.ElementNames.Num(); ElementIdx++)
+			{
+				ElementNames.Add(MakeShared<FJsonValueString>(Parameters.ElementNames[ElementIdx].ToString()));
+			}
+
+			Object->SetArrayField(TEXT("ElementNames"), ElementNames);
+			break;
+		}
+
 		case Action::EType::And:
 		{
 			const Action::FSchemaAndParameters Parameters = ActionSchema.GetAnd(ActionSchemaElement);
