LearningAgents: Fixed bug in frame range set

diff --git a/Engine/Plugins/Experimental/LearningAgents/Source/Learning/Private/LearningFrameRangeSet.cpp b/Engine/Plugins/Experimental/LearningAgents/Source/Learning/Private/LearningFrameRangeSet.cpp
--- a/Engine/Plugins/Experimental/LearningAgents/Source/Learning/Private/LearningFrameRangeSet.cpp
+++ b/Engine/Plugins/Experimental/LearningAgents/Source/Learning/Private/LearningFrameRangeSet.cpp
@@ -117,7 +117,7 @@ namespace UE::Learning
 			for (int32 RangeIdx = 0; RangeIdx < RangeNum; RangeIdx++)
 			{
 				const float StartTime = Starts[RangeIdx] * FrameDeltaTime;
-				const float EndTime = (StartTime + Lengths[RangeIdx] - 1) * FrameDeltaTime;
+				const float EndTime = (Starts[RangeIdx] + Lengths[RangeIdx] - 1) * FrameDeltaTime;
 
 				if (Time >= StartTime && Time < EndTime)
 				{
@@ -166,7 +166,7 @@ namespace UE::Learning
 			for (int32 RangeIdx = 0; RangeIdx < RangeNum; RangeIdx++)
 			{
 				const float StartTime = Starts[RangeIdx] * FrameDeltaTime;
-				const float EndTime = (StartTime + Lengths[RangeIdx] - 1) * FrameDeltaTime;
+				const float EndTime = (Starts[RangeIdx] + Lengths[RangeIdx] - 1) * FrameDeltaTime;
 
 				if (Time >= StartTime && Time < EndTime)
 				{
