Learning Agents: Add visual logger support to MakeActions for controller debugging

#rnx

diff --git a/Engine/Plugins/Experimental/LearningAgents/Source/LearningAgents/Private/LearningAgentsActions.cpp b/Engine/Plugins/Experimental/LearningAgents/Source/LearningAgents/Private/LearningAgentsActions.cpp
--- a/Engine/Plugins/Experimental/LearningAgents/Source/LearningAgents/Private/LearningAgentsActions.cpp
+++ b/Engine/Plugins/Experimental/LearningAgents/Source/LearningAgents/Private/LearningAgentsActions.cpp
@@ -1839,12 +1839,28 @@ FLearningAgentsActionObjectElement ULearningAgentsActions::MakeNullAction(ULearn
 	return { Object->ActionObject.CreateNull(Tag) };
 }
 
-FLearningAgentsActionObjectElement ULearningAgentsActions::MakeContinuousAction(ULearningAgentsActionObject* Object, const TArray<float>& Values, const FName Tag)
+FLearningAgentsActionObjectElement ULearningAgentsActions::MakeContinuousAction(
+	ULearningAgentsActionObject* Object,
+	const TArray<float>& Values,
+	const FName Tag,
+	const bool bVisualLoggerEnabled,
+	ULearningAgentsManagerListener* VisualLoggerListener,
+	const int32 VisualLoggerAgentId,
+	const FVector VisualLoggerLocation,
+	const FLinearColor VisualLoggerColor)
 {
-	return MakeContinuousActionFromArrayView(Object, Values, Tag);
+	return MakeContinuousActionFromArrayView(Object, Values, Tag, bVisualLoggerEnabled, VisualLoggerListener, VisualLoggerAgentId, VisualLoggerLocation, VisualLoggerColor);
 }
 
-FLearningAgentsActionObjectElement ULearningAgentsActions::MakeContinuousActionFromArrayView(ULearningAgentsActionObject* Object, const TArrayView<const float> Values, const FName Tag)
+FLearningAgentsActionObjectElement ULearningAgentsActions::MakeContinuousActionFromArrayView(
+	ULearningAgentsActionObject* Object,
+	const TArrayView<const float> Values,
+	const FName Tag,
+	const bool bVisualLoggerEnabled,
+	ULearningAgentsManagerListener* VisualLoggerListener,
+	const int32 VisualLoggerAgentId,
+	const FVector VisualLoggerLocation,
+	const FLinearColor VisualLoggerColor)
 {
 	if (!Object)
 	{
@@ -1857,10 +1873,34 @@ FLearningAgentsActionObjectElement ULearningAgentsActions::MakeContinuousActionF
 		UE_LOG(LogLearning, Warning, TEXT("MakeContinuousActionFromArrayView: Creating zero-sized Continuous Action."));
 	}
 
+#if UE_LEARNING_AGENTS_ENABLE_VISUAL_LOG
+	if (bVisualLoggerEnabled && VisualLoggerListener)
+	{
+		const ULearningAgentsVisualLoggerObject* VisualLoggerObject = VisualLoggerListener->GetOrAddVisualLoggerObject(Tag);
+
+		UE_LEARNING_AGENTS_VLOG_STRING(VisualLoggerObject, LogLearning, Display, VisualLoggerLocation,
+			VisualLoggerColor.ToFColor(true),
+			TEXT("Listener: %s\nTag: %s\nAgent Id: % 3i\nEncoded: %s\nValues: %s"),
+			*VisualLoggerListener->GetName(),
+			*Tag.ToString(),
+			VisualLoggerAgentId,
+			*UE::Learning::Array::FormatFloat(Values),
+			*UE::Learning::Array::FormatFloat(Values)); // Encoded is identical to provided values
+	}
+#endif
+
 	return { Object->ActionObject.CreateContinuous({ Values }, Tag) };
 }
 
-FLearningAgentsActionObjectElement ULearningAgentsActions::MakeExclusiveDiscreteAction(ULearningAgentsActionObject* Object, const int32 Index, const FName Tag)
+FLearningAgentsActionObjectElement ULearningAgentsActions::MakeExclusiveDiscreteAction(
+	ULearningAgentsActionObject* Object,
+	const int32 Index,
+	const FName Tag,
+	const bool bVisualLoggerEnabled,
+	ULearningAgentsManagerListener* VisualLoggerListener,
+	const int32 VisualLoggerAgentId,
+	const FVector VisualLoggerLocation,
+	const FLinearColor VisualLoggerColor)
 {
 	if (!Object)
 	{
@@ -1874,10 +1914,33 @@ FLearningAgentsActionObjectElement ULearningAgentsActions::MakeExclusiveDiscrete
 		return FLearningAgentsActionObjectElement();
 	}
 
+#if UE_LEARNING_AGENTS_ENABLE_VISUAL_LOG
+	if (bVisualLoggerEnabled && VisualLoggerListener)
+	{
+		const ULearningAgentsVisualLoggerObject* VisualLoggerObject = VisualLoggerListener->GetOrAddVisualLoggerObject(Tag);
+
+		UE_LEARNING_AGENTS_VLOG_STRING(VisualLoggerObject, LogLearning, Display, VisualLoggerLocation,
+			VisualLoggerColor.ToFColor(true),
+			TEXT("Listener: %s\nTag: %s\nAgent Id: % 3i\nIndex: [%i]"),
+			*VisualLoggerListener->GetName(),
+			*Tag.ToString(),
+			VisualLoggerAgentId,
+			Index);
+	}
+#endif
+
 	return { Object->ActionObject.CreateDiscreteExclusive({ Index }, Tag) };
 }
 
-FLearningAgentsActionObjectElement ULearningAgentsActions::MakeNamedExclusiveDiscreteAction(ULearningAgentsActionObject* Object, const FName Name, const FName Tag)
+FLearningAgentsActionObjectElement ULearningAgentsActions::MakeNamedExclusiveDiscreteAction(
+	ULearningAgentsActionObject* Object,
+	const FName Name,
+	const FName Tag,
+	const bool bVisualLoggerEnabled,
+	ULearningAgentsManagerListener* VisualLoggerListener,
+	const int32 VisualLoggerAgentId,
+	const FVector VisualLoggerLocation,
+	const FLinearColor VisualLoggerColor)
 {
 	if (!Object)
 	{
@@ -1885,15 +1948,46 @@ FLearningAgentsActionObjectElement ULearningAgentsActions::MakeNamedExclusiveDis
 		return FLearningAgentsActionObjectElement();
 	}
 
+#if UE_LEARNING_AGENTS_ENABLE_VISUAL_LOG
+	if (bVisualLoggerEnabled && VisualLoggerListener)
+	{
+		const ULearningAgentsVisualLoggerObject* VisualLoggerObject = VisualLoggerListener->GetOrAddVisualLoggerObject(Tag);
+
+		UE_LEARNING_AGENTS_VLOG_STRING(VisualLoggerObject, LogLearning, Display, VisualLoggerLocation,
+			VisualLoggerColor.ToFColor(true),
+			TEXT("Listener: %s\nTag: %s\nAgent Id: % 3i\nName: [%s]"),
+			*VisualLoggerListener->GetName(),
+			*Tag.ToString(),
+			VisualLoggerAgentId,
+			*Name.ToString());
+	}
+#endif
+
 	return { Object->ActionObject.CreateNamedDiscreteExclusive({ Name }, Tag) };
 }
 
-FLearningAgentsActionObjectElement ULearningAgentsActions::MakeInclusiveDiscreteAction(ULearningAgentsActionObject* Object, const TArray<int32>& Indices, const FName Tag)
+FLearningAgentsActionObjectElement ULearningAgentsActions::MakeInclusiveDiscreteAction(
+	ULearningAgentsActionObject* Object,
+	const TArray<int32>& Indices,
+	const FName Tag,
+	const bool bVisualLoggerEnabled,
+	ULearningAgentsManagerListener* VisualLoggerListener,
+	const int32 VisualLoggerAgentId,
+	const FVector VisualLoggerLocation,
+	const FLinearColor VisualLoggerColor)
 {
-	return MakeInclusiveDiscreteActionFromArrayView(Object, Indices, Tag);
+	return MakeInclusiveDiscreteActionFromArrayView(Object, Indices, Tag, bVisualLoggerEnabled, VisualLoggerListener, VisualLoggerAgentId, VisualLoggerLocation, VisualLoggerColor);
 }
 
-FLearningAgentsActionObjectElement ULearningAgentsActions::MakeInclusiveDiscreteActionFromArrayView(ULearningAgentsActionObject* Object, const TArrayView<const int32> Indices, const FName Tag)
+FLearningAgentsActionObjectElement ULearningAgentsActions::MakeInclusiveDiscreteActionFromArrayView(
+	ULearningAgentsActionObject* Object,
+	const TArrayView<const int32> Indices,
+	const FName Tag,
+	const bool bVisualLoggerEnabled,
+	ULearningAgentsManagerListener* VisualLoggerListener,
+	const int32 VisualLoggerAgentId,
+	const FVector VisualLoggerLocation,
+	const FLinearColor VisualLoggerColor)
 {
 	if (!Object)
 	{
@@ -1918,15 +2012,46 @@ FLearningAgentsActionObjectElement ULearningAgentsActions::MakeInclusiveDiscrete
 		}
 	}
 
+#if UE_LEARNING_AGENTS_ENABLE_VISUAL_LOG
+	if (bVisualLoggerEnabled && VisualLoggerListener)
+	{
+		const ULearningAgentsVisualLoggerObject* VisualLoggerObject = VisualLoggerListener->GetOrAddVisualLoggerObject(Tag);
+
+		UE_LEARNING_AGENTS_VLOG_STRING(VisualLoggerObject, LogLearning, Display, VisualLoggerLocation,
+			VisualLoggerColor.ToFColor(true),
+			TEXT("Listener: %s\nTag: %s\nAgent Id: % 3i\nIndices: %s"),
+			*VisualLoggerListener->GetName(),
+			*Tag.ToString(),
+			VisualLoggerAgentId,
+			*UE::Learning::Array::FormatInt32(Indices, 256));
+	}
+#endif
+
 	return { Object->ActionObject.CreateDiscreteInclusive({ Indices }, Tag) };
 }
 
-FLearningAgentsActionObjectElement ULearningAgentsActions::MakeNamedInclusiveDiscreteAction(ULearningAgentsActionObject* Object, const TArray<FName>& Names, const FName Tag)
+FLearningAgentsActionObjectElement ULearningAgentsActions::MakeNamedInclusiveDiscreteAction(
+	ULearningAgentsActionObject* Object,
+	const TArray<FName>& Names,
+	const FName Tag,
+	const bool bVisualLoggerEnabled,
+	ULearningAgentsManagerListener* VisualLoggerListener,
+	const int32 VisualLoggerAgentId,
+	const FVector VisualLoggerLocation,
+	const FLinearColor VisualLoggerColor)
 {
-	return MakeNamedInclusiveDiscreteActionFromArrayView(Object, Names, Tag);
+	return MakeNamedInclusiveDiscreteActionFromArrayView(Object, Names, Tag, bVisualLoggerEnabled, VisualLoggerListener, VisualLoggerAgentId, VisualLoggerLocation, VisualLoggerColor);
 }
 
-FLearningAgentsActionObjectElement ULearningAgentsActions::MakeNamedInclusiveDiscreteActionFromArrayView(ULearningAgentsActionObject* Object, const TArrayView<const FName> Names, const FName Tag)
+FLearningAgentsActionObjectElement ULearningAgentsActions::MakeNamedInclusiveDiscreteActionFromArrayView(
+	ULearningAgentsActionObject* Object,
+	const TArrayView<const FName> Names,
+	const FName Tag,
+	const bool bVisualLoggerEnabled,
+	ULearningAgentsManagerListener* VisualLoggerListener,
+	const int32 VisualLoggerAgentId,
+	const FVector VisualLoggerLocation,
+	const FLinearColor VisualLoggerColor)
 {
 	if (!Object)
 	{
@@ -1940,6 +2065,21 @@ FLearningAgentsActionObjectElement ULearningAgentsActions::MakeNamedInclusiveDis
 		return FLearningAgentsActionObjectElement();
 	}
 
+#if UE_LEARNING_AGENTS_ENABLE_VISUAL_LOG
+	if (bVisualLoggerEnabled && VisualLoggerListener)
+	{
+		const ULearningAgentsVisualLoggerObject* VisualLoggerObject = VisualLoggerListener->GetOrAddVisualLoggerObject(Tag);
+
+		UE_LEARNING_AGENTS_VLOG_STRING(VisualLoggerObject, LogLearning, Display, VisualLoggerLocation,
+			VisualLoggerColor.ToFColor(true),
+			TEXT("Listener: %s\nTag: %s\nAgent Id: % 3i\nNames: %s"),
+			*VisualLoggerListener->GetName(),
+			*Tag.ToString(),
+			VisualLoggerAgentId,
+			*UE::Learning::Agents::Action::Private::FormatNames(Names, 256));
+	}
+#endif
+
 	return { Object->ActionObject.CreateNamedDiscreteInclusive({ Names }, Tag) };
 }
 
@@ -2130,7 +2270,16 @@ FLearningAgentsActionObjectElement ULearningAgentsActions::MakePairAction(ULearn
 	return MakeStructActionFromArrayViews(Object, { TEXT("Key"), TEXT("Value") }, { Key, Value }, Tag);
 }
 
-FLearningAgentsActionObjectElement ULearningAgentsActions::MakeEnumAction(ULearningAgentsActionObject* Object, const UEnum* Enum, const uint8 EnumValue, const FName Tag)
+FLearningAgentsActionObjectElement ULearningAgentsActions::MakeEnumAction(
+	ULearningAgentsActionObject* Object,
+	const UEnum* Enum,
+	const uint8 EnumValue,
+	const FName Tag,
+	const bool bVisualLoggerEnabled,
+	ULearningAgentsManagerListener* VisualLoggerListener,
+	const int32 VisualLoggerAgentId,
+	const FVector VisualLoggerLocation,
+	const FLinearColor VisualLoggerColor)
 {
 	if (!Enum)
 	{
@@ -2146,10 +2295,37 @@ FLearningAgentsActionObjectElement ULearningAgentsActions::MakeEnumAction(ULearn
 		return FLearningAgentsActionObjectElement();
 	}
 
+#if UE_LEARNING_AGENTS_ENABLE_VISUAL_LOG
+	if (bVisualLoggerEnabled && VisualLoggerListener)
+	{
+		const ULearningAgentsVisualLoggerObject* VisualLoggerObject = VisualLoggerListener->GetOrAddVisualLoggerObject(Tag);
+
+		UE_LEARNING_AGENTS_VLOG_STRING(VisualLoggerObject, LogLearning, Display, VisualLoggerLocation,
+			VisualLoggerColor.ToFColor(true),
+			TEXT("Listener: %s\nTag: %s\nAgent Id: % 3i\nEnum: %s\nSize: [%i]\nValue: [%s]\nIndex: [%i]"),
+			*VisualLoggerListener->GetName(),
+			*Tag.ToString(),
+			VisualLoggerAgentId,
+			*Enum->GetName(),
+			Enum->NumEnums() - 1,
+			*Enum->GetDisplayNameTextByValue(EnumValue).ToString(),
+			EnumValueIndex);
+	}
+#endif
+
 	return MakeExclusiveDiscreteAction(Object, EnumValueIndex, Tag);
 }
 
-FLearningAgentsActionObjectElement ULearningAgentsActions::MakeBitmaskAction(ULearningAgentsActionObject* Object, const UEnum* Enum, const int32 BitmaskValue, const FName Tag)
+FLearningAgentsActionObjectElement ULearningAgentsActions::MakeBitmaskAction(
+	ULearningAgentsActionObject* Object,
+	const UEnum* Enum,
+	const int32 BitmaskValue,
+	const FName Tag,
+	const bool bVisualLoggerEnabled,
+	ULearningAgentsManagerListener* VisualLoggerListener,
+	const int32 VisualLoggerAgentId,
+	const FVector VisualLoggerLocation,
+	const FLinearColor VisualLoggerColor)
 {
 	if (!Enum)
 	{
@@ -2174,6 +2350,39 @@ FLearningAgentsActionObjectElement ULearningAgentsActions::MakeBitmaskAction(ULe
 		}
 	}
 
+#if UE_LEARNING_AGENTS_ENABLE_VISUAL_LOG
+	if (bVisualLoggerEnabled && VisualLoggerListener)
+	{
+		const ULearningAgentsVisualLoggerObject* VisualLoggerObject = VisualLoggerListener->GetOrAddVisualLoggerObject(Tag);
+
+		FString ValuesString;
+		FString IndicesString;
+
+		for (int32 EnumIdx = 0; EnumIdx < Enum->NumEnums() - 1; EnumIdx++)
+		{
+			if (BitmaskValue & (1 << EnumIdx))
+			{
+				ValuesString += Enum->GetDisplayNameTextByIndex(EnumIdx).ToString() + TEXT(" ");
+				IndicesString += FString::FromInt(EnumIdx) + TEXT(" ");
+			}
+		}
+
+		ValuesString = ValuesString.TrimEnd();
+		IndicesString = IndicesString.TrimEnd();
+
+		UE_LEARNING_AGENTS_VLOG_STRING(VisualLoggerObject, LogLearning, Display, VisualLoggerLocation,
+			VisualLoggerColor.ToFColor(true),
+			TEXT("Listener: %s\nTag: %s\nAgent Id: % 3i\nEnum: %s\nSize: [%i]\nValues: [%s]\nIndices: [%s]"),
+			*VisualLoggerListener->GetName(),
+			*Tag.ToString(),
+			VisualLoggerAgentId,
+			*Enum->GetName(),
+			Enum->NumEnums() - 1,
+			*ValuesString,
+			*IndicesString);
+	}
+#endif
+
 	return MakeInclusiveDiscreteActionFromArrayView(Object, BitmaskIndices, Tag);
 }
 
@@ -2227,36 +2436,169 @@ FLearningAgentsActionObjectElement ULearningAgentsActions::MakeEncodingAction(UL
 	return { Object->ActionObject.CreateEncoding({ Element.ObjectElement }, Tag) };
 }
 
-FLearningAgentsActionObjectElement ULearningAgentsActions::MakeBoolAction(ULearningAgentsActionObject* Object, const bool bValue, const FName Tag)
+FLearningAgentsActionObjectElement ULearningAgentsActions::MakeBoolAction(
+	ULearningAgentsActionObject* Object,
+	const bool bValue,
+	const FName Tag,
+	const bool bVisualLoggerEnabled,
+	ULearningAgentsManagerListener* VisualLoggerListener,
+	const int32 VisualLoggerAgentId,
+	const FVector VisualLoggerLocation,
+	const FLinearColor VisualLoggerColor)
 {
+#if UE_LEARNING_AGENTS_ENABLE_VISUAL_LOG
+	if (bVisualLoggerEnabled && VisualLoggerListener)
+	{
+		const ULearningAgentsVisualLoggerObject* VisualLoggerObject = VisualLoggerListener->GetOrAddVisualLoggerObject(Tag);
+
+		UE_LEARNING_AGENTS_VLOG_STRING(VisualLoggerObject, LogLearning, Display, VisualLoggerLocation,
+			VisualLoggerColor.ToFColor(true),
+			TEXT("Listener: %s\nTag: %s\nAgent Id: % 3i\nValue: [%s]"),
+			*VisualLoggerListener->GetName(),
+			*Tag.ToString(),
+			VisualLoggerAgentId,
+			bValue ? TEXT("true") : TEXT("false"));
+	}
+#endif
+
 	return MakeExclusiveDiscreteAction(Object, bValue ? 1 : 0, Tag);
 }
 
-FLearningAgentsActionObjectElement ULearningAgentsActions::MakeFloatAction(ULearningAgentsActionObject* Object, const float Value, const FName Tag)
+FLearningAgentsActionObjectElement ULearningAgentsActions::MakeFloatAction(
+	ULearningAgentsActionObject* Object,
+	const float Value,
+	const FName Tag,
+	const bool bVisualLoggerEnabled,
+	ULearningAgentsManagerListener* VisualLoggerListener,
+	const int32 VisualLoggerAgentId,
+	const FVector VisualLoggerLocation,
+	const FLinearColor VisualLoggerColor)
 {
+#if UE_LEARNING_AGENTS_ENABLE_VISUAL_LOG
+	if (bVisualLoggerEnabled && VisualLoggerListener)
+	{
+		const ULearningAgentsVisualLoggerObject* VisualLoggerObject = VisualLoggerListener->GetOrAddVisualLoggerObject(Tag);
+
+		UE_LEARNING_AGENTS_VLOG_STRING(VisualLoggerObject, LogLearning, Display, VisualLoggerLocation,
+			VisualLoggerColor.ToFColor(true),
+			TEXT("Listener: %s\nTag: %s\nAgent Id: % 3i\nValue: [% 6.1f]"),
+			*VisualLoggerListener->GetName(),
+			*Tag.ToString(),
+			VisualLoggerAgentId,
+			Value);
+	}
+#endif
+
 	return MakeContinuousActionFromArrayView(Object, { Value }, Tag);
 }
 
-FLearningAgentsActionObjectElement ULearningAgentsActions::MakeLocationAction(ULearningAgentsActionObject* Object, const FVector Location, const FTransform RelativeTransform, const FName Tag)
+FLearningAgentsActionObjectElement ULearningAgentsActions::MakeLocationAction(
+	ULearningAgentsActionObject* Object,
+	const FVector Location,
+	const FTransform RelativeTransform,
+	const FName Tag,
+	const bool bVisualLoggerEnabled,
+	ULearningAgentsManagerListener* VisualLoggerListener,
+	const int32 VisualLoggerAgentId,
+	const FVector VisualLoggerLocation,
+	const FLinearColor VisualLoggerColor)
 {
 	const FVector LocalLocation = RelativeTransform.InverseTransformPosition(Location);
 
+#if UE_LEARNING_AGENTS_ENABLE_VISUAL_LOG
+	if (bVisualLoggerEnabled && VisualLoggerListener)
+	{
+		const ULearningAgentsVisualLoggerObject* VisualLoggerObject = VisualLoggerListener->GetOrAddVisualLoggerObject(Tag);
+
+		UE_LEARNING_AGENTS_VLOG_LOCATION(VisualLoggerObject, LogLearning, Display,
+			Location,
+			10,
+			VisualLoggerColor.ToFColor(true),
+			TEXT(""));
+
+		UE_LEARNING_AGENTS_VLOG_SEGMENT(VisualLoggerObject, LogLearning, Display,
+			RelativeTransform.GetTranslation(),
+			Location,
+			VisualLoggerColor.ToFColor(true),
+			TEXT(""));
+
+		UE_LEARNING_AGENTS_VLOG_TRANSFORM(VisualLoggerObject, LogLearning, Display,
+			RelativeTransform.GetTranslation(),
+			RelativeTransform.GetRotation(),
+			VisualLoggerColor.ToFColor(true),
+			TEXT(""));
+
+		UE_LEARNING_AGENTS_VLOG_STRING(VisualLoggerObject, LogLearning, Display, VisualLoggerLocation,
+			VisualLoggerColor.ToFColor(true),
+			TEXT("Listener: %s\nTag: %s\nAgent Id: % 3i\nLocal Location: [% 6.1f % 6.1f % 6.1f]\nLocation: [% 6.1f % 6.1f % 6.1f]"),
+			*VisualLoggerListener->GetName(),
+			*Tag.ToString(),
+			VisualLoggerAgentId,
+			LocalLocation.X, LocalLocation.Y, LocalLocation.Z,
+			Location.X, Location.Y, Location.Z);
+	}
+#endif
+
 	return MakeContinuousActionFromArrayView(Object, {
 		(float)LocalLocation.X,
 		(float)LocalLocation.Y,
 		(float)LocalLocation.Z }, Tag);
 }
 
-FLearningAgentsActionObjectElement ULearningAgentsActions::MakeRotationAction(ULearningAgentsActionObject* Object, const FRotator Rotation, const FRotator RelativeRotation, const FName Tag)
+FLearningAgentsActionObjectElement ULearningAgentsActions::MakeRotationAction(
+	ULearningAgentsActionObject* Object,
+	const FRotator Rotation,
+	const FRotator RelativeRotation,
+	const FName Tag,
+	const bool bVisualLoggerEnabled,
+	ULearningAgentsManagerListener* VisualLoggerListener,
+	const int32 VisualLoggerAgentId,
+	const FVector VisualLoggerRotationLocation,
+	const FVector VisualLoggerLocation,
+	const FLinearColor VisualLoggerColor)
 {
-	return MakeRotationActionFromQuat(Object, FQuat::MakeFromRotator(Rotation), FQuat::MakeFromRotator(RelativeRotation), Tag);
+	return MakeRotationActionFromQuat(Object, FQuat::MakeFromRotator(Rotation), FQuat::MakeFromRotator(RelativeRotation), Tag,
+		bVisualLoggerEnabled, VisualLoggerListener, VisualLoggerAgentId, VisualLoggerRotationLocation, VisualLoggerLocation, VisualLoggerColor);
 }
 
-FLearningAgentsActionObjectElement ULearningAgentsActions::MakeRotationActionFromQuat(ULearningAgentsActionObject* Object, const FQuat Rotation, const FQuat RelativeRotation, const FName Tag)
+FLearningAgentsActionObjectElement ULearningAgentsActions::MakeRotationActionFromQuat(
+	ULearningAgentsActionObject* Object,
+	const FQuat Rotation,
+	const FQuat RelativeRotation,
+	const FName Tag,
+	const bool bVisualLoggerEnabled,
+	ULearningAgentsManagerListener* VisualLoggerListener,
+	const int32 VisualLoggerAgentId,
+	const FVector VisualLoggerRotationLocation,
+	const FVector VisualLoggerLocation,
+	const FLinearColor VisualLoggerColor)
 {
 	FQuat LocalRotation = (RelativeRotation.Inverse() * Rotation).GetShortestArcWith(FQuat::Identity);
 	const FVector RotationVector = LocalRotation.ToRotationVector();
 
+#if UE_LEARNING_AGENTS_ENABLE_VISUAL_LOG
+	if (bVisualLoggerEnabled && VisualLoggerListener)
+	{
+		const ULearningAgentsVisualLoggerObject* VisualLoggerObject = VisualLoggerListener->GetOrAddVisualLoggerObject(Tag);
+
+		UE_LEARNING_AGENTS_VLOG_TRANSFORM(VisualLoggerObject, LogLearning, Display,
+			VisualLoggerRotationLocation,
+			LocalRotation.Rotator(),
+			VisualLoggerColor.ToFColor(true),
+			TEXT(""));
+
+		UE_LEARNING_AGENTS_VLOG_STRING(VisualLoggerObject, LogLearning, Display, VisualLoggerLocation,
+			VisualLoggerColor.ToFColor(true),
+			TEXT("Listener: %s\nTag: %s\nAgent Id: % 3i\nLocal Rotation Vector: [% 6.1f % 6.1f % 6.1f]\nLocal Rotation: [% 6.1f % 6.1f % 6.1f % 6.1f]\nRotation: [% 6.1f % 6.1f % 6.1f % 6.1f]"),
+			*VisualLoggerListener->GetName(),
+			*Tag.ToString(),
+			VisualLoggerAgentId,
+			RotationVector.X, RotationVector.Y, RotationVector.Z,
+			LocalRotation.X, LocalRotation.Y, LocalRotation.Z, LocalRotation.W,
+			Rotation.X, Rotation.Y, Rotation.Z, Rotation.W);
+	}
+#endif
+
 	return MakeContinuousActionFromArrayView(Object, {
 		(float)FMath::RadiansToDegrees(RotationVector.X),
 		(float)FMath::RadiansToDegrees(RotationVector.Y),
@@ -2264,12 +2606,37 @@ FLearningAgentsActionObjectElement ULearningAgentsActions::MakeRotationActionFro
 		}, Tag);
 }
 
-FLearningAgentsActionObjectElement ULearningAgentsActions::MakeScaleAction(ULearningAgentsActionObject* Object, const FVector Scale, const FVector RelativeScale, const FName Tag)
+FLearningAgentsActionObjectElement ULearningAgentsActions::MakeScaleAction(
+	ULearningAgentsActionObject* Object,
+	const FVector Scale,
+	const FVector RelativeScale,
+	const FName Tag,
+	const bool bVisualLoggerEnabled,
+	ULearningAgentsManagerListener* VisualLoggerListener,
+	const int32 VisualLoggerAgentId,
+	const FVector VisualLoggerLocation,
+	const FLinearColor VisualLoggerColor)
 {
 	const FVector LocalLogScale =
 		UE::Learning::Agents::Action::Private::VectorLogSafe(Scale) -
 		UE::Learning::Agents::Action::Private::VectorLogSafe(RelativeScale);
 
+#if UE_LEARNING_AGENTS_ENABLE_VISUAL_LOG
+	if (bVisualLoggerEnabled && VisualLoggerListener)
+	{
+		const ULearningAgentsVisualLoggerObject* VisualLoggerObject = VisualLoggerListener->GetOrAddVisualLoggerObject(Tag);
+
+		UE_LEARNING_AGENTS_VLOG_STRING(VisualLoggerObject, LogLearning, Display, VisualLoggerLocation,
+			VisualLoggerColor.ToFColor(true),
+			TEXT("Listener: %s\nTag: %s\nAgent Id: % 3i\nLocal Log Scale: [% 6.1f % 6.1f % 6.1f]\nScale: [% 6.1f % 6.1f % 6.1f]"),
+			*VisualLoggerListener->GetName(),
+			*Tag.ToString(),
+			VisualLoggerAgentId,
+			LocalLogScale.X, LocalLogScale.Y, LocalLogScale.Z,
+			Scale.X, Scale.Y, Scale.Z);
+	}
+#endif
+
 	return MakeContinuousActionFromArrayView(Object, {
 		(float)LocalLogScale.X,
 		(float)LocalLogScale.Y,
@@ -2277,8 +2644,44 @@ FLearningAgentsActionObjectElement ULearningAgentsActions::MakeScaleAction(ULear
 		}, Tag);
 }
 
-FLearningAgentsActionObjectElement ULearningAgentsActions::MakeTransformAction(ULearningAgentsActionObject* Object, const FTransform Transform, const FTransform RelativeTransform, const FName Tag)
+FLearningAgentsActionObjectElement ULearningAgentsActions::MakeTransformAction(
+	ULearningAgentsActionObject* Object,
+	const FTransform Transform,
+	const FTransform RelativeTransform,
+	const FName Tag,
+	const bool bVisualLoggerEnabled,
+	ULearningAgentsManagerListener* VisualLoggerListener,
+	const int32 VisualLoggerAgentId,
+	const FVector VisualLoggerLocation,
+	const FLinearColor VisualLoggerColor)
 {
+	FVector Location = Transform.GetLocation();
+	FQuat Rotation = Transform.GetRotation();
+	FVector Scale = Transform.GetScale3D();
+
+#if UE_LEARNING_AGENTS_ENABLE_VISUAL_LOG
+	if (bVisualLoggerEnabled && VisualLoggerListener)
+	{
+		const ULearningAgentsVisualLoggerObject* VisualLoggerObject = VisualLoggerListener->GetOrAddVisualLoggerObject(Tag);
+
+		UE_LEARNING_AGENTS_VLOG_TRANSFORM(VisualLoggerObject, LogLearning, Display,
+			Location,
+			Rotation,
+			VisualLoggerColor.ToFColor(true),
+			TEXT(""));
+
+		UE_LEARNING_AGENTS_VLOG_STRING(VisualLoggerObject, LogLearning, Display, VisualLoggerLocation,
+			VisualLoggerColor.ToFColor(true),
+			TEXT("Listener: %s\nTag: %s\nAgent Id: % 3i\nLocation: [% 6.1f % 6.1f % 6.1f]\nRotation: [% 6.1f % 6.1f % 6.1f % 6.1f]\nScale: [% 6.1f % 6.1f % 6.1f]"),
+			*VisualLoggerListener->GetName(),
+			*Tag.ToString(),
+			VisualLoggerAgentId,
+			Location.X, Location.Y, Location.Z,
+			Rotation.X, Rotation.Y, Rotation.Z, Rotation.W,
+			Scale.X, Scale.Y, Scale.Z);
+	}
+#endif
+
 	return MakeStructActionFromArrayViews(Object,
 		{
 			TEXT("Location"),
@@ -2286,27 +2689,112 @@ FLearningAgentsActionObjectElement ULearningAgentsActions::MakeTransformAction(U
 			TEXT("Scale")
 		},
 		{
-			MakeLocationAction(Object, Transform.GetLocation(), RelativeTransform),
-			MakeRotationActionFromQuat(Object, Transform.GetRotation(), RelativeTransform.GetRotation()),
-			MakeScaleAction(Object, Transform.GetScale3D(), RelativeTransform.GetScale3D())
+			MakeLocationAction(Object, Location, RelativeTransform),
+			MakeRotationActionFromQuat(Object, Rotation, RelativeTransform.GetRotation()),
+			MakeScaleAction(Object, Scale, RelativeTransform.GetScale3D())
 		},
 		Tag);
 }
 
-FLearningAgentsActionObjectElement ULearningAgentsActions::MakeAngleAction(ULearningAgentsActionObject* Object, const float Angle, const float RelativeAngle, const FName Tag)
+FLearningAgentsActionObjectElement ULearningAgentsActions::MakeAngleAction(
+	ULearningAgentsActionObject* Object,
+	const float Angle,
+	const float RelativeAngle,
+	const FName Tag,
+	const bool bVisualLoggerEnabled,
+	ULearningAgentsManagerListener* VisualLoggerListener,
+	const int32 VisualLoggerAgentId,
+	const FVector VisualLoggerAngleLocation,
+	const FVector VisualLoggerLocation,
+	const FLinearColor VisualLoggerColor)
 {
-	return MakeContinuousActionFromArrayView(Object, { FMath::FindDeltaAngleDegrees(RelativeAngle, Angle) }, Tag);
+	float LocalAngle = FMath::FindDeltaAngleDegrees(RelativeAngle, Angle);
+	
+#if UE_LEARNING_AGENTS_ENABLE_VISUAL_LOG
+	if (bVisualLoggerEnabled && VisualLoggerListener)
+	{
+		const ULearningAgentsVisualLoggerObject* VisualLoggerObject = VisualLoggerListener->GetOrAddVisualLoggerObject(Tag);
+
+		UE_LEARNING_AGENTS_VLOG_ANGLE_DEGREES(VisualLoggerObject, LogLearning, Display,
+			Angle,
+			0.0f,
+			VisualLoggerAngleLocation,
+			10.0f,
+			VisualLoggerColor.ToFColor(true),
+			TEXT(""));
+
+		UE_LEARNING_AGENTS_VLOG_STRING(VisualLoggerObject, LogLearning, Display, VisualLoggerLocation,
+			VisualLoggerColor.ToFColor(true),
+			TEXT("Listener: %s\nTag: %s\nAgent Id: % 3i\nLocal Angle: [% 6.1f]\nAngle: [% 6.1f]"),
+			*VisualLoggerListener->GetName(),
+			*Tag.ToString(),
+			VisualLoggerAgentId,
+			LocalAngle,
+			Angle);
+	}
+#endif
+
+	return MakeContinuousActionFromArrayView(Object, { LocalAngle }, Tag);
 }
 
-FLearningAgentsActionObjectElement ULearningAgentsActions::MakeAngleActionRadians(ULearningAgentsActionObject* Object, const float Angle, const float RelativeAngle, const FName Tag)
+FLearningAgentsActionObjectElement ULearningAgentsActions::MakeAngleActionRadians(
+	ULearningAgentsActionObject* Object,
+	const float Angle,
+	const float RelativeAngle,
+	const FName Tag,
+	const bool bVisualLoggerEnabled,
+	ULearningAgentsManagerListener* VisualLoggerListener,
+	const int32 VisualLoggerAgentId,
+	const FVector VisualLoggerAngleLocation,
+	const FVector VisualLoggerLocation,
+	const FLinearColor VisualLoggerColor)
 {
-	return MakeAngleAction(Object, FMath::RadiansToDegrees(Angle), FMath::RadiansToDegrees(RelativeAngle), Tag);
+	return MakeAngleAction(Object, FMath::RadiansToDegrees(Angle), FMath::RadiansToDegrees(RelativeAngle), Tag,
+		bVisualLoggerEnabled, VisualLoggerListener, VisualLoggerAgentId, VisualLoggerAngleLocation, VisualLoggerLocation, VisualLoggerColor);
 }
 
-FLearningAgentsActionObjectElement ULearningAgentsActions::MakeVelocityAction(ULearningAgentsActionObject* Object, const FVector Velocity, const FTransform RelativeTransform, const FName Tag)
+FLearningAgentsActionObjectElement ULearningAgentsActions::MakeVelocityAction(
+	ULearningAgentsActionObject* Object,
+	const FVector Velocity,
+	const FTransform RelativeTransform,
+	const FName Tag,
+	const bool bVisualLoggerEnabled,
+	ULearningAgentsManagerListener* VisualLoggerListener,
+	const int32 VisualLoggerAgentId,
+	const FVector VisualLoggerVelocityLocation,
+	const FVector VisualLoggerLocation,
+	const FLinearColor VisualLoggerColor)
 {
 	const FVector LocalVelocity = RelativeTransform.InverseTransformVectorNoScale(Velocity);
 
+#if UE_LEARNING_AGENTS_ENABLE_VISUAL_LOG
+	if (bVisualLoggerEnabled && VisualLoggerListener)
+	{
+		const ULearningAgentsVisualLoggerObject* VisualLoggerObject = VisualLoggerListener->GetOrAddVisualLoggerObject(Tag);
+
+		UE_LEARNING_AGENTS_VLOG_ARROW(VisualLoggerObject, LogLearning, Display,
+			VisualLoggerVelocityLocation,
+			VisualLoggerVelocityLocation + Velocity,
+			VisualLoggerColor.ToFColor(true),
+			TEXT(""));
+
+		UE_LEARNING_AGENTS_VLOG_TRANSFORM(VisualLoggerObject, LogLearning, Display,
+			RelativeTransform.GetTranslation(),
+			RelativeTransform.GetRotation(),
+			VisualLoggerColor.ToFColor(true),
+			TEXT(""));
+
+		UE_LEARNING_AGENTS_VLOG_STRING(VisualLoggerObject, LogLearning, Display, VisualLoggerLocation,
+			VisualLoggerColor.ToFColor(true),
+			TEXT("Listener: %s\nTag: %s\nAgent Id: % 3i\nLocal Velocity: [% 6.1f % 6.1f % 6.1f]\nVelocity: [% 6.1f % 6.1f % 6.1f]"),
+			*VisualLoggerListener->GetName(),
+			*Tag.ToString(),
+			VisualLoggerAgentId,
+			LocalVelocity.X, LocalVelocity.Y, LocalVelocity.Z,
+			Velocity.X, Velocity.Y, Velocity.Z);
+	}
+#endif
+
 	return MakeContinuousActionFromArrayView(Object, {
 		(float)LocalVelocity.X,
 		(float)LocalVelocity.Y,
@@ -2314,10 +2802,49 @@ FLearningAgentsActionObjectElement ULearningAgentsActions::MakeVelocityAction(UL
 		}, Tag);
 }
 
-FLearningAgentsActionObjectElement ULearningAgentsActions::MakeDirectionAction(ULearningAgentsActionObject* Object, const FVector Direction, const FTransform RelativeTransform, const FName Tag)
+FLearningAgentsActionObjectElement ULearningAgentsActions::MakeDirectionAction(
+	ULearningAgentsActionObject* Object,
+	const FVector Direction,
+	const FTransform RelativeTransform,
+	const FName Tag,
+	const bool bVisualLoggerEnabled,
+	ULearningAgentsManagerListener* VisualLoggerListener,
+	const int32 VisualLoggerAgentId,
+	const FVector VisualLoggerDirectionLocation,
+	const FVector VisualLoggerLocation,
+	const float VisualLoggerArrowLength,
+	const FLinearColor VisualLoggerColor)
 {
 	const FVector LocalDirection = RelativeTransform.InverseTransformVectorNoScale(Direction).GetSafeNormal(UE_SMALL_NUMBER, FVector::ForwardVector);
 
+#if UE_LEARNING_AGENTS_ENABLE_VISUAL_LOG
+	if (bVisualLoggerEnabled && VisualLoggerListener)
+	{
+		const ULearningAgentsVisualLoggerObject* VisualLoggerObject = VisualLoggerListener->GetOrAddVisualLoggerObject(Tag);
+
+		UE_LEARNING_AGENTS_VLOG_ARROW(VisualLoggerObject, LogLearning, Display,
+			VisualLoggerDirectionLocation,
+			VisualLoggerDirectionLocation + VisualLoggerArrowLength * Direction,
+			VisualLoggerColor.ToFColor(true),
+			TEXT(""));
+
+		UE_LEARNING_AGENTS_VLOG_TRANSFORM(VisualLoggerObject, LogLearning, Display,
+			RelativeTransform.GetTranslation(),
+			RelativeTransform.GetRotation(),
+			VisualLoggerColor.ToFColor(true),
+			TEXT(""));
+
+		UE_LEARNING_AGENTS_VLOG_STRING(VisualLoggerObject, LogLearning, Display, VisualLoggerLocation,
+			VisualLoggerColor.ToFColor(true),
+			TEXT("Listener: %s\nTag: %s\nAgent Id: % 3i\nLocal Direction: [% 6.1f % 6.1f % 6.1f]\nDirection: [% 6.1f % 6.1f % 6.1f]"),
+			*VisualLoggerListener->GetName(),
+			*Tag.ToString(),
+			VisualLoggerAgentId,
+			LocalDirection.X, LocalDirection.Y, LocalDirection.Z,
+			Direction.X, Direction.Y, Direction.Z);
+	}
+#endif
+
 	return MakeContinuousActionFromArrayView(Object, {
 		(float)LocalDirection.X,
 		(float)LocalDirection.Y,
@@ -3999,7 +4526,6 @@ bool ULearningAgentsActions::GetEnumAction(
 
 	OutEnumValue = (uint8)EnumValue;
 
-
 #if UE_LEARNING_AGENTS_ENABLE_VISUAL_LOG
 	if (bVisualLoggerEnabled && VisualLoggerListener)
 	{
diff --git a/Engine/Plugins/Experimental/LearningAgents/Source/LearningAgents/Public/LearningAgentsActions.h b/Engine/Plugins/Experimental/LearningAgents/Source/LearningAgents/Public/LearningAgentsActions.h
--- a/Engine/Plugins/Experimental/LearningAgents/Source/LearningAgents/Public/LearningAgentsActions.h
+++ b/Engine/Plugins/Experimental/LearningAgents/Source/LearningAgents/Public/LearningAgentsActions.h
@@ -679,23 +679,90 @@ public:
 	UFUNCTION(BlueprintPure, Category = "LearningAgents", meta = (AdvancedDisplay = 1))
 	static FLearningAgentsActionObjectElement MakeNullAction(ULearningAgentsActionObject* Object, const FName Tag = TEXT("NullAction"));
 
-	UFUNCTION(BlueprintPure, Category = "LearningAgents", meta = (AdvancedDisplay = 2))
-	static FLearningAgentsActionObjectElement MakeContinuousAction(ULearningAgentsActionObject* Object, const TArray<float>& Values, const FName Tag = TEXT("ContinuousAction"));
-	static FLearningAgentsActionObjectElement MakeContinuousActionFromArrayView(ULearningAgentsActionObject* Object, const TArrayView<const float> Values, const FName Tag = TEXT("ContinuousAction"));
+	UFUNCTION(BlueprintPure, Category = "LearningAgents", meta = (AdvancedDisplay = 2, DefaultToSelf = "VisualLoggerListener"))
+	static FLearningAgentsActionObjectElement MakeContinuousAction(
+		ULearningAgentsActionObject* Object,
+		const TArray<float>& Values,
+		const FName Tag = TEXT("ContinuousAction"),
+		const bool bVisualLoggerEnabled = false,
+		ULearningAgentsManagerListener* VisualLoggerListener = nullptr,
+		const int32 VisualLoggerAgentId = -1,
+		const FVector VisualLoggerLocation = FVector::ZeroVector,
+		const FLinearColor VisualLoggerColor = FLinearColor::Blue);
 
-	UFUNCTION(BlueprintPure, Category = "LearningAgents", meta = (AdvancedDisplay = 2))
-	static FLearningAgentsActionObjectElement MakeExclusiveDiscreteAction(ULearningAgentsActionObject* Object, const int32 Index, const FName Tag = TEXT("DiscreteExclusiveAction"));
+	static FLearningAgentsActionObjectElement MakeContinuousActionFromArrayView(
+		ULearningAgentsActionObject* Object,
+		const TArrayView<const float> Values,
+		const FName Tag = TEXT("ContinuousAction"),
+		const bool bVisualLoggerEnabled = false,
+		ULearningAgentsManagerListener* VisualLoggerListener = nullptr,
+		const int32 VisualLoggerAgentId = -1,
+		const FVector VisualLoggerLocation = FVector::ZeroVector,
+		const FLinearColor VisualLoggerColor = FLinearColor::Blue);
 
-	UFUNCTION(BlueprintPure, Category = "LearningAgents", meta = (AdvancedDisplay = 2))
-	static FLearningAgentsActionObjectElement MakeNamedExclusiveDiscreteAction(ULearningAgentsActionObject* Object, const FName Name, const FName Tag = TEXT("NamedDiscreteExclusiveAction"));
+	UFUNCTION(BlueprintPure, Category = "LearningAgents", meta = (AdvancedDisplay = 2, DefaultToSelf = "VisualLoggerListener"))
+	static FLearningAgentsActionObjectElement MakeExclusiveDiscreteAction(
+		ULearningAgentsActionObject* Object,
+		const int32 Index,
+		const FName Tag = TEXT("DiscreteExclusiveAction"),
+		const bool bVisualLoggerEnabled = false,
+		ULearningAgentsManagerListener* VisualLoggerListener = nullptr,
+		const int32 VisualLoggerAgentId = -1,
+		const FVector VisualLoggerLocation = FVector::ZeroVector,
+		const FLinearColor VisualLoggerColor = FLinearColor::Blue);
 
-	UFUNCTION(BlueprintPure, Category = "LearningAgents", meta = (AdvancedDisplay = 2))
-	static FLearningAgentsActionObjectElement MakeInclusiveDiscreteAction(ULearningAgentsActionObject* Object, const TArray<int32>& Indices, const FName Tag = TEXT("DiscreteInclusiveAction"));
-	static FLearningAgentsActionObjectElement MakeInclusiveDiscreteActionFromArrayView(ULearningAgentsActionObject* Object, const TArrayView<const int32> Indices, const FName Tag = TEXT("DiscreteInclusiveAction"));
+	UFUNCTION(BlueprintPure, Category = "LearningAgents", meta = (AdvancedDisplay = 2, DefaultToSelf = "VisualLoggerListener"))
+	static FLearningAgentsActionObjectElement MakeNamedExclusiveDiscreteAction(
+		ULearningAgentsActionObject* Object,
+		const FName Name,
+		const FName Tag = TEXT("NamedDiscreteExclusiveAction"),
+		const bool bVisualLoggerEnabled = false,
+		ULearningAgentsManagerListener* VisualLoggerListener = nullptr,
+		const int32 VisualLoggerAgentId = -1,
+		const FVector VisualLoggerLocation = FVector::ZeroVector,
+		const FLinearColor VisualLoggerColor = FLinearColor::Blue);
 
-	UFUNCTION(BlueprintPure, Category = "LearningAgents", meta = (AdvancedDisplay = 2))
-	static FLearningAgentsActionObjectElement MakeNamedInclusiveDiscreteAction(ULearningAgentsActionObject* Object, const TArray<FName>& Names, const FName Tag = TEXT("NamedDiscreteInclusiveAction"));
-	static FLearningAgentsActionObjectElement MakeNamedInclusiveDiscreteActionFromArrayView(ULearningAgentsActionObject* Object, const TArrayView<const FName> Names, const FName Tag = TEXT("NamedDiscreteInclusiveAction"));
+	UFUNCTION(BlueprintPure, Category = "LearningAgents", meta = (AdvancedDisplay = 2, DefaultToSelf = "VisualLoggerListener"))
+	static FLearningAgentsActionObjectElement MakeInclusiveDiscreteAction(
+		ULearningAgentsActionObject* Object,
+		const TArray<int32>& Indices,
+		const FName Tag = TEXT("DiscreteInclusiveAction"),
+		const bool bVisualLoggerEnabled = false,
+		ULearningAgentsManagerListener* VisualLoggerListener = nullptr,
+		const int32 VisualLoggerAgentId = -1,
+		const FVector VisualLoggerLocation = FVector::ZeroVector,
+		const FLinearColor VisualLoggerColor = FLinearColor::Blue);
+
+	static FLearningAgentsActionObjectElement MakeInclusiveDiscreteActionFromArrayView(
+		ULearningAgentsActionObject* Object,
+		const TArrayView<const int32> Indices,
+		const FName Tag = TEXT("DiscreteInclusiveAction"),
+		const bool bVisualLoggerEnabled = false,
+		ULearningAgentsManagerListener* VisualLoggerListener = nullptr,
+		const int32 VisualLoggerAgentId = -1,
+		const FVector VisualLoggerLocation = FVector::ZeroVector,
+		const FLinearColor VisualLoggerColor = FLinearColor::Blue);
+
+	UFUNCTION(BlueprintPure, Category = "LearningAgents", meta = (AdvancedDisplay = 2, DefaultToSelf = "VisualLoggerListener"))
+	static FLearningAgentsActionObjectElement MakeNamedInclusiveDiscreteAction(
+		ULearningAgentsActionObject* Object,
+		const TArray<FName>& Names,
+		const FName Tag = TEXT("NamedDiscreteInclusiveAction"),
+		const bool bVisualLoggerEnabled = false,
+		ULearningAgentsManagerListener* VisualLoggerListener = nullptr,
+		const int32 VisualLoggerAgentId = -1,
+		const FVector VisualLoggerLocation = FVector::ZeroVector,
+		const FLinearColor VisualLoggerColor = FLinearColor::Blue);
+
+	static FLearningAgentsActionObjectElement MakeNamedInclusiveDiscreteActionFromArrayView(
+		ULearningAgentsActionObject* Object,
+		const TArrayView<const FName> Names,
+		const FName Tag = TEXT("NamedDiscreteInclusiveAction"),
+		const bool bVisualLoggerEnabled = false,
+		ULearningAgentsManagerListener* VisualLoggerListener = nullptr,
+		const int32 VisualLoggerAgentId = -1,
+		const FVector VisualLoggerLocation = FVector::ZeroVector,
+		const FLinearColor VisualLoggerColor = FLinearColor::Blue);
 
 	UFUNCTION(BlueprintPure, Category = "LearningAgents", meta = (AdvancedDisplay = 2))
 	static FLearningAgentsActionObjectElement MakeStructAction(ULearningAgentsActionObject* Object, const TMap<FName, FLearningAgentsActionObjectElement>& Elements, const FName Tag = TEXT("StructAction"));
@@ -721,11 +788,29 @@ public:
 	UFUNCTION(BlueprintPure, Category = "LearningAgents", meta = (AdvancedDisplay = 3))
 	static FLearningAgentsActionObjectElement MakePairAction(ULearningAgentsActionObject* Object, const FLearningAgentsActionObjectElement Key, const FLearningAgentsActionObjectElement Value, const FName Tag = TEXT("PairAction"));
 
-	UFUNCTION(BlueprintPure, Category = "LearningAgents", meta = (AdvancedDisplay = 2))
-	static FLearningAgentsActionObjectElement MakeEnumAction(ULearningAgentsActionObject* Object, const UEnum* Enum, const uint8 EnumValue, const FName Tag = TEXT("EnumAction"));
+	UFUNCTION(BlueprintPure, Category = "LearningAgents", meta = (AdvancedDisplay = 2, DefaultToSelf = "VisualLoggerListener"))
+	static FLearningAgentsActionObjectElement MakeEnumAction(
+		ULearningAgentsActionObject* Object,
+		const UEnum* Enum,
+		const uint8 EnumValue,
+		const FName Tag = TEXT("EnumAction"),
+		const bool bVisualLoggerEnabled = false,
+		ULearningAgentsManagerListener* VisualLoggerListener = nullptr,
+		const int32 VisualLoggerAgentId = -1,
+		const FVector VisualLoggerLocation = FVector::ZeroVector,
+		const FLinearColor VisualLoggerColor = FLinearColor::Blue);
 
-	UFUNCTION(BlueprintPure, Category = "LearningAgents", meta = (AdvancedDisplay = 2))
-	static FLearningAgentsActionObjectElement MakeBitmaskAction(ULearningAgentsActionObject* Object, const UEnum* Enum, const int32 BitmaskValue, const FName Tag = TEXT("BitmaskAction"));
+	UFUNCTION(BlueprintPure, Category = "LearningAgents", meta = (AdvancedDisplay = 2, DefaultToSelf = "VisualLoggerListener"))
+	static FLearningAgentsActionObjectElement MakeBitmaskAction(
+		ULearningAgentsActionObject* Object,
+		const UEnum* Enum,
+		const int32 BitmaskValue,
+		const FName Tag = TEXT("BitmaskAction"),
+		const bool bVisualLoggerEnabled = false,
+		ULearningAgentsManagerListener* VisualLoggerListener = nullptr,
+		const int32 VisualLoggerAgentId = -1,
+		const FVector VisualLoggerLocation = FVector::ZeroVector,
+		const FLinearColor VisualLoggerColor = FLinearColor::Blue);
 
 	UFUNCTION(BlueprintPure, Category = "LearningAgents", meta = (AdvancedDisplay = 3))
 	static FLearningAgentsActionObjectElement MakeOptionalAction(ULearningAgentsActionObject* Object, const FLearningAgentsActionObjectElement Element, const ELearningAgentsOptionalAction Option, const FName Tag = TEXT("OptionalAction"));
@@ -748,38 +833,142 @@ public:
 	UFUNCTION(BlueprintPure, Category = "LearningAgents", meta = (AdvancedDisplay = 2))
 	static FLearningAgentsActionObjectElement MakeEncodingAction(ULearningAgentsActionObject* Object, const FLearningAgentsActionObjectElement Element, const FName Tag = TEXT("EncodingAction"));
 
-	UFUNCTION(BlueprintPure, Category = "LearningAgents", meta = (AdvancedDisplay = 2))
-	static FLearningAgentsActionObjectElement MakeBoolAction(ULearningAgentsActionObject* Object, const bool bValue, const FName Tag = TEXT("BoolAction"));
+	UFUNCTION(BlueprintPure, Category = "LearningAgents", meta = (AdvancedDisplay = 2, DefaultToSelf = "VisualLoggerListener"))
+	static FLearningAgentsActionObjectElement MakeBoolAction(
+		ULearningAgentsActionObject* Object,
+		const bool bValue,
+		const FName Tag = TEXT("BoolAction"),
+		const bool bVisualLoggerEnabled = false,
+		ULearningAgentsManagerListener * VisualLoggerListener = nullptr,
+		const int32 VisualLoggerAgentId = -1,
+		const FVector VisualLoggerLocation = FVector::ZeroVector,
+		const FLinearColor VisualLoggerColor = FLinearColor::Blue);
 
-	UFUNCTION(BlueprintPure, Category = "LearningAgents", meta = (AdvancedDisplay = 2))
-	static FLearningAgentsActionObjectElement MakeFloatAction(ULearningAgentsActionObject* Object, const float Value, const FName Tag = TEXT("FloatAction"));
+	UFUNCTION(BlueprintPure, Category = "LearningAgents", meta = (AdvancedDisplay = 2, DefaultToSelf = "VisualLoggerListener"))
+	static FLearningAgentsActionObjectElement MakeFloatAction(
+		ULearningAgentsActionObject* Object,
+		const float Value,
+		const FName Tag = TEXT("FloatAction"),
+		const bool bVisualLoggerEnabled = false,
+		ULearningAgentsManagerListener* VisualLoggerListener = nullptr,
+		const int32 VisualLoggerAgentId = -1,
+		const FVector VisualLoggerLocation = FVector::ZeroVector,
+		const FLinearColor VisualLoggerColor = FLinearColor::Blue);
 
-	UFUNCTION(BlueprintPure, Category = "LearningAgents", meta = (AdvancedDisplay = 3))
-	static FLearningAgentsActionObjectElement MakeLocationAction(ULearningAgentsActionObject* Object, const FVector Location, const FTransform RelativeTransform = FTransform(), const FName Tag = TEXT("LocationAction"));
+	UFUNCTION(BlueprintPure, Category = "LearningAgents", meta = (AdvancedDisplay = 3, DefaultToSelf = "VisualLoggerListener"))
+	static FLearningAgentsActionObjectElement MakeLocationAction(
+		ULearningAgentsActionObject* Object,
+		const FVector Location,
+		const FTransform RelativeTransform = FTransform(),
+		const FName Tag = TEXT("LocationAction"),
+		const bool bVisualLoggerEnabled = false,
+		ULearningAgentsManagerListener* VisualLoggerListener = nullptr,
+		const int32 VisualLoggerAgentId = -1,
+		const FVector VisualLoggerLocation = FVector::ZeroVector,
+		const FLinearColor VisualLoggerColor = FLinearColor::Blue);
 
-	UFUNCTION(BlueprintPure, Category = "LearningAgents", meta = (AdvancedDisplay = 3))
-	static FLearningAgentsActionObjectElement MakeRotationAction(ULearningAgentsActionObject* Object, const FRotator Rotation, const FRotator RelativeRotation = FRotator::ZeroRotator, const FName Tag = TEXT("RotationAction"));
+	UFUNCTION(BlueprintPure, Category = "LearningAgents", meta = (AdvancedDisplay = 3, DefaultToSelf = "VisualLoggerListener"))
+	static FLearningAgentsActionObjectElement MakeRotationAction(
+		ULearningAgentsActionObject* Object,
+		const FRotator Rotation,
+		const FRotator RelativeRotation = FRotator::ZeroRotator,
+		const FName Tag = TEXT("RotationAction"),
+		const bool bVisualLoggerEnabled = false,
+		ULearningAgentsManagerListener* VisualLoggerListener = nullptr,
+		const int32 VisualLoggerAgentId = -1,
+		const FVector VisualLoggerRotationLocation = FVector::ZeroVector,
+		const FVector VisualLoggerLocation = FVector::ZeroVector,
+		const FLinearColor VisualLoggerColor = FLinearColor::Blue);
 
-	UFUNCTION(BlueprintPure, Category = "LearningAgents", meta = (AdvancedDisplay = 3))
-	static FLearningAgentsActionObjectElement MakeRotationActionFromQuat(ULearningAgentsActionObject* Object, const FQuat Rotation, const FQuat RelativeRotation, const FName Tag = TEXT("RotationAction"));
+	UFUNCTION(BlueprintPure, Category = "LearningAgents", meta = (AdvancedDisplay = 3, DefaultToSelf = "VisualLoggerListener"))
+	static FLearningAgentsActionObjectElement MakeRotationActionFromQuat(
+		ULearningAgentsActionObject* Object,
+		const FQuat Rotation,
+		const FQuat RelativeRotation,
+		const FName Tag = TEXT("RotationAction"),
+		const bool bVisualLoggerEnabled = false,
+		ULearningAgentsManagerListener* VisualLoggerListener = nullptr,
+		const int32 VisualLoggerAgentId = -1,
+		const FVector VisualLoggerRotationLocation = FVector::ZeroVector,
+		const FVector VisualLoggerLocation = FVector::ZeroVector,
+		const FLinearColor VisualLoggerColor = FLinearColor::Blue);
 
-	UFUNCTION(BlueprintPure, Category = "LearningAgents", meta = (AdvancedDisplay = 3))
-	static FLearningAgentsActionObjectElement MakeScaleAction(ULearningAgentsActionObject* Object, const FVector Scale, const FVector RelativeScale = FVector(1, 1, 1), const FName Tag = TEXT("ScaleAction"));
+	UFUNCTION(BlueprintPure, Category = "LearningAgents", meta = (AdvancedDisplay = 3, DefaultToSelf = "VisualLoggerListener"))
+	static FLearningAgentsActionObjectElement MakeScaleAction(
+		ULearningAgentsActionObject* Object,
+		const FVector Scale,
+		const FVector RelativeScale = FVector(1, 1, 1),
+		const FName Tag = TEXT("ScaleAction"),
+		const bool bVisualLoggerEnabled = false,
+		ULearningAgentsManagerListener* VisualLoggerListener = nullptr,
+		const int32 VisualLoggerAgentId = -1,
+		const FVector VisualLoggerLocation = FVector::ZeroVector,
+		const FLinearColor VisualLoggerColor = FLinearColor::Blue);
 
-	UFUNCTION(BlueprintPure, Category = "LearningAgents", meta = (AdvancedDisplay = 3))
-	static FLearningAgentsActionObjectElement MakeTransformAction(ULearningAgentsActionObject* Object, const FTransform Transform, const FTransform RelativeTransform = FTransform(), const FName Tag = TEXT("TransformAction"));
+	UFUNCTION(BlueprintPure, Category = "LearningAgents", meta = (AdvancedDisplay = 3, DefaultToSelf = "VisualLoggerListener"))
+	static FLearningAgentsActionObjectElement MakeTransformAction(
+		ULearningAgentsActionObject* Object,
+		const FTransform Transform,
+		const FTransform RelativeTransform = FTransform(),
+		const FName Tag = TEXT("TransformAction"),
+		const bool bVisualLoggerEnabled = false,
+		ULearningAgentsManagerListener* VisualLoggerListener = nullptr,
+		const int32 VisualLoggerAgentId = -1,
+		const FVector VisualLoggerLocation = FVector::ZeroVector,
+		const FLinearColor VisualLoggerColor = FLinearColor::Blue);
 
-	UFUNCTION(BlueprintPure, Category = "LearningAgents", meta = (AdvancedDisplay = 3))
-	static FLearningAgentsActionObjectElement MakeAngleAction(ULearningAgentsActionObject* Object, const float Angle, const float RelativeAngle = 0.0f, const FName Tag = TEXT("AngleAction"));
+	UFUNCTION(BlueprintPure, Category = "LearningAgents", meta = (AdvancedDisplay = 3, DefaultToSelf = "VisualLoggerListener"))
+	static FLearningAgentsActionObjectElement MakeAngleAction(
+		ULearningAgentsActionObject* Object,
+		const float Angle,
+		const float RelativeAngle = 0.0f,
+		const FName Tag = TEXT("AngleAction"),
+		const bool bVisualLoggerEnabled = false,
+		ULearningAgentsManagerListener* VisualLoggerListener = nullptr,
+		const int32 VisualLoggerAgentId = -1,
+		const FVector VisualLoggerAngleLocation = FVector::ZeroVector,
+		const FVector VisualLoggerLocation = FVector::ZeroVector,
+		const FLinearColor VisualLoggerColor = FLinearColor::Blue);
 
-	UFUNCTION(BlueprintPure, Category = "LearningAgents", meta = (AdvancedDisplay = 3))
-	static FLearningAgentsActionObjectElement MakeAngleActionRadians(ULearningAgentsActionObject* Object, const float Angle, const float RelativeAngle = 0.0f, const FName Tag = TEXT("AngleAction"));
+	UFUNCTION(BlueprintPure, Category = "LearningAgents", meta = (AdvancedDisplay = 3, DefaultToSelf = "VisualLoggerListener"))
+	static FLearningAgentsActionObjectElement MakeAngleActionRadians(
+		ULearningAgentsActionObject* Object,
+		const float Angle,
+		const float RelativeAngle = 0.0f,
+		const FName Tag = TEXT("AngleAction"),
+		const bool bVisualLoggerEnabled = false,
+		ULearningAgentsManagerListener* VisualLoggerListener = nullptr,
+		const int32 VisualLoggerAgentId = -1,
+		const FVector VisualLoggerAngleLocation = FVector::ZeroVector,
+		const FVector VisualLoggerLocation = FVector::ZeroVector,
+		const FLinearColor VisualLoggerColor = FLinearColor::Blue);
 
-	UFUNCTION(BlueprintPure, Category = "LearningAgents", meta = (AdvancedDisplay = 3))
-	static FLearningAgentsActionObjectElement MakeVelocityAction(ULearningAgentsActionObject* Object, const FVector Velocity, const FTransform RelativeTransform = FTransform(), const FName Tag = TEXT("VelocityAction"));
+	UFUNCTION(BlueprintPure, Category = "LearningAgents", meta = (AdvancedDisplay = 3, DefaultToSelf = "VisualLoggerListener"))
+	static FLearningAgentsActionObjectElement MakeVelocityAction(
+		ULearningAgentsActionObject* Object,
+		const FVector Velocity,
+		const FTransform RelativeTransform = FTransform(),
+		const FName Tag = TEXT("VelocityAction"),
+		const bool bVisualLoggerEnabled = false,
+		ULearningAgentsManagerListener* VisualLoggerListener = nullptr,
+		const int32 VisualLoggerAgentId = -1,
+		const FVector VisualLoggerVelocityLocation = FVector::ZeroVector,
+		const FVector VisualLoggerLocation = FVector::ZeroVector,
+		const FLinearColor VisualLoggerColor = FLinearColor::Blue);
 
-	UFUNCTION(BlueprintPure, Category = "LearningAgents", meta = (AdvancedDisplay = 3))
-	static FLearningAgentsActionObjectElement MakeDirectionAction(ULearningAgentsActionObject* Object, const FVector Direction, const FTransform RelativeTransform = FTransform(), const FName Tag = TEXT("DirectionAction"));
+	UFUNCTION(BlueprintPure, Category = "LearningAgents", meta = (AdvancedDisplay = 3, DefaultToSelf = "VisualLoggerListener"))
+	static FLearningAgentsActionObjectElement MakeDirectionAction(
+		ULearningAgentsActionObject* Object,
+		const FVector Direction,
+		const FTransform RelativeTransform = FTransform(),
+		const FName Tag = TEXT("DirectionAction"),
+		const bool bVisualLoggerEnabled = false,
+		ULearningAgentsManagerListener* VisualLoggerListener = nullptr,
+		const int32 VisualLoggerAgentId = -1,
+		const FVector VisualLoggerDirectionLocation = FVector::ZeroVector,
+		const FVector VisualLoggerLocation = FVector::ZeroVector,
+		const float VisualLoggerArrowLength = 100.0f,
+		const FLinearColor VisualLoggerColor = FLinearColor::Blue);
 
 public:
 
