Learning Agents: Add button to editor for appending records to a recording asset

#rnx

diff --git a/Engine/Plugins/Experimental/LearningAgents/Source/LearningAgentsTraining/Private/LearningAgentsRecording.cpp b/Engine/Plugins/Experimental/LearningAgents/Source/LearningAgentsTraining/Private/LearningAgentsRecording.cpp
--- a/Engine/Plugins/Experimental/LearningAgents/Source/LearningAgentsTraining/Private/LearningAgentsRecording.cpp
+++ b/Engine/Plugins/Experimental/LearningAgents/Source/LearningAgentsTraining/Private/LearningAgentsRecording.cpp
@@ -23,6 +23,11 @@ void ULearningAgentsRecording::ResetRecording()
 	ForceMarkDirty();
 }
 
+void ULearningAgentsRecording::AppendRecording()
+{
+	AppendRecordingFromFile(RecordingFile);
+}
+
 void ULearningAgentsRecording::LoadRecordingFromFile(const FFilePath& File)
 {
 	TArray<uint8> RecordingData;
@@ -31,7 +36,7 @@ void ULearningAgentsRecording::LoadRecordingFromFile(const FFilePath& File)
 	{
 		if (RecordingData.Num() < sizeof(int32) * 3)
 		{
-			UE_LOG(LogLearning, Error, TEXT("%s: Failed to load network. Incorrect Format."), *GetName());
+			UE_LOG(LogLearning, Error, TEXT("%s: Failed to load recording. Incorrect Format."), *GetName());
 			return;
 		}
 
diff --git a/Engine/Plugins/Experimental/LearningAgents/Source/LearningAgentsTraining/Public/LearningAgentsRecording.h b/Engine/Plugins/Experimental/LearningAgents/Source/LearningAgentsTraining/Public/LearningAgentsRecording.h
--- a/Engine/Plugins/Experimental/LearningAgents/Source/LearningAgentsTraining/Public/LearningAgentsRecording.h
+++ b/Engine/Plugins/Experimental/LearningAgents/Source/LearningAgentsTraining/Public/LearningAgentsRecording.h
@@ -66,6 +66,10 @@ public:
 	UFUNCTION(CallInEditor, Category = "LearningAgents")
 	void ResetRecording();
 
+	/** Appends the recording from "RecordingFile" to this recording.  */
+	UFUNCTION(CallInEditor, Category = "LearningAgents")
+	void AppendRecording();
+
 	/** Load this recording from a file. */
 	UFUNCTION(BlueprintCallable, Category = "LearningAgents", meta = (RelativePath))
 	void LoadRecordingFromFile(const FFilePath& File);
@@ -127,6 +131,10 @@ public:
 
 public:
 
+	/** A recording file. Used in combination with the "AppendRecording" button in the editor. */
+	UPROPERTY(EditInstanceOnly, Category = "LearningAgents")
+	FFilePath RecordingFile;
+
 	/** Set of records. */
 	UPROPERTY(EditInstanceOnly, Category = "LearningAgents")
 	TArray<FLearningAgentsRecord> Records;
