LearningAgents: Save schema from recorder

#rnx

diff --git a/Engine/Plugins/Experimental/LearningAgents/Source/LearningAgentsTraining/Private/LearningAgentsRecorder.cpp b/Engine/Plugins/Experimental/LearningAgents/Source/LearningAgentsTraining/Private/LearningAgentsRecorder.cpp
--- a/Engine/Plugins/Experimental/LearningAgents/Source/LearningAgentsTraining/Private/LearningAgentsRecorder.cpp
+++ b/Engine/Plugins/Experimental/LearningAgents/Source/LearningAgentsTraining/Private/LearningAgentsRecorder.cpp
@@ -11,6 +11,11 @@
 #include "UObject/Package.h"
 #include "Engine/World.h"
 #include "Misc/Paths.h"
+#include "Misc/FileHelper.h"
+
+#include "Dom/JsonObject.h"
+#include "Serialization/JsonSerializer.h"
+#include "Serialization/JsonWriter.h"
 
 FLearningAgentsRecorderPathSettings::FLearningAgentsRecorderPathSettings()
 {
@@ -207,6 +212,8 @@ void ULearningAgentsRecorder::SetupRecorder(
 		RecordBuffers.Emplace(ObservationCompatibilityHash, ActionCompatibilityHash);
 	}
 
+	SaveSchemaConfig();
+
 	bIsSetup = true;
 
 	Manager->AddListener(this);
@@ -373,4 +380,18 @@ const ULearningAgentsRecording* ULearningAgentsRecorder::GetRecordingAsset() con
 	}
 
 	return Recording;
-}
\ No newline at end of file
+}
+
+void ULearningAgentsRecorder::SaveSchemaConfig()
+{
+	TSharedRef<FJsonObject> ConfigObject = MakeShared<FJsonObject>();
+	ConfigObject->SetObjectField(TEXT("ObservationSchema"),
+		UE::Learning::Trainer::ConvertObservationSchemaToJSON(Interactor->GetObservationSchema()->ObservationSchema, Interactor->GetObservationSchemaElement().SchemaElement));
+	ConfigObject->SetObjectField(TEXT("ActionSchema"),
+		UE::Learning::Trainer::ConvertActionSchemaToJSON(Interactor->GetActionSchema()->ActionSchema, Interactor->GetActionSchemaElement().SchemaElement));
+
+	FString ConfigString;
+	TSharedRef<TJsonWriter<>> JsonWriter = TJsonWriterFactory<>::Create(&ConfigString, 0);
+	FJsonSerializer::Serialize(ConfigObject, JsonWriter, true);
+	FFileHelper::SaveStringToFile(ConfigString, *(RecordingDirectory / "schema-config.json"));
+}
diff --git a/Engine/Plugins/Experimental/LearningAgents/Source/LearningAgentsTraining/Public/LearningAgentsRecorder.h b/Engine/Plugins/Experimental/LearningAgents/Source/LearningAgentsTraining/Public/LearningAgentsRecorder.h
--- a/Engine/Plugins/Experimental/LearningAgents/Source/LearningAgentsTraining/Public/LearningAgentsRecorder.h
+++ b/Engine/Plugins/Experimental/LearningAgents/Source/LearningAgentsTraining/Public/LearningAgentsRecorder.h
@@ -143,6 +143,11 @@ private:
 	UPROPERTY(VisibleAnywhere, Transient, Category = "LearningAgents")
 	bool bIsRecording = false;
 
+// ----- Private Functions -----
+private:
+	
+	void SaveSchemaConfig();
+
 // ----- Private Data ----- 
 private:
 
