Learning Agents: Support imitation training with potentially skipped fields

#rnx

#virtualized

diff --git a/Engine/Plugins/Experimental/LearningAgents/Source/LearningAgentsTraining/Private/LearningAgentsImitationTrainer.cpp b/Engine/Plugins/Experimental/LearningAgents/Source/LearningAgentsTraining/Private/LearningAgentsImitationTrainer.cpp
--- a/Engine/Plugins/Experimental/LearningAgents/Source/LearningAgentsTraining/Private/LearningAgentsImitationTrainer.cpp
+++ b/Engine/Plugins/Experimental/LearningAgents/Source/LearningAgentsTraining/Private/LearningAgentsImitationTrainer.cpp
@@ -2,15 +2,17 @@
 
 #include "LearningAgentsImitationTrainer.h"
 
-#include "LearningAgentsManager.h"
+#include "LearningAgentsCommunicator.h"
 #include "LearningAgentsInteractor.h"
+#include "LearningAgentsManager.h"
+#include "LearningAgentsPolicy.h"
+#include "LearningAgentsRecording.h"
+
 #include "LearningExperience.h"
+#include "LearningExternalTrainer.h"
 #include "LearningLog.h"
 #include "LearningNeuralNetwork.h"
-#include "LearningExternalTrainer.h"
-#include "LearningAgentsCommunicator.h"
-#include "LearningAgentsRecording.h"
-#include "LearningAgentsPolicy.h"
+#include "LearningTrainer.h"
 
 #include "Dom/JsonObject.h"
 #include "HAL/PlatformMisc.h"
@@ -170,26 +172,21 @@ void ULearningAgentsImitationTrainer::BeginTraining(
 		return;
 	}
 
-	// Sizes
-
-	const int32 ObservationNum = Interactor->GetObservationVectorSize();
-	const int32 ActionNum = Interactor->GetActionVectorSize();
-
 	// Get Number of Steps
 
 	int32 TotalEpisodeNum = 0;
 	int32 TotalStepNum = 0;
 	for (const FLearningAgentsRecord& Record : Recording->Records)
 	{
-		if (Record.ObservationDimNum != ObservationNum)
+		if (!UE::Learning::Trainer::IsObservationSchemaSubsetCompatible(Record.ObservationSchemaJson, Interactor->GetObservationSchema()->ObservationSchema, Interactor->GetObservationSchemaElement().SchemaElement))
 		{
-			UE_LOG(LogLearning, Warning, TEXT("%s: Record has wrong dimensionality for observations, got %i, policy expected %i."), *GetName(), Record.ObservationDimNum, ObservationNum);
+			UE_LOG(LogLearning, Warning, TEXT("%s: Record has incompatible observation schema."), *GetName());
 			continue;
 		}
 
-		if (Record.ActionDimNum != ActionNum)
+		if (!UE::Learning::Trainer::IsActionSchemaSubsetCompatible(Record.ActionSchemaJson, Interactor->GetActionSchema()->ActionSchema, Interactor->GetActionSchemaElement().SchemaElement))
 		{
-			UE_LOG(LogLearning, Warning, TEXT("%s: Record has wrong dimensionality for actions, got %i, policy expected %i."), *GetName(), Record.ActionDimNum, ActionNum);
+			UE_LOG(LogLearning, Warning, TEXT("%s: Record has incompatible action schema."), *GetName());
 			continue;
 		}
 
@@ -203,6 +200,11 @@ void ULearningAgentsImitationTrainer::BeginTraining(
 		return;
 	}
 
+	// Sizes
+
+	const int32 ObservationNum = Interactor->GetObservationVectorSize();
+	const int32 ActionNum = Interactor->GetActionVectorSize();
+
 	// Copy into Flat Arrays
 
 	TLearningArray<1, int32> RecordedEpisodeStarts;
@@ -219,16 +221,51 @@ void ULearningAgentsImitationTrainer::BeginTraining(
 	int32 StepIdx = 0;
 	for (const FLearningAgentsRecord& Record : Recording->Records)
 	{
-		if (Record.ObservationDimNum != ObservationNum) { continue; }
-		if (Record.ActionDimNum != ActionNum) { continue; }
+		if (Record.ObservationDimNum == ObservationNum && Record.ActionDimNum == ActionNum)
+		{
+			TLearningArrayView<2, const float> ObservationsView = TLearningArrayView<2, const float>(Record.ObservationData.GetData(), { Record.StepNum, Record.ObservationDimNum });
+			TLearningArrayView<2, const float> ActionsView = TLearningArrayView<2, const float>(Record.ActionData.GetData(), { Record.StepNum, Record.ActionDimNum });
 
-		TLearningArrayView<2, const float> ObservationsView = TLearningArrayView<2, const float>(Record.ObservationData.GetData(), { Record.StepNum, Record.ObservationDimNum });
-		TLearningArrayView<2, const float> ActionsView = TLearningArrayView<2, const float>(Record.ActionData.GetData(), { Record.StepNum, Record.ActionDimNum });
+			RecordedEpisodeStarts[EpisodeIdx] = StepIdx;
+			RecordedEpisodeLengths[EpisodeIdx] = Record.StepNum;
+			UE::Learning::Array::Copy(RecordedObservations.Slice(StepIdx, Record.StepNum), ObservationsView);
+			UE::Learning::Array::Copy(RecordedActions.Slice(StepIdx, Record.StepNum), ActionsView);
+		}
+		else if (UE::Learning::Trainer::IsObservationSchemaSubsetCompatible(Record.ObservationSchemaJson, Interactor->GetObservationSchema()->ObservationSchema, Interactor->GetObservationSchemaElement().SchemaElement) &&
+			UE::Learning::Trainer::IsActionSchemaSubsetCompatible(Record.ActionSchemaJson, Interactor->GetActionSchema()->ActionSchema, Interactor->GetActionSchemaElement().SchemaElement))
+		{
+			TLearningArrayView<2, const float> ObservationsView = TLearningArrayView<2, const float>(Record.ObservationData.GetData(), { Record.StepNum, Record.ObservationDimNum });
+			TLearningArrayView<2, const float> ActionsView = TLearningArrayView<2, const float>(Record.ActionData.GetData(), { Record.StepNum, Record.ActionDimNum });
+
+			RecordedEpisodeStarts[EpisodeIdx] = StepIdx;
+			RecordedEpisodeLengths[EpisodeIdx] = Record.StepNum;
+
+			TArray<int32> ObservationIndices = UE::Learning::Trainer::ComputeObservationSchemaSubsetIndices(
+				Record.ObservationSchemaJson, Interactor->GetObservationSchema()->ObservationSchema, Interactor->GetObservationSchemaElement().SchemaElement);
+			check(ObservationIndices.Num() == ObservationNum);
+
+			for (int Step = 0; Step < Record.StepNum; Step++)
+			{
+				for (int i = 0; i < ObservationIndices.Num(); i++)
+				{
+					RecordedObservations[StepIdx + Step][i] = ObservationsView[Step][ObservationIndices[i]];
+				}
+			}
+
+			TArray<int32> ActionIndices = UE::Learning::Trainer::ComputeActionSchemaSubsetIndices(
+				Record.ActionSchemaJson, Interactor->GetActionSchema()->ActionSchema, Interactor->GetActionSchemaElement().SchemaElement);
+			check(ActionIndices.Num() == ActionNum);
+
+			for (int Step = 0; Step < Record.StepNum; Step++)
+			{
+				for (int i = 0; i < ActionIndices.Num(); i++)
+				{
+					RecordedActions[StepIdx + Step][i] = ActionsView[Step][ActionIndices[i]];
+				}
+			}
+		}
+		else { continue; }
 
-		RecordedEpisodeStarts[EpisodeIdx] = StepIdx;
-		RecordedEpisodeLengths[EpisodeIdx] = Record.StepNum;
-		UE::Learning::Array::Copy(RecordedObservations.Slice(StepIdx, Record.StepNum), ObservationsView);
-		UE::Learning::Array::Copy(RecordedActions.Slice(StepIdx, Record.StepNum), ActionsView);
 		EpisodeIdx++;
 		StepIdx += Record.StepNum;
 	}
diff --git a/Engine/Plugins/Experimental/LearningAgents/Source/LearningAgentsTraining/Private/LearningAgentsRecorder.cpp b/Engine/Plugins/Experimental/LearningAgents/Source/LearningAgentsTraining/Private/LearningAgentsRecorder.cpp
--- a/Engine/Plugins/Experimental/LearningAgents/Source/LearningAgentsTraining/Private/LearningAgentsRecorder.cpp
+++ b/Engine/Plugins/Experimental/LearningAgents/Source/LearningAgentsTraining/Private/LearningAgentsRecorder.cpp
@@ -24,9 +24,19 @@ FLearningAgentsRecorderPathSettings::FLearningAgentsRecorderPathSettings()
 
 ULearningAgentsRecorder::FAgentRecordBuffer::FAgentRecordBuffer() = default;
 
-ULearningAgentsRecorder::FAgentRecordBuffer::FAgentRecordBuffer(const int32 InObservationCompatibilityHash, const int32 InActionCompatibilityHash) 
+ULearningAgentsRecorder::FAgentRecordBuffer::FAgentRecordBuffer(
+	const int32 InObservationCompatibilityHash,
+	const int32 InActionCompatibilityHash,
+	const UE::Learning::Observation::FSchema InObservationSchema,
+	const UE::Learning::Observation::FSchemaElement InObservationSchemaElement,
+	const UE::Learning::Action::FSchema InActionSchema,
+	const UE::Learning::Action::FSchemaElement InActionSchemaElement)
 	: ObservationCompatibilityHash(InObservationCompatibilityHash)
-	, ActionCompatibilityHash(InActionCompatibilityHash) {};
+	, ActionCompatibilityHash(InActionCompatibilityHash)
+	, ObservationSchema(InObservationSchema)
+	, ObservationSchemaElement(InObservationSchemaElement)
+	, ActionSchema(InActionSchema)
+	, ActionSchemaElement(InActionSchemaElement) {};
 
 TLearningArrayView<1, float> ULearningAgentsRecorder::FAgentRecordBuffer::GetObservation(const int32 SampleIdx)
 {
@@ -87,6 +97,12 @@ void ULearningAgentsRecorder::FAgentRecordBuffer::CopyToRecord(FLearningAgentsRe
 	Record.ObservationData.SetNumUninitialized(StepNum * Record.ObservationDimNum);
 	Record.ActionData.SetNumUninitialized(StepNum * Record.ActionDimNum);
 
+	TSharedRef<TJsonWriter<>> ObsJsonWriter = TJsonWriterFactory<>::Create(&Record.ObservationSchemaJson, 0);
+	FJsonSerializer::Serialize(UE::Learning::Trainer::ConvertObservationSchemaToJSON(ObservationSchema, ObservationSchemaElement).ToSharedRef(), ObsJsonWriter, true);
+	
+	TSharedRef<TJsonWriter<>> ActJsonWriter = TJsonWriterFactory<>::Create(&Record.ActionSchemaJson, 0);
+	FJsonSerializer::Serialize(UE::Learning::Trainer::ConvertActionSchemaToJSON(ActionSchema, ActionSchemaElement).ToSharedRef(), ActJsonWriter, true);
+
 	TLearningArrayView<2, float> ObservationsView = TLearningArrayView<2, float>(Record.ObservationData.GetData(), { Record.StepNum, Record.ObservationDimNum });
 	TLearningArrayView<2, float> ActionsView = TLearningArrayView<2, float>(Record.ActionData.GetData(), { Record.StepNum, Record.ActionDimNum });
 
@@ -209,7 +225,13 @@ void ULearningAgentsRecorder::SetupRecorder(
 	RecordBuffers.Empty(Manager->GetMaxAgentNum());
 	for (int32 RecordBufferIdx = 0; RecordBufferIdx < Manager->GetMaxAgentNum(); RecordBufferIdx++)
 	{
-		RecordBuffers.Emplace(ObservationCompatibilityHash, ActionCompatibilityHash);
+		RecordBuffers.Emplace(
+			ObservationCompatibilityHash,
+			ActionCompatibilityHash,
+			Interactor->GetObservationSchema()->ObservationSchema,
+			Interactor->GetObservationSchemaElement().SchemaElement,
+			Interactor->GetActionSchema()->ActionSchema,
+			Interactor->GetActionSchemaElement().SchemaElement);
 	}
 
 	SaveSchemaConfig();
diff --git a/Engine/Plugins/Experimental/LearningAgents/Source/LearningAgentsTraining/Public/LearningAgentsRecorder.h b/Engine/Plugins/Experimental/LearningAgents/Source/LearningAgentsTraining/Public/LearningAgentsRecorder.h
--- a/Engine/Plugins/Experimental/LearningAgents/Source/LearningAgentsTraining/Public/LearningAgentsRecorder.h
+++ b/Engine/Plugins/Experimental/LearningAgents/Source/LearningAgentsTraining/Public/LearningAgentsRecorder.h
@@ -4,7 +4,10 @@
 
 #include "LearningAgentsManagerListener.h"
 
+#include "LearningAction.h"
 #include "LearningArray.h"
+#include "LearningObservation.h"
+
 #include "Containers/Map.h"
 #include "UObject/ObjectPtr.h"
 
@@ -158,12 +161,22 @@ private:
 	struct FAgentRecordBuffer
 	{
 		FAgentRecordBuffer();
-		FAgentRecordBuffer(const int32 InObservationCompatibilityHash, const int32 InActionCompatibilityHash);
+		FAgentRecordBuffer(
+			const int32 InObservationCompatibilityHash,
+			const int32 InActionCompatibilityHash,
+			const UE::Learning::Observation::FSchema InObservationSchema,
+			const UE::Learning::Observation::FSchemaElement InObservationSchemaElement,
+			const UE::Learning::Action::FSchema InActionSchema,
+			const UE::Learning::Action::FSchemaElement InActionSchemaElement);
 
 		const int32 ChunkSize = 1024;
 		int32 StepNum = 0;
 		int32 ObservationCompatibilityHash = 0;
 		int32 ActionCompatibilityHash = 0;
+		UE::Learning::Observation::FSchema ObservationSchema;
+		UE::Learning::Observation::FSchemaElement ObservationSchemaElement;
+		UE::Learning::Action::FSchema ActionSchema;
+		UE::Learning::Action::FSchemaElement ActionSchemaElement;
 		TArray<TLearningArray<2, float>, TInlineAllocator<16>> Observations;
 		TArray<TLearningArray<2, float>, TInlineAllocator<16>> Actions;
 
diff --git a/Engine/Plugins/Experimental/LearningAgents/Source/LearningAgentsTraining/Public/LearningAgentsRecording.h b/Engine/Plugins/Experimental/LearningAgents/Source/LearningAgentsTraining/Public/LearningAgentsRecording.h
--- a/Engine/Plugins/Experimental/LearningAgents/Source/LearningAgentsTraining/Public/LearningAgentsRecording.h
+++ b/Engine/Plugins/Experimental/LearningAgents/Source/LearningAgentsTraining/Public/LearningAgentsRecording.h
@@ -38,6 +38,12 @@ public:
 	UPROPERTY(BlueprintReadOnly, VisibleAnywhere, Category = "LearningAgents")
 	int32 ActionCompatibilityHash = 0;
 
+	UPROPERTY(BlueprintReadOnly, VisibleAnywhere, Category = "LearningAgents")
+	FString ObservationSchemaJson;
+	
+	UPROPERTY(BlueprintReadOnly, VisibleAnywhere, Category = "LearningAgents")
+	FString ActionSchemaJson;
+
 	/** Observation data */
 	UPROPERTY(BlueprintReadOnly, Category = "LearningAgents")
 	TArray<float> ObservationData;
diff --git a/Engine/Plugins/Experimental/LearningAgents/Source/LearningTraining/Private/LearningTrainer.cpp b/Engine/Plugins/Experimental/LearningAgents/Source/LearningTraining/Private/LearningTrainer.cpp
--- a/Engine/Plugins/Experimental/LearningAgents/Source/LearningTraining/Private/LearningTrainer.cpp
+++ b/Engine/Plugins/Experimental/LearningAgents/Source/LearningTraining/Private/LearningTrainer.cpp
@@ -5,9 +5,11 @@
 #include "LearningObservation.h"
 #include "LearningAction.h"
 
-#include "HAL/Platform.h"
 #include "Dom/JsonObject.h"
+#include "HAL/Platform.h"
 #include "Misc/Paths.h"
+#include "Serialization/JsonReader.h"
+#include "Serialization/JsonSerializer.h"
 
 namespace UE::Learning
 {
@@ -431,6 +433,977 @@ namespace UE::Learning::Trainer
 		return Object;
 	}
 
+	bool IsObservationSchemaSubsetCompatible(const FString& SourceJsonString, const Observation::FSchema& ObservationSchema, const Observation::FSchemaElement& ObservationSchemaElement)
+	{
+		TSharedPtr<FJsonObject> JsonObject;
+		TSharedRef<TJsonReader<>> Reader = TJsonReaderFactory<>::Create(SourceJsonString);
+
+		if (FJsonSerializer::Deserialize(Reader, JsonObject) && JsonObject.IsValid())
+		{
+			return IsObservationSchemaSubsetCompatible(JsonObject.ToSharedRef(), ObservationSchema, ObservationSchemaElement);
+		}
+
+		return false;
+	}
+
+	bool IsObservationSchemaSubsetCompatible(const TSharedRef<FJsonObject> Object, const Observation::FSchema& ObservationSchema, const Observation::FSchemaElement& ObservationSchemaElement)
+	{
+		if (Object->GetNumberField(TEXT("VectorSize")) < ObservationSchema.GetObservationVectorSize(ObservationSchemaElement) ||
+			Object->GetNumberField(TEXT("EncodedSize")) < ObservationSchema.GetEncodedVectorSize(ObservationSchemaElement))
+		{
+			return false;
+		}
+
+		switch (ObservationSchema.GetType(ObservationSchemaElement))
+		{
+		case Observation::EType::Null:
+		{
+			return Object->GetStringField(TEXT("Type")) == TEXT("Null");
+		}
+
+		case Observation::EType::Continuous:
+		{
+			const Observation::FSchemaContinuousParameters Parameters = ObservationSchema.GetContinuous(ObservationSchemaElement);
+
+			return
+				Object->GetStringField(TEXT("Type")) == TEXT("Continuous") &&
+				Object->GetNumberField(TEXT("Num")) == Parameters.Num;
+		}
+
+		case Observation::EType::DiscreteExclusive:
+		{
+			const Observation::FSchemaDiscreteExclusiveParameters Parameters = ObservationSchema.GetDiscreteExclusive(ObservationSchemaElement);
+			
+			return
+				Object->GetStringField(TEXT("Type")) == TEXT("DiscreteExclusive") &&
+				Object->GetNumberField(TEXT("Num")) == Parameters.Num;
+		}
+
+		case Observation::EType::DiscreteInclusive:
+		{
+			const Observation::FSchemaDiscreteInclusiveParameters Parameters = ObservationSchema.GetDiscreteInclusive(ObservationSchemaElement);
+
+			return
+				Object->GetStringField(TEXT("Type")) == TEXT("DiscreteInclusive") &&
+				Object->GetNumberField(TEXT("Num")) == Parameters.Num;
+		}
+
+		case Observation::EType::NamedDiscreteExclusive:
+		{
+			const Observation::FSchemaNamedDiscreteExclusiveParameters Parameters = ObservationSchema.GetNamedDiscreteExclusive(ObservationSchemaElement);
+
+			if (Object->GetStringField(TEXT("Type")) != TEXT("NamedDiscreteExclusive"))
+			{
+				return false;
+			}
+
+			TArray<TSharedPtr<FJsonValue>> ElementNames = Object->GetArrayField(TEXT("ElementNames"));
+			for (int32 ElementIdx = 0; ElementIdx < Parameters.ElementNames.Num(); ElementIdx++)
+			{
+				bool bFound = false;
+				for (const TSharedPtr<FJsonValue>& ElementName : ElementNames)
+				{
+					if (ElementName->Type == EJson::String && ElementName->AsString() == Parameters.ElementNames[ElementIdx].ToString())
+					{
+						bFound = true;
+						break;
+					}
+				}
+
+				if (!bFound)
+				{
+					return false;
+				}
+			}
+
+			return true;
+		}
+
+		case Observation::EType::NamedDiscreteInclusive:
+		{
+			const Observation::FSchemaNamedDiscreteInclusiveParameters Parameters = ObservationSchema.GetNamedDiscreteInclusive(ObservationSchemaElement);
+
+			if (Object->GetStringField(TEXT("Type")) != TEXT("NamedDiscreteInclusive"))
+			{
+				return false;
+			}
+
+			TArray<TSharedPtr<FJsonValue>> ElementNames = Object->GetArrayField(TEXT("ElementNames"));
+			for (int32 ElementIdx = 0; ElementIdx < Parameters.ElementNames.Num(); ElementIdx++)
+			{
+				bool bFound = false;
+				for (const TSharedPtr<FJsonValue>& ElementName : ElementNames)
+				{
+					if (ElementName->Type == EJson::String && ElementName->AsString() == Parameters.ElementNames[ElementIdx].ToString())
+					{
+						bFound = true;
+						break;
+					}
+				}
+
+				if (!bFound)
+				{
+					return false;
+				}
+			}
+
+			return true;
+		}
+
+		case Observation::EType::And:
+		{
+			const Observation::FSchemaAndParameters Parameters = ObservationSchema.GetAnd(ObservationSchemaElement);
+
+			if (Object->GetStringField(TEXT("Type")) != TEXT("And"))
+			{
+				return false;
+			}
+
+			TSharedPtr<FJsonObject> SubObject = Object->GetObjectField(TEXT("Elements"));
+			for (int32 SubElementIdx = 0; SubElementIdx < Parameters.Elements.Num(); SubElementIdx++)
+			{
+				TSharedPtr<FJsonObject> SubElement = SubObject->GetObjectField(Parameters.ElementNames[SubElementIdx].ToString());
+				if (SubElement.IsValid())
+				{
+					bool bResult = IsObservationSchemaSubsetCompatible(SubElement.ToSharedRef(), ObservationSchema, Parameters.Elements[SubElementIdx]);
+					if (!bResult)
+					{
+						return false;
+					}
+				}
+				else
+				{
+					return false;
+				}
+			}
+
+			return true;
+		}
+
+		case Observation::EType::OrExclusive:
+		{
+			const Observation::FSchemaOrExclusiveParameters Parameters = ObservationSchema.GetOrExclusive(ObservationSchemaElement);
+
+			if (Object->GetStringField(TEXT("Type")) != TEXT("OrExclusive"))
+			{
+				return false;
+			}
+
+			TSharedPtr<FJsonObject> SubObject = Object->GetObjectField(TEXT("Elements"));
+			for (int32 SubElementIdx = 0; SubElementIdx < Parameters.Elements.Num(); SubElementIdx++)
+			{
+				TSharedPtr<FJsonObject> SubElement = SubObject->GetObjectField(Parameters.ElementNames[SubElementIdx].ToString());
+				if (SubElement.IsValid())
+				{
+					bool bResult = IsObservationSchemaSubsetCompatible(SubElement.ToSharedRef(), ObservationSchema, Parameters.Elements[SubElementIdx]);
+					if (!bResult)
+					{
+						return false;
+					}
+				}
+				else
+				{
+					return false;
+				}
+			}
+
+			return true;
+		}
+
+		case Observation::EType::OrInclusive:
+		{
+			const Observation::FSchemaOrInclusiveParameters Parameters = ObservationSchema.GetOrInclusive(ObservationSchemaElement);
+
+			if (Object->GetStringField(TEXT("Type")) != TEXT("OrInclusive"))
+			{
+				return false;
+			}
+
+			TSharedPtr<FJsonObject> SubObject = Object->GetObjectField(TEXT("Elements"));
+			for (int32 SubElementIdx = 0; SubElementIdx < Parameters.Elements.Num(); SubElementIdx++)
+			{
+				TSharedPtr<FJsonObject> SubElement = SubObject->GetObjectField(Parameters.ElementNames[SubElementIdx].ToString());
+				if (SubElement.IsValid())
+				{
+					bool bResult = IsObservationSchemaSubsetCompatible(SubElement.ToSharedRef(), ObservationSchema, Parameters.Elements[SubElementIdx]);
+					if (!bResult)
+					{
+						return false;
+					}
+				}
+				else
+				{
+					return false;
+				}
+			}
+
+			return true;
+		}
+
+		case Observation::EType::Array:
+		{
+			const Observation::FSchemaArrayParameters Parameters = ObservationSchema.GetArray(ObservationSchemaElement);
+
+			if (Object->GetStringField(TEXT("Type")) != TEXT("Array") ||
+				Object->GetNumberField(TEXT("Num")) != Parameters.Num)
+			{
+				return false;
+			}
+
+			return IsObservationSchemaSubsetCompatible(Object->GetObjectField(TEXT("Element")).ToSharedRef(), ObservationSchema, Parameters.Element);
+		}
+
+		case Observation::EType::Set:
+		{
+			const Observation::FSchemaSetParameters Parameters = ObservationSchema.GetSet(ObservationSchemaElement);
+
+			if (Object->GetStringField(TEXT("Type")) != TEXT("Set") ||
+				Object->GetNumberField(TEXT("MaxNum")) != Parameters.MaxNum)
+			{
+				return false;
+			}
+
+			return IsObservationSchemaSubsetCompatible(Object->GetObjectField(TEXT("Element")).ToSharedRef(), ObservationSchema, Parameters.Element);
+		}
+
+		case Observation::EType::Encoding:
+		{
+			const Observation::FSchemaEncodingParameters Parameters = ObservationSchema.GetEncoding(ObservationSchemaElement);
+
+			if (Object->GetStringField(TEXT("Type")) != TEXT("Encoding"))
+			{
+				return false;
+			}
+
+			return IsObservationSchemaSubsetCompatible(Object->GetObjectField(TEXT("Element")).ToSharedRef(), ObservationSchema, Parameters.Element);
+		}
+
+		default:
+			checkNoEntry();
+		}
+
+		return false;
+	}
+
+	TArray<int32> ComputeObservationSchemaSubsetIndices(const FString& SourceJsonString, const Observation::FSchema& ObservationSchema, const Observation::FSchemaElement& ObservationSchemaElement)
+	{
+		TSharedPtr<FJsonObject> JsonObject;
+		TSharedRef<TJsonReader<>> Reader = TJsonReaderFactory<>::Create(SourceJsonString);
+
+		TArray<int32> Output;
+
+		if (FJsonSerializer::Deserialize(Reader, JsonObject) && JsonObject.IsValid())
+		{
+			int32 NextIndex = 0;
+			ComputeObservationSchemaSubsetIndices(NextIndex, Output, JsonObject.ToSharedRef(), ObservationSchema, ObservationSchemaElement);
+		}
+		
+		return Output;
+	}
+
+	void ComputeObservationSchemaSubsetIndices(
+		int32& NextIndex,
+		TArray<int32>& Indices,
+		const TSharedRef<FJsonObject> Object,
+		const Observation::FSchema& ObservationSchema,
+		const Observation::FSchemaElement& ObservationSchemaElement)
+	{
+		if (Object->GetStringField(TEXT("Type")) == TEXT("Null"))
+		{
+			// Do nothing since Null is zero-width
+			return;
+		}
+		else if (Object->GetStringField(TEXT("Type")) == TEXT("Continuous"))
+		{
+			const Observation::FSchemaContinuousParameters Parameters = ObservationSchema.GetContinuous(ObservationSchemaElement);
+
+			for (int32 i = 0; i < Parameters.Num; i++)
+			{
+				Indices.Add(NextIndex);
+				NextIndex++;
+			}
+
+			return;
+		}
+		else if (Object->GetStringField(TEXT("Type")) == TEXT("DiscreteExclusive"))
+		{
+			const Observation::FSchemaDiscreteExclusiveParameters Parameters = ObservationSchema.GetDiscreteExclusive(ObservationSchemaElement);
+			
+			for (int32 i = 0; i < Parameters.Num; i++)
+			{
+				Indices.Add(NextIndex);
+				NextIndex++;
+			}
+
+			return;
+		}
+		else if (Object->GetStringField(TEXT("Type")) == TEXT("DiscreteInclusive"))
+		{
+			const Observation::FSchemaDiscreteInclusiveParameters Parameters = ObservationSchema.GetDiscreteInclusive(ObservationSchemaElement);
+
+			for (int32 i = 0; i < Parameters.Num; i++)
+			{
+				Indices.Add(NextIndex);
+				NextIndex++;
+			}
+
+			return;
+		}
+		else if (Object->GetStringField(TEXT("Type")) == TEXT("NamedDiscreteExclusive"))
+		{
+			const Observation::FSchemaNamedDiscreteExclusiveParameters Parameters = ObservationSchema.GetNamedDiscreteExclusive(ObservationSchemaElement);
+
+			TArray<TSharedPtr<FJsonValue>> ElementNames = Object->GetArrayField(TEXT("ElementNames"));
+			for (const TSharedPtr<FJsonValue>& ElementName : ElementNames)
+			{
+				bool bFound = false;
+				for (int32 ElementIdx = 0; ElementIdx < Parameters.ElementNames.Num(); ElementIdx++)
+				{
+					if (ElementName->Type == EJson::String && ElementName->AsString() == Parameters.ElementNames[ElementIdx].ToString())
+					{
+						bFound = true;
+						break;
+					}
+				}
+
+				if (bFound)
+				{
+					Indices.Add(NextIndex);
+				}
+
+				NextIndex++;
+			}
+
+			return;
+		}
+		else if (Object->GetStringField(TEXT("Type")) == TEXT("NamedDiscreteInclusive"))
+		{
+			const Observation::FSchemaNamedDiscreteInclusiveParameters Parameters = ObservationSchema.GetNamedDiscreteInclusive(ObservationSchemaElement);
+
+			TArray<TSharedPtr<FJsonValue>> ElementNames = Object->GetArrayField(TEXT("ElementNames"));
+			for (const TSharedPtr<FJsonValue>& ElementName : ElementNames)
+			{
+				bool bFound = false;
+				for (int32 ElementIdx = 0; ElementIdx < Parameters.ElementNames.Num(); ElementIdx++)
+				{
+					if (ElementName->Type == EJson::String && ElementName->AsString() == Parameters.ElementNames[ElementIdx].ToString())
+					{
+						bFound = true;
+						break;
+					}
+				}
+
+				if (bFound)
+				{
+					Indices.Add(NextIndex);
+				}
+
+				NextIndex++;
+			}
+
+			return;
+		}
+		else if (Object->GetStringField(TEXT("Type")) == TEXT("And"))
+		{
+			const Observation::FSchemaAndParameters Parameters = ObservationSchema.GetAnd(ObservationSchemaElement);
+
+			TSharedPtr<FJsonObject> SubObject = Object->GetObjectField(TEXT("Elements"));
+			TArray<FString> ElementNames;
+			SubObject->Values.GetKeys(ElementNames);
+
+			for (const FString& ElementName : ElementNames)
+			{
+				TSharedPtr<FJsonObject> SubElement = SubObject->GetObjectField(ElementName);
+
+				bool bFound = false;
+				for (int32 ElementIdx = 0; ElementIdx < Parameters.ElementNames.Num(); ElementIdx++)
+				{
+					if (ElementName == Parameters.ElementNames[ElementIdx].ToString())
+					{
+						bFound = true;
+						ComputeObservationSchemaSubsetIndices(NextIndex, Indices, SubElement.ToSharedRef(), ObservationSchema, Parameters.Elements[ElementIdx]);
+						break;
+					}
+				}
+
+				if (!bFound)
+				{
+					int32 VectorSize = SubElement->GetNumberField(TEXT("VectorSize"));
+					NextIndex += VectorSize;
+				}
+			}
+
+			return;
+		}
+		else if (Object->GetStringField(TEXT("Type")) == TEXT("OrExclusive"))
+		{
+			const Observation::FSchemaOrExclusiveParameters Parameters = ObservationSchema.GetOrExclusive(ObservationSchemaElement);
+
+			TSharedPtr<FJsonObject> SubObject = Object->GetObjectField(TEXT("Elements"));
+			TArray<FString> ElementNames;
+			SubObject->Values.GetKeys(ElementNames);
+
+			for (const FString& ElementName : ElementNames)
+			{
+				TSharedPtr<FJsonObject> SubElement = SubObject->GetObjectField(ElementName);
+
+				bool bFound = false;
+				for (int32 ElementIdx = 0; ElementIdx < Parameters.ElementNames.Num(); ElementIdx++)
+				{
+					if (ElementName == Parameters.ElementNames[ElementIdx].ToString())
+					{
+						bFound = true;
+						ComputeObservationSchemaSubsetIndices(NextIndex, Indices, SubElement.ToSharedRef(), ObservationSchema, Parameters.Elements[ElementIdx]);
+						break;
+					}
+				}
+
+				if (!bFound)
+				{
+					int32 VectorSize = SubElement->GetNumberField(TEXT("VectorSize"));
+					NextIndex += VectorSize;
+				}
+			}
+
+			return;
+		}
+		else if (Object->GetStringField(TEXT("Type")) == TEXT("OrInclusive"))
+		{
+			const Observation::FSchemaOrInclusiveParameters Parameters = ObservationSchema.GetOrInclusive(ObservationSchemaElement);
+
+			TSharedPtr<FJsonObject> SubObject = Object->GetObjectField(TEXT("Elements"));
+			TArray<FString> ElementNames;
+			SubObject->Values.GetKeys(ElementNames);
+
+			for (const FString& ElementName : ElementNames)
+			{
+				TSharedPtr<FJsonObject> SubElement = SubObject->GetObjectField(ElementName);
+
+				bool bFound = false;
+				for (int32 ElementIdx = 0; ElementIdx < Parameters.ElementNames.Num(); ElementIdx++)
+				{
+					if (ElementName == Parameters.ElementNames[ElementIdx].ToString())
+					{
+						bFound = true;
+						ComputeObservationSchemaSubsetIndices(NextIndex, Indices, SubElement.ToSharedRef(), ObservationSchema, Parameters.Elements[ElementIdx]);
+						break;
+					}
+				}
+
+				if (!bFound)
+				{
+					int32 VectorSize = SubElement->GetNumberField(TEXT("VectorSize"));
+					NextIndex += VectorSize;
+				}
+			}
+
+			return;
+		}
+		else if (Object->GetStringField(TEXT("Type")) == TEXT("Array"))
+		{
+			const Observation::FSchemaArrayParameters Parameters = ObservationSchema.GetArray(ObservationSchemaElement);
+
+			for (int i = 0; i < Parameters.Num; i++)
+			{
+				ComputeObservationSchemaSubsetIndices(NextIndex, Indices, Object->GetObjectField(TEXT("Element")).ToSharedRef(), ObservationSchema, Parameters.Element);
+			}
+
+			return;
+		}
+		else if (Object->GetStringField(TEXT("Type")) == TEXT("Set"))
+		{
+			const Observation::FSchemaSetParameters Parameters = ObservationSchema.GetSet(ObservationSchemaElement);
+
+			for (int i = 0; i < Parameters.MaxNum; i++)
+			{
+				ComputeObservationSchemaSubsetIndices(NextIndex, Indices, Object->GetObjectField(TEXT("Element")).ToSharedRef(), ObservationSchema, Parameters.Element);
+			}
+
+			return;
+		}
+		else if (Object->GetStringField(TEXT("Type")) == TEXT("Encoding"))
+		{
+			const Observation::FSchemaEncodingParameters Parameters = ObservationSchema.GetEncoding(ObservationSchemaElement);
+			
+			ComputeObservationSchemaSubsetIndices(NextIndex, Indices, Object->GetObjectField(TEXT("Element")).ToSharedRef(), ObservationSchema, Parameters.Element);
+		}
+
+		// We didn't return earlier so we received an invalid Type
+		checkNoEntry();
+	}
+
+	bool IsActionSchemaSubsetCompatible(const FString& SourceJsonString, const Action::FSchema& ActionSchema, const Action::FSchemaElement& ActionSchemaElement)
+	{
+		TSharedPtr<FJsonObject> JsonObject;
+		TSharedRef<TJsonReader<>> Reader = TJsonReaderFactory<>::Create(SourceJsonString);
+
+		if (FJsonSerializer::Deserialize(Reader, JsonObject) && JsonObject.IsValid())
+		{
+			return IsActionSchemaSubsetCompatible(JsonObject.ToSharedRef(), ActionSchema, ActionSchemaElement);
+		}
+
+		return false;
+	}
+
+	bool IsActionSchemaSubsetCompatible(const TSharedRef<FJsonObject> Object, const Action::FSchema& ActionSchema, const Action::FSchemaElement& ActionSchemaElement)
+	{
+		if (Object->GetNumberField(TEXT("VectorSize")) < ActionSchema.GetActionVectorSize(ActionSchemaElement) ||
+			Object->GetNumberField(TEXT("EncodedSize")) < ActionSchema.GetEncodedVectorSize(ActionSchemaElement))
+		{
+			return false;
+		}
+
+		switch (ActionSchema.GetType(ActionSchemaElement))
+		{
+		case Action::EType::Null:
+		{
+			return Object->GetStringField(TEXT("Type")) == TEXT("Null");
+		}
+
+		case Action::EType::Continuous:
+		{
+			const Action::FSchemaContinuousParameters Parameters = ActionSchema.GetContinuous(ActionSchemaElement);
+
+			return
+				Object->GetStringField(TEXT("Type")) == TEXT("Continuous") &&
+				Object->GetNumberField(TEXT("Num")) == Parameters.Num;
+		}
+
+		case Action::EType::DiscreteExclusive:
+		{
+			const Action::FSchemaDiscreteExclusiveParameters Parameters = ActionSchema.GetDiscreteExclusive(ActionSchemaElement);
+
+			return
+				Object->GetStringField(TEXT("Type")) == TEXT("DiscreteExclusive") &&
+				Object->GetNumberField(TEXT("Num")) == Parameters.Num;
+		}
+
+		case Action::EType::DiscreteInclusive:
+		{
+			const Action::FSchemaDiscreteInclusiveParameters Parameters = ActionSchema.GetDiscreteInclusive(ActionSchemaElement);
+
+			return
+				Object->GetStringField(TEXT("Type")) == TEXT("DiscreteInclusive") &&
+				Object->GetNumberField(TEXT("Num")) == Parameters.Num;
+		}
+
+		case Action::EType::NamedDiscreteExclusive:
+		{
+			const Action::FSchemaNamedDiscreteExclusiveParameters Parameters = ActionSchema.GetNamedDiscreteExclusive(ActionSchemaElement);
+
+			if (Object->GetStringField(TEXT("Type")) != TEXT("NamedDiscreteExclusive"))
+			{
+				return false;
+			}
+
+			TArray<TSharedPtr<FJsonValue>> ElementNames = Object->GetArrayField(TEXT("ElementNames"));
+			for (int32 ElementIdx = 0; ElementIdx < Parameters.ElementNames.Num(); ElementIdx++)
+			{
+				bool bFound = false;
+				for (const TSharedPtr<FJsonValue>& ElementName : ElementNames)
+				{
+					if (ElementName->Type == EJson::String && ElementName->AsString() == Parameters.ElementNames[ElementIdx].ToString())
+					{
+						bFound = true;
+						break;
+					}
+				}
+
+				if (!bFound)
+				{
+					return false;
+				}
+			}
+
+			return true;
+		}
+
+		case Action::EType::NamedDiscreteInclusive:
+		{
+			const Action::FSchemaNamedDiscreteInclusiveParameters Parameters = ActionSchema.GetNamedDiscreteInclusive(ActionSchemaElement);
+
+			if (Object->GetStringField(TEXT("Type")) != TEXT("NamedDiscreteInclusive"))
+			{
+				return false;
+			}
+
+			TArray<TSharedPtr<FJsonValue>> ElementNames = Object->GetArrayField(TEXT("ElementNames"));
+			for (int32 ElementIdx = 0; ElementIdx < Parameters.ElementNames.Num(); ElementIdx++)
+			{
+				bool bFound = false;
+				for (const TSharedPtr<FJsonValue>& ElementName : ElementNames)
+				{
+					if (ElementName->Type == EJson::String && ElementName->AsString() == Parameters.ElementNames[ElementIdx].ToString())
+					{
+						bFound = true;
+						break;
+					}
+				}
+
+				if (!bFound)
+				{
+					return false;
+				}
+			}
+
+			return true;
+		}
+
+		case Action::EType::And:
+		{
+			const Action::FSchemaAndParameters Parameters = ActionSchema.GetAnd(ActionSchemaElement);
+
+			if (Object->GetStringField(TEXT("Type")) != TEXT("And"))
+			{
+				return false;
+			}
+
+			TSharedPtr<FJsonObject> SubObject = Object->GetObjectField(TEXT("Elements"));
+			for (int32 SubElementIdx = 0; SubElementIdx < Parameters.Elements.Num(); SubElementIdx++)
+			{
+				TSharedPtr<FJsonObject> SubElement = SubObject->GetObjectField(Parameters.ElementNames[SubElementIdx].ToString());
+				if (SubElement.IsValid())
+				{
+					bool bResult = IsActionSchemaSubsetCompatible(SubElement.ToSharedRef(), ActionSchema, Parameters.Elements[SubElementIdx]);
+					if (!bResult)
+					{
+						return false;
+					}
+				}
+				else
+				{
+					return false;
+				}
+			}
+
+			return true;
+		}
+
+		case Action::EType::OrExclusive:
+		{
+			const Action::FSchemaOrExclusiveParameters Parameters = ActionSchema.GetOrExclusive(ActionSchemaElement);
+
+			if (Object->GetStringField(TEXT("Type")) != TEXT("OrExclusive"))
+			{
+				return false;
+			}
+
+			TSharedPtr<FJsonObject> SubObject = Object->GetObjectField(TEXT("Elements"));
+			for (int32 SubElementIdx = 0; SubElementIdx < Parameters.Elements.Num(); SubElementIdx++)
+			{
+				TSharedPtr<FJsonObject> SubElement = SubObject->GetObjectField(Parameters.ElementNames[SubElementIdx].ToString());
+				if (SubElement.IsValid())
+				{
+					bool bResult = IsActionSchemaSubsetCompatible(SubElement.ToSharedRef(), ActionSchema, Parameters.Elements[SubElementIdx]);
+					if (!bResult)
+					{
+						return false;
+					}
+				}
+				else
+				{
+					return false;
+				}
+			}
+
+			return true;
+		}
+
+		case Action::EType::OrInclusive:
+		{
+			const Action::FSchemaOrInclusiveParameters Parameters = ActionSchema.GetOrInclusive(ActionSchemaElement);
+
+			if (Object->GetStringField(TEXT("Type")) != TEXT("OrInclusive"))
+			{
+				return false;
+			}
+
+			TSharedPtr<FJsonObject> SubObject = Object->GetObjectField(TEXT("Elements"));
+			for (int32 SubElementIdx = 0; SubElementIdx < Parameters.Elements.Num(); SubElementIdx++)
+			{
+				TSharedPtr<FJsonObject> SubElement = SubObject->GetObjectField(Parameters.ElementNames[SubElementIdx].ToString());
+				if (SubElement.IsValid())
+				{
+					bool bResult = IsActionSchemaSubsetCompatible(SubElement.ToSharedRef(), ActionSchema, Parameters.Elements[SubElementIdx]);
+					if (!bResult)
+					{
+						return false;
+					}
+				}
+				else
+				{
+					return false;
+				}
+			}
+
+			return true;
+		}
+
+		case Action::EType::Array:
+		{
+			const Action::FSchemaArrayParameters Parameters = ActionSchema.GetArray(ActionSchemaElement);
+
+			if (Object->GetStringField(TEXT("Type")) != TEXT("Array") ||
+				Object->GetNumberField(TEXT("Num")) != Parameters.Num)
+			{
+				return false;
+			}
+
+			return IsActionSchemaSubsetCompatible(Object->GetObjectField(TEXT("Element")).ToSharedRef(), ActionSchema, Parameters.Element);
+		}
+
+		case Action::EType::Encoding:
+		{
+			const Action::FSchemaEncodingParameters Parameters = ActionSchema.GetEncoding(ActionSchemaElement);
+
+			if (Object->GetStringField(TEXT("Type")) != TEXT("Encoding") ||
+				Object->GetNumberField(TEXT("EncodingSize")) != Parameters.EncodingSize)
+			{
+				return false;
+			}
+
+			return IsActionSchemaSubsetCompatible(Object->GetObjectField(TEXT("Element")).ToSharedRef(), ActionSchema, Parameters.Element);
+		}
+
+		default:
+			checkNoEntry();
+		}
+
+		return false;
+	}
+
+	TArray<int32> ComputeActionSchemaSubsetIndices(const FString& SourceJsonString, const Action::FSchema& ActionSchema, const Action::FSchemaElement& ActionSchemaElement)
+	{
+		TSharedPtr<FJsonObject> JsonObject;
+		TSharedRef<TJsonReader<>> Reader = TJsonReaderFactory<>::Create(SourceJsonString);
+
+		TArray<int32> Output;
+
+		if (FJsonSerializer::Deserialize(Reader, JsonObject) && JsonObject.IsValid())
+		{
+			int32 NextIndex = 0;
+			ComputeActionSchemaSubsetIndices(NextIndex, Output, JsonObject.ToSharedRef(), ActionSchema, ActionSchemaElement);
+		}
+
+		return Output;
+	}
+
+	void ComputeActionSchemaSubsetIndices(int32& NextIndex, TArray<int32>& Indices,	const TSharedRef<FJsonObject> Object, const Action::FSchema& ActionSchema, const Action::FSchemaElement& ActionSchemaElement)
+	{
+
+		if (Object->GetStringField(TEXT("Type")) == TEXT("Null"))
+		{
+			// Do nothing since Null is zero-width
+			return;
+		}
+		else if (Object->GetStringField(TEXT("Type")) == TEXT("Continuous"))
+		{
+			const Action::FSchemaContinuousParameters Parameters = ActionSchema.GetContinuous(ActionSchemaElement);
+
+			for (int32 i = 0; i < Parameters.Num; i++)
+			{
+				Indices.Add(NextIndex);
+				NextIndex++;
+			}
+
+			return;
+		}
+		else if (Object->GetStringField(TEXT("Type")) == TEXT("DiscreteExclusive"))
+		{
+			const Action::FSchemaDiscreteExclusiveParameters Parameters = ActionSchema.GetDiscreteExclusive(ActionSchemaElement);
+
+			for (int32 i = 0; i < Parameters.Num; i++)
+			{
+				Indices.Add(NextIndex);
+				NextIndex++;
+			}
+
+			return;
+		}
+		else if (Object->GetStringField(TEXT("Type")) == TEXT("DiscreteInclusive"))
+		{
+			const Action::FSchemaDiscreteInclusiveParameters Parameters = ActionSchema.GetDiscreteInclusive(ActionSchemaElement);
+			
+			for (int32 i = 0; i < Parameters.Num; i++)
+			{
+				Indices.Add(NextIndex);
+				NextIndex++;
+			}
+
+			return;
+		}
+		else if (Object->GetStringField(TEXT("Type")) == TEXT("NamedDiscreteExclusive"))
+		{
+			const Action::FSchemaNamedDiscreteExclusiveParameters Parameters = ActionSchema.GetNamedDiscreteExclusive(ActionSchemaElement);
+
+			TArray<TSharedPtr<FJsonValue>> ElementNames = Object->GetArrayField(TEXT("ElementNames"));
+			for (const TSharedPtr<FJsonValue>& ElementName : ElementNames)
+			{
+				bool bFound = false;
+				for (int32 ElementIdx = 0; ElementIdx < Parameters.ElementNames.Num(); ElementIdx++)
+				{
+					if (ElementName->Type == EJson::String && ElementName->AsString() == Parameters.ElementNames[ElementIdx].ToString())
+					{
+						bFound = true;
+						break;
+					}
+				}
+
+				if (bFound)
+				{
+					Indices.Add(NextIndex);
+				}
+
+				NextIndex++;
+			}
+
+			return;
+		}
+		else if (Object->GetStringField(TEXT("Type")) == TEXT("NamedDiscreteInclusive"))
+		{
+			const Action::FSchemaNamedDiscreteInclusiveParameters Parameters = ActionSchema.GetNamedDiscreteInclusive(ActionSchemaElement);
+
+			TArray<TSharedPtr<FJsonValue>> ElementNames = Object->GetArrayField(TEXT("ElementNames"));
+			for (const TSharedPtr<FJsonValue>& ElementName : ElementNames)
+			{
+				bool bFound = false;
+				for (int32 ElementIdx = 0; ElementIdx < Parameters.ElementNames.Num(); ElementIdx++)
+				{
+					if (ElementName->Type == EJson::String && ElementName->AsString() == Parameters.ElementNames[ElementIdx].ToString())
+					{
+						bFound = true;
+						break;
+					}
+				}
+
+				if (bFound)
+				{
+					Indices.Add(NextIndex);
+				}
+
+				NextIndex++;
+			}
+
+			return;
+		}
+		else if (Object->GetStringField(TEXT("Type")) == TEXT("And"))
+		{
+			const Action::FSchemaAndParameters Parameters = ActionSchema.GetAnd(ActionSchemaElement);
+
+			TSharedPtr<FJsonObject> SubObject = Object->GetObjectField(TEXT("Elements"));
+			TArray<FString> ElementNames;
+			SubObject->Values.GetKeys(ElementNames);
+
+			for (const FString& ElementName : ElementNames)
+			{
+				TSharedPtr<FJsonObject> SubElement = SubObject->GetObjectField(ElementName);
+
+				bool bFound = false;
+				for (int32 ElementIdx = 0; ElementIdx < Parameters.ElementNames.Num(); ElementIdx++)
+				{
+					if (ElementName == Parameters.ElementNames[ElementIdx].ToString())
+					{
+						bFound = true;
+						ComputeActionSchemaSubsetIndices(NextIndex, Indices, SubElement.ToSharedRef(), ActionSchema, Parameters.Elements[ElementIdx]);
+						break;
+					}
+				}
+
+				if (!bFound)
+				{
+					int32 VectorSize = SubElement->GetNumberField(TEXT("VectorSize"));
+					NextIndex += VectorSize;
+				}
+			}
+
+			return;
+		}
+		else if (Object->GetStringField(TEXT("Type")) == TEXT("OrExclusive"))
+		{
+			const Action::FSchemaOrExclusiveParameters Parameters = ActionSchema.GetOrExclusive(ActionSchemaElement);
+
+			TSharedPtr<FJsonObject> SubObject = Object->GetObjectField(TEXT("Elements"));
+			TArray<FString> ElementNames;
+			SubObject->Values.GetKeys(ElementNames);
+
+			for (const FString& ElementName : ElementNames)
+			{
+				TSharedPtr<FJsonObject> SubElement = SubObject->GetObjectField(ElementName);
+
+				bool bFound = false;
+				for (int32 ElementIdx = 0; ElementIdx < Parameters.ElementNames.Num(); ElementIdx++)
+				{
+					if (ElementName == Parameters.ElementNames[ElementIdx].ToString())
+					{
+						bFound = true;
+						ComputeActionSchemaSubsetIndices(NextIndex, Indices, SubElement.ToSharedRef(), ActionSchema, Parameters.Elements[ElementIdx]);
+						break;
+					}
+				}
+
+				if (!bFound)
+				{
+					int32 VectorSize = SubElement->GetNumberField(TEXT("VectorSize"));
+					NextIndex += VectorSize;
+				}
+			}
+
+			return;
+		}
+		else if (Object->GetStringField(TEXT("Type")) == TEXT("OrInclusive"))
+		{
+			const Action::FSchemaOrInclusiveParameters Parameters = ActionSchema.GetOrInclusive(ActionSchemaElement);
+
+			TSharedPtr<FJsonObject> SubObject = Object->GetObjectField(TEXT("Elements"));
+			TArray<FString> ElementNames;
+			SubObject->Values.GetKeys(ElementNames);
+
+			for (const FString& ElementName : ElementNames)
+			{
+				TSharedPtr<FJsonObject> SubElement = SubObject->GetObjectField(ElementName);
+
+				bool bFound = false;
+				for (int32 ElementIdx = 0; ElementIdx < Parameters.ElementNames.Num(); ElementIdx++)
+				{
+					if (ElementName == Parameters.ElementNames[ElementIdx].ToString())
+					{
+						bFound = true;
+						ComputeActionSchemaSubsetIndices(NextIndex, Indices, SubElement.ToSharedRef(), ActionSchema, Parameters.Elements[ElementIdx]);
+						break;
+					}
+				}
+
+				if (!bFound)
+				{
+					int32 VectorSize = SubElement->GetNumberField(TEXT("VectorSize"));
+					NextIndex += VectorSize;
+				}
+			}
+
+			return;
+		}
+		else if (Object->GetStringField(TEXT("Type")) == TEXT("Array"))
+		{
+			const Action::FSchemaArrayParameters Parameters = ActionSchema.GetArray(ActionSchemaElement);
+
+			for (int i = 0; i < Parameters.Num; i++)
+			{
+				ComputeActionSchemaSubsetIndices(NextIndex, Indices, Object->GetObjectField(TEXT("Element")).ToSharedRef(), ActionSchema, Parameters.Element);
+			}
+
+			return;
+		}
+		else if (Object->GetStringField(TEXT("Type")) == TEXT("Encoding"))
+		{
+			const Action::FSchemaEncodingParameters Parameters = ActionSchema.GetEncoding(ActionSchemaElement);
+
+			ComputeActionSchemaSubsetIndices(NextIndex, Indices, Object->GetObjectField(TEXT("Element")).ToSharedRef(), ActionSchema, Parameters.Element);
+		}
+
+		// We didn't return earlier so we received an invalid Type
+		checkNoEntry();
+	}
+
 	const TCHAR* GetDeviceString(const ETrainerDevice Device)
 	{
 		switch (Device)
diff --git a/Engine/Plugins/Experimental/LearningAgents/Source/LearningTraining/Public/LearningTrainer.h b/Engine/Plugins/Experimental/LearningAgents/Source/LearningTraining/Public/LearningTrainer.h
--- a/Engine/Plugins/Experimental/LearningAgents/Source/LearningTraining/Public/LearningTrainer.h
+++ b/Engine/Plugins/Experimental/LearningAgents/Source/LearningTraining/Public/LearningTrainer.h
@@ -210,6 +210,58 @@ namespace UE::Learning
 		LEARNINGTRAINING_API TSharedPtr<FJsonObject> ConvertActionSchemaToJSON(
 			const Action::FSchema& ActionSchema,
 			const Action::FSchemaElement& ActionSchemaElement);
+
+		/* Returns true if the schema and element are a subset of the provided JSON. */
+		LEARNINGTRAINING_API bool IsObservationSchemaSubsetCompatible(
+			const FString& SourceJsonString,
+			const Observation::FSchema& ObservationSchema,
+			const Observation::FSchemaElement& ObservationSchemaElement);
+
+		/* Returns true if the schema and element are a subset of the provided JSON. */
+		LEARNINGTRAINING_API bool IsObservationSchemaSubsetCompatible(
+			const TSharedRef<FJsonObject> Object,
+			const Observation::FSchema& ObservationSchema,
+			const Observation::FSchemaElement& ObservationSchemaElement);
+
+		/* Computes the indices to copy from a subset compatible JSON. */
+		LEARNINGTRAINING_API TArray<int32> ComputeObservationSchemaSubsetIndices(
+			const FString& SourceJsonString,
+			const Observation::FSchema& ObservationSchema,
+			const Observation::FSchemaElement& ObservationSchemaElement);
+
+		/* Computes the indices to copy from a subset compatible JSON. */
+		LEARNINGTRAINING_API void ComputeObservationSchemaSubsetIndices(
+			int32& NextIndex,
+			TArray<int32>& Indices,
+			const TSharedRef<FJsonObject> Object,
+			const Observation::FSchema& ObservationSchema,
+			const Observation::FSchemaElement& ObservationSchemaElement);
+
+		/* Returns true if the schema and element are a subset of the provided JSON. */
+		LEARNINGTRAINING_API bool IsActionSchemaSubsetCompatible(
+			const FString& SourceJsonString,
+			const Action::FSchema& ActionSchema,
+			const Action::FSchemaElement& ActionSchemaElement);
+
+		/* Returns true if the schema and element are a subset of the provided JSON. */
+		LEARNINGTRAINING_API bool IsActionSchemaSubsetCompatible(
+			const TSharedRef<FJsonObject> Object,
+			const Action::FSchema& ActionSchema,
+			const Action::FSchemaElement& ActionSchemaElement);
+
+		/* Computes the indices to copy from a subset compatible JSON. */
+		LEARNINGTRAINING_API TArray<int32> ComputeActionSchemaSubsetIndices(
+			const FString& SourceJsonString,
+			const Action::FSchema& ActionSchema,
+			const Action::FSchemaElement& ActionSchemaElement);
+
+		/* Computes the indices to copy from a subset compatible JSON. */
+		LEARNINGTRAINING_API void ComputeActionSchemaSubsetIndices(
+			int32& NextIndex,
+			TArray<int32>& Indices,
+			const TSharedRef<FJsonObject> Object,
+			const Action::FSchema& ActionSchema,
+			const Action::FSchemaElement& ActionSchemaElement);
 	}
 
 }
\ No newline at end of file
