LearningAgents: Fix public dependencies

diff --git a/Engine/Plugins/Experimental/LearningAgents/Source/LearningAgents/LearningAgents.Build.cs b/Engine/Plugins/Experimental/LearningAgents/Source/LearningAgents/LearningAgents.Build.cs
--- a/Engine/Plugins/Experimental/LearningAgents/Source/LearningAgents/LearningAgents.Build.cs
+++ b/Engine/Plugins/Experimental/LearningAgents/Source/LearningAgents/LearningAgents.Build.cs
@@ -7,39 +7,21 @@ public class LearningAgents : ModuleRules
 	public LearningAgents(ReadOnlyTargetRules Target) : base(Target)
 	{
 		PCHUsage = ModuleRules.PCHUsageMode.UseExplicitOrSharedPCHs;
-
-		PublicIncludePaths.AddRange(
-			new string[] {
-				// ... add public include paths required here ...
-			});
-		
-		PrivateIncludePaths.AddRange(
-			new string[] {
-				// ... add other private include paths required here ...
-			});
 		
 		PublicDependencyModuleNames.AddRange(
 			new string[]
 			{
 				"Core",
+				"CoreUObject",
+				"Engine",
 				"Learning",
-				// ... add other public dependencies that you statically link with here ...
 			});
 		
 		PrivateDependencyModuleNames.AddRange(
 			new string[]
 			{
-				"CoreUObject",
-				"Engine",
 				"NNE",
 				"NNERuntimeBasicCpu",
-				// ... add private dependencies that you statically link with here ...	
-			});
-		
-		DynamicallyLoadedModuleNames.AddRange(
-			new string[]
-			{
-				// ... add any modules that your module loads dynamically here ...
 			});
 	}
 }
