[LearningAgents] [NNE] Added Conv1d (zeros & circular padding modes)

- Conv1d NNE Layer
- Conv1d observation in LearningAgents
- Lidar observation in LearningAgents

diff --git a/Engine/Plugins/Experimental/LearningAgents/Source/Learning/Private/LearningObservation.cpp b/Engine/Plugins/Experimental/LearningAgents/Source/Learning/Private/LearningObservation.cpp
--- a/Engine/Plugins/Experimental/LearningAgents/Source/Learning/Private/LearningObservation.cpp
+++ b/Engine/Plugins/Experimental/LearningAgents/Source/Learning/Private/LearningObservation.cpp
@@ -298,6 +298,33 @@ namespace UE::Learning::Observation
 		return { Index, Generation };
 	}
 
+	FSchemaElement FSchema::CreateConv1d(const FSchemaConv1dParameters Parameters, const FName Tag)
+	{
+		FConv1dData ElementData;
+		ElementData.ElementIndex = SubElementObjects.Num();
+		ElementData.InputLength = Parameters.InputLength;
+		ElementData.InChannels = Parameters.InChannels;
+		ElementData.OutChannels = Parameters.OutChannels;
+		ElementData.KernelSize = Parameters.KernelSize;
+		ElementData.Padding = Parameters.Padding;
+		ElementData.PaddingMode = Parameters.PaddingMode;
+		ElementData.ActivationFunction = Parameters.ActivationFunction;
+
+		SubElementNames.Add(NAME_None);
+		SubElementObjects.Add(Parameters.Element);
+
+		const int32 Index = Types.Add(EType::Conv1d);
+		Tags.Add(Tag);
+		
+		ObservationVectorSizes.Add(GetObservationVectorSize(Parameters.Element));
+		int32 OutputLength = Parameters.InputLength + 2 * Parameters.Padding - Parameters.KernelSize + 1;
+		EncodedVectorSizes.Add(OutputLength * Parameters.OutChannels);
+
+		TypeDataIndices.Add(Conv1dData.Add(ElementData));
+
+		return { Index, Generation };
+	}
+
 	bool FSchema::IsValid(const FSchemaElement Element) const
 	{
 		return Element.Generation == Generation && Element.Index != INDEX_NONE;
@@ -453,6 +480,23 @@ namespace UE::Learning::Observation
 		return Parameters;
 	}
 
+	FSchemaConv1dParameters FSchema::GetConv1d(const FSchemaElement Element) const
+	{
+		check(IsValid(Element) && GetType(Element) == EType::Conv1d);
+		const FConv1dData& ElementData = Conv1dData[TypeDataIndices[Element.Index]];
+
+		FSchemaConv1dParameters Parameters;
+		Parameters.Element = SubElementObjects[ElementData.ElementIndex];
+		Parameters.InputLength = ElementData.InputLength;
+		Parameters.InChannels = ElementData.InChannels;
+		Parameters.OutChannels = ElementData.OutChannels;
+		Parameters.KernelSize = ElementData.KernelSize;
+		Parameters.Padding = ElementData.Padding;
+		Parameters.PaddingMode = ElementData.PaddingMode;
+		Parameters.ActivationFunction = ElementData.ActivationFunction;
+		return Parameters;
+	}
+
 	uint32 FSchema::GetGeneration() const
 	{
 		return Generation;
@@ -773,6 +817,28 @@ namespace UE::Learning::Observation
 		return { Index, Generation };
 	}
 
+	FObjectElement FObject::CreateConv1d(const FObjectConv1dParameters Parameters, const FName Tag)
+	{
+		check(IsValid(Parameters.Element));
+
+		const int32 Index = Types.Add(EType::Conv1d);
+		Tags.Add(Tag);
+
+		ContinuousDataOffsets.Add(ContinuousValues.Num());
+		ContinuousDataNums.Add(0);
+
+		DiscreteDataOffsets.Add(DiscreteValues.Num());
+		DiscreteDataNums.Add(0);
+
+		SubElementDataOffsets.Add(SubElementObjects.Num());
+		SubElementDataNums.Add(1);
+
+		SubElementNames.Add(NAME_Name);
+		SubElementObjects.Add(Parameters.Element);
+
+		return { Index, Generation };
+	}
+
 	bool FObject::IsValid(const FObjectElement Element) const
 	{
 		return Element.Generation == Generation && Element.Index != INDEX_NONE;
@@ -892,6 +958,15 @@ namespace UE::Learning::Observation
 		return Parameters;
 	}
 
+	FObjectConv1dParameters FObject::GetConv1d(const FObjectElement Element) const
+	{
+		check(IsValid(Element) && GetType(Element) == EType::Conv1d);
+
+		FObjectConv1dParameters Parameters;
+		Parameters.Element = SubElementObjects[SubElementDataOffsets[Element.Index]];
+		return Parameters;
+	}
+
 	uint32 FObject::GetGeneration() const
 	{
 		return Generation;
@@ -960,6 +1035,16 @@ namespace UE::Learning::Observation
 			}
 		}
 
+		static inline NNE::RuntimeBasic::FModelBuilder::EPaddingMode GetNNEPaddingMode(const EPaddingMode PaddingMode)
+		{
+			switch (PaddingMode)
+			{
+			case EPaddingMode::Zeros: return NNE::RuntimeBasic::FModelBuilder::EPaddingMode::Zeros;
+			case EPaddingMode::Circular: return NNE::RuntimeBasic::FModelBuilder::EPaddingMode::Circular;
+			default: checkNoEntry(); return NNE::RuntimeBasic::FModelBuilder::EPaddingMode::Zeros;
+			}
+		}
+
 		static inline int32 HashFNameStable(const FName Name)
 		{
 			const FString NameString = Name.ToString().ToLower();
@@ -1085,6 +1170,19 @@ namespace UE::Learning::Observation
 			return GetSchemaObjectsCompatibilityHash(Schema, Parameters.Element, Salt);
 		}
 
+		case EType::Conv1d:
+		{
+			const FSchemaConv1dParameters Parameters = Schema.GetConv1d(SchemaElement);
+			return Private::HashCombine({
+				Hash, 
+				Private::HashInt(Parameters.InputLength), 
+				Private::HashInt(Parameters.InChannels), 
+				Private::HashInt(Parameters.OutChannels), 
+				Private::HashInt(Parameters.KernelSize),
+				Private::HashInt(Parameters.Padding),
+				GetSchemaObjectsCompatibilityHash(Schema, Parameters.Element, Salt) });
+		}
+
 		default:
 		{
 			checkNoEntry();
@@ -1230,6 +1328,19 @@ namespace UE::Learning::Observation
 			return false;
 		}
 
+		case EType::Conv1d:
+		{
+			const FSchemaConv1dParameters ParametersA = SchemaA.GetConv1d(SchemaElementA);
+			const FSchemaConv1dParameters ParametersB = SchemaB.GetConv1d(SchemaElementB);
+
+			return (ParametersA.InputLength == ParametersB.InputLength)
+				&& (ParametersA.InChannels == ParametersB.InChannels)
+				&& (ParametersA.OutChannels == ParametersB.OutChannels)
+				&& (ParametersA.KernelSize == ParametersB.KernelSize)
+				&& (ParametersA.Padding == ParametersB.Padding)
+				&& AreSchemaObjectsCompatible(SchemaA, ParametersA.Element, SchemaB, ParametersB.Element);
+		}
+
 		default:
 		{
 			checkNoEntry();
@@ -1350,7 +1461,7 @@ namespace UE::Learning::Observation
 
 				switch (NetworkSettings.WeightInitialization)
 				{
-				case EWeightInitialization::KaimingGaussian: LinearLayerSettings.WeightInitializationSettings.Type = 
+				case EWeightInitialization::KaimingGaussian: LinearLayerSettings.WeightInitializationSettings.Type =
 					NNE::RuntimeBasic::FModelBuilder::EWeightInitializationType::KaimingGaussian; break;
 				case EWeightInitialization::KaimingUniform: LinearLayerSettings.WeightInitializationSettings.Type =
 					NNE::RuntimeBasic::FModelBuilder::EWeightInitializationType::KaimingUniform; break;
@@ -1389,7 +1500,7 @@ namespace UE::Learning::Observation
 
 				switch (NetworkSettings.WeightInitialization)
 				{
-				case EWeightInitialization::KaimingGaussian: LinearLayerSettings.WeightInitializationSettings.Type = 
+				case EWeightInitialization::KaimingGaussian: LinearLayerSettings.WeightInitializationSettings.Type =
 					NNE::RuntimeBasic::FModelBuilder::EWeightInitializationType::KaimingGaussian; break;
 				case EWeightInitialization::KaimingUniform: LinearLayerSettings.WeightInitializationSettings.Type =
 					NNE::RuntimeBasic::FModelBuilder::EWeightInitializationType::KaimingUniform; break;
@@ -1439,7 +1550,7 @@ namespace UE::Learning::Observation
 
 			switch (NetworkSettings.WeightInitialization)
 			{
-			case EWeightInitialization::KaimingGaussian: LinearLayerSettings.WeightInitializationSettings.Type = 
+			case EWeightInitialization::KaimingGaussian: LinearLayerSettings.WeightInitializationSettings.Type =
 				NNE::RuntimeBasic::FModelBuilder::EWeightInitializationType::KaimingGaussian; break;
 			case EWeightInitialization::KaimingUniform: LinearLayerSettings.WeightInitializationSettings.Type =
 				NNE::RuntimeBasic::FModelBuilder::EWeightInitializationType::KaimingUniform; break;
@@ -1475,7 +1586,7 @@ namespace UE::Learning::Observation
 
 			switch (NetworkSettings.WeightInitialization)
 			{
-			case EWeightInitialization::KaimingGaussian: LinearLayerSettings.WeightInitializationSettings.Type = 
+			case EWeightInitialization::KaimingGaussian: LinearLayerSettings.WeightInitializationSettings.Type =
 				NNE::RuntimeBasic::FModelBuilder::EWeightInitializationType::KaimingGaussian; break;
 			case EWeightInitialization::KaimingUniform: LinearLayerSettings.WeightInitializationSettings.Type =
 				NNE::RuntimeBasic::FModelBuilder::EWeightInitializationType::KaimingUniform; break;
@@ -1497,6 +1608,38 @@ namespace UE::Learning::Observation
 			break;
 		}
 
+		case EType::Conv1d:
+		{
+			const FSchemaConv1dParameters Parameters = Schema.GetConv1d(SchemaElement);
+			const int32 SubElementEncodedSize = Schema.GetEncodedVectorSize(Parameters.Element);
+			NNE::RuntimeBasic::FModelBuilderElement BuilderSubElement;
+			MakeEncoderNetworkModelBuilderElementFromSchema(BuilderSubElement, Builder, Schema, Parameters.Element, NetworkSettings);
+			
+			NNE::RuntimeBasic::FModelBuilder::FWeightInitializationSettings WeightInitializationSettings;
+			switch (NetworkSettings.WeightInitialization)
+			{
+			case EWeightInitialization::KaimingGaussian: WeightInitializationSettings.Type =
+				NNE::RuntimeBasic::FModelBuilder::EWeightInitializationType::KaimingGaussian; break;
+			case EWeightInitialization::KaimingUniform: WeightInitializationSettings.Type =
+				NNE::RuntimeBasic::FModelBuilder::EWeightInitializationType::KaimingUniform; break;
+			default: checkNoEntry();
+			}
+
+			OutElement = Builder.MakeConv1d(
+				Parameters.InputLength,
+				Parameters.InChannels,
+				Parameters.OutChannels,
+				Parameters.KernelSize,
+				Parameters.Padding,
+				Private::GetNNEPaddingMode(Parameters.PaddingMode),
+				Private::GetNNEActivationFunction(Parameters.ActivationFunction),
+				Builder.MakeInitialWeights(Parameters.InChannels * Parameters.KernelSize, Parameters.OutChannels, WeightInitializationSettings),
+				Builder.MakeInitialBiases(Parameters.OutChannels, WeightInitializationSettings),
+				BuilderSubElement);
+			
+			break;
+		}
+
 		default:
 		{
 			checkNoEntry();
@@ -1825,6 +1968,21 @@ namespace UE::Learning::Observation
 			return;
 		}
 
+		case EType::Conv1d:
+		{
+			const FSchemaConv1dParameters SchemaParameters = Schema.GetConv1d(SchemaElement);
+			const FObjectConv1dParameters ObjectParameters = Object.GetConv1d(ObjectElement);
+
+			SetVectorFromObject(
+				OutObservationVector,
+				Schema,
+				SchemaParameters.Element,
+				Object,
+				ObjectParameters.Element);
+
+			return;
+		}
+
 		default:
 		{
 			checkNoEntry();
@@ -2158,6 +2316,22 @@ namespace UE::Learning::Observation
 			return;
 		}
 
+		case EType::Conv1d:
+		{
+			const FSchemaConv1dParameters Parameters = Schema.GetConv1d(SchemaElement);
+
+			FObjectElement SubElement;
+			GetObjectFromVector(
+				OutObject,
+				SubElement,
+				Schema,
+				Parameters.Element,
+				ObservationVector);
+
+			OutObjectElement = OutObject.CreateConv1d({ SubElement }, SchemaElementTag);
+			return;
+		}
+
 		default:
 		{
 			checkNoEntry();
@@ -2317,7 +2491,7 @@ namespace UE::Learning::Observation
 			{
 				AddGaussianNoiseToVector(
 					InOutRandomState,
-					InOutObservationVector.Slice(ElementIdx* SubElementSize, SubElementSize),
+					InOutObservationVector.Slice(ElementIdx * SubElementSize, SubElementSize),
 					Schema,
 					Parameters.Element,
 					NoiseScale);
@@ -2365,6 +2539,18 @@ namespace UE::Learning::Observation
 			return;
 		}
 
+		case EType::Conv1d:
+		{
+			AddGaussianNoiseToVector(
+				InOutRandomState,
+				InOutObservationVector,
+				Schema,
+				Schema.GetConv1d(SchemaElement).Element,
+				NoiseScale);
+
+			return;
+		}
+
 		default:
 		{
 			checkNoEntry();
diff --git a/Engine/Plugins/Experimental/LearningAgents/Source/Learning/Public/LearningObservation.h b/Engine/Plugins/Experimental/LearningAgents/Source/Learning/Public/LearningObservation.h
--- a/Engine/Plugins/Experimental/LearningAgents/Source/Learning/Public/LearningObservation.h
+++ b/Engine/Plugins/Experimental/LearningAgents/Source/Learning/Public/LearningObservation.h
@@ -22,10 +22,10 @@ namespace UE::Learning::Observation
 	enum class EType : uint8
 	{
 		// Empty Observation
-		Null		= 0,	
+		Null = 0,
 
 		// Vector of continuous float observations
-		Continuous	= 1,	
+		Continuous = 1,
 
 		// Set of exclusive discrete actions
 		DiscreteExclusive = 2,
@@ -40,7 +40,7 @@ namespace UE::Learning::Observation
 		NamedDiscreteInclusive = 5,
 
 		// Combination of multiple observations
-		And			= 6,
+		And = 6,
 
 		// Exclusive choice from a set of observations
 		OrExclusive = 7,
@@ -49,19 +49,22 @@ namespace UE::Learning::Observation
 		OrInclusive = 8,
 
 		// Fixed sized array of observations
-		Array		= 9,
+		Array = 9,
 
 		// Variable sized, unordered set of observations
-		Set			= 10,
+		Set = 10,
 
 		// Encoding of another observation 
-		Encoding	= 11,
+		Encoding = 11,
+
+		// 1d convolution of another observation 
+		Conv1d = 12,
 	};
 
 	/**
 	 * Observation Schema Element
 	 *
-	 * A single element in the observation schema representing a part of an observation. Internally this consists of a index used by the schema to 
+	 * A single element in the observation schema representing a part of an observation. Internally this consists of a index used by the schema to
 	 * look up the associated observation data and a generation id which can be used to check when this index is no longer valid.
 	 */
 	struct FSchemaElement { int32 Index = INDEX_NONE; uint32 Generation = INDEX_NONE; };
@@ -174,6 +177,13 @@ namespace UE::Learning::Observation
 		GELU = 3,
 	};
 
+	/** Enum of Padding Mode */
+	enum class EPaddingMode : uint8
+	{
+		Zeros = 0,
+		Circular = 1,
+	};
+
 	struct FSchemaEncodingParameters
 	{
 		// The sub-element.
@@ -189,11 +199,38 @@ namespace UE::Learning::Observation
 		EEncodingActivationFunction ActivationFunction = EEncodingActivationFunction::ELU;
 	};
 
+	struct FSchemaConv1dParameters
+	{
+		// The sub-element.
+		FSchemaElement Element;
+
+		// Length of the 1d convolution input.
+		int32 InputLength;
+
+		// Number of input channels.
+		int32 InChannels;
+
+		// Number of output channels after convolution.
+		int32 OutChannels;
+
+		// Size of the convolving kernel.
+		int32 KernelSize;
+
+		// The amount of padding to add to both sides of the input.
+		int32 Padding;
+
+		// Zeros or Circular.
+		EPaddingMode PaddingMode = EPaddingMode::Zeros;
+
+		// The activation function to use for the convolution layer.
+		EEncodingActivationFunction ActivationFunction = EEncodingActivationFunction::ReLU;
+	};
+
 	/**
 	 * Observation Schema
 	 *
-	 * This object allows you to construct a description of the kind of observations you might want to provide as input to a policy. Internally this 
-	 * object contains a pool of individual elements. This allows them to be constructed performantly and in a cache efficient way. This object is 
+	 * This object allows you to construct a description of the kind of observations you might want to provide as input to a policy. Internally this
+	 * object contains a pool of individual elements. This allows them to be constructed performantly and in a cache efficient way. This object is
 	 * therefore required to access any data about the individual observation elements that are created.
 	 */
 	struct LEARNING_API FSchema
@@ -210,6 +247,7 @@ namespace UE::Learning::Observation
 		FSchemaElement CreateArray(const FSchemaArrayParameters Parameters, const FName Tag = NAME_None);
 		FSchemaElement CreateSet(const FSchemaSetParameters Parameters, const FName Tag = NAME_None);
 		FSchemaElement CreateEncoding(const FSchemaEncodingParameters Parameters, const FName Tag = NAME_None);
+		FSchemaElement CreateConv1d(const FSchemaConv1dParameters Parameters, const FName Tag = NAME_None);
 
 		FSchemaContinuousParameters GetContinuous(const FSchemaElement Element) const;
 		FSchemaDiscreteExclusiveParameters GetDiscreteExclusive(const FSchemaElement Element) const;
@@ -222,6 +260,7 @@ namespace UE::Learning::Observation
 		FSchemaArrayParameters GetArray(const FSchemaElement Element) const;
 		FSchemaSetParameters GetSet(const FSchemaElement Element) const;
 		FSchemaEncodingParameters GetEncoding(const FSchemaElement Element) const;
+		FSchemaConv1dParameters GetConv1d(const FSchemaElement Element) const;
 
 		// Checks if the given element is valid
 		bool IsValid(const FSchemaElement Element) const;
@@ -325,6 +364,18 @@ namespace UE::Learning::Observation
 			EEncodingActivationFunction ActivationFunction = EEncodingActivationFunction::ELU;
 		};
 
+		struct FConv1dData
+		{
+			int32 ElementIndex = INDEX_NONE;
+			int32 InputLength = INDEX_NONE;
+			int32 InChannels = INDEX_NONE;
+			int32 OutChannels = INDEX_NONE;
+			int32 KernelSize = INDEX_NONE;
+			int32 Padding = INDEX_NONE;
+			EPaddingMode PaddingMode = EPaddingMode::Zeros;
+			EEncodingActivationFunction ActivationFunction = EEncodingActivationFunction::ReLU;
+		};
+
 		uint32 Generation = 0;
 
 		/** These have entries for each Schema Element */
@@ -333,7 +384,7 @@ namespace UE::Learning::Observation
 		TArray<int32> ObservationVectorSizes;
 		TArray<int32> EncodedVectorSizes;
 		TArray<int32> TypeDataIndices;
-		
+
 		/** These are indexed based on the Type using TypeDataIndices */
 		TArray<FContinuousData> ContinuousData;
 		TArray<FDiscreteExclusiveData> DiscreteExclusiveData;
@@ -346,6 +397,7 @@ namespace UE::Learning::Observation
 		TArray<FArrayData> ArrayData;
 		TArray<FSetData> SetData;
 		TArray<FEncodingData> EncodingData;
+		TArray<FConv1dData> Conv1dData;
 
 		/** This is an array of all the SubElements and their names, referenced by other elements */
 		TArray<FName> SubElementNames;
@@ -355,7 +407,7 @@ namespace UE::Learning::Observation
 	/**
 	 * Observation Object Element
 	 *
-	 * A single element in the observation object representing part of an observation. Internally this consists of a index used by the object to look 
+	 * A single element in the observation object representing part of an observation. Internally this consists of a index used by the object to look
 	 * up the associated observation data and a generation id which can be used to check when this index is no longer valid.
 	 */
 	struct FObjectElement { int32 Index = INDEX_NONE; uint32 Generation = INDEX_NONE; };
@@ -435,6 +487,12 @@ namespace UE::Learning::Observation
 		FObjectElement Element;
 	};
 
+	struct FObjectConv1dParameters
+	{
+		// Encoded sub-element.
+		FObjectElement Element;
+	};
+
 	/**
 	 * Observation Object
 	 *
@@ -456,6 +514,7 @@ namespace UE::Learning::Observation
 		FObjectElement CreateArray(const FObjectArrayParameters Parametes, const FName Tag = NAME_None);
 		FObjectElement CreateSet(const FObjectSetParameters Parameters, const FName Tag = NAME_None);
 		FObjectElement CreateEncoding(const FObjectEncodingParameters Parameters, const FName Tag = NAME_None);
+		FObjectElement CreateConv1d(const FObjectConv1dParameters Parameters, const FName Tag = NAME_None);
 
 		FObjectContinuousParameters GetContinuous(const FObjectElement Element) const;
 		FObjectDiscreteExclusiveParameters GetDiscreteExclusive(const FObjectElement Element) const;
@@ -468,6 +527,7 @@ namespace UE::Learning::Observation
 		FObjectArrayParameters GetArray(const FObjectElement Element) const;
 		FObjectSetParameters GetSet(const FObjectElement Element) const;
 		FObjectEncodingParameters GetEncoding(const FObjectElement Element) const;
+		FObjectConv1dParameters GetConv1d(const FObjectElement Element) const;
 
 		// Checks if the given element is valid
 		bool IsValid(const FObjectElement Element) const;
diff --git a/Engine/Plugins/Experimental/LearningAgents/Source/LearningAgents/Private/LearningAgentsObservations.cpp b/Engine/Plugins/Experimental/LearningAgents/Source/LearningAgents/Private/LearningAgentsObservations.cpp
--- a/Engine/Plugins/Experimental/LearningAgents/Source/LearningAgents/Private/LearningAgentsObservations.cpp
+++ b/Engine/Plugins/Experimental/LearningAgents/Source/LearningAgents/Private/LearningAgentsObservations.cpp
@@ -415,6 +415,18 @@ namespace UE::Learning::Agents::Observation::Private
 				ObjectParameters.Element);
 		}
 
+		case Learning::Observation::EType::Conv1d:
+		{
+			const Learning::Observation::FSchemaConv1dParameters SchemaParameters = Schema.GetConv1d(SchemaElement);
+			const Learning::Observation::FObjectConv1dParameters ObjectParameters = Object.GetConv1d(ObjectElement);
+
+			return ValidateObservationObjectMatchesSchema(
+				Schema,
+				SchemaParameters.Element,
+				Object,
+				ObjectParameters.Element);
+		}
+
 		default:
 		{
 			checkNoEntry();
@@ -576,6 +588,16 @@ namespace UE::Learning::Agents::Observation::Private
 			FMath::Exp(V.Z));
 	}
 
+	static inline Learning::Observation::EPaddingMode GetPaddingMode(const ELearningAgentsPaddingMode PaddingMode)
+	{
+		switch (PaddingMode)
+		{
+		case ELearningAgentsPaddingMode::Zeros: return Learning::Observation::EPaddingMode::Zeros;
+		case ELearningAgentsPaddingMode::Circular: return Learning::Observation::EPaddingMode::Circular;
+		default: checkNoEntry(); return Learning::Observation::EPaddingMode::Zeros;
+		}
+	}
+
 	static inline Learning::Observation::EEncodingActivationFunction GetEncodingActivationFunction(const ELearningAgentsActivationFunction ActivationFunction)
 	{
 		switch (ActivationFunction)
@@ -1203,6 +1225,63 @@ FLearningAgentsObservationSchemaElement ULearningAgentsObservations::SpecifyBool
 	return SpecifyContinuousObservation(Schema, 1, 1.0f, Tag);
 }
 
+FLearningAgentsObservationSchemaElement ULearningAgentsObservations::SpecifyConv1dObservation(ULearningAgentsObservationSchema* Schema, const FLearningAgentsObservationSchemaElement Element, const int32 InputLength, const int32 InChannels, const int32 OutChannels, const int32 KernelSize, const int32 Padding, const ELearningAgentsPaddingMode PaddingMode, const ELearningAgentsActivationFunction ActivationFunction, const FName Tag /*= TEXT("Conv1dObservation")*/)
+{
+	if (!Schema)
+	{
+		UE_LOG(LogLearning, Error, TEXT("SpecifyConv1dObservation: Schema is nullptr."));
+		return FLearningAgentsObservationSchemaElement();
+	}
+	if (!Schema->ObservationSchema.IsValid(Element.SchemaElement))
+	{
+		UE_LOG(LogLearning, Error, TEXT("SpecifyConv1dObservation: Invalid Observation Object."));
+		return FLearningAgentsObservationSchemaElement();
+	}
+	if (InputLength < 1)
+	{
+		UE_LOG(LogLearning, Error, TEXT("SpecifyConv1dObservation: Invalid Observation InputLength '%i' - must be greater than zero."), InputLength);
+		return FLearningAgentsObservationSchemaElement();
+	}
+	if (InChannels < 1)
+	{
+		UE_LOG(LogLearning, Error, TEXT("SpecifyConv1dObservation: Invalid Observation InChannels '%i' - must be greater than zero."), InChannels);
+		return FLearningAgentsObservationSchemaElement();
+	}
+	if (OutChannels < 1)
+	{
+		UE_LOG(LogLearning, Error, TEXT("SpecifyConv1dObservation: Invalid Observation OutChannels '%i' - must be greater than zero."), OutChannels);
+		return FLearningAgentsObservationSchemaElement();
+	}
+	if (KernelSize < 1)
+	{
+		UE_LOG(LogLearning, Error, TEXT("SpecifyConv1dObservation: Invalid Observation KernelSize '%i' - must be greater than zero."), KernelSize);
+		return FLearningAgentsObservationSchemaElement();
+	}
+	if (Padding < 1)
+	{
+		UE_LOG(LogLearning, Error, TEXT("SpecifyConv1dObservation: Invalid Observation Padding '%i' - must be greater than zero."), Padding);
+		return FLearningAgentsObservationSchemaElement();
+	}
+	if (KernelSize > InputLength + 2 * Padding)
+	{
+		UE_LOG(LogLearning, Error, TEXT("SpecifyConv1dObservation: Invalid KernelSize '%i' - Kernel size is too large."), KernelSize);
+		return FLearningAgentsObservationSchemaElement();
+	}
+
+	return { Schema->ObservationSchema.CreateConv1d(
+	{
+		Element.SchemaElement,
+		InputLength,
+		InChannels,
+		OutChannels,
+		KernelSize,
+		Padding,
+		UE::Learning::Agents::Observation::Private::GetPaddingMode(PaddingMode),
+		UE::Learning::Agents::Observation::Private::GetEncodingActivationFunction(ActivationFunction)
+	},
+	Tag) };
+}
+
 FLearningAgentsObservationSchemaElement ULearningAgentsObservations::SpecifyFloatObservation(ULearningAgentsObservationSchema* Schema, const float FloatScale, const FName Tag)
 {
 	return SpecifyContinuousObservation(Schema, 1, FloatScale, Tag);
@@ -1284,6 +1363,11 @@ FLearningAgentsObservationSchemaElement ULearningAgentsObservations::SpecifyProp
 	return SpecifyFloatObservation(Schema, 1.0f, Tag);
 }
 
+FLearningAgentsObservationSchemaElement ULearningAgentsObservations::SpecifyLidarObservation(ULearningAgentsObservationSchema* Schema, const FName Tag)
+{
+	return SpecifyFloatObservation(Schema, 1.0f, Tag);
+}
+
 void ULearningAgentsObservations::LogObservation(const ULearningAgentsObservationObject* Object, const FLearningAgentsObservationObjectElement Element)
 {
 	if (!Object)
@@ -2111,6 +2195,23 @@ FLearningAgentsObservationObjectElement ULearningAgentsObservations::MakeBoolObs
 	return MakeContinuousObservationFromArrayView(Object, { EncodedBool }, Tag);
 }
 
+FLearningAgentsObservationObjectElement ULearningAgentsObservations::MakeConv1dObservation(ULearningAgentsObservationObject* Object, const FLearningAgentsObservationObjectElement Element, const FName Tag /*= TEXT("Conv1dObservation")*/)
+{
+	if (!Object)
+	{
+		UE_LOG(LogLearning, Error, TEXT("MakeConv1dObservation: Object is nullptr."));
+		return FLearningAgentsObservationObjectElement();
+	}
+
+	if (!Object->ObservationObject.IsValid(Element.ObjectElement))
+	{
+		UE_LOG(LogLearning, Error, TEXT("MakeConv1dObservation: Invalid Observation Object."));
+		return FLearningAgentsObservationObjectElement();
+	}
+
+	return { Object->ObservationObject.CreateConv1d({ Element.ObjectElement }, Tag) };
+}
+
 FLearningAgentsObservationObjectElement ULearningAgentsObservations::MakeFloatObservation(
 	ULearningAgentsObservationObject* Object,
 	const float Value,
@@ -2724,6 +2825,85 @@ FLearningAgentsObservationObjectElement ULearningAgentsObservations::MakeProport
 	return MakeFloatObservation(Object, Encoded, Tag);
 }
 
+FLearningAgentsObservationObjectElement ULearningAgentsObservations::MakeLidarObservation(
+	ULearningAgentsObservationObject* Object, 
+	const FVector RayStart,
+	const FVector RayEnd,
+	const FTransform RayTransform, 
+	const ECollisionChannel CollisionChannel, 
+	const AActor* ActorToIgnore,
+	const FName Tag,
+	const bool bVisualLoggerEnabled,
+	ULearningAgentsManagerListener* VisualLoggerListener,
+	const int32 VisualLoggerAgentId,
+	const FVector VisualLoggerLocation,
+	const FLinearColor VisualLoggerColor)
+{
+	const FVector RayStartWorld = RayTransform.TransformPosition(RayStart);
+	const FVector RayEndWorld = RayTransform.TransformPosition(RayEnd);
+	const float RayRange = FVector::Dist(RayEndWorld, RayStartWorld);
+
+	if (RayRange <= 0)
+	{
+		UE_LOG(LogLearning, Error, TEXT("MakeLiDARObservation: Ray range cannot be 0."));
+		return FLearningAgentsObservationObjectElement();
+	}
+
+	FCollisionObjectQueryParams ObjectQueryParams;
+	ObjectQueryParams.AddObjectTypesToQuery(CollisionChannel);
+
+	FCollisionQueryParams TraceParams;
+	TraceParams.AddIgnoredActor(ActorToIgnore);
+
+	FHitResult TraceHit;
+	const bool bHit = Object->GetWorld()->LineTraceSingleByObjectType(TraceHit, RayStartWorld, RayEndWorld, ObjectQueryParams, TraceParams);
+
+	const float Encoded = bHit ? TraceHit.Distance / RayRange : 1.0f; 
+
+#if UE_LEARNING_AGENTS_ENABLE_VISUAL_LOG
+	if (bVisualLoggerEnabled && VisualLoggerListener)
+	{
+		const ULearningAgentsVisualLoggerObject* VisualLoggerObject = VisualLoggerListener->GetOrAddVisualLoggerObject(Tag);
+
+		UE_LEARNING_AGENTS_VLOG_SEGMENT(VisualLoggerObject, LogLearning, Display,
+			RayStartWorld,
+			RayEndWorld,
+			VisualLoggerColor.ToFColor(true),
+			TEXT(""));
+
+		if (bHit)
+		{
+			UE_LEARNING_AGENTS_VLOG_LOCATION(VisualLoggerObject, LogLearning, Display,
+				TraceHit.Location,
+				10.f,
+				VisualLoggerColor.ToFColor(true),
+				TEXT(""));
+		}
+
+		UE_LEARNING_AGENTS_VLOG_TRANSFORM(VisualLoggerObject, LogLearning, Display,
+			RayTransform.GetTranslation(),
+			RayTransform.GetRotation(),
+			VisualLoggerColor.ToFColor(true),
+			TEXT(""));
+
+		UE_LEARNING_AGENTS_VLOG_STRING(VisualLoggerObject, LogLearning, Display, VisualLoggerLocation,
+			VisualLoggerColor.ToFColor(true),
+			TEXT("Listener: %s\nTag: %s\nAgent Id: % 3i\nRay Start: [% 6.1f % 6.1f % 6.1f]\nRay End: [% 6.1f % 6.1f % 6.1f]\nRay Start World: [% 6.1f % 6.1f % 6.1f]\nRay End World: [% 6.1f % 6.1f % 6.1f]\nCollision Channel: [%s]\nHit: [%s]\nEncoded: [% 6.2f]"),
+			*VisualLoggerListener->GetName(),
+			*Tag.ToString(),
+			VisualLoggerAgentId,
+			RayStart.X, RayStart.Y, RayStart.Z,
+			RayEnd.X, RayEnd.Y, RayEnd.Z,
+			RayStartWorld.X, RayStartWorld.Y, RayStartWorld.Z,
+			RayEndWorld.X, RayEndWorld.Y, RayEndWorld.Z,
+			*StaticEnum<ECollisionChannel>()->GetNameStringByValue(CollisionChannel),
+			bHit ? TEXT("true") : TEXT("false"),
+			Encoded);
+	}
+#endif
+
+	return MakeFloatObservation(Object, Encoded, Tag);
+}
 
 bool ULearningAgentsObservations::GetNullObservation(const ULearningAgentsObservationObject* Object, const FLearningAgentsObservationObjectElement Element, const FName Tag)
 {
@@ -4367,3 +4547,39 @@ bool ULearningAgentsObservations::GetProportionAlongRayObservation(float& OutPro
 	return true;
 }
 
+bool ULearningAgentsObservations::GetConv1dObservation(FLearningAgentsObservationObjectElement& OutElement, const ULearningAgentsObservationObject* Object, const FLearningAgentsObservationObjectElement Element, const FName Tag)
+{
+	if (!Object)
+	{
+		UE_LOG(LogLearning, Error, TEXT("GetConv1dObservation: Object is nullptr."));
+		OutElement = FLearningAgentsObservationObjectElement();
+		return false;
+	}
+
+	if (!Object->ObservationObject.IsValid(Element.ObjectElement))
+	{
+		UE_LOG(LogLearning, Error, TEXT("GetConv1dObservation: Invalid Observation Object."));
+		OutElement = FLearningAgentsObservationObjectElement();
+		return false;
+	}
+
+	if (Object->ObservationObject.GetTag(Element.ObjectElement) != Tag)
+	{
+		UE_LOG(LogLearning, Warning, TEXT("GetConv1dObservation: Observation tag does not match. Observation is '%s' but asked for '%s'."), *Object->ObservationObject.GetTag(Element.ObjectElement).ToString(), *Tag.ToString());
+	}
+
+	if (Object->ObservationObject.GetType(Element.ObjectElement) != UE::Learning::Observation::EType::Conv1d)
+	{
+		UE_LOG(LogLearning, Error, TEXT("GetConv1dObservation: Observation '%s' type does not match. Observation is '%s' but asked for '%s'."),
+			*Object->ObservationObject.GetTag(Element.ObjectElement).ToString(),
+			UE::Learning::Agents::Observation::Private::GetObservationTypeString(Object->ObservationObject.GetType(Element.ObjectElement)),
+			UE::Learning::Agents::Observation::Private::GetObservationTypeString(UE::Learning::Observation::EType::Conv1d));
+		OutElement = FLearningAgentsObservationObjectElement();
+		return false;
+	}
+
+	OutElement = { Object->ObservationObject.GetConv1d(Element.ObjectElement).Element };
+
+	return true;
+}
+
diff --git a/Engine/Plugins/Experimental/LearningAgents/Source/LearningAgents/Public/LearningAgentsObservations.h b/Engine/Plugins/Experimental/LearningAgents/Source/LearningAgents/Public/LearningAgentsObservations.h
--- a/Engine/Plugins/Experimental/LearningAgents/Source/LearningAgents/Public/LearningAgentsObservations.h
+++ b/Engine/Plugins/Experimental/LearningAgents/Source/LearningAgents/Public/LearningAgentsObservations.h
@@ -112,6 +112,14 @@ enum class ELearningAgentsOptionalObservation : uint8
 	Valid,
 };
 
+/** Padding Modes for Convolutions. */
+UENUM(BlueprintType, Category = "LearningAgents")
+enum class ELearningAgentsPaddingMode : uint8
+{
+	Zeros = 0	UMETA(DisplayName = "Zeros"),
+	Circular = 1	UMETA(DisplayName = "Circular"),
+};
+
 UCLASS()
 class LEARNINGAGENTS_API ULearningAgentsObservations : public UBlueprintFunctionLibrary
 {
@@ -522,6 +530,21 @@ public:
 	UFUNCTION(BlueprintPure, Category = "LearningAgents", meta = (AdvancedDisplay = 1))
 	static FLearningAgentsObservationSchemaElement SpecifyBoolObservation(ULearningAgentsObservationSchema* Schema, const FName Tag = TEXT("BoolObservation"));
 
+	/**
+	 * Specifies a new 1d convolution observation of another observation.
+	 *
+	 * @param Schema The Observation Schema
+	 * @param Element The sub-observation to be convolved.
+	 * @param InputLength Length of the convolution input.
+	 * @param InChannels Number of channels in the input.
+	 * @param OutChannels Number of channels produced by the convolution.
+	 * @param KernelSize Size of the convolving kernel.
+	 * @param Tag The tag of this new observation. Used during observation object validation and debugging.
+	 * @return The newly created observation schema element.
+	 */
+	UFUNCTION(BlueprintPure, Category = "LearningAgents", meta = (AdvancedDisplay = 9))
+	static FLearningAgentsObservationSchemaElement SpecifyConv1dObservation(ULearningAgentsObservationSchema* Schema, const FLearningAgentsObservationSchemaElement Element, const int32 InputLength, const int32 InChannels, const int32 OutChannels, const int32 KernelSize, const int32 Padding, const ELearningAgentsPaddingMode PaddingMode, const ELearningAgentsActivationFunction ActivationFunction, const FName Tag = TEXT("Conv1dObservation"));
+
 	/**
 	 * Specifies a new float observation. A simple observation which can be used as a catch-all for situations where a
 	 * type-specific observation does not exist.
@@ -653,6 +676,17 @@ public:
 	 */
 	UFUNCTION(BlueprintPure, Category = "LearningAgents", meta = (AdvancedDisplay = 1))
 	static FLearningAgentsObservationSchemaElement SpecifyProportionAlongRayObservation(ULearningAgentsObservationSchema* Schema, const FName Tag = TEXT("ProportionAlongRayObservation"));
+	
+	/**
+	 * Specifies a new LiDAR observation. This observes how far a you can travel along a ray before collision. Rays that can travel 
+	 * the full distance are encoded as one, while rays that collide instantly are encoded as zero.
+	 *
+	 * @param Schema The Observation Schema
+	 * @param Tag The tag of this new observation. Used during observation object validation and debugging.
+	 * @return The newly created observation schema element.
+	 */
+	UFUNCTION(BlueprintPure, Category = "LearningAgents", meta = (AdvancedDisplay = 1))
+	static FLearningAgentsObservationSchemaElement SpecifyLidarObservation(ULearningAgentsObservationSchema* Schema, const FName Tag = TEXT("LidarObservation"));
 
 public:
 
@@ -1234,6 +1268,17 @@ public:
 		const FVector VisualLoggerLocation = FVector::ZeroVector,
 		const FLinearColor VisualLoggerColor = FLinearColor::Red);
 
+	/**
+	 * Make a new 1d convolution observation.
+	 *
+	 * @param Object The Observation Object
+	 * @param Element The Observation Element to be encoded.
+	 * @return The newly created observation object element.
+	 */
+	UFUNCTION(BlueprintPure, Category = "LearningAgents", meta = (AdvancedDisplay = 2, DefaultToSelf = "VisualLoggerListener"))
+	static FLearningAgentsObservationObjectElement MakeConv1dObservation(ULearningAgentsObservationObject* Object, const FLearningAgentsObservationObjectElement Element, const FName Tag = TEXT("Conv1dObservation"));
+
+
 	/**
 	 * Make a new float observation.
 	 *
@@ -1618,6 +1663,38 @@ public:
 		const FVector VisualLoggerLocation = FVector::ZeroVector,
 		const FLinearColor VisualLoggerColor = FLinearColor::Red);
 
+	/**
+	 * Make a new LiDAR observation. 
+	 *
+	 * @param Object The Observation Object
+	 * @param RayStart The local ray start location.
+	 * @param RayEnd The local ray end location.
+	 * @param RayTransform The transform to use to transform the local ray starts and ends into the world space.
+	 * @param CollisionChannel The collision channel to collide against.
+	 * @param ActorToIgnore The actor to ignore in the raycast.
+	 * @param Tag The tag of the corresponding observation. Must match the tag given during Specify.
+	 * @param bVisualLoggerEnabled When true, debug data will be sent to the visual logger.
+	 * @param VisualLoggerListener The listener object which is making this observation. This must be set to use logging.
+	 * @param VisualLoggerAgentId The agent id associated with this observation.
+	 * @param VisualLoggerLocation A location for the visual logger information in the world.
+	 * @param VisualLoggerColor The color for the visual logger display.
+	 * @return The newly created observation object element.
+	 */
+	UFUNCTION(BlueprintPure, Category = "LearningAgents", meta = (AdvancedDisplay = 6, DefaultToSelf = "VisualLoggerListener"))
+	static FLearningAgentsObservationObjectElement MakeLidarObservation(
+		ULearningAgentsObservationObject* Object,
+		const FVector RayStart,
+		const FVector RayEnd,
+		const FTransform RayTransform = FTransform(),
+		const ECollisionChannel CollisionChannel = ECollisionChannel::ECC_WorldStatic,
+		const AActor* ActorToIgnore = nullptr,
+		const FName Tag = TEXT("LidarObservation"),
+		const bool bVisualLoggerEnabled = false,
+		ULearningAgentsManagerListener* VisualLoggerListener = nullptr,
+		const int32 VisualLoggerAgentId = -1,
+		const FVector VisualLoggerLocation = FVector::ZeroVector,
+		const FLinearColor VisualLoggerColor = FLinearColor::Red);
+
 public:
 
 	// Get Basic Observations
@@ -1782,6 +1859,8 @@ public:
 	UFUNCTION(BlueprintPure, Category = "LearningAgents", meta = (AdvancedDisplay = 3, ReturnDisplayName = "Success"))
 	static bool GetProportionAlongRayObservation(float& OutProportion, const ULearningAgentsObservationObject* Object, const FLearningAgentsObservationObjectElement Element, const FName Tag = TEXT("ProportionAlongRayObservation"));
 
+	UFUNCTION(BlueprintPure = false, Category = "LearningAgents", meta = (AdvancedDisplay = 3, ReturnDisplayName = "Success"))
+	static bool GetConv1dObservation(FLearningAgentsObservationObjectElement& OutElement, const ULearningAgentsObservationObject* Object, const FLearningAgentsObservationObjectElement Element, const FName Tag = TEXT("Conv1dObservation"));
 };
 
 
diff --git a/Engine/Plugins/Experimental/LearningAgents/Source/LearningTraining/Private/LearningTrainer.cpp b/Engine/Plugins/Experimental/LearningAgents/Source/LearningTraining/Private/LearningTrainer.cpp
--- a/Engine/Plugins/Experimental/LearningAgents/Source/LearningTraining/Private/LearningTrainer.cpp
+++ b/Engine/Plugins/Experimental/LearningAgents/Source/LearningTraining/Private/LearningTrainer.cpp
@@ -102,7 +102,6 @@ namespace UE::Learning
 
 namespace UE::Learning::Trainer
 {
-
 	TSharedPtr<FJsonObject> ConvertObservationSchemaToJSON(
 		const Observation::FSchema& ObservationSchema,
 		const Observation::FSchemaElement& ObservationSchemaElement)
@@ -268,6 +267,20 @@ namespace UE::Learning::Trainer
 			break;
 		}
 
+		case Observation::EType::Conv1d:
+		{
+			const Observation::FSchemaConv1dParameters Parameters = ObservationSchema.GetConv1d(ObservationSchemaElement);
+			Object->SetStringField(TEXT("Type"), TEXT("Conv1d"));
+			Object->SetNumberField(TEXT("InputLength"), Parameters.InputLength);
+			Object->SetNumberField(TEXT("InChannels"), Parameters.InChannels);
+			Object->SetNumberField(TEXT("OutChannels"), Parameters.OutChannels);
+			Object->SetNumberField(TEXT("KernelSize"), Parameters.KernelSize);
+			Object->SetNumberField(TEXT("Padding"), Parameters.Padding);
+			Object->SetNumberField(TEXT("PaddingMode"), (uint32)Parameters.Padding);
+			Object->SetObjectField(TEXT("Element"), ConvertObservationSchemaToJSON(ObservationSchema, Parameters.Element));
+			break;
+		}
+
 		default:
 			checkNoEntry();
 		}
@@ -678,6 +691,18 @@ namespace UE::Learning::Trainer
 			return IsObservationSchemaSubsetCompatible(Object->GetObjectField(TEXT("Element")).ToSharedRef(), ObservationSchema, Parameters.Element);
 		}
 
+		case Observation::EType::Conv1d:
+		{
+			const Observation::FSchemaConv1dParameters Parameters = ObservationSchema.GetConv1d(ObservationSchemaElement);
+
+			if (Object->GetStringField(TEXT("Type")) != TEXT("Conv1d"))
+			{
+				return false;
+			}
+
+			return IsObservationSchemaSubsetCompatible(Object->GetObjectField(TEXT("Element")).ToSharedRef(), ObservationSchema, Parameters.Element);
+		}
+
 		default:
 			checkNoEntry();
 		}
diff --git a/Engine/Plugins/Experimental/NNERuntimeBasicCpu/Source/NNERuntimeBasicCpu/Private/NNERuntimeBasicCpu.ispc b/Engine/Plugins/Experimental/NNERuntimeBasicCpu/Source/NNERuntimeBasicCpu/Private/NNERuntimeBasicCpu.ispc
--- a/Engine/Plugins/Experimental/NNERuntimeBasicCpu/Source/NNERuntimeBasicCpu/Private/NNERuntimeBasicCpu.ispc
+++ b/Engine/Plugins/Experimental/NNERuntimeBasicCpu/Source/NNERuntimeBasicCpu/Private/NNERuntimeBasicCpu.ispc
@@ -690,4 +690,67 @@ export void NNERuntimeBasicCPUOperatorLayerFiLM(
 				Condition[BatchIdx * ConditionStride + Idx]) +
 			Condition[BatchIdx * ConditionStride + InputOutputNum + Idx];
 	}
+}
+
+export void NNERuntimeBasicCPUOperatorConv1d(
+	uniform float Output[],
+	uniform const float Input[],
+	uniform const float Weights[],
+	uniform const float Biases[],
+	uniform const int BatchNum,
+	uniform const int InChannelsNum,
+	uniform const int OutChannelsNum,
+	uniform const int InputNum,
+	uniform const int KernelSize,
+	uniform const int Padding,
+	uniform const int PaddingMode,
+	uniform const int OutputStride,
+	uniform const int InputStride)
+{
+	uniform int OutputNum = InputNum + 2 * Padding - KernelSize + 1;
+
+	for (uniform int BatchIdx = 0; BatchIdx < BatchNum; BatchIdx++)
+	{
+		uniform const float* InputItr = Input + BatchIdx * InputStride;
+		uniform float* OutputItr = Output + BatchIdx * OutputStride;
+
+		for (uniform int OutChannelsIdx = 0; OutChannelsIdx < OutChannelsNum; OutChannelsIdx++)
+		{
+			uniform const float* WeightsItr = Weights + OutChannelsIdx * InChannelsNum * KernelSize;
+			uniform float* OutChannelsItr = OutputItr + OutChannelsIdx * OutputNum;
+			uniform float Bias = Biases[OutChannelsIdx];
+
+			foreach (OutputIdx = 0 ... OutputNum)
+			{
+				float Sum = Bias;
+
+				for (uniform int InChannelsIdx = 0; InChannelsIdx < InChannelsNum; InChannelsIdx++)
+				{
+					uniform const float* WeightsInChannelsItr = WeightsItr + InChannelsIdx * KernelSize;
+					uniform const float* InChannelsItr = InputItr + InChannelsIdx * InputNum;
+
+					if (PaddingMode == 0) //Zeros
+					{
+						for (uniform int KernelIdx = 0; KernelIdx < KernelSize; KernelIdx++)
+						{
+							int Idx = OutputIdx + KernelIdx - Padding;
+							float Val = (Idx < 0 || Idx >= InputNum) ? 0.0f : InChannelsItr[Idx];
+							Sum += Val * WeightsInChannelsItr[KernelIdx];
+						}
+					}
+					else if (PaddingMode == 1) //Circular
+					{
+						for (uniform int KernelIdx = 0; KernelIdx < KernelSize; KernelIdx++)
+						{
+							int Idx = OutputIdx + KernelIdx - Padding;
+							Idx = Idx % InputNum;
+							if (Idx < 0) {Idx += InputNum;}
+							Sum += InChannelsItr[Idx] * WeightsInChannelsItr[KernelIdx];
+						}
+					}
+				}
+				OutChannelsItr[OutputIdx] = Sum;
+			}
+		}
+	}
 }
\ No newline at end of file
diff --git a/Engine/Plugins/Experimental/NNERuntimeBasicCpu/Source/NNERuntimeBasicCpu/Private/NNERuntimeBasicCpuModel.cpp b/Engine/Plugins/Experimental/NNERuntimeBasicCpu/Source/NNERuntimeBasicCpu/Private/NNERuntimeBasicCpuModel.cpp
--- a/Engine/Plugins/Experimental/NNERuntimeBasicCpu/Source/NNERuntimeBasicCpu/Private/NNERuntimeBasicCpuModel.cpp
+++ b/Engine/Plugins/Experimental/NNERuntimeBasicCpu/Source/NNERuntimeBasicCpu/Private/NNERuntimeBasicCpuModel.cpp
@@ -422,6 +422,105 @@ namespace UE::NNE::RuntimeBasic
 #endif
 		}
 
+		static inline void OperatorConv1d(
+			float* RESTRICT Output, 
+			const float* RESTRICT Input,
+			const float* RESTRICT Weights,
+			const float* RESTRICT Biases,
+			const uint32 BatchSize,
+			const uint32 InChannels,
+			const uint32 OutChannels,
+			const uint32 InputLength,
+			const uint32 KernelSize,
+			const uint32 Padding,
+			const FModelBuilder::EPaddingMode PaddingMode,
+			const uint32 OutputStride,
+			const uint32 InputStride)
+		{
+			NNE_RUNTIME_BASIC_TRACE_SCOPE(NNE::RuntimeBasic::Private::OperatorConv1d);
+			NNE_RUNTIME_BASIC_CHECK_ALIASING_SIZES(Output, Input, (InputLength + 2 * Padding - KernelSize + 1) * OutChannels, InputLength * InChannels);
+
+#if NNE_RUNTIME_BASIC_ENABLE_ISPC
+			ispc::NNERuntimeBasicCPUOperatorConv1d(
+				Output,
+				Input,
+				Weights,
+				Biases,
+				BatchSize,
+				InChannels,
+				OutChannels,
+				InputLength,
+				KernelSize,
+				Padding,
+				static_cast<uint32>(PaddingMode),
+				OutputStride,
+				InputStride);
+#else
+			check(KernelSize <= InputLength + 2 * Padding);
+			const uint32 OutputLength = InputLength + 2 * Padding - KernelSize + 1;
+			for (uint32 BatchIdx = 0; BatchIdx < BatchSize; BatchIdx++)
+			{
+				const float* InputItr = Input + BatchIdx * InputStride;
+				float* OutputItr = Output + BatchIdx * OutputStride;
+				
+				for (uint32 OutChannelsIdx = 0; OutChannelsIdx < OutChannels; OutChannelsIdx++)
+				{
+					const float* WeightsItr = Weights + OutChannelsIdx * InChannels * KernelSize;
+					float* OutChannelsItr = OutputItr + OutChannelsIdx * OutputLength;
+					
+					for (uint32 OutputIdx = 0; OutputIdx < OutputLength; OutputIdx++)
+					{
+						float Sum = Biases[OutChannelsIdx];
+
+						for (uint32 InChannelsIdx = 0; InChannelsIdx < InChannels; InChannelsIdx++)
+						{
+							const float* WeightsInChannelsItr = WeightsItr + InChannelsIdx * KernelSize;
+							const float* InChannelsItr = InputItr + InChannelsIdx * InputLength;
+							
+							for (uint32 KernelIdx = 0; KernelIdx < KernelSize; KernelIdx++)
+							{
+								int32 Idx = static_cast<int32_t>(OutputIdx) + static_cast<int32_t>(KernelIdx) - static_cast<int32_t>(Padding);
+								float Val;
+
+								switch (PaddingMode)
+								{
+								case FModelBuilder::EPaddingMode::Circular:
+								{
+									Idx = Idx % static_cast<int32_t>(InputLength);
+									if (Idx < 0)
+									{
+										Idx += static_cast<int32_t>(InputLength);
+									}
+									Val = InChannelsItr[Idx];
+									break;
+								}
+								case FModelBuilder::EPaddingMode::Zeros:
+								{
+									if (Idx < 0 or Idx >= static_cast<int32_t>(InputLength))
+									{
+										Val = 0.f;
+									}
+									else
+									{
+										Val = InChannelsItr[Idx];
+									}
+									break;
+								}
+								default:
+								{
+									checkNoEntry();
+								}
+								}
+								Sum += Val * WeightsInChannelsItr[KernelIdx];
+							}
+						}
+						OutChannelsItr[OutputIdx] = Sum;
+					}
+				}
+			}
+#endif
+		}
+
 		static inline void OperatorLinear(
 			float* RESTRICT Output,
 			const float* RESTRICT Input,
@@ -1664,6 +1763,7 @@ namespace UE::NNE::RuntimeBasic
 			Residual = 26,
 			FiLM = 27,
 			GRUCell = 28,
+			Conv1d = 29,
 		};
 
 		//--------------------------------------------------------------------------
@@ -4984,6 +5084,144 @@ namespace UE::NNE::RuntimeBasic
 		}
 
 
+		//--------------------------------------------------------------------------
+
+		struct FConv1dLayer;
+
+		struct FConv1dInstance : public ILayerInstance
+		{
+			FConv1dInstance(const FConv1dLayer& InConv1dLayer);
+
+			void SetMaxBatchSize(const uint32 MaxBatchSize) override final;
+			virtual ELayerType GetLayerType() const override final { return ELayerType::Conv1d; }
+
+			const FConv1dLayer& Conv1dLayer;
+			TSharedPtr<ILayerInstance> SubLayerInstance;
+			TArray<float> SubLayerOutputBuffer;
+		};
+
+		struct FConv1dLayer : public ILayer
+		{
+			virtual TSharedPtr<ILayerInstance> MakeInstance() const { return MakeShared<FConv1dInstance>(*this); };
+			virtual ELayerType GetLayerType() const override final { return ELayerType::Conv1d; }
+			virtual uint32 GetInputSize() const override final { return SubLayer->GetInputSize(); }
+			virtual uint32 GetOutputSize() const override final 
+			{ 
+				uint32 OutputLength = InputLength + 2 * Padding - KernelSize + 1;
+				return OutputLength * OutChannels;
+			}
+			uint32 GetPaddingMode() const { return static_cast<uint32_t>(PaddingMode); }
+			
+			void SerializationLoadPaddingMode(uint64& InOutOffset, TConstArrayView<uint8> Data) 
+			{
+				uint32 LoadedMode;
+				Serialization::Load(InOutOffset, LoadedMode, Data);
+				PaddingMode = static_cast<FModelBuilder::EPaddingMode>(LoadedMode);
+			}
+
+			virtual void SerializationSize(uint64& InOutOffset) const override final
+			{
+				Serialization::Size(InOutOffset, InputLength);
+				Serialization::Size(InOutOffset, InChannels);
+				Serialization::Size(InOutOffset, OutChannels);
+				Serialization::Size(InOutOffset, KernelSize);
+				Serialization::Size(InOutOffset, Padding);
+				Serialization::Size(InOutOffset, GetPaddingMode());
+				Serialization::Size(InOutOffset, Weights);
+				Serialization::Size(InOutOffset, Biases);
+				Serialization::Size(InOutOffset, SubLayer);
+			}
+
+			virtual void SerializationLoad(uint64& InOutOffset, TConstArrayView<uint8> Data) override final
+			{
+				Serialization::Load(InOutOffset, InputLength, Data);
+				Serialization::Load(InOutOffset, InChannels, Data);
+				Serialization::Load(InOutOffset, OutChannels, Data);
+				Serialization::Load(InOutOffset, KernelSize, Data);
+				Serialization::Load(InOutOffset, Padding, Data);
+				SerializationLoadPaddingMode(InOutOffset, Data);
+				Serialization::Load(InOutOffset, Weights, Data, KernelSize * InChannels * OutChannels);
+				Serialization::Load(InOutOffset, Biases, Data, OutChannels);
+				Serialization::Load(InOutOffset, SubLayer, Data);
+			}
+
+			virtual void SerializationSave(uint64& InOutOffset, TArrayView<uint8> Data) const override final
+			{
+				Serialization::Save(InOutOffset, InputLength, Data);
+				Serialization::Save(InOutOffset, InChannels, Data);
+				Serialization::Save(InOutOffset, OutChannels, Data);
+				Serialization::Save(InOutOffset, KernelSize, Data);
+				Serialization::Save(InOutOffset, Padding, Data);
+				Serialization::Save(InOutOffset, GetPaddingMode(), Data);
+				Serialization::Save(InOutOffset, Weights, Data);
+				Serialization::Save(InOutOffset, Biases, Data);
+				Serialization::Save(InOutOffset, SubLayer, Data);
+			}
+
+			virtual void Evaluate(
+				ILayerInstance* Instance,
+				float* OutputBuffer,
+				const float* InputBuffer,
+				const uint32 BatchSize,
+				const uint32 OutputBufferSize,
+				const uint32 InputBufferSize,
+				const uint32 OutputBufferStride,
+				const uint32 InputBufferStride) override final
+			{
+				NNE_RUNTIME_BASIC_TRACE_SCOPE(NNE::RuntimeBasic::Private::FConv1dLayer::Evaluate);
+				check(OutputBufferSize == GetOutputSize() && InputBufferSize == GetInputSize());
+				check(OutputBufferStride >= GetOutputSize() && InputBufferStride >= GetInputSize());
+				check(Instance && Instance->GetLayerType() == GetLayerType());
+
+				FConv1dInstance* Conv1dLayerInstance = StaticCast<FConv1dInstance*>(Instance);
+
+				OperatorNanCheck(InputBuffer, BatchSize, InputBufferSize, InputBufferStride);
+
+				SubLayer->Evaluate(
+					Conv1dLayerInstance->SubLayerInstance.Get(),
+					Conv1dLayerInstance->SubLayerOutputBuffer.GetData(),
+					InputBuffer,
+					BatchSize,
+					SubLayer->GetOutputSize(),
+					SubLayer->GetInputSize(),
+					SubLayer->GetOutputSize(), 
+					SubLayer->GetInputSize());
+
+				OperatorConv1d(
+					OutputBuffer,
+					Conv1dLayerInstance->SubLayerOutputBuffer.GetData(),
+					Weights.GetData(),
+					Biases.GetData(),
+					BatchSize,
+					InChannels,
+					OutChannels,
+					InputLength,
+					KernelSize,
+					Padding,
+					PaddingMode,
+					OutputBufferStride,
+					InputLength * InChannels);
+
+				OperatorNanCheck(OutputBuffer, BatchSize, OutputBufferSize, OutputBufferStride);
+			}
+			uint32 InputLength, InChannels, OutChannels, KernelSize, Padding;
+			FModelBuilder::EPaddingMode PaddingMode;
+			TConstArrayView<float> Weights, Biases;
+			TSharedPtr<ILayer> SubLayer;
+		};
+
+		FConv1dInstance::FConv1dInstance(const FConv1dLayer& InConv1dLayer) : Conv1dLayer(InConv1dLayer)
+		{
+			SubLayerInstance = Conv1dLayer.SubLayer->MakeInstance();
+		}
+
+		void FConv1dInstance::SetMaxBatchSize(const uint32 MaxBatchSize)
+		{
+			NNE_RUNTIME_BASIC_TRACE_SCOPE(NNE::RuntimeBasic::Private::FConv1dLayerInstance::SetMaxBatchSize);
+			SubLayerInstance->SetMaxBatchSize(MaxBatchSize);
+			SubLayerOutputBuffer.SetNumUninitialized(MaxBatchSize * Conv1dLayer.SubLayer->GetOutputSize(), EAllowShrinking::No);
+		}
+		
 		//--------------------------------------------------------------------------
 		// Layer Serialization
 		//--------------------------------------------------------------------------
@@ -5040,6 +5278,7 @@ namespace UE::NNE::RuntimeBasic
 					case ELayerType::Residual: OutLayer = MakeShared<FResidualLayer>(); break;
 					case ELayerType::FiLM: OutLayer = MakeShared<FFiLMLayer>(); break;
 					case ELayerType::GRUCell: OutLayer = MakeShared<FGRUCellLayer>(); break;
+					case ELayerType::Conv1d: OutLayer = MakeShared<FConv1dLayer>(); break;
 					default: checkf(false, TEXT("Unknown Layer Id %i"), LayerTypeId);
 					}
 				}
@@ -6228,6 +6467,39 @@ namespace UE::NNE::RuntimeBasic
 		return StaticCastSharedPtr<Private::ILayer>(Layer);
 	}
 
+	FModelBuilderElement FModelBuilder::MakeConv1d(
+		const uint32 InputLength,
+		const uint32 InChannels,
+		const uint32 OutChannels,
+		const uint32 KernelSize,
+		const uint32 Padding,
+		const EPaddingMode PaddingMode,
+		const EActivationFunction ActivationFunction,
+		const TConstArrayView<float> Weights,
+		const TConstArrayView<float> Biases,
+		const FModelBuilderElement& SubLayer)
+	{
+		check(Weights.Num() == InChannels * OutChannels * KernelSize);
+		check(Biases.Num() == OutChannels);
+
+		const TSharedPtr<Private::FConv1dLayer> Conv1dLayer = MakeShared<Private::FConv1dLayer>();
+		Conv1dLayer->InputLength = InputLength;
+		Conv1dLayer->InChannels = InChannels;
+		Conv1dLayer->OutChannels = OutChannels;
+		Conv1dLayer->KernelSize = KernelSize;
+		Conv1dLayer->Padding = Padding;
+		Conv1dLayer->PaddingMode = PaddingMode;
+		Conv1dLayer->Weights = Weights;
+		Conv1dLayer->Biases = Biases;
+		Conv1dLayer->SubLayer = SubLayer.Layer;
+
+		TArray<FModelBuilderElement, TInlineAllocator<32>> Layers;
+		Layers.Emplace(StaticCastSharedPtr<Private::ILayer>(Conv1dLayer));
+		Layers.Emplace(MakeActivation(InputLength * OutChannels, ActivationFunction));
+		
+		return MakeSequence(Layers);
+	}
+
 	void FModelBuilder::Reset()
 	{
 		Rng = RngInitialState;
diff --git a/Engine/Plugins/Experimental/NNERuntimeBasicCpu/Source/NNERuntimeBasicCpu/Public/NNERuntimeBasicCpuBuilder.h b/Engine/Plugins/Experimental/NNERuntimeBasicCpu/Source/NNERuntimeBasicCpu/Public/NNERuntimeBasicCpuBuilder.h
--- a/Engine/Plugins/Experimental/NNERuntimeBasicCpu/Source/NNERuntimeBasicCpu/Public/NNERuntimeBasicCpuBuilder.h
+++ b/Engine/Plugins/Experimental/NNERuntimeBasicCpu/Source/NNERuntimeBasicCpu/Public/NNERuntimeBasicCpuBuilder.h
@@ -76,6 +76,15 @@ namespace UE::NNE::RuntimeBasic
 			KaimingUniform = 1,
 		};
 
+		/** Padding Mode for Convolutions */
+		enum class EPaddingMode : uint8
+		{
+			// The input will be padded with zeroes on both sides
+			Zeros = 0, 
+			// The input will be padded with wrap around values
+			Circular = 1,
+		};
+
 		/*
 		 * There is a clang compiler bug that requires the following classes to declare a default constructor in the cpp file as a workaround:
 		 * 
@@ -87,7 +96,7 @@ namespace UE::NNE::RuntimeBasic
 		struct FWeightInitializationSettings
 		{
 			// Workaround for clang compiler bug...
-			FWeightInitializationSettings();
+			UE_API FWeightInitializationSettings();
 
 			// Type of initialization to use
 			EWeightInitializationType Type = EWeightInitializationType::KaimingGaussian;
@@ -548,6 +557,30 @@ namespace UE::NNE::RuntimeBasic
 			const FModelBuilderElement& Prefix,
 			const FModelBuilderElement& Condition,
 			const FModelBuilderElement& Postfix);
+		
+		/**
+		 * Make a new 1D Convolution Layer.
+		 *
+		 * @param InputLength			Length of 1 dimensional input
+		 * @param InChannels			Number of channels in the input
+		 * @param OutChannels			Number of channels produced by the convolution
+		 * @param KernelSize			Size of the convolving kernel
+		 * @param Padding				The amount of padding added to both sides of the input
+		 * @param PaddingMode			"Zeros" or "Circular"
+		 * @param Weights				Convolutional layer weights
+		 * @param Biases				Convolutional layer biases
+		 */
+		UE_API FModelBuilderElement MakeConv1d(
+			const uint32 InputLength,
+			const uint32 InChannels,
+			const uint32 OutChannels,
+			const uint32 KernelSize,
+			const uint32 Padding, 
+			const EPaddingMode PaddingMode,
+			const EActivationFunction ActivationFunction, 
+			const TConstArrayView<float> Weights,
+			const TConstArrayView<float> Biases,
+			const FModelBuilderElement& SubLayer);
 
 	public:
 
