LearningAgents: Fixed small bug in frame attributes length calculation

diff --git a/Engine/Plugins/Experimental/LearningAgents/Source/Learning/Private/LearningFrameAttribute.cpp b/Engine/Plugins/Experimental/LearningAgents/Source/Learning/Private/LearningFrameAttribute.cpp
--- a/Engine/Plugins/Experimental/LearningAgents/Source/Learning/Private/LearningFrameAttribute.cpp
+++ b/Engine/Plugins/Experimental/LearningAgents/Source/Learning/Private/LearningFrameAttribute.cpp
@@ -1728,7 +1728,7 @@ namespace UE::Learning
 				const TLearningArrayView<1, const int32> RangeOffsets,
 				const TLearningArrayView<1, const int32> RangeLengths) {
 
-					const int32 ChannelNum = Out.GetChannelNum();
+					const int32 ChannelNum = In.GetChannelNum();
 
 					Array::Zero(Out.GetAttributeData());
 
