Learning Agents: Load schema when recording is loaded

#rnx

diff --git a/Engine/Plugins/Experimental/LearningAgents/Source/LearningAgentsTraining/Private/LearningAgentsRecording.cpp b/Engine/Plugins/Experimental/LearningAgents/Source/LearningAgentsTraining/Private/LearningAgentsRecording.cpp
--- a/Engine/Plugins/Experimental/LearningAgents/Source/LearningAgentsTraining/Private/LearningAgentsRecording.cpp
+++ b/Engine/Plugins/Experimental/LearningAgents/Source/LearningAgentsTraining/Private/LearningAgentsRecording.cpp
@@ -6,6 +6,10 @@
 
 #include "UObject/Package.h"
 #include "Misc/FileHelper.h"
+#include "Serialization/JsonReader.h"
+#include "Serialization/JsonSerializer.h"
+#include "Serialization/JsonWriter.h"
+#include "Dom/JsonObject.h"
 
 ULearningAgentsRecording::ULearningAgentsRecording() = default;
 ULearningAgentsRecording::ULearningAgentsRecording(FVTableHelper& Helper) : Super(Helper) {}
@@ -25,11 +29,43 @@ void ULearningAgentsRecording::ResetRecording()
 
 void ULearningAgentsRecording::AppendRecording()
 {
-	AppendRecordingFromFile(RecordingFile);
+	AppendRecordingFromFile(RecordingFile, SchemaFile);
 }
 
-void ULearningAgentsRecording::LoadRecordingFromFile(const FFilePath& File)
+void ULearningAgentsRecording::LoadRecordingFromFile(const FFilePath& File, const FFilePath& Schema)
 {
+	// Load JSON
+	FString ObservationSchemaString;
+	FString ActionSchemaString;
+
+	FString FileString;
+	if (FFileHelper::LoadFileToString(FileString, *Schema.FilePath))
+	{
+		TSharedRef<TJsonReader<>> Reader = TJsonReaderFactory<>::Create(FileString);
+		TSharedPtr<FJsonObject> JsonObject;
+		if (FJsonSerializer::Deserialize(Reader, JsonObject) && JsonObject.IsValid())
+		{
+			const TSharedPtr<FJsonObject>& ObservationJson = JsonObject->GetObjectField(TEXT("ObservationSchema"));
+			TSharedRef<TJsonWriter<>> ObservationWriter = TJsonWriterFactory<>::Create(&ObservationSchemaString);
+			FJsonSerializer::Serialize(ObservationJson.ToSharedRef(), ObservationWriter);
+
+			const TSharedPtr<FJsonObject>& ActionJson = JsonObject->GetObjectField(TEXT("ActionSchema"));
+			TSharedRef<TJsonWriter<>> ActionWriter = TJsonWriterFactory<>::Create(&ActionSchemaString);
+			FJsonSerializer::Serialize(ActionJson.ToSharedRef(), ActionWriter);
+		}
+		else
+		{
+			UE_LOG(LogLearning, Error, TEXT("%s: Failed to load schema. Failed to parse JSON from: \"%s\""), *GetName(), *Schema.FilePath);
+		}
+	}
+	else
+	{
+		UE_LOG(LogLearning, Error, TEXT("%s: Failed to load schema. File not found: \"%s\""), *GetName(), *Schema.FilePath);
+		return;
+	}
+
+	// Write Record
+
 	TArray<uint8> RecordingData;
 
 	if (FFileHelper::LoadFileToArray(RecordingData, *File.FilePath))
@@ -74,6 +110,9 @@ void ULearningAgentsRecording::LoadRecordingFromFile(const FFilePath& File)
 			UE::Learning::DeserializeFromBytes(Offset, RecordingData, Records[RecordIdx].ActionCompatibilityHash);
 			UE::Learning::Array::DeserializeFromBytes32<2, float>(Offset, RecordingData, Records[RecordIdx].ObservationData);
 			UE::Learning::Array::DeserializeFromBytes32<2, float>(Offset, RecordingData, Records[RecordIdx].ActionData);
+			
+			Records[RecordIdx].ObservationSchemaJson = ObservationSchemaString;
+			Records[RecordIdx].ActionSchemaJson = ActionSchemaString;
 		}
 
 		check(Offset == RecordingData.Num());
@@ -133,10 +172,10 @@ void ULearningAgentsRecording::SaveRecordingToFile(const FFilePath& File) const
 	}
 }
 
-void ULearningAgentsRecording::AppendRecordingFromFile(const FFilePath& File)
+void ULearningAgentsRecording::AppendRecordingFromFile(const FFilePath& File, const FFilePath& Schema)
 {
 	ULearningAgentsRecording* TempRecording = NewObject<ULearningAgentsRecording>(this);
-	TempRecording->LoadRecordingFromFile(File);
+	TempRecording->LoadRecordingFromFile(File, Schema);
 
 	Records.Append(TempRecording->Records);
 	ForceMarkDirty();
diff --git a/Engine/Plugins/Experimental/LearningAgents/Source/LearningAgentsTraining/Public/LearningAgentsRecording.h b/Engine/Plugins/Experimental/LearningAgents/Source/LearningAgentsTraining/Public/LearningAgentsRecording.h
--- a/Engine/Plugins/Experimental/LearningAgents/Source/LearningAgentsTraining/Public/LearningAgentsRecording.h
+++ b/Engine/Plugins/Experimental/LearningAgents/Source/LearningAgentsTraining/Public/LearningAgentsRecording.h
@@ -80,7 +80,7 @@ public:
 
 	/** Load this recording from a file. */
 	UFUNCTION(BlueprintCallable, Category = "LearningAgents", meta = (RelativePath))
-	UE_API void LoadRecordingFromFile(const FFilePath& File);
+	UE_API void LoadRecordingFromFile(const FFilePath& File, const FFilePath& Schema);
 
 	/** Save this recording to a file. */
 	UFUNCTION(BlueprintCallable, BlueprintPure = false, Category = "LearningAgents", meta = (RelativePath))
@@ -88,7 +88,7 @@ public:
 
 	/** Append to this recording from a file. */
 	UFUNCTION(BlueprintCallable, Category = "LearningAgents", meta = (RelativePath))
-	UE_API void AppendRecordingFromFile(const FFilePath& File);
+	UE_API void AppendRecordingFromFile(const FFilePath& File, const FFilePath& Schema);
 
 	/** Loads this recording from the given recording asset */
 	UFUNCTION(BlueprintCallable, Category = "LearningAgents")
@@ -143,6 +143,10 @@ public:
 	UPROPERTY(EditInstanceOnly, Category = "LearningAgents")
 	FFilePath RecordingFile;
 
+	/** A schema file. Used in combination with the "AppendRecording" button in the editor. */
+	UPROPERTY(EditInstanceOnly, Category = "LearningAgents")
+	FFilePath SchemaFile;
+
 	/** Set of records. */
 	UPROPERTY(EditInstanceOnly, Category = "LearningAgents")
 	TArray<FLearningAgentsRecord> Records;
