Replace some usages of FORCEINLINE with inline in ML modules.

For context about this change, see this doc: https://docs.google.com/document/d/1LjpT7NpzTecfXdFNyv5KM6sCu67827BRljNi1iaWVSE/edit?tab=t.0#heading=h.fvtz8gjukz7x

Reviewers were automatically determined. Please comment if anyone else should be added.

#rnx

diff --git a/Engine/Plugins/AI/MLAdapter/Source/MLAdapter/Public/MLAdapterAsync.h b/Engine/Plugins/AI/MLAdapter/Source/MLAdapter/Public/MLAdapterAsync.h
--- a/Engine/Plugins/AI/MLAdapter/Source/MLAdapter/Public/MLAdapterAsync.h
+++ b/Engine/Plugins/AI/MLAdapter/Source/MLAdapter/Public/MLAdapterAsync.h
@@ -22,7 +22,7 @@ RetType CallOnGameThread(TFunction<RetType()> InFunction)
 
 // fire (InFunction to be called on the GameThread) and forget
 template<>
-FORCEINLINE void CallOnGameThread(TFunction<void()> InFunction)
+inline void CallOnGameThread(TFunction<void()> InFunction)
 {
 	AsyncTask(ENamedThreads::GameThread, InFunction);
 }
diff --git a/Engine/Plugins/AI/MLAdapter/Source/MLAdapter/Public/Managers/MLAdapterManager.h b/Engine/Plugins/AI/MLAdapter/Source/MLAdapter/Public/Managers/MLAdapterManager.h
--- a/Engine/Plugins/AI/MLAdapter/Source/MLAdapter/Public/Managers/MLAdapterManager.h
+++ b/Engine/Plugins/AI/MLAdapter/Source/MLAdapter/Public/Managers/MLAdapterManager.h
@@ -138,10 +138,10 @@ public:
 	FOnGenericRPCServerDelegate& GetOnAddServerFunctions() { return OnAddServerFunctions; }
 
 	/** Get the current manager instance. */
-	FORCEINLINE static UMLAdapterManager& Get();
+	inline static UMLAdapterManager& Get();
 
 	/** Returns true if the manager instance exists. */
-	FORCEINLINE static bool IsReady();
+	inline static bool IsReady();
 
 	/** Get this manager's librarian. */
 	const FMLAdapterLibrarian& GetLibrarian() const { return Librarian; }
diff --git a/Engine/Plugins/AI/MLAdapter/Source/MLAdapter/Public/Sensors/MLAdapterSensor_AIPerception.h b/Engine/Plugins/AI/MLAdapter/Source/MLAdapter/Public/Sensors/MLAdapterSensor_AIPerception.h
--- a/Engine/Plugins/AI/MLAdapter/Source/MLAdapter/Public/Sensors/MLAdapterSensor_AIPerception.h
+++ b/Engine/Plugins/AI/MLAdapter/Source/MLAdapter/Public/Sensors/MLAdapterSensor_AIPerception.h
@@ -58,7 +58,7 @@ protected:
 	virtual void DescribeSelfToGameplayDebugger(FGameplayDebuggerCategory& DebuggerCategory) const override;
 #endif // WITH_GAMEPLAY_DEBUGGER
 
-	FORCEINLINE float Sanify(const float DegreeAngles) const
+	inline float Sanify(const float DegreeAngles) const
 	{
 		return (DegreeAngles < -180.f) 
 			? (DegreeAngles + 360.f)
@@ -66,7 +66,7 @@ protected:
 				? (DegreeAngles - 360.f) 
 				: DegreeAngles;
 	}
-	FORCEINLINE FRotator Sanify(const FRotator Rotator) const
+	inline FRotator Sanify(const FRotator Rotator) const
 	{
 		ensureMsgf(Rotator.Roll == 0.f, TEXT("UMLAdapterSensor_AIPerception is expected to deal only with 0-roll rotators"));
 		return FRotator(Sanify(Rotator.Pitch), Sanify(Rotator.Yaw), 0.f);
diff --git a/Engine/Plugins/AI/MLAdapter/Source/MLAdapter/Public/Sessions/MLAdapterSession.h b/Engine/Plugins/AI/MLAdapter/Source/MLAdapter/Public/Sessions/MLAdapterSession.h
--- a/Engine/Plugins/AI/MLAdapter/Source/MLAdapter/Public/Sessions/MLAdapterSession.h
+++ b/Engine/Plugins/AI/MLAdapter/Source/MLAdapter/Public/Sessions/MLAdapterSession.h
@@ -181,7 +181,7 @@ public:
 protected:
 	virtual void SetGameMode(AGameModeBase* GameModeBase);
 
-	FORCEINLINE static uint32 HashAvatar(const AActor& Avatar)
+	inline static uint32 HashAvatar(const AActor& Avatar)
 	{
 		return Avatar.GetUniqueID();
 	}
diff --git a/Engine/Plugins/Experimental/LearningAgents/Source/Learning/Public/MultiArray.h b/Engine/Plugins/Experimental/LearningAgents/Source/Learning/Public/MultiArray.h
--- a/Engine/Plugins/Experimental/LearningAgents/Source/Learning/Public/MultiArray.h
+++ b/Engine/Plugins/Experimental/LearningAgents/Source/Learning/Public/MultiArray.h
@@ -77,7 +77,7 @@ public:
 	 *
 	 * @returns Reference to indexed element.
 	 */
-	FORCEINLINE TMultiArrayView<DimNum - 1, ElementType> operator[](int64 Index)
+	inline TMultiArrayView<DimNum - 1, ElementType> operator[](int64 Index)
 	{
 		RangeCheck(0, Index);
 
@@ -97,7 +97,7 @@ public:
 	 *
 	 * @returns Reference to indexed element.
 	 */
-	FORCEINLINE TMultiArrayView<DimNum - 1, const ElementType> operator[](int64 Index) const
+	inline TMultiArrayView<DimNum - 1, const ElementType> operator[](int64 Index) const
 	{
 		RangeCheck(0, Index);
 
@@ -115,7 +115,7 @@ public:
 	 *
 	 * @returns Flattened array view.
 	 */
-	FORCEINLINE TMultiArrayView<1, ElementType> Flatten()
+	inline TMultiArrayView<1, ElementType> Flatten()
 	{
 		return TMultiArrayView<1, ElementType>(GetData(), Num());
 	}
@@ -125,7 +125,7 @@ public:
 	 *
 	 * @returns Flattened array view.
 	 */
-	FORCEINLINE TMultiArrayView<1, const ElementType> Flatten() const
+	inline TMultiArrayView<1, const ElementType> Flatten() const
 	{
 		return TMultiArrayView<1, const ElementType>(GetData(), Num());
 	}
@@ -136,7 +136,7 @@ public:
 	 *
 	 * @returns View flattened on the given dimension.
 	 */
-	FORCEINLINE TMultiArrayView<DimNum - 1, ElementType> Flatten(uint8 Dimension)
+	inline TMultiArrayView<DimNum - 1, ElementType> Flatten(uint8 Dimension)
 	{
 		checkf((Dimension < DimNum - 1), TEXT("MultiArray flatten dimension out of bounds: %lld from a rank of %lld"), (long long)Dimension, (long long)DimNum); // & for one branch
 
@@ -165,7 +165,7 @@ public:
 	 *
 	 * @returns View flattened on the given dimension.
 	 */
-	FORCEINLINE TMultiArrayView<DimNum - 1, const ElementType> Flatten(uint8 Dimension) const
+	inline TMultiArrayView<DimNum - 1, const ElementType> Flatten(uint8 Dimension) const
 	{
 		checkf((Dimension < DimNum - 1), TEXT("MultiArray flatten dimension out of bounds: %lld from a rank of %lld"), (long long)Dimension, (long long)DimNum); // & for one branch
 
@@ -195,7 +195,7 @@ public:
 	 * @returns View flattened on the given dimension.
 	 */
 	template<uint8 Dimension>
-	FORCEINLINE TMultiArrayView<DimNum - 1, ElementType> Flatten()
+	inline TMultiArrayView<DimNum - 1, ElementType> Flatten()
 	{
 		static_assert(Dimension < DimNum - 1, "MultiArray flatten dimension out of bounds");
 
@@ -225,7 +225,7 @@ public:
 	 * @returns View flattened on the given dimension.
 	 */
 	template<uint8 Dimension>
-	FORCEINLINE TMultiArrayView<DimNum - 1, const ElementType> Flatten() const
+	inline TMultiArrayView<DimNum - 1, const ElementType> Flatten() const
 	{
 		static_assert(Dimension < DimNum - 1, "MultiArray flatten dimension out of bounds");
 
@@ -254,7 +254,7 @@ public:
 	/**
 	 * Constructor, initializes element number counters.
 	 */
-	FORCEINLINE TMultiArray() = default;
+	inline TMultiArray() = default;
 
 	/**
 	 * Constructor from a raw array of elements.
@@ -262,7 +262,7 @@ public:
 	 * @param Ptr   A pointer to an array of elements to copy.
 	 * @param Nums The number of elements on each dimension.
 	 */
-	FORCEINLINE TMultiArray(const ElementType* Ptr, TMultiArrayShape<DimNum> Shape)
+	inline TMultiArray(const ElementType* Ptr, TMultiArrayShape<DimNum> Shape)
 	{
 		check(Ptr != nullptr || Shape.Total() == 0);
 
@@ -270,7 +270,7 @@ public:
 	}
 
 	template <typename OtherElementType>
-	FORCEINLINE explicit TMultiArray(const TMultiArrayView<DimNum, OtherElementType>& Other)
+	inline explicit TMultiArray(const TMultiArrayView<DimNum, OtherElementType>& Other)
 		: TMultiArray(Other.GetData(), Other.Shape()) {}
 
 	/**
@@ -283,7 +283,7 @@ public:
 		typename OtherAllocator
 		UE_REQUIRES(UE::MultiArray::Private::TMultiArrayElementsAreCompatible_V<ElementType, const OtherElementType&>)
 	>
-		FORCEINLINE explicit TMultiArray(const TMultiArray<DimNum, OtherElementType, OtherAllocator>& Other)
+		inline explicit TMultiArray(const TMultiArray<DimNum, OtherElementType, OtherAllocator>& Other)
 	{
 		CopyToEmpty(Other.GetData(), Other.Shape());
 	}
@@ -293,7 +293,7 @@ public:
 	 *
 	 * @param Other The source array to copy.
 	 */
-	FORCEINLINE TMultiArray(const TMultiArray& Other)
+	inline TMultiArray(const TMultiArray& Other)
 	{
 		CopyToEmpty(Other.GetData(), Other.Shape());
 	}
@@ -344,7 +344,7 @@ private:
 	 * @param FromArray Array to move from.
 	 */
 	template <typename FromArrayType, typename ToArrayType>
-	static FORCEINLINE void MoveOrCopy(ToArrayType& ToArray, FromArrayType& FromArray)
+	static inline void MoveOrCopy(ToArrayType& ToArray, FromArrayType& FromArray)
 	{
 		ToArray.CopyToEmpty(FromArray.GetData(), FromArray.Shape());
 	}
@@ -355,7 +355,7 @@ public:
 	 *
 	 * @param Other Array to move from.
 	 */
-	FORCEINLINE TMultiArray(TMultiArray&& Other)
+	inline TMultiArray(TMultiArray&& Other)
 	{
 		MoveOrCopy(*this, Other);
 	}
@@ -370,7 +370,7 @@ public:
 		typename OtherAllocator
 		UE_REQUIRES(UE::MultiArray::Private::TMultiArrayElementsAreCompatible_V<ElementType, OtherElementType&&>)
 	>
-		FORCEINLINE explicit TMultiArray(TMultiArray<DimNum, OtherElementType, OtherAllocator>&& Other)
+		inline explicit TMultiArray(TMultiArray<DimNum, OtherElementType, OtherAllocator>&& Other)
 	{
 		MoveOrCopy(*this, Other);
 	}
@@ -404,7 +404,7 @@ public:
 	 *
 	 * @returns Pointer to first array entry or nullptr if empty
 	 */
-	FORCEINLINE ElementType* GetData()
+	inline ElementType* GetData()
 	{
 		return (ElementType*)AllocatorInstance.GetAllocation();
 	}
@@ -414,7 +414,7 @@ public:
 	 *
 	 * @returns Pointer to first array entry or nullptr if empty
 	 */
-	FORCEINLINE const ElementType* GetData() const
+	inline const ElementType* GetData() const
 	{
 		return (const ElementType*)AllocatorInstance.GetAllocation();
 	}
@@ -423,22 +423,22 @@ public:
 	 * DO NOT USE DIRECTLY
 	 * STL-like iterators to enable range-based for loop support.
 	 */
-	FORCEINLINE ElementType* begin() { return GetData(); }
-	FORCEINLINE ElementType* end() { return GetData() + Num(); }
+	inline ElementType* begin() { return GetData(); }
+	inline ElementType* end() { return GetData() + Num(); }
 
 	/**
 	 * DO NOT USE DIRECTLY
 	 * STL-like iterators to enable range-based for loop support.
 	 */
-	FORCEINLINE const ElementType* begin() const { return GetData(); }
-	FORCEINLINE const ElementType* end() const { return GetData() + Num(); }
+	inline const ElementType* begin() const { return GetData(); }
+	inline const ElementType* end() const { return GetData() + Num(); }
 
 	/**
 	 * Helper function returning the size of the inner type.
 	 *
 	 * @returns Size in bytes of array type.
 	 */
-	FORCEINLINE uint32 GetTypeSize() const
+	inline uint32 GetTypeSize() const
 	{
 		return sizeof(ElementType);
 	}
@@ -447,7 +447,7 @@ public:
 	 * Checks array invariants: if array size is greater than zero and less
 	 * than maximum.
 	 */
-	FORCEINLINE void CheckInvariants() const
+	inline void CheckInvariants() const
 	{
 		for (uint8 Idx = 0; Idx < DimNum; Idx++)
 		{
@@ -460,7 +460,7 @@ public:
 	 *
 	 * @param Dimension Dimension of the array.
 	 */
-	FORCEINLINE void DimensionCheck(uint8 Dimension) const
+	inline void DimensionCheck(uint8 Dimension) const
 	{
 		checkf((Dimension < DimNum), TEXT("MultiArray dimension out of bounds: %lld from a rank of %lld"), (long long)Dimension, (long long)DimNum); // & for one branch
 	}
@@ -471,7 +471,7 @@ public:
 	 * @param Dimension Dimension of the array.
 	 * @param Index Index to check.
 	 */
-	FORCEINLINE void RangeCheck(uint8 Dimension, int64 Index) const
+	inline void RangeCheck(uint8 Dimension, int64 Index) const
 	{
 		DimensionCheck(Dimension);
 		CheckInvariants();
@@ -487,7 +487,7 @@ public:
 	 * @param Index Starting index of the slice.
 	 * @param InNum Length of the slice.
 	 */
-	FORCEINLINE void SliceRangeCheck(uint8 Dimension, int64 Index, int64 InNum) const
+	inline void SliceRangeCheck(uint8 Dimension, int64 Index, int64 InNum) const
 	{
 		DimensionCheck(Dimension);
 
@@ -526,7 +526,7 @@ public:
 	 *
 	 * @returns Number of dimensions in array.
 	 */
-	FORCEINLINE uint8 Rank() const
+	inline uint8 Rank() const
 	{
 		return DimNum;
 	}
@@ -536,7 +536,7 @@ public:
 	 *
 	 * @returns Total number of elements in array.
 	 */
-	FORCEINLINE int64 Num() const
+	inline int64 Num() const
 	{
 		return ArrayShape.Total();
 	}
@@ -546,7 +546,7 @@ public:
 	 *
 	 * @returns Total number of bytes used by the array.
 	 */
-	FORCEINLINE int64 NumBytes() const
+	inline int64 NumBytes() const
 	{
 		return ArrayShape.Total() * sizeof(ElementType);
 	}
@@ -557,7 +557,7 @@ public:
 	 * @returns Number of elements in array.
 	 */
 	template<uint8 InDimIdx>
-	FORCEINLINE int64 Num() const
+	inline int64 Num() const
 	{
 		static_assert(InDimIdx < DimNum);
 		return ArrayShape[InDimIdx];
@@ -569,7 +569,7 @@ public:
 	 * @param Dimension Dimension of the array.
 	 * @returns Number of elements in array.
 	 */
-	FORCEINLINE int64 Num(uint8 Dimension) const
+	inline int64 Num(uint8 Dimension) const
 	{
 		DimensionCheck(Dimension);
 
@@ -581,7 +581,7 @@ public:
 	 *
 	 * @returns Number of elements in each dimension.
 	 */
-	FORCEINLINE TMultiArrayShape<DimNum> Shape() const
+	inline TMultiArrayShape<DimNum> Shape() const
 	{
 		return ArrayShape;
 	}
@@ -592,7 +592,7 @@ public:
 	 * @param Dimension Dimension of the array.
 	 * @returns Stride of that dimension.
 	 */
-	FORCEINLINE int64 Stride(uint8 Dimension) const
+	inline int64 Stride(uint8 Dimension) const
 	{
 		DimensionCheck(Dimension);
 
@@ -611,7 +611,7 @@ public:
 	 * @param InNum Number of elements in the new view
 	 * @returns Sliced view
 	 */
-	[[nodiscard]] FORCEINLINE TMultiArrayView<DimNum, ElementType> Slice(int64 Index, int64 InNum)
+	[[nodiscard]] inline TMultiArrayView<DimNum, ElementType> Slice(int64 Index, int64 InNum)
 	{
 		SliceRangeCheck(0, Index, InNum);
 		TMultiArrayShape NewShape = ArrayShape;
@@ -626,7 +626,7 @@ public:
 	 * @param InNum Number of elements in the new view
 	 * @returns Sliced view
 	 */
-	[[nodiscard]] FORCEINLINE TMultiArrayView<DimNum, const ElementType> Slice(int64 Index, int64 InNum) const
+	[[nodiscard]] inline TMultiArrayView<DimNum, const ElementType> Slice(int64 Index, int64 InNum) const
 	{
 		SliceRangeCheck(0, Index, InNum);
 		TMultiArrayShape NewShape = ArrayShape;
@@ -783,12 +783,12 @@ public:
 
 public:
 
-	FORCEINLINE operator TMultiArrayView<DimNum, ElementType>()
+	inline operator TMultiArrayView<DimNum, ElementType>()
 	{
 		return TMultiArrayView<DimNum, ElementType>(GetData(), ArrayShape);
 	}
 
-	FORCEINLINE operator TMultiArrayView<DimNum, const ElementType>() const
+	inline operator TMultiArrayView<DimNum, const ElementType>() const
 	{
 		return TMultiArrayView<DimNum, const ElementType>(GetData(), ArrayShape);
 	}
@@ -849,7 +849,7 @@ public:
 	 *
 	 * @returns Reference to indexed element.
 	 */
-	FORCEINLINE ElementType& operator[](int64 Index)
+	inline ElementType& operator[](int64 Index)
 	{
 		RangeCheck(0, Index);
 		return GetData()[Index];
@@ -862,7 +862,7 @@ public:
 	 *
 	 * @returns Reference to indexed element.
 	 */
-	FORCEINLINE const ElementType& operator[](int64 Index) const
+	inline const ElementType& operator[](int64 Index) const
 	{
 		RangeCheck(0, Index);
 		return GetData()[Index];
@@ -873,7 +873,7 @@ public:
 	 *
 	 * @returns Flattened array view.
 	 */
-	FORCEINLINE TMultiArrayView<1, ElementType> Flatten()
+	inline TMultiArrayView<1, ElementType> Flatten()
 	{
 		return TMultiArrayView<1, ElementType>(GetData(), Num());
 	}
@@ -883,31 +883,31 @@ public:
 	 *
 	 * @returns Flattened array view.
 	 */
-	FORCEINLINE TMultiArrayView<1, const ElementType> Flatten() const
+	inline TMultiArrayView<1, const ElementType> Flatten() const
 	{
 		return TMultiArrayView<1, const ElementType>(GetData(), Num());
 	}
 
 	/** Explicit conversion to TArrayView. */
-	FORCEINLINE TArrayView<ElementType> ArrayView()
+	inline TArrayView<ElementType> ArrayView()
 	{
 		return TArrayView<ElementType>(GetData(), Num());
 	}
 
 	/** Explicit conversion to TArrayView. */
-	FORCEINLINE TArrayView<const ElementType> ArrayView() const
+	inline TArrayView<const ElementType> ArrayView() const
 	{
 		return TArrayView<const ElementType>(GetData(), Num());
 	}
 
 	/** Implicit cast to TArrayView. */
-	FORCEINLINE operator TArrayView<ElementType>()
+	inline operator TArrayView<ElementType>()
 	{
 		return TArrayView<ElementType>(GetData(), Num());
 	}
 
 	/** Implicit cast to TArrayView. */
-	FORCEINLINE operator TArrayView<const ElementType>() const
+	inline operator TArrayView<const ElementType>() const
 	{
 		return TArrayView<const ElementType>(GetData(), Num());
 	}
@@ -917,7 +917,7 @@ public:
 	/**
 	 * Constructor, initializes element number counters.
 	 */
-	FORCEINLINE TMultiArray() = default;
+	inline TMultiArray() = default;
 
 	/**
 	 * Constructor from a raw array of elements.
@@ -925,7 +925,7 @@ public:
 	 * @param Ptr   A pointer to an array of elements to copy.
 	 * @param Nums The number of elements on each dimension.
 	 */
-	FORCEINLINE TMultiArray(const ElementType* Ptr, TMultiArrayShape<DimNum> Shape)
+	inline TMultiArray(const ElementType* Ptr, TMultiArrayShape<DimNum> Shape)
 	{
 		check(Ptr != nullptr || Shape.Total() == 0);
 
@@ -946,7 +946,7 @@ public:
 		typename OtherAllocator
 		UE_REQUIRES(UE::MultiArray::Private::TMultiArrayElementsAreCompatible_V<ElementType, const OtherElementType&>)
 	>
-		FORCEINLINE explicit TMultiArray(const TMultiArray<DimNum, OtherElementType, OtherAllocator>& Other)
+		inline explicit TMultiArray(const TMultiArray<DimNum, OtherElementType, OtherAllocator>& Other)
 	{
 		CopyToEmpty(Other.GetData(), Other.Shape());
 	}
@@ -956,7 +956,7 @@ public:
 	 *
 	 * @param Other The source array to copy.
 	 */
-	FORCEINLINE TMultiArray(const TMultiArray& Other)
+	inline TMultiArray(const TMultiArray& Other)
 	{
 		CopyToEmpty(Other.GetData(), Other.Shape());
 	}
@@ -1007,7 +1007,7 @@ private:
 	 * @param FromArray Array to move from.
 	 */
 	template <typename FromArrayType, typename ToArrayType>
-	static FORCEINLINE void MoveOrCopy(ToArrayType& ToArray, FromArrayType& FromArray)
+	static inline void MoveOrCopy(ToArrayType& ToArray, FromArrayType& FromArray)
 	{
 		ToArray.CopyToEmpty(FromArray.GetData(), FromArray.Shape());
 	}
@@ -1018,7 +1018,7 @@ public:
 	 *
 	 * @param Other Array to move from.
 	 */
-	FORCEINLINE TMultiArray(TMultiArray&& Other)
+	inline TMultiArray(TMultiArray&& Other)
 	{
 		MoveOrCopy(*this, Other);
 	}
@@ -1033,7 +1033,7 @@ public:
 		typename OtherAllocator
 		UE_REQUIRES(UE::MultiArray::Private::TMultiArrayElementsAreCompatible_V<ElementType, OtherElementType&&>)
 	>
-		FORCEINLINE explicit TMultiArray(TMultiArray<DimNum, OtherElementType, OtherAllocator>&& Other)
+		inline explicit TMultiArray(TMultiArray<DimNum, OtherElementType, OtherAllocator>&& Other)
 	{
 		MoveOrCopy(*this, Other);
 	}
@@ -1067,7 +1067,7 @@ public:
 	 *
 	 * @returns Pointer to first array entry or nullptr if empty
 	 */
-	FORCEINLINE ElementType* GetData()
+	inline ElementType* GetData()
 	{
 		return (ElementType*)AllocatorInstance.GetAllocation();
 	}
@@ -1077,7 +1077,7 @@ public:
 	 *
 	 * @returns Pointer to first array entry or nullptr if empty
 	 */
-	FORCEINLINE const ElementType* GetData() const
+	inline const ElementType* GetData() const
 	{
 		return (const ElementType*)AllocatorInstance.GetAllocation();
 	}
@@ -1086,23 +1086,23 @@ public:
 	 * DO NOT USE DIRECTLY
 	 * STL-like iterators to enable range-based for loop support.
 	 */
-	FORCEINLINE ElementType* begin() { return GetData(); }
-	FORCEINLINE ElementType* end() { return GetData() + Num(); }
+	inline ElementType* begin() { return GetData(); }
+	inline ElementType* end() { return GetData() + Num(); }
 
 
 	/**
 	 * DO NOT USE DIRECTLY
 	 * STL-like iterators to enable range-based for loop support.
 	 */
-	FORCEINLINE const ElementType* begin() const { return GetData(); }
-	FORCEINLINE const ElementType* end() const { return GetData() + Num(); }
+	inline const ElementType* begin() const { return GetData(); }
+	inline const ElementType* end() const { return GetData() + Num(); }
 
 	/**
 	 * Helper function returning the size of the inner type.
 	 *
 	 * @returns Size in bytes of array type.
 	 */
-	FORCEINLINE uint32 GetTypeSize() const
+	inline uint32 GetTypeSize() const
 	{
 		return sizeof(ElementType);
 	}
@@ -1111,7 +1111,7 @@ public:
 	 * Checks array invariants: if array size is greater than zero and less
 	 * than maximum.
 	 */
-	FORCEINLINE void CheckInvariants() const
+	inline void CheckInvariants() const
 	{
 		for (uint8 Idx = 0; Idx < DimNum; Idx++)
 		{
@@ -1124,7 +1124,7 @@ public:
 	 *
 	 * @param Dimension Dimension of the array.
 	 */
-	FORCEINLINE void DimensionCheck(uint8 Dimension) const
+	inline void DimensionCheck(uint8 Dimension) const
 	{
 		checkf((Dimension < DimNum), TEXT("MultiArray dimension out of bounds: %lld from a rank of %lld"), (long long)Dimension, (long long)DimNum); // & for one branch
 	}
@@ -1135,7 +1135,7 @@ public:
 	 * @param Dimension Dimension of the array.
 	 * @param Index Index to check.
 	 */
-	FORCEINLINE void RangeCheck(uint8 Dimension, int64 Index) const
+	inline void RangeCheck(uint8 Dimension, int64 Index) const
 	{
 		DimensionCheck(Dimension);
 		CheckInvariants();
@@ -1151,7 +1151,7 @@ public:
 	 * @param Index Starting index of the slice.
 	 * @param InNum Length of the slice.
 	 */
-	FORCEINLINE void SliceRangeCheck(uint8 Dimension, int64 Index, int64 InNum) const
+	inline void SliceRangeCheck(uint8 Dimension, int64 Index, int64 InNum) const
 	{
 		DimensionCheck(Dimension);
 
@@ -1190,7 +1190,7 @@ public:
 	 *
 	 * @returns Number of dimensions in array.
 	 */
-	FORCEINLINE uint8 Rank() const
+	inline uint8 Rank() const
 	{
 		return DimNum;
 	}
@@ -1200,7 +1200,7 @@ public:
 	 *
 	 * @returns Total number of elements in array.
 	 */
-	FORCEINLINE int64 Num() const
+	inline int64 Num() const
 	{
 		return ArrayShape.Total();
 	}
@@ -1210,7 +1210,7 @@ public:
 	 *
 	 * @returns Total number of bytes used by the array.
 	 */
-	FORCEINLINE int64 NumBytes() const
+	inline int64 NumBytes() const
 	{
 		return ArrayShape.Total() * sizeof(ElementType);
 	}
@@ -1221,7 +1221,7 @@ public:
 	 * @returns Number of elements in array.
 	 */
 	template<uint8 InDimIdx>
-	FORCEINLINE int64 Num() const
+	inline int64 Num() const
 	{
 		static_assert(InDimIdx < DimNum);
 		return ArrayShape[InDimIdx];
@@ -1233,7 +1233,7 @@ public:
 	 * @param Dimension Dimension of the array.
 	 * @returns Number of elements in array.
 	 */
-	FORCEINLINE int64 Num(uint8 Dimension) const
+	inline int64 Num(uint8 Dimension) const
 	{
 		DimensionCheck(Dimension);
 
@@ -1245,7 +1245,7 @@ public:
 	 *
 	 * @returns Number of elements in each dimension.
 	 */
-	FORCEINLINE TMultiArrayShape<DimNum> Shape() const
+	inline TMultiArrayShape<DimNum> Shape() const
 	{
 		return ArrayShape;
 	}
@@ -1256,7 +1256,7 @@ public:
 	 * @param Dimension Dimension of the array.
 	 * @returns Stride of that dimension.
 	 */
-	FORCEINLINE int64 Stride(uint8 Dimension) const
+	inline int64 Stride(uint8 Dimension) const
 	{
 		DimensionCheck(Dimension);
 
@@ -1275,7 +1275,7 @@ public:
 	 * @param InNum Number of elements in the new view
 	 * @returns Sliced view
 	 */
-	[[nodiscard]] FORCEINLINE TMultiArrayView<DimNum, ElementType> Slice(int64 Index, int64 InNum)
+	[[nodiscard]] inline TMultiArrayView<DimNum, ElementType> Slice(int64 Index, int64 InNum)
 	{
 		SliceRangeCheck(0, Index, InNum);
 		TMultiArrayShape NewShape = ArrayShape;
@@ -1290,7 +1290,7 @@ public:
 	 * @param InNum Number of elements in the new view
 	 * @returns Sliced view
 	 */
-	[[nodiscard]] FORCEINLINE TMultiArrayView<DimNum, const ElementType> Slice(int64 Index, int64 InNum) const
+	[[nodiscard]] inline TMultiArrayView<DimNum, const ElementType> Slice(int64 Index, int64 InNum) const
 	{
 		SliceRangeCheck(0, Index, InNum);
 		TMultiArrayShape NewShape = ArrayShape;
@@ -1448,12 +1448,12 @@ public:
 
 public:
 
-	FORCEINLINE operator TMultiArrayView<DimNum, ElementType>()
+	inline operator TMultiArrayView<DimNum, ElementType>()
 	{
 		return TMultiArrayView<DimNum, ElementType>(GetData(), ArrayShape);
 	}
 
-	FORCEINLINE operator TMultiArrayView<DimNum, const ElementType>() const
+	inline operator TMultiArrayView<DimNum, const ElementType>() const
 	{
 		return TMultiArrayView<DimNum, const ElementType>(GetData(), ArrayShape);
 	}
diff --git a/Engine/Plugins/Experimental/LearningAgents/Source/Learning/Public/MultiArrayView.h b/Engine/Plugins/Experimental/LearningAgents/Source/Learning/Public/MultiArrayView.h
--- a/Engine/Plugins/Experimental/LearningAgents/Source/Learning/Public/MultiArrayView.h
+++ b/Engine/Plugins/Experimental/LearningAgents/Source/Learning/Public/MultiArrayView.h
@@ -29,7 +29,7 @@ namespace UE::MultiArrayView::Private
 	// Simply forwards to an unqualified GetData(), but can be called from within TArrayView
 	// where GetData() is already a member and so hides any others.
 	template <typename T>
-	FORCEINLINE decltype(auto) GetDataHelper(T&& Arg)
+	inline decltype(auto) GetDataHelper(T&& Arg)
 	{
 		return GetData(Forward<T>(Arg));
 	}
@@ -38,7 +38,7 @@ namespace UE::MultiArrayView::Private
 	//
 	// Gets the data from the passed argument and proceeds to reinterpret the resulting elements
 	template <typename T>
-	FORCEINLINE decltype(auto) GetReinterpretedDataHelper(T&& Arg)
+	inline decltype(auto) GetReinterpretedDataHelper(T&& Arg)
 	{
 		auto NaturalPtr = GetData(Forward<T>(Arg));
 		using NaturalElementType = std::remove_pointer_t<decltype(NaturalPtr)>;
@@ -183,7 +183,7 @@ public:
 	 * @param InData	The data to view
 	 * @param InShape	The number of elements on each dimension
 	 */
-	FORCEINLINE TMultiArrayView(ElementType* InData, TMultiArrayShape<DimNum> InShape)
+	inline TMultiArrayView(ElementType* InData, TMultiArrayShape<DimNum> InShape)
 		: DataPtr(InData)
 		, ArrayShape(InShape)
 	{
@@ -195,7 +195,7 @@ public:
 	 *
 	 * @returns Reference to indexed element.
 	 */
-	FORCEINLINE TMultiArrayView<DimNum - 1, ElementType> operator[](int64 Index) const
+	inline TMultiArrayView<DimNum - 1, ElementType> operator[](int64 Index) const
 	{
 		RangeCheck(0, Index);
 
@@ -213,7 +213,7 @@ public:
 	 *
 	 * @returns Flattened array view.
 	 */
-	FORCEINLINE TMultiArrayView<1, ElementType> Flatten() const
+	inline TMultiArrayView<1, ElementType> Flatten() const
 	{
 		return TMultiArrayView<1, ElementType>(DataPtr, Num());
 	}
@@ -224,7 +224,7 @@ public:
 	 *
 	 * @returns View flattened on the given dimension.
 	 */
-	FORCEINLINE TMultiArrayView<DimNum - 1, ElementType> Flatten(uint8 Dimension) const
+	inline TMultiArrayView<DimNum - 1, ElementType> Flatten(uint8 Dimension) const
 	{
 		checkf((Dimension < DimNum - 1), TEXT("MultiArray flatten dimension out of bounds: %lld from a rank of %lld"), (long long)Dimension, (long long)DimNum); // & for one branch
 
@@ -254,7 +254,7 @@ public:
 	 * @returns View flattened on the given dimension.
 	 */
 	template<uint8 Dimension>
-	FORCEINLINE TMultiArrayView<DimNum - 1, ElementType> Flatten() const
+	inline TMultiArrayView<DimNum - 1, ElementType> Flatten() const
 	{
 		static_assert(Dimension < DimNum - 1, "MultiArray flatten dimension out of bounds");
 
@@ -284,7 +284,7 @@ public:
 	 *
 	 * @returns Pointer to first array entry or nullptr if empty
 	 */
-	FORCEINLINE ElementType* GetData() const
+	inline ElementType* GetData() const
 	{
 		return DataPtr;
 	}
@@ -294,15 +294,15 @@ public:
 	 * DO NOT USE DIRECTLY
 	 * STL-like iterators to enable range-based for loop support.
 	 */
-	FORCEINLINE ElementType* begin() const { return GetData(); }
-	FORCEINLINE ElementType* end() const { return GetData() + Num(); }
+	inline ElementType* begin() const { return GetData(); }
+	inline ElementType* end() const { return GetData() + Num(); }
 
 	/**
 	 * Helper function returning the size of the inner type.
 	 *
 	 * @returns Size in bytes of array type.
 	 */
-	FORCEINLINE static constexpr size_t GetTypeSize()
+	inline static constexpr size_t GetTypeSize()
 	{
 		return sizeof(ElementType);
 	}
@@ -310,7 +310,7 @@ public:
 	/**
 	 * Helper function returning the alignment of the inner type.
 	 */
-	FORCEINLINE static constexpr size_t GetTypeAlignment()
+	inline static constexpr size_t GetTypeAlignment()
 	{
 		return alignof(ElementType);
 	}
@@ -318,7 +318,7 @@ public:
 	/**
 	 * Checks array invariants: if array size is greater than zero on each dimension
 	 */
-	FORCEINLINE void CheckInvariants() const
+	inline void CheckInvariants() const
 	{
 		for (uint8 Idx = 0; Idx < DimNum; Idx++)
 		{
@@ -331,7 +331,7 @@ public:
 	 *
 	 * @param Dimension Dimension of the array.
 	 */
-	FORCEINLINE void DimensionCheck(uint8 Dimension) const
+	inline void DimensionCheck(uint8 Dimension) const
 	{
 		checkf((Dimension < DimNum), TEXT("MultiArray dimension out of bounds: %lld from a rank of %lld"), (long long)Dimension, (long long)DimNum); // & for one branch
 	}
@@ -342,7 +342,7 @@ public:
 	 * @param Dimension Dimension of the array.
 	 * @param Index Index to check.
 	 */
-	FORCEINLINE void RangeCheck(uint8 Dimension, int64 Index) const
+	inline void RangeCheck(uint8 Dimension, int64 Index) const
 	{
 		DimensionCheck(Dimension);
 		CheckInvariants();
@@ -358,7 +358,7 @@ public:
 	 * @param Index Starting index of the slice.
 	 * @param InNum Length of the slice.
 	 */
-	FORCEINLINE void SliceRangeCheck(uint8 Dimension, int64 Index, int64 InNum) const
+	inline void SliceRangeCheck(uint8 Dimension, int64 Index, int64 InNum) const
 	{
 		DimensionCheck(Dimension);
 
@@ -397,7 +397,7 @@ public:
 	 *
 	 * @returns Number of dimensions in array.
 	 */
-	FORCEINLINE uint8 Rank() const
+	inline uint8 Rank() const
 	{
 		return DimNum;
 	}
@@ -407,7 +407,7 @@ public:
 	 *
 	 * @returns Total number of elements in array.
 	 */
-	FORCEINLINE int64 Num() const
+	inline int64 Num() const
 	{
 		return ArrayShape.Total();
 	}
@@ -417,7 +417,7 @@ public:
 	 *
 	 * @returns Total number of bytes used by the array.
 	 */
-	FORCEINLINE int64 NumBytes() const
+	inline int64 NumBytes() const
 	{
 		return ArrayShape.Total() * sizeof(ElementType);
 	}
@@ -428,7 +428,7 @@ public:
 	 * @returns Number of elements in array.
 	 */
 	template<uint8 InDimIdx>
-	FORCEINLINE int64 Num() const
+	inline int64 Num() const
 	{
 		static_assert(InDimIdx < DimNum);
 		return ArrayShape[InDimIdx];
@@ -440,7 +440,7 @@ public:
 	 * @param Dimension Dimension of the array.
 	 * @returns Number of elements in array.
 	 */
-	FORCEINLINE int64 Num(uint8 Dimension) const
+	inline int64 Num(uint8 Dimension) const
 	{
 		DimensionCheck(Dimension);
 
@@ -452,7 +452,7 @@ public:
 	 *
 	 * @returns Number of elements in each dimension.
 	 */
-	FORCEINLINE TMultiArrayShape<DimNum> Shape() const
+	inline TMultiArrayShape<DimNum> Shape() const
 	{
 		return ArrayShape;
 	}
@@ -463,7 +463,7 @@ public:
 	 * @param Dimension Dimension of the array.
 	 * @returns Stride of that dimension.
 	 */
-	FORCEINLINE int64 Stride(uint8 Dimension) const
+	inline int64 Stride(uint8 Dimension) const
 	{
 		DimensionCheck(Dimension);
 
@@ -482,7 +482,7 @@ public:
 	 * @param InNum Number of elements in the new view
 	 * @returns Sliced view
 	 */
-	[[nodiscard]] FORCEINLINE TMultiArrayView Slice(int64 Index, int64 InNum) const
+	[[nodiscard]] inline TMultiArrayView Slice(int64 Index, int64 InNum) const
 	{
 		SliceRangeCheck(0, Index, InNum);
 		TMultiArrayShape NewShape = ArrayShape;
@@ -509,7 +509,7 @@ public:
 	}
 
 	/** Implicit cast for const. */
-	FORCEINLINE operator TMultiArrayView<DimNum, const ElementType>() const
+	inline operator TMultiArrayView<DimNum, const ElementType>() const
 	{
 		return TMultiArrayView<DimNum, const ElementType>(DataPtr, ArrayShape);
 	}
@@ -562,7 +562,7 @@ public:
 			>::Value
 		)
 	>
-	FORCEINLINE TMultiArrayView(OtherRangeType&& Other)
+	inline TMultiArrayView(OtherRangeType&& Other)
 		: DataPtr(std::conditional_t<
 			TIsCompatibleRangeType<OtherRangeType>::Value,
 			TIsCompatibleRangeType<OtherRangeType>,
@@ -580,7 +580,7 @@ public:
 	 * @param InData	The data to view
 	 * @param InNums	The number of elements on each dimension
 	 */
-	FORCEINLINE TMultiArrayView(ElementType* InData, TMultiArrayShape<DimNum> InShape)
+	inline TMultiArrayView(ElementType* InData, TMultiArrayShape<DimNum> InShape)
 		: DataPtr(InData)
 		, ArrayShape(InShape)
 	{
@@ -597,7 +597,7 @@ public:
 		typename OtherElementType
 		UE_REQUIRES(UE::MultiArrayView::Private::TIsCompatibleElementType_V<OtherElementType, ElementType>)
 	>
-	FORCEINLINE TMultiArrayView(OtherElementType* InData, int64 InNum)
+	inline TMultiArrayView(OtherElementType* InData, int64 InNum)
 		: DataPtr(InData)
 	{
 		ArrayShape[0] = InNum;
@@ -609,7 +609,7 @@ public:
 	 *
 	 * The caller is responsible for ensuring that the view does not outlive the initializer list.
 	 */
-	FORCEINLINE TMultiArrayView(std::initializer_list<ElementType> List)
+	inline TMultiArrayView(std::initializer_list<ElementType> List)
 		: DataPtr(UE::MultiArrayView::Private::GetDataHelper(List))
 	{
 		ArrayShape[0] = GetNum(List);
@@ -623,32 +623,32 @@ public:
 	 *
 	 * @returns Reference to indexed element.
 	 */
-	FORCEINLINE ElementType& operator[](int64 Index) const
+	inline ElementType& operator[](int64 Index) const
 	{
 		RangeCheck(0, Index);
 		return DataPtr[Index];
 	}
 
 	/** Explicit conversion to TArrayView. */
-	FORCEINLINE TArrayView<ElementType> ArrayView()
+	inline TArrayView<ElementType> ArrayView()
 	{
 		return TArrayView<ElementType>(GetData(), Num());
 	}
 
 	/** Explicit conversion to TArrayView. */
-	FORCEINLINE TArrayView<const ElementType> ArrayView() const
+	inline TArrayView<const ElementType> ArrayView() const
 	{
 		return TArrayView<const ElementType>(GetData(), Num());
 	}
 
 	/** Implicit cast to TArrayView. */
-	FORCEINLINE operator TArrayView<ElementType>()
+	inline operator TArrayView<ElementType>()
 	{
 		return TArrayView<ElementType>(DataPtr, ArrayShape[0]);
 	}
 
 	/** Implicit cast to TArrayView. */
-	FORCEINLINE operator TArrayView<const ElementType>() const
+	inline operator TArrayView<const ElementType>() const
 	{
 		return TArrayView<const ElementType>(DataPtr, ArrayShape[0]);
 	}
@@ -660,7 +660,7 @@ public:
 	 *
 	 * @returns Pointer to first array entry or nullptr if empty
 	 */
-	FORCEINLINE ElementType* GetData() const
+	inline ElementType* GetData() const
 	{
 		return DataPtr;
 	}
@@ -669,15 +669,15 @@ public:
 	 * DO NOT USE DIRECTLY
 	 * STL-like iterators to enable range-based for loop support.
 	 */
-	FORCEINLINE ElementType* begin() const { return GetData(); }
-	FORCEINLINE ElementType* end() const { return GetData() + Num(); }
+	inline ElementType* begin() const { return GetData(); }
+	inline ElementType* end() const { return GetData() + Num(); }
 
 	/**
 	 * Helper function returning the size of the inner type.
 	 *
 	 * @returns Size in bytes of array type.
 	 */
-	FORCEINLINE static constexpr size_t GetTypeSize()
+	inline static constexpr size_t GetTypeSize()
 	{
 		return sizeof(ElementType);
 	}
@@ -685,7 +685,7 @@ public:
 	/**
 	 * Helper function returning the alignment of the inner type.
 	 */
-	FORCEINLINE static constexpr size_t GetTypeAlignment()
+	inline static constexpr size_t GetTypeAlignment()
 	{
 		return alignof(ElementType);
 	}
@@ -693,7 +693,7 @@ public:
 	/**
 	 * Checks array invariants: if array size is greater than zero on each dimension
 	 */
-	FORCEINLINE void CheckInvariants() const
+	inline void CheckInvariants() const
 	{
 		for (uint8 Idx = 0; Idx < DimNum; Idx++)
 		{
@@ -706,7 +706,7 @@ public:
 	 *
 	 * @param Dimension Dimension of the array.
 	 */
-	FORCEINLINE void DimensionCheck(uint8 Dimension) const
+	inline void DimensionCheck(uint8 Dimension) const
 	{
 		checkf((Dimension < DimNum), TEXT("MultiArray dimension out of bounds: %lld from a rank of %lld"), (long long)Dimension, (long long)DimNum); // & for one branch
 	}
@@ -717,7 +717,7 @@ public:
 	 * @param Dimension Dimension of the array.
 	 * @param Index Index to check.
 	 */
-	FORCEINLINE void RangeCheck(uint8 Dimension, int64 Index) const
+	inline void RangeCheck(uint8 Dimension, int64 Index) const
 	{
 		DimensionCheck(Dimension);
 		CheckInvariants();
@@ -733,7 +733,7 @@ public:
 	 * @param Index Starting index of the slice.
 	 * @param InNum Length of the slice.
 	 */
-	FORCEINLINE void SliceRangeCheck(uint8 Dimension, int64 Index, int64 InNum) const
+	inline void SliceRangeCheck(uint8 Dimension, int64 Index, int64 InNum) const
 	{
 		DimensionCheck(Dimension);
 
@@ -772,7 +772,7 @@ public:
 	 *
 	 * @returns Number of dimensions in array.
 	 */
-	FORCEINLINE uint8 Rank() const
+	inline uint8 Rank() const
 	{
 		return DimNum;
 	}
@@ -782,7 +782,7 @@ public:
 	 *
 	 * @returns Total number of elements in array.
 	 */
-	FORCEINLINE int64 Num() const
+	inline int64 Num() const
 	{
 		return ArrayShape.Total();
 	}
@@ -792,7 +792,7 @@ public:
 	 *
 	 * @returns Total number of bytes used by the array.
 	 */
-	FORCEINLINE int64 NumBytes() const
+	inline int64 NumBytes() const
 	{
 		return ArrayShape.Total() * sizeof(ElementType);
 	}
@@ -803,7 +803,7 @@ public:
 	 * @returns Number of elements in array.
 	 */
 	template<uint8 InDimIdx>
-	FORCEINLINE int64 Num() const
+	inline int64 Num() const
 	{
 		static_assert(InDimIdx < DimNum);
 		return ArrayShape[InDimIdx];
@@ -815,7 +815,7 @@ public:
 	 * @param Dimension Dimension of the array.
 	 * @returns Number of elements in array.
 	 */
-	FORCEINLINE int64 Num(uint8 Dimension) const
+	inline int64 Num(uint8 Dimension) const
 	{
 		DimensionCheck(Dimension);
 
@@ -827,7 +827,7 @@ public:
 	 *
 	 * @returns Number of elements in each dimension.
 	 */
-	FORCEINLINE TMultiArrayShape<DimNum> Shape() const
+	inline TMultiArrayShape<DimNum> Shape() const
 	{
 		return ArrayShape;
 	}
@@ -838,7 +838,7 @@ public:
 	 * @param Dimension Dimension of the array.
 	 * @returns Stride of that dimension.
 	 */
-	FORCEINLINE int64 Stride(uint8 Dimension) const
+	inline int64 Stride(uint8 Dimension) const
 	{
 		DimensionCheck(Dimension);
 
@@ -857,7 +857,7 @@ public:
 	 * @param InNum Number of elements in the new view
 	 * @returns Sliced view
 	 */
-	[[nodiscard]] FORCEINLINE TMultiArrayView Slice(int64 Index, int64 InNum) const
+	[[nodiscard]] inline TMultiArrayView Slice(int64 Index, int64 InNum) const
 	{
 		SliceRangeCheck(0, Index, InNum);
 		TMultiArrayShape NewShape = ArrayShape;
@@ -884,7 +884,7 @@ public:
 	}
 
 	/** Implicit cast for const. */
-	FORCEINLINE operator TMultiArrayView<DimNum, const ElementType>() const
+	inline operator TMultiArrayView<DimNum, const ElementType>() const
 	{
 		return TMultiArrayView<DimNum, const ElementType>(DataPtr, ArrayShape);
 	}
