[LearningAgents] Thread safe wrapper of add agent and remove agent to prevent synchronous mutations of agent arrays

diff --git a/Engine/Plugins/Experimental/LearningAgents/Source/LearningAgents/Private/LearningAgentsManager.cpp b/Engine/Plugins/Experimental/LearningAgents/Source/LearningAgents/Private/LearningAgentsManager.cpp
--- a/Engine/Plugins/Experimental/LearningAgents/Source/LearningAgents/Private/LearningAgentsManager.cpp
+++ b/Engine/Plugins/Experimental/LearningAgents/Source/LearningAgents/Private/LearningAgentsManager.cpp
@@ -8,6 +8,7 @@
 
 #include "UObject/ScriptInterface.h"
 #include "EngineDefines.h"
+#include "Async/Async.h"
 
 ULearningAgentsManager::ULearningAgentsManager() : Super(FObjectInitializer::Get()) { PrimaryComponentTick.bCanEverTick = true; }
 ULearningAgentsManager::ULearningAgentsManager(FVTableHelper& Helper) : Super(Helper) {}
@@ -160,6 +161,19 @@ void ULearningAgentsManager::AddAgents(TArray<int32>& OutAgentIds, const TArray<
 	}
 }
 
+void ULearningAgentsManager::RequestAddAgent(UObject* Agent)
+{
+	if (!IsInGameThread())
+	{
+		AsyncTask(ENamedThreads::GameThread, [this, Agent]()
+			{
+				AddAgent(Agent);
+			});
+		return;
+	}
+	AddAgent(Agent);
+}
+
 void ULearningAgentsManager::RemoveAgent(const int32 AgentId)
 {
 	if (AgentId == INDEX_NONE)
@@ -245,6 +259,19 @@ void ULearningAgentsManager::RemoveAllAgents()
 	RemoveAgents(RemoveAgentIds);
 }
 
+void ULearningAgentsManager::RequestRemoveAgent(const int32 AgentId)
+{
+	if (!IsInGameThread())
+	{
+		AsyncTask(ENamedThreads::GameThread, [this, AgentId]()
+			{
+				RemoveAgent(AgentId);
+			});
+		return;
+	}
+	RemoveAgent(AgentId);
+}
+
 void ULearningAgentsManager::ResetAgent(const int32 AgentId)
 {
 	if (AgentId == INDEX_NONE)
diff --git a/Engine/Plugins/Experimental/LearningAgents/Source/LearningAgents/Public/LearningAgentsManager.h b/Engine/Plugins/Experimental/LearningAgents/Source/LearningAgents/Public/LearningAgentsManager.h
--- a/Engine/Plugins/Experimental/LearningAgents/Source/LearningAgents/Public/LearningAgentsManager.h
+++ b/Engine/Plugins/Experimental/LearningAgents/Source/LearningAgents/Public/LearningAgentsManager.h
@@ -71,6 +71,14 @@ public:
 	UFUNCTION(BlueprintCallable, Category = "LearningAgents")
 	UE_API void AddAgents(TArray<int32>& OutAgentIds, const TArray<UObject*>& InAgents);
 
+	/**
+	 * Schedules a task to add the given object as an agent to the manager, preventing synchronous mutations of agent arrays on the game thread.
+	 * @param Agent The object to be added.
+	 * @return The agent's newly assigned id.
+	 */
+	UFUNCTION(BlueprintCallable, Category = "LearningAgents")
+	UE_API void RequestAddAgent(UObject* Agent);
+
 public:
 
 	/**
@@ -91,6 +99,13 @@ public:
 	UFUNCTION(BlueprintCallable, Category = "LearningAgents")
 	UE_API void RemoveAllAgents();
 
+	/**
+	 * Schedules a task to remove the agent with the given ID from the manager, preventing synchronous mutations of agent arrays on the game thread.
+	 * @param AgentId The id of the agent to remove.
+	 */
+	UFUNCTION(BlueprintCallable, Category = "LearningAgents", meta = (AgentId = "-1"))
+	UE_API void RequestRemoveAgent(const int32 AgentId);
+
 public:
 
 	/**
