[LearningAgentsTraining] In-Editor IL Training

diff --git a/Engine/Plugins/Experimental/LearningAgents/LearningAgents.uplugin b/Engine/Plugins/Experimental/LearningAgents/LearningAgents.uplugin
--- a/Engine/Plugins/Experimental/LearningAgents/LearningAgents.uplugin
+++ b/Engine/Plugins/Experimental/LearningAgents/LearningAgents.uplugin
@@ -35,6 +35,11 @@
 			"Name": "LearningAgentsTraining",
 			"Type": "Runtime",
 			"LoadingPhase": "Default"
+		},
+		{
+			"Name": "LearningAgentsTrainingEditor",
+			"Type": "Editor",
+			"LoadingPhase": "Default"
 		}
 	],
 	"Plugins": [
diff --git a/Engine/Plugins/Experimental/LearningAgents/Source/LearningAgentsTrainingEditor/LearningAgentsTrainingEditor.Build.cs b/Engine/Plugins/Experimental/LearningAgents/Source/LearningAgentsTrainingEditor/LearningAgentsTrainingEditor.Build.cs
new file mode 100644
--- /dev/null
+++ b/Engine/Plugins/Experimental/LearningAgents/Source/LearningAgentsTrainingEditor/LearningAgentsTrainingEditor.Build.cs
@@ -0,0 +1,46 @@
+// Copyright Epic Games, Inc. All Rights Reserved.
+
+using UnrealBuildTool;
+
+public class LearningAgentsTrainingEditor : ModuleRules
+{
+	public LearningAgentsTrainingEditor(ReadOnlyTargetRules Target) : base(Target)
+	{
+		PCHUsage = ModuleRules.PCHUsageMode.UseExplicitOrSharedPCHs;
+
+		PublicIncludePaths.AddRange(
+			new string[] {
+			});
+
+		PrivateIncludePaths.AddRange(
+			new string[] {
+			});
+
+		PublicDependencyModuleNames.AddRange(
+			new string[]
+			{
+				"Core",
+			});
+
+		PrivateDependencyModuleNames.AddRange(
+			new string[]
+			{
+				"CoreUObject",
+				"Engine",
+				"Learning",
+				"LearningAgents",
+				"LearningTraining",
+				"LearningAgentsTraining",
+				"Slate",
+				"SlateCore",
+				"EditorStyle",
+				"PropertyEditor",
+				"UnrealEd"
+			});
+
+		DynamicallyLoadedModuleNames.AddRange(
+			new string[]
+			{
+			});
+	}
+}
diff --git a/Engine/Plugins/Experimental/LearningAgents/Source/LearningAgentsTrainingEditor/Private/LearningAgentsImitationTrainerEditor.cpp b/Engine/Plugins/Experimental/LearningAgents/Source/LearningAgentsTrainingEditor/Private/LearningAgentsImitationTrainerEditor.cpp
new file mode 100644
--- /dev/null
+++ b/Engine/Plugins/Experimental/LearningAgents/Source/LearningAgentsTrainingEditor/Private/LearningAgentsImitationTrainerEditor.cpp
@@ -0,0 +1,69 @@
+// Copyright Epic Games, Inc. All Rights Reserved.
+
+#include "LearningAgentsImitationTrainerEditor.h"
+#include "LearningLog.h"
+#include "Containers/Ticker.h"
+
+ALearningAgentsImitationTrainerEditor::ALearningAgentsImitationTrainerEditor()
+{
+	LearningAgentsManager = CreateDefaultSubobject<ULearningAgentsManager>(TEXT("LearningAgentsManager"));
+}
+
+void ALearningAgentsImitationTrainerEditor::StartTraining()
+{
+	SetupTraining();
+	
+	if (!LearningAgentsImitationTrainer)
+	{
+		UE_LOG(LogLearning, Error, TEXT("%s: Training has failed. Invalid imitation trainer. Was the imitation trainer set with SetupTraining?"), *GetName());
+		return;
+	}
+
+	if (LearningAgentsImitationTrainer->IsTraining())
+	{
+		UE_LOG(LogLearning, Warning, TEXT("%s: Cannot start training when an existing training process is active!"), *GetName());
+		return;
+	}
+
+	UE_LOG(LogLearning, Display, TEXT("%s: Starting training..."), *GetName());
+
+	LearningAgentsImitationTrainer->BeginTraining(Recording, ImitationTrainerSettings, ImitationTrainerTrainingSettings, ImitationTrainerPathSettings);
+
+	FTSTicker::GetCoreTicker().AddTicker(
+		FTickerDelegate::CreateLambda([this](float DeltaTime)->bool {
+			if (LearningAgentsImitationTrainer->IsTraining())
+			{
+				LearningAgentsImitationTrainer->IterateTraining();
+				return true;
+			}
+			return false;
+			}),
+		TrainingTickInterval
+	);
+}
+
+void ALearningAgentsImitationTrainerEditor::StopTraining()
+{
+	if (!LearningAgentsImitationTrainer)
+	{
+		UE_LOG(LogLearning, Error, TEXT("%s: Cannot stop training with an invalid imitation trainer."), *GetName());
+		return;
+	}
+	if (!LearningAgentsImitationTrainer->IsTraining())
+	{
+		UE_LOG(LogLearning, Warning, TEXT("%s: There is no training active. Cannot stop training when training has not started."), *GetName());
+		return;
+	}
+
+	UE_LOG(LogLearning, Display, TEXT("%s: Ending training..."), *GetName());
+	LearningAgentsImitationTrainer->EndTraining();
+}
+
+bool ALearningAgentsImitationTrainerEditor::IsTraining() const
+{
+	if (LearningAgentsImitationTrainer)
+	{
+		return LearningAgentsImitationTrainer->IsTraining();
+	}
+	return false;
+}
diff --git a/Engine/Plugins/Experimental/LearningAgents/Source/LearningAgentsTrainingEditor/Private/LearningAgentsImitationTrainerEditorDetails.cpp b/Engine/Plugins/Experimental/LearningAgents/Source/LearningAgentsTrainingEditor/Private/LearningAgentsImitationTrainerEditorDetails.cpp
new file mode 100644
--- /dev/null
+++ b/Engine/Plugins/Experimental/LearningAgents/Source/LearningAgentsTrainingEditor/Private/LearningAgentsImitationTrainerEditorDetails.cpp
@@ -0,0 +1,101 @@
+// Copyright Epic Games, Inc. All Rights Reserved.
+
+#include "LearningAgentsImitationTrainerEditorDetails.h"
+#include "LearningAgentsImitationTrainerEditor.h"
+#include "LearningLog.h"
+#include "DetailCategoryBuilder.h"
+#include "DetailWidgetRow.h"
+#include "DetailLayoutBuilder.h"
+#include "Misc/Attribute.h"
+#include "Widgets/Input/SButton.h"
+#include "Widgets/Text/STextBlock.h"
+#include "Widgets/SBoxPanel.h"
+
+#define LOCTEXT_NAMESPACE "LearningAgentsImitationTrainerEditorDetails"
+
+TSharedRef<IDetailCustomization> FLearningAgentsImitationTrainerEditorDetails::MakeInstance()
+{
+	return MakeShareable(new FLearningAgentsImitationTrainerEditorDetails);
+}
+
+void FLearningAgentsImitationTrainerEditorDetails::CustomizeDetails(IDetailLayoutBuilder& DetailBuilder)
+{
+	DetailBuilder.GetObjectsBeingCustomized(EditedObjects);
+	IDetailCategoryBuilder& Category = DetailBuilder.EditCategory("LearningAgents");
+
+	Category.AddCustomRow(LOCTEXT("RunTrainingCategory", "Run Training"))
+		.NameContent()
+		[
+			SNew(STextBlock)
+				.Text(LOCTEXT("RunTrainingLabel", "Run Training"))
+				.Font(DetailBuilder.GetDetailFont())
+		]
+		.ValueContent()
+		[
+			SNew(SHorizontalBox)
+				+ SHorizontalBox::Slot()
+				.AutoWidth()
+				.Padding(2.0f, 0.0f)
+				[
+					SNew(SButton)
+						.Text(LOCTEXT("RunTrainingButton", "Run"))
+						.OnClicked(FOnClicked::CreateSP(this, &FLearningAgentsImitationTrainerEditorDetails::OnRunClicked))
+						.IsEnabled_Lambda([this]() 
+						{
+							if (EditedObjects.Num() > 0)
+							{
+								if (ALearningAgentsImitationTrainerEditor* ImitationTrainerEditor = Cast<ALearningAgentsImitationTrainerEditor>(EditedObjects[0].Get()))
+								{
+									return !ImitationTrainerEditor->IsTraining();
+								}
+							}
+							return false;
+						})
+				]
+				+ SHorizontalBox::Slot()
+				.AutoWidth()
+				.Padding(2.0f, 0.0f)
+				[
+					SNew(SButton)
+						.Text(LOCTEXT("StopTrainingButton", "Stop"))
+						.OnClicked(FOnClicked::CreateSP(this, &FLearningAgentsImitationTrainerEditorDetails::OnStopClicked))
+						.IsEnabled_Lambda([this]()
+						{
+							if (EditedObjects.Num() > 0)
+							{
+								if (ALearningAgentsImitationTrainerEditor* ImitationTrainerEditor = Cast<ALearningAgentsImitationTrainerEditor>(EditedObjects[0].Get()))
+								{
+									return ImitationTrainerEditor->IsTraining();
+								}
+							}
+							return false;
+						})
+				]
+		];
+}
+
+FReply FLearningAgentsImitationTrainerEditorDetails::OnRunClicked()
+{
+	if (EditedObjects.Num() > 0)
+	{
+		if (ALearningAgentsImitationTrainerEditor* ImitationTrainerEditor = Cast<ALearningAgentsImitationTrainerEditor>(EditedObjects[0].Get()))
+		{
+			ImitationTrainerEditor->StartTraining();
+		}
+	}
+	return FReply::Handled();
+}
+
+FReply FLearningAgentsImitationTrainerEditorDetails::OnStopClicked()
+{
+	if (EditedObjects.Num() > 0)
+	{
+		if (ALearningAgentsImitationTrainerEditor* ImitationTrainerEditor = Cast<ALearningAgentsImitationTrainerEditor>(EditedObjects[0].Get()))
+		{
+			ImitationTrainerEditor->StopTraining();
+		}
+	}
+	return FReply::Handled();
+}
+
+#undef LOCTEXT_NAMESPACE
diff --git a/Engine/Plugins/Experimental/LearningAgents/Source/LearningAgentsTrainingEditor/Private/LearningAgentsTrainingEditor.cpp b/Engine/Plugins/Experimental/LearningAgents/Source/LearningAgentsTrainingEditor/Private/LearningAgentsTrainingEditor.cpp
new file mode 100644
--- /dev/null
+++ b/Engine/Plugins/Experimental/LearningAgents/Source/LearningAgentsTrainingEditor/Private/LearningAgentsTrainingEditor.cpp
@@ -0,0 +1,28 @@
+// Copyright Epic Games, Inc. All Rights Reserved.
+
+#include "LearningAgentsTrainingEditor.h"
+#include "LearningAgentsImitationTrainerEditor.h"
+#include "LearningAgentsImitationTrainerEditorDetails.h"
+#include "PropertyEditorModule.h"
+#include "PropertyEditorDelegates.h"
+#include "Misc/CoreDelegates.h"
+
+void FLearningAgentsTrainingEditorModule::StartupModule()
+{
+	FPropertyEditorModule& PropertyModule = FModuleManager::LoadModuleChecked<FPropertyEditorModule>("PropertyEditor");
+	PropertyModule.RegisterCustomClassLayout(
+		ALearningAgentsImitationTrainerEditor::StaticClass()->GetFName(),
+		FOnGetDetailCustomizationInstance::CreateStatic(&FLearningAgentsImitationTrainerEditorDetails::MakeInstance)
+	);
+}
+
+void FLearningAgentsTrainingEditorModule::ShutdownModule()
+{
+	if (FModuleManager::Get().IsModuleLoaded("PropertyEditor"))
+	{
+		FPropertyEditorModule& PropertyModule = FModuleManager::GetModuleChecked<FPropertyEditorModule>("PropertyEditor");
+		PropertyModule.UnregisterCustomClassLayout(ALearningAgentsImitationTrainerEditor::StaticClass()->GetFName());
+	}
+}
+	
+IMPLEMENT_MODULE(FLearningAgentsTrainingEditorModule, LearningAgentsTrainingEditor)
diff --git a/Engine/Plugins/Experimental/LearningAgents/Source/LearningAgentsTrainingEditor/Public/LearningAgentsImitationTrainerEditor.h b/Engine/Plugins/Experimental/LearningAgents/Source/LearningAgentsTrainingEditor/Public/LearningAgentsImitationTrainerEditor.h
new file mode 100644
--- /dev/null
+++ b/Engine/Plugins/Experimental/LearningAgents/Source/LearningAgentsTrainingEditor/Public/LearningAgentsImitationTrainerEditor.h
@@ -0,0 +1,59 @@
+// Copyright Epic Games, Inc. All Rights Reserved.
+
+#pragma once
+
+#include "CoreMinimal.h"
+#include "GameFramework/Actor.h"
+#include "LearningAgentsManager.h"
+#include "LearningAgentsImitationTrainer.h"
+
+#include "LearningAgentsImitationTrainerEditor.generated.h"
+
+#define UE_API LEARNINGAGENTSTRAININGEDITOR_API
+
+/** Editor callable imitation learning trainer. */
+UCLASS(BlueprintType, Blueprintable)
+class ALearningAgentsImitationTrainerEditor : public AActor
+{
+	GENERATED_BODY()
+
+public:
+	UE_API ALearningAgentsImitationTrainerEditor();
+
+	/** Setup the imitation trainer with necessary components in blueprints. */
+	UFUNCTION(BlueprintImplementableEvent, Category = "LearningAgents")
+	UE_API void SetupTraining();
+
+	/** Start training. */
+	UFUNCTION(BlueprintCallable, Category = "LearningAgents")
+	UE_API void StartTraining();
+
+	/** Stop training. */
+	UFUNCTION(BlueprintCallable, Category = "LearningAgents")
+	UE_API void StopTraining();
+
+	UE_API bool IsTraining() const;
+
+	UPROPERTY(EditDefaultsOnly, Category = "LearningAgents")
+	float TrainingTickInterval = 0.0f;
+	
+	UPROPERTY(VisibleAnywhere, BlueprintReadWrite, Category = "LearningAgents")
+    TObjectPtr<ULearningAgentsManager> LearningAgentsManager;
+
+	UPROPERTY(BlueprintReadWrite, Category = "LearningAgents")
+	TObjectPtr<ULearningAgentsImitationTrainer> LearningAgentsImitationTrainer;
+
+	UPROPERTY(EditDefaultsOnly, BlueprintReadWrite, Category = "LearningAgents")
+	TObjectPtr<ULearningAgentsRecording> Recording;
+
+	UPROPERTY(EditDefaultsOnly, BlueprintReadWrite, Category = "LearningAgents")
+	FLearningAgentsImitationTrainerSettings ImitationTrainerSettings;
+
+	UPROPERTY(EditDefaultsOnly, BlueprintReadWrite, Category = "LearningAgents")
+	FLearningAgentsImitationTrainerTrainingSettings ImitationTrainerTrainingSettings;
+
+	UPROPERTY(EditDefaultsOnly, BlueprintReadWrite, Category = "LearningAgents")
+	FLearningAgentsTrainerProcessSettings ImitationTrainerPathSettings;
+};
+
+#undef UE_API
diff --git a/Engine/Plugins/Experimental/LearningAgents/Source/LearningAgentsTrainingEditor/Public/LearningAgentsImitationTrainerEditorDetails.h b/Engine/Plugins/Experimental/LearningAgents/Source/LearningAgentsTrainingEditor/Public/LearningAgentsImitationTrainerEditorDetails.h
new file mode 100644
--- /dev/null
+++ b/Engine/Plugins/Experimental/LearningAgents/Source/LearningAgentsTrainingEditor/Public/LearningAgentsImitationTrainerEditorDetails.h
@@ -0,0 +1,23 @@
+// Copyright Epic Games, Inc. All Rights Reserved.
+
+#pragma once
+
+#include "IDetailCustomization.h"
+#include "Input/Reply.h"
+
+#define UE_API LEARNINGAGENTSTRAININGEDITOR_API
+
+/** Details customization for the imitation learning trainer */
+class FLearningAgentsImitationTrainerEditorDetails : public IDetailCustomization
+{
+public:
+	UE_API static TSharedRef<IDetailCustomization> MakeInstance();
+	UE_API virtual void CustomizeDetails(IDetailLayoutBuilder& DetailBuilder) override;
+
+private:
+	TArray<TWeakObjectPtr<UObject>> EditedObjects;
+	FReply OnRunClicked();
+	FReply OnStopClicked();
+};
+
+#undef UE_API
diff --git a/Engine/Plugins/Experimental/LearningAgents/Source/LearningAgentsTrainingEditor/Public/LearningAgentsTrainingEditor.h b/Engine/Plugins/Experimental/LearningAgents/Source/LearningAgentsTrainingEditor/Public/LearningAgentsTrainingEditor.h
new file mode 100644
--- /dev/null
+++ b/Engine/Plugins/Experimental/LearningAgents/Source/LearningAgentsTrainingEditor/Public/LearningAgentsTrainingEditor.h
@@ -0,0 +1,13 @@
+// Copyright Epic Games, Inc. All Rights Reserved.
+
+#pragma once
+
+#include "Modules/ModuleManager.h"
+
+class FLearningAgentsTrainingEditorModule : public IModuleInterface
+{
+public:
+	/** IModuleInterface implementation */
+	virtual void StartupModule() override;
+	virtual void ShutdownModule() override;
+};
