LearningAgents: Added some new functions to FrameRangeSet

diff --git a/Engine/Plugins/Experimental/LearningAgents/Source/Learning/Private/LearningFrameRangeSet.cpp b/Engine/Plugins/Experimental/LearningAgents/Source/Learning/Private/LearningFrameRangeSet.cpp
--- a/Engine/Plugins/Experimental/LearningAgents/Source/Learning/Private/LearningFrameRangeSet.cpp
+++ b/Engine/Plugins/Experimental/LearningAgents/Source/Learning/Private/LearningFrameRangeSet.cpp
@@ -105,6 +105,25 @@ namespace UE::Learning
 			return false;
 		}
 
+		static inline bool RangesIntersectsRange(
+			const TLearningArrayView<1, const int32> Starts,
+			const TLearningArrayView<1, const int32> Lengths,
+			const int32 Start,
+			const int32 Length)
+		{
+			check(Starts.Num() == Lengths.Num());
+
+			const int32 RangeNum = Starts.Num();
+			for (int32 RangeIdx = 0; RangeIdx < RangeNum; RangeIdx++)
+			{
+				if (Starts[RangeIdx] < Start + Length && Start < Starts[RangeIdx] + Lengths[RangeIdx])
+				{
+					return true;
+				}
+			}
+			return false;
+		}
+
 		static inline bool RangesContainsTime(
 			const TLearningArrayView<1, const int32> Starts,
 			const TLearningArrayView<1, const int32> Lengths,
@@ -1399,6 +1418,17 @@ namespace UE::Learning
 			Frame);
 	}
 
+	bool FFrameRangeSet::IntersectsRange(const int32 Sequence, const int32 Start, const int32 Length) const
+	{
+		const int32 EntryIdx = FindSequenceEntry(Sequence);
+
+		return EntryIdx != INDEX_NONE && FrameRangeSet::Private::RangesIntersectsRange(
+			GetEntryRangeStarts(EntryIdx),
+			GetEntryRangeLengths(EntryIdx),
+			Start,
+			Length);
+	}
+
 	bool FFrameRangeSet::ContainsTime(const int32 Sequence, const float Time, const float FrameDeltaTime) const
 	{
 		const int32 EntryIdx = FindSequenceEntry(Sequence);
@@ -2551,8 +2581,8 @@ namespace UE::Learning
 
 			Out.Empty();
 
-			TArray<int32> RangeStartsAdded;
-			TArray<int32> RangeLengthsAdded;
+			TArray<int32, TInlineAllocator<64>> RangeStartsAdded;
+			TArray<int32, TInlineAllocator<64>> RangeLengthsAdded;
 
 			int32 RangeOffset = 0;
 			for (int32 EntryIdx = 0; EntryIdx < EntryNum; EntryIdx++)
@@ -2599,8 +2629,8 @@ namespace UE::Learning
 
 			Out.Empty();
 
-			TArray<int32> RangeStartsAdded;
-			TArray<int32> RangeLengthsAdded;
+			TArray<int32, TInlineAllocator<64>> RangeStartsAdded;
+			TArray<int32, TInlineAllocator<64>> RangeLengthsAdded;
 
 			int32 RangeOffset = 0;
 			for (int32 EntryIdx = 0; EntryIdx < EntryNum; EntryIdx++)
@@ -2659,6 +2689,43 @@ namespace UE::Learning
 			Out.Check();
 		}
 
+		void GatherRanges(FFrameRangeSet& Out, const FFrameRangeSet& FrameRangeSet, const FIndexSet Indices)
+		{
+			FrameRangeSet.Check();
+
+			const int32 EntryNum = FrameRangeSet.GetEntryNum();
+
+			Out.Empty();
+
+			TArray<int32, TInlineAllocator<64>> RangeStartsAdded;
+			TArray<int32, TInlineAllocator<64>> RangeLengthsAdded;
+
+			int32 RangeOffset = 0;
+			for (int32 EntryIdx = 0; EntryIdx < EntryNum; EntryIdx++)
+			{
+				const int32 Sequence = FrameRangeSet.GetEntrySequence(EntryIdx);
+				const int32 RangeNum = FrameRangeSet.GetEntryRangeNum(EntryIdx);
+
+				RangeStartsAdded.Reset();
+				RangeLengthsAdded.Reset();
+
+				for (int32 RangeIdx = 0; RangeIdx < RangeNum; RangeIdx++)
+				{
+					if (Indices.Contains(RangeOffset))
+					{
+						RangeStartsAdded.Add(FrameRangeSet.GetEntryRangeStart(EntryIdx, RangeIdx));
+						RangeLengthsAdded.Add(FrameRangeSet.GetEntryRangeLength(EntryIdx, RangeIdx));
+					}
+
+					RangeOffset++;
+				}
+
+				Out.AddEntry(Sequence, RangeStartsAdded, RangeLengthsAdded);
+			}
+
+			Out.Check();
+		}
+
 		void MakeFromFrameSet(FFrameRangeSet& OutFrameRangeSet, const FFrameSet& FrameSet)
 		{
 			OutFrameRangeSet.EntrySequences = FrameSet.EntrySequences;
@@ -2700,8 +2767,8 @@ namespace UE::Learning
 			const int32 RangeSetEntryNum = FrameRangeSet.GetEntryNum();
 			const int32 FrameSetEntryNum = FrameSet.GetEntryNum();
 
-			TArray<int32> RangeStartsAdded;
-			TArray<int32> RangeLengthsAdded;
+			TArray<int32, TInlineAllocator<64>> RangeStartsAdded;
+			TArray<int32, TInlineAllocator<64>> RangeLengthsAdded;
 
 			int32 RangeOffset = 0;
 			for (int32 RangeSetEntryIdx = 0; RangeSetEntryIdx < RangeSetEntryNum; RangeSetEntryIdx++)
@@ -2746,8 +2813,8 @@ namespace UE::Learning
 			const int32 RangeSetEntryNum = FrameRangeSet.GetEntryNum();
 			const int32 FrameSetEntryNum = FrameSet.GetEntryNum();
 
-			TArray<int32> RangeStartsAdded;
-			TArray<int32> RangeLengthsAdded;
+			TArray<int32, TInlineAllocator<64>> RangeStartsAdded;
+			TArray<int32, TInlineAllocator<64>> RangeLengthsAdded;
 
 			int32 RangeOffset = 0;
 			for (int32 RangeSetEntryIdx = 0; RangeSetEntryIdx < RangeSetEntryNum; RangeSetEntryIdx++)
diff --git a/Engine/Plugins/Experimental/LearningAgents/Source/Learning/Public/LearningFrameRangeSet.h b/Engine/Plugins/Experimental/LearningAgents/Source/Learning/Public/LearningFrameRangeSet.h
--- a/Engine/Plugins/Experimental/LearningAgents/Source/Learning/Public/LearningFrameRangeSet.h
+++ b/Engine/Plugins/Experimental/LearningAgents/Source/Learning/Public/LearningFrameRangeSet.h
@@ -114,6 +114,9 @@ namespace UE::Learning
 		/** Checks if this FrameRangeSet contains a given sequence and frame in that sequence */
 		UE_API bool Contains(const int32 Sequence, const int32 Frame) const;
 
+		/** Checks if this FrameRangeSet intersects a sequence and any part of a range in that sequence */
+		UE_API bool IntersectsRange(const int32 Sequence, const int32 Start, const int32 Length) const;
+
 		/** Checks if this FrameRangeSet contains a given sequence and time in that sequence */
 		UE_API bool ContainsTime(const int32 Sequence, const float Time, const float FrameDeltaTime) const;
 
@@ -209,6 +212,9 @@ namespace UE::Learning
 		/** Creates a new frame range set with both the start and the end of all ranges padded by the given numbers of frames */
 		LEARNING_API void Pad(FFrameRangeSet& Out, const FFrameRangeSet& FrameRangeSet, const int32 PadStartFrameNum, int32 PadEndFrameNum);
 
+		/** Creates a new frame range set with just the ranges given by the provided indices */
+		LEARNING_API void GatherRanges(FFrameRangeSet& Out, const FFrameRangeSet& FrameRangeSet, const FIndexSet Indices);
+
 		/** Makes a Frame Range Set consisting of single-frame ranges from a frame set */
 		LEARNING_API void MakeFromFrameSet(FFrameRangeSet& OutFrameRangeSet, const FFrameSet& FrameSet);
 
