[NNE] [LearningAgents] Trivial - Quality of Life Structs
Improves qol in blueprint land

#

diff --git a/Engine/Plugins/Experimental/LearningAgents/Source/LearningAgents/Private/LearningAgentsObservations.cpp b/Engine/Plugins/Experimental/LearningAgents/Source/LearningAgents/Private/LearningAgentsObservations.cpp
--- a/Engine/Plugins/Experimental/LearningAgents/Source/LearningAgents/Private/LearningAgentsObservations.cpp
+++ b/Engine/Plugins/Experimental/LearningAgents/Source/LearningAgents/Private/LearningAgentsObservations.cpp
@@ -1239,7 +1239,7 @@ FLearningAgentsObservationSchemaElement ULearningAgentsObservations::SpecifyBool
 	return SpecifyContinuousObservation(Schema, 1, 1.0f, Tag);
 }
 
-FLearningAgentsObservationSchemaElement ULearningAgentsObservations::SpecifyConv1dObservation(ULearningAgentsObservationSchema* Schema, const FLearningAgentsObservationSchemaElement Element, const int32 InputLength, const int32 InChannels, const int32 OutChannels, const int32 KernelSize, const int32 Padding, const ELearningAgentsPaddingMode PaddingMode, const ELearningAgentsActivationFunction ActivationFunction, const FName Tag)
+FLearningAgentsObservationSchemaElement ULearningAgentsObservations::SpecifyConv1dObservation(ULearningAgentsObservationSchema* Schema, const FLearningAgentsObservationSchemaElement Element, const FConv1dObservationParams& Conv1dParams, const FName Tag)
 {
 	if (!Schema)
 	{
@@ -1251,52 +1251,52 @@ FLearningAgentsObservationSchemaElement ULearningAgentsObservations::SpecifyConv
 		UE_LOG(LogLearning, Error, TEXT("SpecifyConv1dObservation: Invalid Observation Object."));
 		return FLearningAgentsObservationSchemaElement();
 	}
-	if (InputLength < 1)
+	if (Conv1dParams.InputLength < 1)
 	{
-		UE_LOG(LogLearning, Error, TEXT("SpecifyConv1dObservation: Invalid Observation InputLength '%i' - must be greater than zero."), InputLength);
+		UE_LOG(LogLearning, Error, TEXT("SpecifyConv1dObservation: Invalid Observation InputLength '%i' - must be greater than zero."), Conv1dParams.InputLength);
 		return FLearningAgentsObservationSchemaElement();
 	}
-	if (InChannels < 1)
+	if (Conv1dParams.InChannels < 1)
 	{
-		UE_LOG(LogLearning, Error, TEXT("SpecifyConv1dObservation: Invalid Observation InChannels '%i' - must be greater than zero."), InChannels);
+		UE_LOG(LogLearning, Error, TEXT("SpecifyConv1dObservation: Invalid Observation InChannels '%i' - must be greater than zero."), Conv1dParams.InChannels);
 		return FLearningAgentsObservationSchemaElement();
 	}
-	if (OutChannels < 1)
+	if (Conv1dParams.OutChannels < 1)
 	{
-		UE_LOG(LogLearning, Error, TEXT("SpecifyConv1dObservation: Invalid Observation OutChannels '%i' - must be greater than zero."), OutChannels);
+		UE_LOG(LogLearning, Error, TEXT("SpecifyConv1dObservation: Invalid Observation OutChannels '%i' - must be greater than zero."), Conv1dParams.OutChannels);
 		return FLearningAgentsObservationSchemaElement();
 	}
-	if (KernelSize < 1)
+	if (Conv1dParams.KernelSize < 1)
 	{
-		UE_LOG(LogLearning, Error, TEXT("SpecifyConv1dObservation: Invalid Observation KernelSize '%i' - must be greater than zero."), KernelSize);
+		UE_LOG(LogLearning, Error, TEXT("SpecifyConv1dObservation: Invalid Observation KernelSize '%i' - must be greater than zero."), Conv1dParams.KernelSize);
 		return FLearningAgentsObservationSchemaElement();
 	}
-	if (Padding < 1)
+	if (Conv1dParams.Padding < 1)
 	{
-		UE_LOG(LogLearning, Error, TEXT("SpecifyConv1dObservation: Invalid Observation Padding '%i' - must be greater than zero."), Padding);
+		UE_LOG(LogLearning, Error, TEXT("SpecifyConv1dObservation: Invalid Observation Padding '%i' - must be greater than zero."), Conv1dParams.Padding);
 		return FLearningAgentsObservationSchemaElement();
 	}
-	if (KernelSize > InputLength + 2 * Padding)
+	if (Conv1dParams.KernelSize > Conv1dParams.InputLength + 2 * Conv1dParams.Padding)
 	{
-		UE_LOG(LogLearning, Error, TEXT("SpecifyConv1dObservation: Invalid KernelSize '%i' - Kernel size is too large."), KernelSize);
+		UE_LOG(LogLearning, Error, TEXT("SpecifyConv1dObservation: Invalid KernelSize '%i' - Kernel size is too large."), Conv1dParams.KernelSize);
 		return FLearningAgentsObservationSchemaElement();
 	}
 
 	return { Schema->ObservationSchema.CreateConv1d(
 	{
 		Element.SchemaElement,
-		InputLength,
-		InChannels,
-		OutChannels,
-		KernelSize,
-		Padding,
-		UE::Learning::Agents::Observation::Private::GetPaddingMode(PaddingMode),
-		UE::Learning::Agents::Observation::Private::GetEncodingActivationFunction(ActivationFunction)
+		Conv1dParams.InputLength,
+		Conv1dParams.InChannels,
+		Conv1dParams.OutChannels,
+		Conv1dParams.KernelSize,
+		Conv1dParams.Padding,
+		UE::Learning::Agents::Observation::Private::GetPaddingMode(Conv1dParams.PaddingMode),
+		UE::Learning::Agents::Observation::Private::GetEncodingActivationFunction(Conv1dParams.ActivationFunction)
 	},
 	Tag) };
 }
 
-FLearningAgentsObservationSchemaElement ULearningAgentsObservations::SpecifyConv2dObservation(ULearningAgentsObservationSchema* Schema, const FLearningAgentsObservationSchemaElement Element, const int32 InputHeight, const int32 InputWidth, const int32 InChannels, const int32 OutChannels, const int32 KernelSize, const int32 Stride, const int32 Padding, const ELearningAgentsActivationFunction ActivationFunction, const FName Tag)
+FLearningAgentsObservationSchemaElement ULearningAgentsObservations::SpecifyConv2dObservation(ULearningAgentsObservationSchema* Schema, const FLearningAgentsObservationSchemaElement Element, const FConv2dObservationParams& Conv2dParams, const FName Tag)
 {
 	if (!Schema)
 	{
@@ -1308,54 +1308,54 @@ FLearningAgentsObservationSchemaElement ULearningAgentsObservations::SpecifyConv
 		UE_LOG(LogLearning, Error, TEXT("SpecifyConv2dObservation: Invalid Observation Object."));
 		return FLearningAgentsObservationSchemaElement();
 	}
-	if (InputHeight < 1)
+	if (Conv2dParams.InputHeight < 1)
 	{
-		UE_LOG(LogLearning, Error, TEXT("SpecifyConv2dObservation: Invalid Observation InputHeight '%i' - must be greater than zero."), InputHeight);
+		UE_LOG(LogLearning, Error, TEXT("SpecifyConv2dObservation: Invalid Observation InputHeight '%i' - must be greater than zero."), Conv2dParams.InputHeight);
 		return FLearningAgentsObservationSchemaElement();
 	}
-	if (InputWidth < 1)
+	if (Conv2dParams.InputWidth < 1)
 	{
-		UE_LOG(LogLearning, Error, TEXT("SpecifyConv2dObservation: Invalid Observation InputWidth '%i' - must be greater than zero."), InputWidth);
+		UE_LOG(LogLearning, Error, TEXT("SpecifyConv2dObservation: Invalid Observation InputWidth '%i' - must be greater than zero."), Conv2dParams.InputWidth);
 		return FLearningAgentsObservationSchemaElement();
 	}
-	if (InChannels < 1)
+	if (Conv2dParams.InChannels < 1)
 	{
-		UE_LOG(LogLearning, Error, TEXT("SpecifyConv2dObservation: Invalid Observation InChannels '%i' - must be greater than zero."), InChannels);
+		UE_LOG(LogLearning, Error, TEXT("SpecifyConv2dObservation: Invalid Observation InChannels '%i' - must be greater than zero."), Conv2dParams.InChannels);
 		return FLearningAgentsObservationSchemaElement();
 	}
-	if (OutChannels < 1)
+	if (Conv2dParams.OutChannels < 1)
 	{
-		UE_LOG(LogLearning, Error, TEXT("SpecifyConv2dObservation: Invalid Observation OutChannels '%i' - must be greater than zero."), OutChannels);
+		UE_LOG(LogLearning, Error, TEXT("SpecifyConv2dObservation: Invalid Observation OutChannels '%i' - must be greater than zero."), Conv2dParams.OutChannels);
 		return FLearningAgentsObservationSchemaElement();
 	}
-	if (KernelSize < 1)
+	if (Conv2dParams.KernelSize < 1)
 	{
-		UE_LOG(LogLearning, Error, TEXT("SpecifyConv2dObservation: Invalid Observation KernelSize '%i' - must be greater than zero."), KernelSize);
+		UE_LOG(LogLearning, Error, TEXT("SpecifyConv2dObservation: Invalid Observation KernelSize '%i' - must be greater than zero."), Conv2dParams.KernelSize);
 		return FLearningAgentsObservationSchemaElement();
 	}
-	if (Stride < 1)
+	if (Conv2dParams.Stride < 1)
 	{
-		UE_LOG(LogLearning, Error, TEXT("SpecifyConv2dObservation: Invalid Observation Stride '%i' - must be greater than zero."), Stride);
+		UE_LOG(LogLearning, Error, TEXT("SpecifyConv2dObservation: Invalid Observation Stride '%i' - must be greater than zero."), Conv2dParams.Stride);
 		return FLearningAgentsObservationSchemaElement();
 	}
-	if (Padding < 0)
+	if (Conv2dParams.Padding < 0)
 	{
-		UE_LOG(LogLearning, Error, TEXT("SpecifyConv2dObservation: Invalid Observation Padding '%i' - must be greater than zero."), Padding);
+		UE_LOG(LogLearning, Error, TEXT("SpecifyConv2dObservation: Invalid Observation Padding '%i' - must be greater than zero."), Conv2dParams.Padding);
 		return FLearningAgentsObservationSchemaElement();
 	}
 
 	return { Schema->ObservationSchema.CreateConv2d(
 	{
 		Element.SchemaElement,
-		InputHeight,
-		InputWidth,
-		InChannels,
-		OutChannels,
-		KernelSize,
-		Stride,
-		Padding,
+		Conv2dParams.InputHeight,
+		Conv2dParams.InputWidth,
+		Conv2dParams.InChannels,
+		Conv2dParams.OutChannels,
+		Conv2dParams.KernelSize,
+		Conv2dParams.Stride,
+		Conv2dParams.Padding,
 		UE::Learning::Agents::Observation::Private::GetPaddingMode(ELearningAgentsPaddingMode::Zeros),
-		UE::Learning::Agents::Observation::Private::GetEncodingActivationFunction(ActivationFunction)
+		UE::Learning::Agents::Observation::Private::GetEncodingActivationFunction(Conv2dParams.ActivationFunction)
 	},
 	Tag) };
 }
diff --git a/Engine/Plugins/Experimental/LearningAgents/Source/LearningAgents/Public/LearningAgentsObservations.h b/Engine/Plugins/Experimental/LearningAgents/Source/LearningAgents/Public/LearningAgentsObservations.h
--- a/Engine/Plugins/Experimental/LearningAgents/Source/LearningAgents/Public/LearningAgentsObservations.h
+++ b/Engine/Plugins/Experimental/LearningAgents/Source/LearningAgents/Public/LearningAgentsObservations.h
@@ -122,6 +122,122 @@ enum class ELearningAgentsPaddingMode : uint8
 	Circular = 1	UMETA(DisplayName = "Circular"),
 };
 
+/** A wrapper of Conv1d observation parameters */
+USTRUCT(BlueprintType)
+struct FConv1dObservationParams
+{
+    GENERATED_BODY()
+
+    /** Length of the 1D input. */
+    UPROPERTY(EditAnywhere, BlueprintReadWrite, Category="LearningAgents")
+    int32 InputLength = 0;
+
+    /** Number of input channels. */
+    UPROPERTY(EditAnywhere, BlueprintReadWrite, Category="LearningAgents")
+    int32 InChannels = 0;
+
+    /** Number of output channels. */
+    UPROPERTY(EditAnywhere, BlueprintReadWrite, Category="LearningAgents")
+    int32 OutChannels = 0;
+
+    /** Size of the convolution kernel. */
+    UPROPERTY(EditAnywhere, BlueprintReadWrite, Category="LearningAgents")
+    int32 KernelSize = 0;
+
+    /** Amount of padding to add. */
+    UPROPERTY(EditAnywhere, BlueprintReadWrite, Category="LearningAgents")
+    int32 Padding = 0;
+
+    /** The padding mode. */
+    UPROPERTY(EditAnywhere, BlueprintReadWrite, Category="LearningAgents")
+    ELearningAgentsPaddingMode PaddingMode = ELearningAgentsPaddingMode::Circular;
+
+    /** Activation applied after the convolution. */
+    UPROPERTY(EditAnywhere, BlueprintReadWrite, Category="LearningAgents")
+    ELearningAgentsActivationFunction ActivationFunction = ELearningAgentsActivationFunction::ReLU;
+
+	FConv1dObservationParams() = default;
+
+	FConv1dObservationParams(
+		int32 InInputLength,
+		int32 InInChannels,
+		int32 InOutChannels,
+		int32 InKernelSize,
+		int32 InPadding,
+		ELearningAgentsPaddingMode InPaddingMode,
+		ELearningAgentsActivationFunction InActivationFunction = ELearningAgentsActivationFunction::ReLU
+	):
+		InputLength(InInputLength),
+		InChannels(InInChannels),
+		OutChannels(InOutChannels),
+		KernelSize(InKernelSize),
+		Padding(InPadding),
+		PaddingMode(InPaddingMode),
+		ActivationFunction(InActivationFunction)
+	{}
+};
+
+/** A wrapper of Conv2d observation parameters */
+USTRUCT(BlueprintType)
+struct FConv2dObservationParams
+{
+    GENERATED_BODY()
+
+    /** Height of the 2D input. */
+    UPROPERTY(EditAnywhere, BlueprintReadWrite, Category="LearningAgents")
+    int32 InputHeight = 0;
+
+    /** Width of the 2D input. */
+    UPROPERTY(EditAnywhere, BlueprintReadWrite, Category="LearningAgents")
+    int32 InputWidth = 0;
+
+    /** Number of input channels. */
+    UPROPERTY(EditAnywhere, BlueprintReadWrite, Category="LearningAgents")
+    int32 InChannels = 0;
+
+    /** Number of output channels. */
+    UPROPERTY(EditAnywhere, BlueprintReadWrite, Category="LearningAgents")
+    int32 OutChannels = 0;
+
+    /** Size of the convolution kernel. (3 -> (3x3), 5 -> (5x5) etc) */
+    UPROPERTY(EditAnywhere, BlueprintReadWrite, Category="LearningAgents")
+    int32 KernelSize = 0;
+
+    /** Stride of the convolution. */
+    UPROPERTY(EditAnywhere, BlueprintReadWrite, Category="LearningAgents")
+    int32 Stride = 0;
+
+    /** Amount of zero‑padding on each side. */
+    UPROPERTY(EditAnywhere, BlueprintReadWrite, Category="LearningAgents")
+    int32 Padding = 0;
+
+    /** Activation applied after the convolution. */
+    UPROPERTY(EditAnywhere, BlueprintReadWrite, Category="LearningAgents")
+    ELearningAgentsActivationFunction ActivationFunction = ELearningAgentsActivationFunction::ReLU;
+
+	FConv2dObservationParams() = default;
+
+	FConv2dObservationParams(
+		int32 InInputHeight,
+		int32 InInputWidth,
+		int32 InInChannels,
+		int32 InOutChannels,
+		int32 InKernelSize,
+		int32 InStride,
+		int32 InPadding,
+		ELearningAgentsActivationFunction InActivationFunction = ELearningAgentsActivationFunction::ReLU
+	):
+		InputHeight(InInputHeight),
+		InputWidth(InInputWidth),
+		InChannels(InInChannels),
+		OutChannels(InOutChannels),
+		KernelSize(InKernelSize),
+		Stride(InStride),
+		Padding(InPadding),
+		ActivationFunction(InActivationFunction)
+	{}
+};
+
 UCLASS(MinimalAPI)
 class ULearningAgentsObservations : public UBlueprintFunctionLibrary
 {
@@ -537,35 +653,24 @@ public:
 	 *
 	 * @param Schema The Observation Schema
 	 * @param Element The sub-observation to be convolved.
-	 * @param InputLength The length of the convolution input.
-	 * @param InChannels The Number of channels in the input.
-	 * @param OutChannels The Number of channels produced by the convolution.
-	 * @param KernelSize The size of the convolving kernel.
-	 * @param Padding The Amount of padding applied to the input.
-	 * @param PaddingMode 'Zeros' or 'Circular'
+	 * @param FConv1dObservationParams The convolution parameters.
 	 * @param Tag The tag of this new observation. Used during observation object validation and debugging.
 	 * @return The newly created observation schema element.
 	 */
 	UFUNCTION(BlueprintPure, Category = "LearningAgents", meta = (AdvancedDisplay = 9))
-	static UE_API FLearningAgentsObservationSchemaElement SpecifyConv1dObservation(ULearningAgentsObservationSchema* Schema, const FLearningAgentsObservationSchemaElement Element, const int32 InputLength, const int32 InChannels, const int32 OutChannels, const int32 KernelSize, const int32 Padding, const ELearningAgentsPaddingMode PaddingMode = ELearningAgentsPaddingMode::Circular, const ELearningAgentsActivationFunction ActivationFunction = ELearningAgentsActivationFunction::ReLU, const FName Tag = TEXT("Conv1dObservation"));
+	static UE_API FLearningAgentsObservationSchemaElement SpecifyConv1dObservation(ULearningAgentsObservationSchema* Schema, const FLearningAgentsObservationSchemaElement Element, const FConv1dObservationParams& Conv1dParams, const FName Tag = TEXT("Conv1dObservation"));
 
 	/**
 	 * Specifies a new 2d convolution observation of another observation (Flat buffer input, Channels x Input Height x Input Width), Zeros Padding.
 	 *
 	 * @param Schema The Observation Schema
 	 * @param Element The sub-observation to be convolved.
-	 * @param InputHeight The height of the convolution input.
-	 * @param InputWidth The width of the convolution input.
-	 * @param InChannels The Number of channels in the input.
-	 * @param OutChannels The Number of channels produced by the convolution.
-	 * @param KernelSize The size of the convolving kernel (3 -> (3x3), 5 -> (5x5) etc).
-	 * @param Stride The stride of the convolution.
-	 * @param Padding The amount of padding applied to the input.
+	 * @param FConv2dObservationParams The convolution parameters.
 	 * @param Tag The tag of this new observation. Used during observation object validation and debugging.
 	 * @return The newly created observation schema element.
 	 */
 	 UFUNCTION(BlueprintPure, Category = "LearningAgents", meta = (AdvancedDisplay = 10))
-	static UE_API FLearningAgentsObservationSchemaElement SpecifyConv2dObservation(ULearningAgentsObservationSchema* Schema, const FLearningAgentsObservationSchemaElement Element, const int32 InputHeight, const int32 InputWidth, const int32 InChannels, const int32 OutChannels, const int32 KernelSize, const int32 Stride = 1, const int32 Padding = 0, const ELearningAgentsActivationFunction ActivationFunction = ELearningAgentsActivationFunction::ReLU, const FName Tag = TEXT("Conv2dObservation"));
+	static UE_API FLearningAgentsObservationSchemaElement SpecifyConv2dObservation(ULearningAgentsObservationSchema* Schema, const FLearningAgentsObservationSchemaElement Element, const FConv2dObservationParams& Conv2dParams, const FName Tag = TEXT("Conv2dObservation"));
 
 	/**
 	 * Specifies a new float observation. A simple observation which can be used as a catch-all for situations where a
diff --git a/Engine/Plugins/Experimental/NNERuntimeBasicCpu/Source/NNERuntimeBasicCpu/Private/NNERuntimeBasicCpuModel.cpp b/Engine/Plugins/Experimental/NNERuntimeBasicCpu/Source/NNERuntimeBasicCpu/Private/NNERuntimeBasicCpuModel.cpp
--- a/Engine/Plugins/Experimental/NNERuntimeBasicCpu/Source/NNERuntimeBasicCpu/Private/NNERuntimeBasicCpuModel.cpp
+++ b/Engine/Plugins/Experimental/NNERuntimeBasicCpu/Source/NNERuntimeBasicCpu/Private/NNERuntimeBasicCpuModel.cpp
@@ -542,7 +542,6 @@ namespace UE::NNE::RuntimeBasic
 			NNE_RUNTIME_BASIC_CHECK_ALIASING_SIZES(Output, Input, ((InputHeight + 2 * Padding - KernelSize) / Stride + 1) * ((InputWidth + 2 * Padding - KernelSize) / Stride + 1) * OutChannels, InputHeight * InputWidth * InChannels);
 
 			check(PaddingMode == FModelBuilder::EPaddingMode::Zeros);
-			check((InputHeight + 2 * Padding - KernelSize) % Stride == 0);
 
 #if NNE_RUNTIME_BASIC_ENABLE_ISPC
 			ispc::NNERuntimeBasicCPUOperatorConv2d(
