[LearningAgents] DepthMapComponent

diff --git a/Engine/Plugins/Experimental/LearningAgents/Source/LearningAgents/Private/LearningAgentsDepthMapComponent.cpp b/Engine/Plugins/Experimental/LearningAgents/Source/LearningAgents/Private/LearningAgentsDepthMapComponent.cpp
new file mode 100644
--- /dev/null
+++ b/Engine/Plugins/Experimental/LearningAgents/Source/LearningAgents/Private/LearningAgentsDepthMapComponent.cpp
@@ -0,0 +1,108 @@
+// Copyright Epic Games, Inc. All Rights Reserved.
+
+#include "LearningAgentsDepthMapComponent.h"
+#include "LearningLog.h"
+#include "GameFramework/Actor.h"
+#include "CollisionQueryParams.h"
+#include "Engine/HitResult.h"
+#include "Engine/World.h"
+
+ULearningAgentsDepthMapComponent::ULearningAgentsDepthMapComponent()
+{
+	PrimaryComponentTick.bCanEverTick = true;
+}
+
+const TArray<float>& ULearningAgentsDepthMapComponent::GetDepthMapFlatArray() const
+{
+	return DepthMapFlatArray;
+}
+
+void ULearningAgentsDepthMapComponent::TickComponent(float DeltaTime, ELevelTick TickType, FActorComponentTickFunction* ThisTickFunction)
+{
+	Super::TickComponent(DeltaTime, TickType, ThisTickFunction);
+	BatchUpdateDepthMap();
+}
+
+void ULearningAgentsDepthMapComponent::BeginPlay()
+{
+	Super::BeginPlay();
+
+	const int32 Size = DepthMapConfig.Height * DepthMapConfig.Width;
+	DepthMapFlatArray.SetNumUninitialized(Size);
+	DepthMapDirections.SetNumUninitialized(Size);
+	GenerateDepthMapDirections();
+}
+
+void ULearningAgentsDepthMapComponent::GenerateDepthMapDirections()
+{
+	const float HalfHFOV = FMath::DegreesToRadians(DepthMapConfig.HorizontalFOV / 2);
+	const float HalfVFOV = HalfHFOV / DepthMapConfig.AspectRatio;
+
+	for (int32 H = 0; H < DepthMapConfig.Height; H++)
+	{
+		for (int32 W = 0; W < DepthMapConfig.Width; W++)
+		{
+			const float U = -1 * (((W + 0.5f) / DepthMapConfig.Width) * 2.0f - 1.0f);
+			const float V = -1 * (((H + 0.5f) / DepthMapConfig.Height) * 2.0f - 1.0f);
+
+			FVector LocalDirection(U * FMath::Tan(HalfHFOV), 1.0f, V * FMath::Tan(HalfVFOV));
+			LocalDirection.Normalize();
+
+			DepthMapDirections[H * DepthMapConfig.Width + W] = LocalDirection;
+		}
+	}
+}
+
+float ULearningAgentsDepthMapComponent::CastDepthRay(const FVector RayDirection)
+{
+	TObjectPtr<AActor> OwnerActor = GetOwner();
+	TObjectPtr<UWorld> World = GetWorld();
+
+	if (!OwnerActor)
+	{
+		UE_LOG(LogLearning, Error, TEXT("Cannot Cast Depth Ray due to invalid owner actor."));
+
+		return EncodeDepthValue(0.0f);
+	}
+
+	if (!World)
+	{
+		UE_LOG(LogLearning, Error, TEXT("Cannot Cast Depth Ray due to invalid world object."));
+
+		return EncodeDepthValue(0.0f);
+	}
+
+	const FTransform RayTransform = OwnerActor->GetActorTransform();
+	const FVector RayStartWorld = RayTransform.TransformPosition(FVector::ZeroVector);
+	const FVector RayEndWorld = RayTransform.TransformPosition(RayDirection * DepthMapConfig.MaxDistance);
+
+	FCollisionObjectQueryParams ObjectQueryParams;
+	ObjectQueryParams.AddObjectTypesToQuery(CollisionChannel);
+
+	FCollisionQueryParams TraceParams;
+	TraceParams.AddIgnoredActor(GetOwner());
+
+	FHitResult TraceHit;
+	const bool bHit = World->LineTraceSingleByObjectType(TraceHit, RayStartWorld, RayEndWorld, ObjectQueryParams, TraceParams);
+
+	return EncodeDepthValue(bHit ? 1.0f - TraceHit.Distance / DepthMapConfig.MaxDistance : 0.0f);
+}
+
+void ULearningAgentsDepthMapComponent::BatchUpdateDepthMap()
+{ 
+	const int32 Size = DepthMapDirections.Num();
+
+	int32 Start = BatchCastingCounter;
+	int32 End = (DepthRaysBatchSize > 0) ? FMath::Min(Start + DepthRaysBatchSize, Size): Size;
+	BatchCastingCounter = (End == Size) ? 0 : End;
+
+	for (int32 i = Start; i < End; i++)
+	{
+		DepthMapFlatArray[i] = CastDepthRay(DepthMapDirections[i]);
+	}
+}
+
+float ULearningAgentsDepthMapComponent::EncodeDepthValue(const float DepthValue)
+{
+	return DepthMapConfig.bInvertEncoding ? 1.0f - DepthValue : DepthValue;
+}
diff --git a/Engine/Plugins/Experimental/LearningAgents/Source/LearningAgents/Public/LearningAgentsDepthMapComponent.h b/Engine/Plugins/Experimental/LearningAgents/Source/LearningAgents/Public/LearningAgentsDepthMapComponent.h
new file mode 100644
--- /dev/null
+++ b/Engine/Plugins/Experimental/LearningAgents/Source/LearningAgents/Public/LearningAgentsDepthMapComponent.h
@@ -0,0 +1,77 @@
+// Copyright Epic Games, Inc. All Rights Reserved.
+
+#pragma once
+
+#include "CoreMinimal.h"
+#include "Components/ActorComponent.h"
+#include "LearningAgentsDepthMapComponent.generated.h"
+
+#define UE_API LEARNINGAGENTS_API
+
+USTRUCT(BlueprintType, Category = "LearningAgents")
+struct FDepthMapConfig
+{
+	GENERATED_BODY()
+
+	UPROPERTY(BlueprintReadWrite, EditAnywhere, Category = "LearningAgents")
+	float HorizontalFOV = 90.0f;
+
+	UPROPERTY(BlueprintReadWrite, EditAnywhere, Category = "LearningAgents")
+	float AspectRatio = 1.0f;
+
+	UPROPERTY(BlueprintReadWrite, EditAnywhere, Category = "LearningAgents")
+	int32 Width = 42; 
+
+	UPROPERTY(BlueprintReadWrite, EditAnywhere, Category = "LearningAgents")
+	int32 Height = 42; 
+
+	UPROPERTY(BlueprintReadWrite, EditAnywhere, Category = "LearningAgents")
+	float MaxDistance = 1500.0f;
+
+	UPROPERTY(BlueprintReadWrite, EditAnywhere, Category = "LearningAgents")
+	bool bInvertEncoding = false;
+};
+
+UCLASS(BlueprintType, Blueprintable, Category = "LearningAgents")
+class ULearningAgentsDepthMapComponent : public UActorComponent
+{
+	GENERATED_BODY()
+
+public:
+	UE_API ULearningAgentsDepthMapComponent();
+
+	UPROPERTY(BlueprintReadWrite, EditAnywhere, Category = "LearningAgents")
+	FDepthMapConfig DepthMapConfig;
+
+	// If specified as a positive integer, each tick will only update a batch of the depth rays in order to improve performance.
+	UPROPERTY(BlueprintReadWrite, EditAnywhere, Category = "LearningAgents")
+	int32 DepthRaysBatchSize = 0;
+
+	UPROPERTY(BlueprintReadWrite, EditAnywhere, Category = "LearningAgents")
+	TEnumAsByte<ECollisionChannel> CollisionChannel = ECollisionChannel::ECC_WorldStatic;
+
+	UFUNCTION(BlueprintCallable, Category = "LearningAgents")
+	UE_API const TArray<float>& GetDepthMapFlatArray() const;
+
+	virtual void TickComponent(float DeltaTime, ELevelTick TickType, FActorComponentTickFunction* ThisTickFunction) override;
+
+protected:
+	virtual void BeginPlay() override;
+
+private:
+	int32 BatchCastingCounter = 0;
+	
+	TArray<float> DepthMapFlatArray;
+
+	TArray<FVector> DepthMapDirections;	
+
+	void GenerateDepthMapDirections();
+
+	float CastDepthRay(const FVector RayDirection);
+	
+	void BatchUpdateDepthMap();
+
+	float EncodeDepthValue(const float DepthValue);
+};
+
+#undef UE_API
diff --git a/Engine/Plugins/Experimental/LearningAgents/Source/LearningAgentsTrainingEditor/LearningAgentsTrainingEditor.Build.cs b/Engine/Plugins/Experimental/LearningAgents/Source/LearningAgentsTrainingEditor/LearningAgentsTrainingEditor.Build.cs
--- a/Engine/Plugins/Experimental/LearningAgents/Source/LearningAgentsTrainingEditor/LearningAgentsTrainingEditor.Build.cs
+++ b/Engine/Plugins/Experimental/LearningAgents/Source/LearningAgentsTrainingEditor/LearningAgentsTrainingEditor.Build.cs
@@ -35,7 +35,8 @@ public class LearningAgentsTrainingEditor : ModuleRules
 				"SlateCore",
 				"EditorStyle",
 				"PropertyEditor",
-				"UnrealEd"
+				"UnrealEd",
+				"UMG"
 			});
 
 		DynamicallyLoadedModuleNames.AddRange(
diff --git a/Engine/Plugins/Experimental/LearningAgents/Source/LearningAgentsTrainingEditor/Private/LearningAgentsDepthMapVisualizer.cpp b/Engine/Plugins/Experimental/LearningAgents/Source/LearningAgentsTrainingEditor/Private/LearningAgentsDepthMapVisualizer.cpp
new file mode 100644
--- /dev/null
+++ b/Engine/Plugins/Experimental/LearningAgents/Source/LearningAgentsTrainingEditor/Private/LearningAgentsDepthMapVisualizer.cpp
@@ -0,0 +1,93 @@
+// Copyright Epic Games, Inc. All Rights Reserved.
+
+#include "LearningAgentsDepthMapVisualizer.h"
+#include "Engine/Canvas.h"
+#include "Blueprint/WidgetTree.h"
+#include "Components/CanvasPanel.h"
+#include "Components/Image.h"
+#include "Components/CanvasPanelSlot.h"
+#include "LearningLog.h"
+
+ULearningAgentsDepthMapVisualizerComponent::ULearningAgentsDepthMapVisualizerComponent()
+{
+	PrimaryComponentTick.bCanEverTick = true;
+	MarkAsEditorOnlySubobject();
+}
+
+void ULearningAgentsDepthMapVisualizerComponent::TickComponent(float DeltaTime, ELevelTick TickType, FActorComponentTickFunction* ThisTickFunction)
+{
+	Super::TickComponent(DeltaTime, TickType, ThisTickFunction);
+	
+	if (!DepthMapComp)
+	{
+		UE_LOG(LogLearning, Error, TEXT("Unable to find a corresponding LearningAgentsDepthMapComponent! Failed to render texture."));
+		return;
+	}
+
+	Widget->SetDepthValues(DepthMapComp->GetDepthMapFlatArray());
+}
+
+void ULearningAgentsDepthMapVisualizerComponent::BeginPlay()
+{
+	Super::BeginPlay();
+
+	DepthMapComp = GetOwner()->FindComponentByClass<ULearningAgentsDepthMapComponent>();
+	if (!DepthMapComp)
+	{
+		UE_LOG(LogLearning, Error, TEXT("Unable to find a corresponding LearningAgentsDepthMapComponent! Failed to render texture."));
+		return;
+	}
+
+	SetComponentTickInterval(DepthMapComp->PrimaryComponentTick.TickInterval);
+
+	Widget = CreateWidget<ULearningAgentsDepthMapWidget>(GetWorld(), ULearningAgentsDepthMapWidget::StaticClass());
+	Widget->InitializeView(RenderSize, RenderPosition, DepthMapComp->DepthMapConfig.Height, DepthMapComp->DepthMapConfig.Width);
+	Widget->AddToViewport();
+}
+
+void ULearningAgentsDepthMapWidget::SetDepthValues(const TArray<float>& InDepthValues)
+{
+	DepthValues = InDepthValues;
+
+	if (DepthMapRenderer)
+	{
+		DepthMapRenderer->UpdateResource();
+	}
+}
+
+void ULearningAgentsDepthMapWidget::InitializeView(FVector2D RenderSize, FVector2D RenderPosition, int32 DMapHeight, int32 DMapWidth)
+{
+	DepthValues.SetNumUninitialized(DMapWidth * DMapHeight);
+
+	DepthMapRenderer = UCanvasRenderTarget2D::CreateCanvasRenderTarget2D(GetWorld(), UCanvasRenderTarget2D::StaticClass(), DMapWidth, DMapHeight);
+	DepthMapRenderer->OnCanvasRenderTargetUpdate.AddDynamic(this, &ULearningAgentsDepthMapWidget::OnCanvasRenderTargetUpdate);
+	DepthMapRenderer->UpdateResource();
+
+	UCanvasPanel* Root = WidgetTree->ConstructWidget<UCanvasPanel>(UCanvasPanel::StaticClass());
+	WidgetTree->RootWidget = Root;
+
+	UImage* Image = WidgetTree->ConstructWidget<UImage>(UImage::StaticClass());
+	Root->AddChildToCanvas(Image);
+
+	UCanvasPanelSlot* CanvasSlot = Cast<UCanvasPanelSlot>(Image->Slot);
+	CanvasSlot->SetSize(RenderSize);
+	CanvasSlot->SetPosition(RenderPosition);
+
+	FSlateBrush Brush;
+	Brush.SetResourceObject(DepthMapRenderer);
+	Image->SetBrush(Brush);
+	Image->InvalidateLayoutAndVolatility();
+}
+
+void ULearningAgentsDepthMapWidget::OnCanvasRenderTargetUpdate(UCanvas* Canvas, int32 Width, int32 Height)
+{
+	for (int32 H = 0; H < Height; H++)
+	{
+		for (int32 W = 0; W < Width; W++)
+		{
+			float DepthValue = DepthValues[H * Width + W];
+			FLinearColor DepthColor(DepthValue, DepthValue, DepthValue, 1.0f);
+			Canvas->K2_DrawBox(FVector2D(W, H), FVector2D(1.0f, 1.0f), 1.0f, DepthColor);
+		}
+	}
+}
diff --git a/Engine/Plugins/Experimental/LearningAgents/Source/LearningAgentsTrainingEditor/Public/LearningAgentsDepthMapVisualizer.h b/Engine/Plugins/Experimental/LearningAgents/Source/LearningAgentsTrainingEditor/Public/LearningAgentsDepthMapVisualizer.h
new file mode 100644
--- /dev/null
+++ b/Engine/Plugins/Experimental/LearningAgents/Source/LearningAgentsTrainingEditor/Public/LearningAgentsDepthMapVisualizer.h
@@ -0,0 +1,55 @@
+// Copyright Epic Games, Inc. All Rights Reserved.
+
+#pragma once
+
+#include "CoreMinimal.h"
+#include "Blueprint/UserWidget.h"
+#include "Engine/CanvasRenderTarget2D.h"
+#include "LearningAgentsDepthMapComponent.h"
+#include "LearningAgentsDepthMapVisualizer.generated.h"
+
+#define UE_API LEARNINGAGENTSTRAININGEDITOR_API
+
+UCLASS()
+class ULearningAgentsDepthMapWidget : public UUserWidget
+{
+    GENERATED_BODY()
+
+public:
+	UE_API void SetDepthValues(const TArray<float>& InDepthValues);
+	UE_API void InitializeView(FVector2D RenderSize, FVector2D RenderPosition, int32 DMapHeight, int32 DMapWidth);
+
+private:
+	TArray<float> DepthValues;
+
+	TObjectPtr<UCanvasRenderTarget2D> DepthMapRenderer;
+
+	UFUNCTION()
+	void OnCanvasRenderTargetUpdate(UCanvas* Canvas, int32 Width, int32 Height);
+};
+
+UCLASS(BlueprintType, Blueprintable, Category = "LearningAgents")
+class ULearningAgentsDepthMapVisualizerComponent : public UActorComponent
+{
+	GENERATED_BODY()
+
+public:
+	ULearningAgentsDepthMapVisualizerComponent();
+
+	UPROPERTY(BlueprintReadWrite, EditAnywhere, Category = "LearningAgents")
+	FVector2D RenderSize = FVector2D(512, 512);
+
+	UPROPERTY(BlueprintReadWrite, EditAnywhere, Category = "LearningAgents")
+	FVector2D RenderPosition = FVector2D(100, 100);
+
+	virtual void TickComponent(float DeltaTime, ELevelTick TickType, FActorComponentTickFunction* ThisTickFunction) override;
+
+protected:
+	virtual void BeginPlay() override;
+
+private:
+	TObjectPtr<ULearningAgentsDepthMapComponent> DepthMapComp;
+	TObjectPtr<ULearningAgentsDepthMapWidget> Widget;
+};
+
+#undef UE_API
