[LearningAgents] DepthMap Perf

diff --git a/Engine/Plugins/Experimental/LearningAgents/Source/LearningAgents/Private/LearningAgentsDepthMapComponent.cpp b/Engine/Plugins/Experimental/LearningAgents/Source/LearningAgents/Private/LearningAgentsDepthMapComponent.cpp
--- a/Engine/Plugins/Experimental/LearningAgents/Source/LearningAgents/Private/LearningAgentsDepthMapComponent.cpp
+++ b/Engine/Plugins/Experimental/LearningAgents/Source/LearningAgents/Private/LearningAgentsDepthMapComponent.cpp
@@ -3,9 +3,9 @@
 #include "LearningAgentsDepthMapComponent.h"
 #include "LearningLog.h"
 #include "GameFramework/Actor.h"
-#include "CollisionQueryParams.h"
 #include "Engine/HitResult.h"
 #include "Engine/World.h"
+#include "DrawDebugHelpers.h"
 
 ULearningAgentsDepthMapComponent::ULearningAgentsDepthMapComponent()
 {
@@ -21,6 +21,11 @@ void ULearningAgentsDepthMapComponent::TickComponent(float DeltaTime, ELevelTick
 {
 	Super::TickComponent(DeltaTime, TickType, ThisTickFunction);
 	BatchUpdateDepthMap();
+
+	if (bDrawDepthFrustum)
+	{
+		DebugDrawCornerRays();
+	}
 }
 
 void ULearningAgentsDepthMapComponent::BeginPlay()
@@ -28,8 +33,12 @@ void ULearningAgentsDepthMapComponent::BeginPlay()
 	Super::BeginPlay();
 
 	const int32 Size = DepthMapConfig.Height * DepthMapConfig.Width;
-	DepthMapFlatArray.SetNumUninitialized(Size);
-	DepthMapDirections.SetNumUninitialized(Size);
+	DepthMapFlatArray.SetNumZeroed(Size);
+	DepthMapDirections.SetNumZeroed(Size);
+
+	CollisionObjectQueryParams.AddObjectTypesToQuery(CollisionChannel);
+	CollisionQueryParams.AddIgnoredActor(GetOwner());
+
 	GenerateDepthMapDirections();
 }
 
@@ -53,7 +62,7 @@ void ULearningAgentsDepthMapComponent::GenerateDepthMapDirections()
 	}
 }
 
-float ULearningAgentsDepthMapComponent::CastDepthRay(const FVector RayDirection)
+float ULearningAgentsDepthMapComponent::CastDepthRay(const FVector& RayStartWorld, const FVector& RayEndWorld) const
 {
 	TObjectPtr<AActor> OwnerActor = GetOwner();
 	TObjectPtr<UWorld> World = GetWorld();
@@ -72,37 +81,84 @@ float ULearningAgentsDepthMapComponent::CastDepthRay(const FVector RayDirection)
 		return EncodeDepthValue(0.0f);
 	}
 
-	const FTransform RayTransform = OwnerActor->GetActorTransform();
-	const FVector RayStartWorld = RayTransform.TransformPosition(FVector::ZeroVector);
-	const FVector RayEndWorld = RayTransform.TransformPosition(RayDirection * DepthMapConfig.MaxDistance);
-
-	FCollisionObjectQueryParams ObjectQueryParams;
-	ObjectQueryParams.AddObjectTypesToQuery(CollisionChannel);
-
-	FCollisionQueryParams TraceParams;
-	TraceParams.AddIgnoredActor(GetOwner());
-
 	FHitResult TraceHit;
-	const bool bHit = World->LineTraceSingleByObjectType(TraceHit, RayStartWorld, RayEndWorld, ObjectQueryParams, TraceParams);
+	const bool bHit = World->LineTraceSingleByObjectType(TraceHit, RayStartWorld, RayEndWorld, CollisionObjectQueryParams, CollisionQueryParams);
 
 	return EncodeDepthValue(bHit ? 1.0f - TraceHit.Distance / DepthMapConfig.MaxDistance : 0.0f);
 }
 
 void ULearningAgentsDepthMapComponent::BatchUpdateDepthMap()
 { 
+	TObjectPtr<AActor> OwnerActor = GetOwner();
+	TObjectPtr<UWorld> World = GetWorld();
+
+	if (!OwnerActor)
+	{
+		UE_LOG(LogLearning, Error, TEXT("Cannot Cast Depth Ray due to invalid owner actor."));
+		return;
+	}
+	if (!World)
+	{
+		UE_LOG(LogLearning, Error, TEXT("Cannot Cast Depth Ray due to invalid world object."));
+		return;
+	}
+
 	const int32 Size = DepthMapDirections.Num();
 
 	int32 Start = BatchCastingCounter;
 	int32 End = (DepthRaysBatchSize > 0) ? FMath::Min(Start + DepthRaysBatchSize, Size): Size;
 	BatchCastingCounter = (End == Size) ? 0 : End;
 
+	const FTransform RayTransform = OwnerActor->GetActorTransform();
+	const FVector RayStartWorld = RayTransform.TransformPosition(FVector::ZeroVector) + DepthMapConfig.FrustumOffset;
+
 	for (int32 i = Start; i < End; i++)
 	{
-		DepthMapFlatArray[i] = CastDepthRay(DepthMapDirections[i]);
+		const FVector RayEndWorld = RayTransform.TransformPosition(DepthMapDirections[i] * DepthMapConfig.MaxDistance) + DepthMapConfig.FrustumOffset;
+		DepthMapFlatArray[i] = CastDepthRay(RayStartWorld, RayEndWorld);
 	}
 }
 
-float ULearningAgentsDepthMapComponent::EncodeDepthValue(const float DepthValue)
+float ULearningAgentsDepthMapComponent::EncodeDepthValue(const float DepthValue) const
 {
 	return DepthMapConfig.bInvertEncoding ? 1.0f - DepthValue : DepthValue;
 }
+
+void ULearningAgentsDepthMapComponent::DebugDrawCornerRays()
+{
+	AActor* Owner = GetOwner();
+	UWorld* World = GetWorld();
+	if (!Owner || !World) return;
+
+	const int32 W = DepthMapConfig.Width;
+	const int32 H = DepthMapConfig.Height;
+	const FVector Origin = Owner->GetActorLocation() + DepthMapConfig.FrustumOffset;
+
+	// Four corners of the depth map
+	const int32 TopLeft = 0;
+	const int32 TopRight = W - 1;
+	const int32 BottomLeft = (H - 1) * W;
+	const int32 BottomRight = H * W - 1;
+	const int32 Center = H / 2 * W + W / 2;
+
+	const TArray<int32> Indices = { TopLeft, TopRight, BottomLeft, BottomRight, Center };
+
+	const FTransform RayTransform = Owner->GetActorTransform();
+	const FVector RayStartWorld = RayTransform.TransformPosition(FVector::ZeroVector) + DepthMapConfig.FrustumOffset;
+
+	for (int32 Idx : Indices)
+	{
+		const FVector RayEndWorld = RayTransform.TransformPosition(DepthMapDirections[Idx] * DepthMapConfig.MaxDistance) + DepthMapConfig.FrustumOffset;
+
+		DrawDebugLine(
+			World,
+			RayStartWorld,
+			RayEndWorld,
+			FColor::Yellow,
+			false,
+			PrimaryComponentTick.TickInterval,
+			0, 
+			1.0f
+		);
+	}
+}
diff --git a/Engine/Plugins/Experimental/LearningAgents/Source/LearningAgents/Public/LearningAgentsDepthMapComponent.h b/Engine/Plugins/Experimental/LearningAgents/Source/LearningAgents/Public/LearningAgentsDepthMapComponent.h
--- a/Engine/Plugins/Experimental/LearningAgents/Source/LearningAgents/Public/LearningAgentsDepthMapComponent.h
+++ b/Engine/Plugins/Experimental/LearningAgents/Source/LearningAgents/Public/LearningAgentsDepthMapComponent.h
@@ -4,6 +4,7 @@
 
 #include "CoreMinimal.h"
 #include "Components/ActorComponent.h"
+#include "CollisionQueryParams.h"
 #include "LearningAgentsDepthMapComponent.generated.h"
 
 #define UE_API LEARNINGAGENTS_API
@@ -28,11 +29,14 @@ struct FDepthMapConfig
 	UPROPERTY(BlueprintReadWrite, EditAnywhere, Category = "LearningAgents")
 	float MaxDistance = 1500.0f;
 
+	UPROPERTY(BlueprintReadWrite, EditAnywhere, Category = "LearningAgents")
+	FVector FrustumOffset = FVector::ZeroVector;
+
 	UPROPERTY(BlueprintReadWrite, EditAnywhere, Category = "LearningAgents")
 	bool bInvertEncoding = false;
 };
 
-UCLASS(BlueprintType, Blueprintable, Category = "LearningAgents")
+UCLASS(BlueprintType, Blueprintable, ClassGroup = (LearningAgents), meta = (BlueprintSpawnableComponent), config=Engine)
 class ULearningAgentsDepthMapComponent : public UActorComponent
 {
 	GENERATED_BODY()
@@ -40,14 +44,17 @@ class ULearningAgentsDepthMapComponent : public UActorComponent
 public:
 	UE_API ULearningAgentsDepthMapComponent();
 
-	UPROPERTY(BlueprintReadWrite, EditAnywhere, Category = "LearningAgents")
+	UPROPERTY(BlueprintReadWrite, EditDefaultsOnly, Category = "LearningAgents")
+	bool bDrawDepthFrustum = false;
+
+	UPROPERTY(Config, BlueprintReadWrite, EditDefaultsOnly, Category = "LearningAgents")
 	FDepthMapConfig DepthMapConfig;
 
 	// If specified as a positive integer, each tick will only update a batch of the depth rays in order to improve performance.
-	UPROPERTY(BlueprintReadWrite, EditAnywhere, Category = "LearningAgents")
+	UPROPERTY(Config, BlueprintReadWrite, EditDefaultsOnly, Category = "LearningAgents")
 	int32 DepthRaysBatchSize = 0;
 
-	UPROPERTY(BlueprintReadWrite, EditAnywhere, Category = "LearningAgents")
+	UPROPERTY(Config, BlueprintReadWrite, EditDefaultsOnly, Category = "LearningAgents")
 	TEnumAsByte<ECollisionChannel> CollisionChannel = ECollisionChannel::ECC_WorldStatic;
 
 	UFUNCTION(BlueprintCallable, Category = "LearningAgents")
@@ -59,6 +66,12 @@ protected:
 	virtual void BeginPlay() override;
 
 private:
+	// Set as class variable for perf
+	FCollisionObjectQueryParams CollisionObjectQueryParams;
+
+	// Set as class variable for perf
+	FCollisionQueryParams CollisionQueryParams;
+
 	int32 BatchCastingCounter = 0;
 	
 	TArray<float> DepthMapFlatArray;
@@ -67,11 +80,13 @@ private:
 
 	void GenerateDepthMapDirections();
 
-	float CastDepthRay(const FVector RayDirection);
+	float CastDepthRay(const FVector& RayStartWorld, const FVector& RayEndWorld) const;
 	
 	void BatchUpdateDepthMap();
 
-	float EncodeDepthValue(const float DepthValue);
+	float EncodeDepthValue(const float DepthValue) const;
+
+	void DebugDrawCornerRays();
 };
 
 #undef UE_API
diff --git a/Engine/Plugins/Experimental/LearningAgents/Source/LearningAgentsTrainingEditor/Public/LearningAgentsDepthMapVisualizer.h b/Engine/Plugins/Experimental/LearningAgents/Source/LearningAgentsTrainingEditor/Public/LearningAgentsDepthMapVisualizer.h
--- a/Engine/Plugins/Experimental/LearningAgents/Source/LearningAgentsTrainingEditor/Public/LearningAgentsDepthMapVisualizer.h
+++ b/Engine/Plugins/Experimental/LearningAgents/Source/LearningAgentsTrainingEditor/Public/LearningAgentsDepthMapVisualizer.h
@@ -28,7 +28,7 @@ private:
 	void OnCanvasRenderTargetUpdate(UCanvas* Canvas, int32 Width, int32 Height);
 };
 
-UCLASS(BlueprintType, Blueprintable, Category = "LearningAgents")
+UCLASS(Blueprintable, ClassGroup = (LearningAgents), meta = (BlueprintSpawnableComponent))
 class ULearningAgentsDepthMapVisualizerComponent : public UActorComponent
 {
 	GENERATED_BODY()
