[LearningAgents] [Sandbox]

#virtualized

diff --git a/Engine/Plugins/Experimental/LearningAgents/Source/LearningAgents/Public/LearningAgentsDepthMapComponent.h b/Engine/Plugins/Experimental/LearningAgents/Source/LearningAgents/Public/LearningAgentsDepthMapComponent.h
--- a/Engine/Plugins/Experimental/LearningAgents/Source/LearningAgents/Public/LearningAgentsDepthMapComponent.h
+++ b/Engine/Plugins/Experimental/LearningAgents/Source/LearningAgents/Public/LearningAgentsDepthMapComponent.h
@@ -44,7 +44,7 @@ class ULearningAgentsDepthMapComponent : public UActorComponent
 public:
 	UE_API ULearningAgentsDepthMapComponent();
 
-	UPROPERTY(BlueprintReadWrite, EditDefaultsOnly, Category = "LearningAgents")
+	UPROPERTY(Config, BlueprintReadWrite, EditDefaultsOnly, Category = "LearningAgents")
 	bool bDrawDepthFrustum = false;
 
 	UPROPERTY(Config, BlueprintReadWrite, EditDefaultsOnly, Category = "LearningAgents")
