[Learning Agents] Wrapped learning agents Python code into a learning_agents module. Added auto-registering of NNERuntimeBasicCPU and learning_agents Python code, thus avoiding relative imports in train.py. In UE Python interpreter learning_agents can be accessed directly by `import learning_agents`.

#rnx

diff --git a/Engine/Plugins/Experimental/LearningAgents/Source/LearningAgentsTraining/Public/LearningAgentsTrainer.h b/Engine/Plugins/Experimental/LearningAgents/Source/LearningAgentsTraining/Public/LearningAgentsTrainer.h
--- a/Engine/Plugins/Experimental/LearningAgents/Source/LearningAgentsTraining/Public/LearningAgentsTrainer.h
+++ b/Engine/Plugins/Experimental/LearningAgents/Source/LearningAgentsTraining/Public/LearningAgentsTrainer.h
@@ -156,7 +156,7 @@ public:
 
 	/** Trainer file name. The name of the python file to use for training. Do NOT include the '.py' file extension. */
 	UPROPERTY(EditAnywhere, Category = "LearningAgents")
-	FString TrainerFileName = TEXT("train_ppo");
+	FString TrainerFileName = TEXT("learning_agents.train_ppo");
 
 public:
 
diff --git a/Engine/Plugins/Experimental/LearningAgents/Source/LearningTraining/Private/LearningPPOTrainer.cpp b/Engine/Plugins/Experimental/LearningAgents/Source/LearningTraining/Private/LearningPPOTrainer.cpp
--- a/Engine/Plugins/Experimental/LearningAgents/Source/LearningTraining/Private/LearningPPOTrainer.cpp
+++ b/Engine/Plugins/Experimental/LearningAgents/Source/LearningTraining/Private/LearningPPOTrainer.cpp
@@ -86,7 +86,7 @@ int32 ULearningSocketPPOTrainerServerCommandlet::Main(const FString& Commandline
 	UE::Learning::FSocketTrainerServerProcess ServerProcess(
 		TEXT("Training"),
 		UE::Learning::Trainer::GetProjectPythonContentPath(),
-		TEXT("train_ppo"),
+		TEXT("learning_agents.train_ppo"),
 		PythonExecutiblePath,
 		PythonContentPath,
 		IntermediatePath,
diff --git a/Engine/Plugins/Experimental/LearningAgents/Source/LearningTraining/Private/LearningTraining.cpp b/Engine/Plugins/Experimental/LearningAgents/Source/LearningTraining/Private/LearningTraining.cpp
--- a/Engine/Plugins/Experimental/LearningAgents/Source/LearningTraining/Private/LearningTraining.cpp
+++ b/Engine/Plugins/Experimental/LearningAgents/Source/LearningTraining/Private/LearningTraining.cpp
@@ -1,10 +1,13 @@
 // Copyright Epic Games, Inc. All Rights Reserved.
 
 #include "LearningTraining.h"
+#include "Misc/FileHelper.h"
+#include "Misc/Paths.h"
 
 void FLearningTrainingModule::StartupModule()
 {
 	// This code will execute after your module is loaded into memory; the exact timing is specified in the .uplugin file per-module
+	RegisterPythonPackages();
 }
 
 void FLearningTrainingModule::ShutdownModule()
@@ -12,5 +15,18 @@ void FLearningTrainingModule::ShutdownModule()
 	// This function may be called during shutdown to clean up your module.  For modules that support dynamic reloading,
 	// we call this function before unloading the module.
 }
+
+void FLearningTrainingModule::RegisterPythonPackages()
+{
+	const FString ExperimentalPluginsPath = FPaths::Combine(FPaths::EnginePluginsDir(), TEXT("Experimental"));
+	const FString SitePackagesPath = FPaths::ConvertRelativePathToFull(FPaths::ProjectIntermediateDir() / TEXT("PipInstall") / TEXT("Lib") / TEXT("site-packages"));
+	const FString PyPluginsSitePackageFile = FPaths::ConvertRelativePathToFull(SitePackagesPath / TEXT("learning_agents.pth"));
+	
+	TArray<FString> PluginSitePackagePaths;
+	PluginSitePackagePaths.Add(FPaths::ConvertRelativePathToFull(FPaths::Combine(ExperimentalPluginsPath, TEXT("LearningAgents/Content/Python"))));
+	PluginSitePackagePaths.Add(FPaths::ConvertRelativePathToFull(FPaths::Combine(ExperimentalPluginsPath, TEXT("NNERuntimeBasicCpu/Content/Python"))));
 	
-IMPLEMENT_MODULE(FLearningTrainingModule, LearningTraining)
\ No newline at end of file
+	FFileHelper::SaveStringArrayToFile(PluginSitePackagePaths, *PyPluginsSitePackageFile, FFileHelper::EEncodingOptions::ForceUTF8WithoutBOM);
+}
+
+IMPLEMENT_MODULE(FLearningTrainingModule, LearningTraining)
diff --git a/Engine/Plugins/Experimental/LearningAgents/Source/LearningTraining/Public/LearningTraining.h b/Engine/Plugins/Experimental/LearningAgents/Source/LearningTraining/Public/LearningTraining.h
--- a/Engine/Plugins/Experimental/LearningAgents/Source/LearningTraining/Public/LearningTraining.h
+++ b/Engine/Plugins/Experimental/LearningAgents/Source/LearningTraining/Public/LearningTraining.h
@@ -11,6 +11,9 @@ public:
 	/** IModuleInterface implementation */
 	virtual void StartupModule() override;
 	virtual void ShutdownModule() override;
+
+	/** Register any internal Python dependencies to make them available in Python sys.path */
+	virtual void RegisterPythonPackages();
 };
 
 
