[LearningAgents] IterationsPerSnapshot

diff --git a/Engine/Plugins/Experimental/LearningAgents/Source/LearningAgentsTraining/Private/LearningAgentsPPOTrainer.cpp b/Engine/Plugins/Experimental/LearningAgents/Source/LearningAgentsTraining/Private/LearningAgentsPPOTrainer.cpp
--- a/Engine/Plugins/Experimental/LearningAgents/Source/LearningAgentsTraining/Private/LearningAgentsPPOTrainer.cpp
+++ b/Engine/Plugins/Experimental/LearningAgents/Source/LearningAgentsTraining/Private/LearningAgentsPPOTrainer.cpp
@@ -56,6 +56,7 @@ TSharedRef<FJsonObject> FLearningAgentsPPOTrainingSettings::AsJsonConfig() const
 	ConfigObject->SetStringField(TEXT("Device"), UE::Learning::Trainer::GetDeviceString(UE::Learning::Agents::GetTrainingDevice(Device)));
 	ConfigObject->SetBoolField(TEXT("UseTensorBoard"), bUseTensorboard);
 	ConfigObject->SetBoolField(TEXT("SaveSnapshots"), bSaveSnapshots);
+	ConfigObject->SetNumberField(TEXT("IterationsPerSnapshot"), IterationsPerSnapshot);
 	ConfigObject->SetBoolField(TEXT("UseMLflow"), bUseMLflow);
 	ConfigObject->SetStringField(TEXT("MLflowTrackingUri"), MLflowTrackingUri);
 
diff --git a/Engine/Plugins/Experimental/LearningAgents/Source/LearningAgentsTraining/Public/LearningAgentsPPOTrainer.h b/Engine/Plugins/Experimental/LearningAgents/Source/LearningAgentsTraining/Public/LearningAgentsPPOTrainer.h
--- a/Engine/Plugins/Experimental/LearningAgents/Source/LearningAgentsTraining/Public/LearningAgentsPPOTrainer.h
+++ b/Engine/Plugins/Experimental/LearningAgents/Source/LearningAgentsTraining/Public/LearningAgentsPPOTrainer.h
@@ -244,6 +244,10 @@ public:
 	UPROPERTY(EditAnywhere, Category = "LearningAgents")
 	bool bSaveSnapshots = false;
 
+	/** The iterations interval to save new networks snapshot. */
+	UPROPERTY(EditAnywhere, Category = "LearningAgents", meta = (ClampMin = "0", UIMin = "0"))
+	int32 IterationsPerSnapshot = 1000;
+
 	/** If true, MLflow will be used for experiment tracking. */
 	UPROPERTY(EditAnywhere, Category = "LearningAgents")
 	bool bUseMLflow = false;
