Learning: Fixed bug in nearest frame search for FrameSets

diff --git a/Engine/Plugins/Experimental/LearningAgents/Source/Learning/Private/LearningFrameSet.cpp b/Engine/Plugins/Experimental/LearningAgents/Source/Learning/Private/LearningFrameSet.cpp
--- a/Engine/Plugins/Experimental/LearningAgents/Source/Learning/Private/LearningFrameSet.cpp
+++ b/Engine/Plugins/Experimental/LearningAgents/Source/Learning/Private/LearningFrameSet.cpp
@@ -300,7 +300,7 @@ namespace UE::Learning
 			{
 				const int32 FrameDifference = GetEntryFrame(EntryIdx, FrameIdx) - Frame;
 
-				if (OutFrameDifference == INT32_MAX || FMath::Abs(FrameDifference) < OutFrameDifference)
+				if (OutFrameDifference == INT32_MAX || FMath::Abs(FrameDifference) < FMath::Abs(OutFrameDifference))
 				{
 					OutFrameDifference = FrameDifference;
 					OutEntryIdx = EntryIdx;
@@ -337,7 +337,7 @@ namespace UE::Learning
 				{
 					const int32 FrameDifference = EntryFrame - Frame;
 
-					if (OutFrameDifference == INT32_MAX || FMath::Abs(FrameDifference) < OutFrameDifference)
+					if (OutFrameDifference == INT32_MAX || FMath::Abs(FrameDifference) < FMath::Abs(OutFrameDifference))
 					{
 						OutFrameDifference = FrameDifference;
 						OutEntryIdx = EntryIdx;
