[LearningAgents] Fixes missing case of conv2d schema check in LearningTrainer

diff --git a/Engine/Plugins/Experimental/LearningAgents/Source/LearningTraining/Private/LearningTrainer.cpp b/Engine/Plugins/Experimental/LearningAgents/Source/LearningTraining/Private/LearningTrainer.cpp
--- a/Engine/Plugins/Experimental/LearningAgents/Source/LearningTraining/Private/LearningTrainer.cpp
+++ b/Engine/Plugins/Experimental/LearningAgents/Source/LearningTraining/Private/LearningTrainer.cpp
@@ -719,6 +719,18 @@ namespace UE::Learning::Trainer
 			return IsObservationSchemaSubsetCompatible(Object->GetObjectField(TEXT("Element")).ToSharedRef(), ObservationSchema, Parameters.Element);
 		}
 
+		case Observation::EType::Conv2d:
+		{
+			const Observation::FSchemaConv2dParameters Parameters = ObservationSchema.GetConv2d(ObservationSchemaElement);
+
+			if (Object->GetStringField(TEXT("Type")) != TEXT("Conv2d"))
+			{
+				return false;
+			}
+
+			return IsObservationSchemaSubsetCompatible(Object->GetObjectField(TEXT("Element")).ToSharedRef(), ObservationSchema, Parameters.Element);
+		}
+
 		default:
 			checkNoEntry();
 		}
