LearningAgents: Added functionality for clamping observation values.

diff --git a/Engine/Plugins/Experimental/LearningAgents/Source/Learning/Private/LearningObservation.cpp b/Engine/Plugins/Experimental/LearningAgents/Source/Learning/Private/LearningObservation.cpp
--- a/Engine/Plugins/Experimental/LearningAgents/Source/Learning/Private/LearningObservation.cpp
+++ b/Engine/Plugins/Experimental/LearningAgents/Source/Learning/Private/LearningObservation.cpp
@@ -55,6 +55,16 @@ namespace UE::Learning::Observation
 			return Size;
 		}
 
+		static inline int32 GetTotalObservationDistributionVectorSize(const FSchema& Schema, const TArrayView<const FSchemaElement> Elements)
+		{
+			int32 Size = 0;
+			for (const FSchemaElement SubElement : Elements)
+			{
+				Size += Schema.GetObservationDistributionVectorSize(SubElement);
+			}
+			return Size;
+		}
+
 		static inline bool CheckAllValid(const FObject& Object, const TArrayView<const FObjectElement> Elements)
 		{
 			for (const FObjectElement SubElement : Elements)
@@ -71,6 +81,7 @@ namespace UE::Learning::Observation
 		Tags.Add(Tag);
 		ObservationVectorSizes.Add(0);
 		EncodedVectorSizes.Add(0);
+		ObservationDistributionVectorSizes.Add(0);
 		TypeDataIndices.Add(INDEX_NONE);
 
 		return { Index, Generation };
@@ -89,6 +100,7 @@ namespace UE::Learning::Observation
 		Tags.Add(Tag);
 		ObservationVectorSizes.Add(Parameters.Num);
 		EncodedVectorSizes.Add(Parameters.Num);
+		ObservationDistributionVectorSizes.Add(Parameters.Num);
 		TypeDataIndices.Add(ContinuousData.Add(ElementData));
 
 		return { Index, Generation };
@@ -103,6 +115,7 @@ namespace UE::Learning::Observation
 		Tags.Add(Tag);
 		ObservationVectorSizes.Add(Parameters.Num);
 		EncodedVectorSizes.Add(Parameters.Num);
+		ObservationDistributionVectorSizes.Add(Parameters.Num);
 		TypeDataIndices.Add(DiscreteExclusiveData.Add(ElementData));
 
 		return { Index, Generation };
@@ -117,6 +130,7 @@ namespace UE::Learning::Observation
 		Tags.Add(Tag);
 		ObservationVectorSizes.Add(Parameters.Num);
 		EncodedVectorSizes.Add(Parameters.Num);
+		ObservationDistributionVectorSizes.Add(Parameters.Num);
 		TypeDataIndices.Add(DiscreteInclusiveData.Add(ElementData));
 
 		return { Index, Generation };
@@ -137,6 +151,7 @@ namespace UE::Learning::Observation
 		Tags.Add(Tag);
 		ObservationVectorSizes.Add(Parameters.ElementNames.Num());
 		EncodedVectorSizes.Add(Parameters.ElementNames.Num());
+		ObservationDistributionVectorSizes.Add(Parameters.ElementNames.Num());
 		TypeDataIndices.Add(NamedDiscreteExclusiveData.Add(ElementData));
 
 		return { Index, Generation };
@@ -157,6 +172,7 @@ namespace UE::Learning::Observation
 		Tags.Add(Tag);
 		ObservationVectorSizes.Add(Parameters.ElementNames.Num());
 		EncodedVectorSizes.Add(Parameters.ElementNames.Num());
+		ObservationDistributionVectorSizes.Add(Parameters.ElementNames.Num());
 		TypeDataIndices.Add(NamedDiscreteInclusiveData.Add(ElementData));
 
 		return { Index, Generation };
@@ -179,6 +195,7 @@ namespace UE::Learning::Observation
 		Tags.Add(Tag);
 		ObservationVectorSizes.Add(Private::GetTotalObservationVectorSize(*this, Parameters.Elements));
 		EncodedVectorSizes.Add(Private::GetTotalEncodedObservationVectorSize(*this, Parameters.Elements));
+		ObservationDistributionVectorSizes.Add(Private::GetTotalObservationDistributionVectorSize(*this, Parameters.Elements));
 		TypeDataIndices.Add(AndData.Add(ElementData));
 
 		return { Index, Generation };
@@ -202,6 +219,7 @@ namespace UE::Learning::Observation
 		Tags.Add(Tag);
 		ObservationVectorSizes.Add(Private::GetMaxObservationVectorSize(*this, Parameters.Elements) + Parameters.Elements.Num());
 		EncodedVectorSizes.Add(Parameters.EncodingSize + Parameters.Elements.Num());
+		ObservationDistributionVectorSizes.Add(Private::GetTotalObservationDistributionVectorSize(*this, Parameters.Elements) + Parameters.Elements.Num());
 		TypeDataIndices.Add(OrExclusiveData.Add(ElementData));
 
 		return { Index, Generation };
@@ -227,6 +245,7 @@ namespace UE::Learning::Observation
 		Tags.Add(Tag);
 		ObservationVectorSizes.Add(Private::GetTotalObservationVectorSize(*this, Parameters.Elements) + Parameters.Elements.Num());
 		EncodedVectorSizes.Add(Parameters.AttentionHeadNum * Parameters.ValueEncodingSize + Parameters.Elements.Num());
+		ObservationDistributionVectorSizes.Add(Private::GetTotalObservationDistributionVectorSize(*this, Parameters.Elements) + Parameters.Elements.Num());
 		TypeDataIndices.Add(OrInclusiveData.Add(ElementData));
 
 		return { Index, Generation };
@@ -248,6 +267,7 @@ namespace UE::Learning::Observation
 		Tags.Add(Tag);
 		ObservationVectorSizes.Add(GetObservationVectorSize(Parameters.Element) * Parameters.Num);
 		EncodedVectorSizes.Add(GetEncodedVectorSize(Parameters.Element) * Parameters.Num);
+		ObservationDistributionVectorSizes.Add(GetObservationDistributionVectorSize(Parameters.Element) * Parameters.Num);
 		TypeDataIndices.Add(ArrayData.Add(ElementData));
 
 		return { Index, Generation };
@@ -271,6 +291,7 @@ namespace UE::Learning::Observation
 		Tags.Add(Tag);
 		ObservationVectorSizes.Add(GetObservationVectorSize(Parameters.Element) * Parameters.MaxNum + Parameters.MaxNum);
 		EncodedVectorSizes.Add(Parameters.ValueEncodingSize * Parameters.AttentionHeadNum + 1);
+		ObservationDistributionVectorSizes.Add(GetObservationDistributionVectorSize(Parameters.Element));
 		TypeDataIndices.Add(SetData.Add(ElementData));
 
 		return { Index, Generation };
@@ -293,6 +314,7 @@ namespace UE::Learning::Observation
 		Tags.Add(Tag);
 		ObservationVectorSizes.Add(GetObservationVectorSize(Parameters.Element));
 		EncodedVectorSizes.Add(Parameters.EncodingSize);
+		ObservationDistributionVectorSizes.Add(GetObservationDistributionVectorSize(Parameters.Element));
 		TypeDataIndices.Add(EncodingData.Add(ElementData));
 
 		return { Index, Generation };
@@ -317,8 +339,9 @@ namespace UE::Learning::Observation
 		Tags.Add(Tag);
 		
 		ObservationVectorSizes.Add(GetObservationVectorSize(Parameters.Element));
-		int32 OutputLength = Parameters.InputLength + 2 * Parameters.Padding - Parameters.KernelSize + 1;
+		const int32 OutputLength = Parameters.InputLength + 2 * Parameters.Padding - Parameters.KernelSize + 1;
 		EncodedVectorSizes.Add(OutputLength * Parameters.OutChannels);
+		ObservationDistributionVectorSizes.Add(GetObservationDistributionVectorSize(Parameters.Element));
 
 		TypeDataIndices.Add(Conv1dData.Add(ElementData));
 
@@ -346,9 +369,10 @@ namespace UE::Learning::Observation
 		Tags.Add(Tag);
 
 		ObservationVectorSizes.Add(GetObservationVectorSize(Parameters.Element));
-		int32 OutputHeight = (Parameters.InputHeight + 2 * Parameters.Padding - Parameters.KernelSize) / Parameters.Stride + 1;
-		int32 OutputWidth = (Parameters.InputWidth + 2 * Parameters.Padding - Parameters.KernelSize) / Parameters.Stride + 1;
+		const int32 OutputHeight = (Parameters.InputHeight + 2 * Parameters.Padding - Parameters.KernelSize) / Parameters.Stride + 1;
+		const int32 OutputWidth = (Parameters.InputWidth + 2 * Parameters.Padding - Parameters.KernelSize) / Parameters.Stride + 1;
 		EncodedVectorSizes.Add(OutputHeight * OutputWidth * Parameters.OutChannels);
+		ObservationDistributionVectorSizes.Add(GetObservationDistributionVectorSize(Parameters.Element));
 
 		TypeDataIndices.Add(Conv2dData.Add(ElementData));
 
@@ -384,6 +408,12 @@ namespace UE::Learning::Observation
 		return EncodedVectorSizes[Element.Index];
 	}
 
+	int32 FSchema::GetObservationDistributionVectorSize(const FSchemaElement Element) const
+	{
+		check(IsValid(Element));
+		return ObservationDistributionVectorSizes[Element.Index];
+	}
+
 	FSchemaContinuousParameters FSchema::GetContinuous(const FSchemaElement Element) const
 	{
 		check(IsValid(Element) && GetType(Element) == EType::Continuous);
@@ -557,6 +587,7 @@ namespace UE::Learning::Observation
 		Tags.Empty();
 		ObservationVectorSizes.Empty();
 		EncodedVectorSizes.Empty();
+		ObservationDistributionVectorSizes.Empty();
 		TypeDataIndices.Empty();
 
 		ContinuousData.Empty();
@@ -588,6 +619,7 @@ namespace UE::Learning::Observation
 		Tags.Reset();
 		ObservationVectorSizes.Reset();
 		EncodedVectorSizes.Reset();
+		ObservationDistributionVectorSizes.Reset();
 		TypeDataIndices.Reset();
 
 		ContinuousData.Reset();
@@ -2524,20 +2556,21 @@ namespace UE::Learning::Observation
 		}
 	}
 
-	void AddGaussianNoiseToVector(
-		uint32& InOutRandomState,
-		TLearningArrayView<1, float> InOutObservationVector,
+	void GetContinuousActionVectorMask(
+		TLearningArrayView<1, bool> OutObservationVectorMask,
+		const TLearningArrayView<1, const float> InObservationVector,
 		const FSchema& Schema,
-		const FSchemaElement SchemaElement,
-		const float NoiseScale)
+		const FSchemaElement SchemaElement)
 	{
 		check(Schema.IsValid(SchemaElement));
 
 		const EType SchemaElementType = Schema.GetType(SchemaElement);
 
-		const int32 ObservationVectorSize = InOutObservationVector.Num();
+		const int32 ObservationVectorSize = OutObservationVectorMask.Num();
 		check(ObservationVectorSize == Schema.GetObservationVectorSize(SchemaElement));
 
+		Array::Zero(OutObservationVectorMask);
+
 		switch (SchemaElementType)
 		{
 
@@ -2548,17 +2581,8 @@ namespace UE::Learning::Observation
 
 		case EType::Continuous:
 		{
-			const FSchemaContinuousParameters SchemaParameters = Schema.GetContinuous(SchemaElement);
-			check(ObservationVectorSize == SchemaParameters.Num);
-
-			TLearningArray<1, float, TInlineAllocator<32>> NoiseValues;
-			NoiseValues.SetNumUninitialized({ SchemaParameters.Num });
-			Random::SampleGaussianArray(NoiseValues, InOutRandomState, 0.0f, NoiseScale);
-			for (int32 ValueIdx = 0; ValueIdx < SchemaParameters.Num; ValueIdx++)
-			{
-				InOutObservationVector[ValueIdx] += NoiseValues[ValueIdx];
-			}
-
+			check(ObservationVectorSize == Schema.GetContinuous(SchemaElement).Num);
+			Array::Set(OutObservationVectorMask, true);
 			return;
 		}
 
@@ -2591,12 +2615,11 @@ namespace UE::Learning::Observation
 			{
 				const int32 SubElementSize = Schema.GetObservationVectorSize(Parameters.Elements[SchemaElementIdx]);
 
-				AddGaussianNoiseToVector(
-					InOutRandomState,
-					InOutObservationVector.Slice(SubElementOffset, SubElementSize),
+				GetContinuousActionVectorMask(
+					OutObservationVectorMask.Slice(SubElementOffset, SubElementSize),
+					InObservationVector.Slice(SubElementOffset, SubElementSize),
 					Schema,
-					Parameters.Elements[SchemaElementIdx],
-					NoiseScale);
+					Parameters.Elements[SchemaElementIdx]);
 
 				SubElementOffset += SubElementSize;
 			}
@@ -2613,8 +2636,8 @@ namespace UE::Learning::Observation
 			int32 SchemaElementIndex = INDEX_NONE;
 			for (int32 SubElementIdx = 0; SubElementIdx < Parameters.Elements.Num(); SubElementIdx++)
 			{
-				check(InOutObservationVector[MaxSubElementSize + SubElementIdx] == 0.0f || InOutObservationVector[MaxSubElementSize + SubElementIdx] == 1.0f);
-				if (InOutObservationVector[MaxSubElementSize + SubElementIdx])
+				check(InObservationVector[MaxSubElementSize + SubElementIdx] == 0.0f || InObservationVector[MaxSubElementSize + SubElementIdx] == 1.0f);
+				if (InObservationVector[MaxSubElementSize + SubElementIdx])
 				{
 					SchemaElementIndex = SubElementIdx;
 					break;
@@ -2624,12 +2647,11 @@ namespace UE::Learning::Observation
 
 			const int32 SubElementSize = Schema.GetObservationVectorSize(Parameters.Elements[SchemaElementIndex]);
 
-			AddGaussianNoiseToVector(
-				InOutRandomState,
-				InOutObservationVector.Slice(0, SubElementSize),
+			GetContinuousActionVectorMask(
+				OutObservationVectorMask.Slice(0, SubElementSize),
+				InObservationVector.Slice(0, SubElementSize),
 				Schema,
-				Parameters.Elements[SchemaElementIndex],
-				NoiseScale);
+				Parameters.Elements[SchemaElementIndex]);
 
 			return;
 		}
@@ -2645,22 +2667,471 @@ namespace UE::Learning::Observation
 				const int32 SubElementSize = Schema.GetObservationVectorSize(Parameters.Elements[SubElementIdx]);
 
 				check(
-					InOutObservationVector[TotalSubElementSize + SubElementIdx] == 0.0f ||
-					InOutObservationVector[TotalSubElementSize + SubElementIdx] == 1.0f);
+					InObservationVector[TotalSubElementSize + SubElementIdx] == 0.0f ||
+					InObservationVector[TotalSubElementSize + SubElementIdx] == 1.0f);
+
+				if (InObservationVector[TotalSubElementSize + SubElementIdx] == 1.0f)
+				{
+					GetContinuousActionVectorMask(
+						OutObservationVectorMask.Slice(SubElementOffset, SubElementSize),
+						InObservationVector.Slice(SubElementOffset, SubElementSize),
+						Schema,
+						Parameters.Elements[SubElementIdx]);
+				}
+
+				SubElementOffset += SubElementSize;
+			}
+			check(SubElementOffset + Parameters.Elements.Num() == ObservationVectorSize);
+
+			return;
+		}
+
+		case EType::Array:
+		{
+			const FSchemaArrayParameters Parameters = Schema.GetArray(SchemaElement);
+			const int32 SubElementSize = Schema.GetObservationVectorSize(Parameters.Element);
+
+			for (int32 ElementIdx = 0; ElementIdx < Parameters.Num; ElementIdx++)
+			{
+				GetContinuousActionVectorMask(
+					OutObservationVectorMask.Slice(ElementIdx * SubElementSize, SubElementSize),
+					InObservationVector.Slice(ElementIdx * SubElementSize, SubElementSize),
+					Schema,
+					Parameters.Element);
+			}
+
+			return;
+		}
+
+		case EType::Set:
+		{
+			const FSchemaSetParameters Parameters = Schema.GetSet(SchemaElement);
+			const int32 SubElementSize = Schema.GetObservationVectorSize(Parameters.Element);
+
+			for (int32 SubElementIdx = 0; SubElementIdx < Parameters.MaxNum; SubElementIdx++)
+			{
+				check(
+					InObservationVector[SubElementSize * Parameters.MaxNum + SubElementIdx] == 0.0f ||
+					InObservationVector[SubElementSize * Parameters.MaxNum + SubElementIdx] == 1.0f);
+
+				if (InObservationVector[SubElementSize * Parameters.MaxNum + SubElementIdx] == 0.0f)
+				{
+					break;
+				}
+
+				GetContinuousActionVectorMask(
+					OutObservationVectorMask.Slice(SubElementIdx* SubElementSize, SubElementSize),
+					InObservationVector.Slice(SubElementIdx* SubElementSize, SubElementSize),
+					Schema,
+					Parameters.Element);
+			}
+
+			return;
+		}
+
+		case EType::Encoding:
+		{
+			GetContinuousActionVectorMask(
+				OutObservationVectorMask,
+				InObservationVector,
+				Schema,
+				Schema.GetEncoding(SchemaElement).Element);
+
+			return;
+		}
+
+		case EType::Conv1d:
+		{
+			GetContinuousActionVectorMask(
+				OutObservationVectorMask,
+				InObservationVector,
+				Schema,
+				Schema.GetConv1d(SchemaElement).Element);
+
+			return;
+		}
+
+		case EType::Conv2d:
+		{
+			GetContinuousActionVectorMask(
+				OutObservationVectorMask,
+				InObservationVector,
+				Schema,
+				Schema.GetConv2d(SchemaElement).Element);
+
+			return;
+		}
+
+		default:
+		{
+			checkNoEntry();
+			return;
+		}
+		}
+	}
+
+	void InitObservationDistributionMinMax(
+		TLearningArrayView<1, float> OutObservationDistributionMin,
+		TLearningArrayView<1, float> OutObservationDistributionMax)
+	{
+		Array::Set(OutObservationDistributionMin, +UE_MAX_FLT);
+		Array::Set(OutObservationDistributionMax, -UE_MAX_FLT);
+	}
+
+	void FitObservationDistributionMinMax(
+		TLearningArrayView<1, float> InOutObservationDistributionMin,
+		TLearningArrayView<1, float> InOutObservationDistributionMax,
+		const TLearningArrayView<1, const float> InObservationVector,
+		const FSchema& Schema,
+		const FSchemaElement SchemaElement)
+	{
+		check(Schema.IsValid(SchemaElement));
+
+		const EType SchemaElementType = Schema.GetType(SchemaElement);
+
+		const int32 ObservationVectorSize = InObservationVector.Num();
+		const int32 ObservationDistributionVectorSize = InOutObservationDistributionMin.Num();
+		check(ObservationVectorSize == Schema.GetObservationVectorSize(SchemaElement));
+		check(ObservationDistributionVectorSize == Schema.GetObservationDistributionVectorSize(SchemaElement));
+
+		switch (SchemaElementType)
+		{
+
+		case EType::Null:
+		case EType::DiscreteExclusive:
+		case EType::DiscreteInclusive:
+		case EType::NamedDiscreteExclusive:
+		case EType::NamedDiscreteInclusive:
+		{
+			return;
+		}
+
+		case EType::Continuous:
+		{
+			check(ObservationVectorSize == Schema.GetContinuous(SchemaElement).Num);
+			check(ObservationVectorSize == ObservationDistributionVectorSize);
+			for (int32 Idx = 0; Idx < ObservationVectorSize; Idx++)
+			{
+				InOutObservationDistributionMin[Idx] = FMath::Min(InOutObservationDistributionMin[Idx], InObservationVector[Idx]);
+				InOutObservationDistributionMax[Idx] = FMath::Max(InOutObservationDistributionMax[Idx], InObservationVector[Idx]);
+			}
+			return;
+		}
+
+		case EType::And:
+		{
+			const FSchemaAndParameters Parameters = Schema.GetAnd(SchemaElement);
+
+			int32 SubElementOffset = 0;
+			int32 SubElementDistributionOffset = 0;
+			for (int32 SchemaElementIdx = 0; SchemaElementIdx < Parameters.Elements.Num(); SchemaElementIdx++)
+			{
+				const int32 SubElementSize = Schema.GetObservationVectorSize(Parameters.Elements[SchemaElementIdx]);
+				const int32 SubElementDistributionSize = Schema.GetObservationDistributionVectorSize(Parameters.Elements[SchemaElementIdx]);
+
+				FitObservationDistributionMinMax(
+					InOutObservationDistributionMin.Slice(SubElementDistributionOffset, SubElementDistributionSize),
+					InOutObservationDistributionMax.Slice(SubElementDistributionOffset, SubElementDistributionSize),
+					InObservationVector.Slice(SubElementOffset, SubElementSize),
+					Schema,
+					Parameters.Elements[SchemaElementIdx]);
+
+				SubElementOffset += SubElementSize;
+				SubElementDistributionOffset += SubElementDistributionSize;
+			}
+			check(SubElementOffset == ObservationVectorSize);
+			check(SubElementDistributionOffset == ObservationDistributionVectorSize);
+
+			return;
+		}
+
+		case EType::OrExclusive:
+		{
+			const FSchemaOrExclusiveParameters Parameters = Schema.GetOrExclusive(SchemaElement);
+			const int32 MaxSubElementSize = Private::GetMaxObservationVectorSize(Schema, Parameters.Elements);
+
+			int32 SubElementDistributionOffset = 0;
+			for (int32 SchemaElementIdx = 0; SchemaElementIdx < Parameters.Elements.Num(); SchemaElementIdx++)
+			{
+				const int32 SubElementDistributionSize = Schema.GetObservationDistributionVectorSize(Parameters.Elements[SchemaElementIdx]);
+
+				check(InObservationVector[MaxSubElementSize + SchemaElementIdx] == 0.0f || InObservationVector[MaxSubElementSize + SchemaElementIdx] == 1.0f);
+				if (InObservationVector[MaxSubElementSize + SchemaElementIdx])
+				{
+					const int32 SubElementSize = Schema.GetObservationVectorSize(Parameters.Elements[SchemaElementIdx]);
+
+					FitObservationDistributionMinMax(
+						InOutObservationDistributionMin.Slice(SubElementDistributionOffset, SubElementDistributionSize),
+						InOutObservationDistributionMax.Slice(SubElementDistributionOffset, SubElementDistributionSize),
+						InObservationVector.Slice(0, SubElementSize),
+						Schema,
+						Parameters.Elements[SchemaElementIdx]);
+
+					break;
+				}
+
+				SubElementDistributionOffset += SubElementDistributionSize;
+			}
+
+			return;
+		}
+
+		case EType::OrInclusive:
+		{
+			const FSchemaOrInclusiveParameters Parameters = Schema.GetOrInclusive(SchemaElement);
+			const int32 TotalSubElementSize = Private::GetTotalObservationVectorSize(Schema, Parameters.Elements);
+
+			int32 SubElementOffset = 0;
+			int32 SubElementDistributionOffset = 0;
+			for (int32 SchemaElementIdx = 0; SchemaElementIdx < Parameters.Elements.Num(); SchemaElementIdx++)
+			{
+				const int32 SubElementSize = Schema.GetObservationVectorSize(Parameters.Elements[SchemaElementIdx]);
+				const int32 SubElementDistributionSize = Schema.GetObservationDistributionVectorSize(Parameters.Elements[SchemaElementIdx]);
+
+				check(
+					InObservationVector[TotalSubElementSize + SchemaElementIdx] == 0.0f ||
+					InObservationVector[TotalSubElementSize + SchemaElementIdx] == 1.0f);
+
+				if (InObservationVector[TotalSubElementSize + SchemaElementIdx] == 1.0f)
+				{
+					FitObservationDistributionMinMax(
+						InOutObservationDistributionMin.Slice(SubElementDistributionOffset, SubElementDistributionSize),
+						InOutObservationDistributionMax.Slice(SubElementDistributionOffset, SubElementDistributionSize),
+						InObservationVector.Slice(SubElementOffset, SubElementSize),
+						Schema,
+						Parameters.Elements[SchemaElementIdx]);
+				}
+
+				SubElementOffset += SubElementSize;
+				SubElementDistributionOffset += SubElementDistributionSize;
+			}
+			check(SubElementOffset + Parameters.Elements.Num() == ObservationVectorSize);
+			check(SubElementDistributionOffset + Parameters.Elements.Num() == ObservationDistributionVectorSize);
+
+			return;
+		}
+
+		case EType::Array:
+		{
+			const FSchemaArrayParameters Parameters = Schema.GetArray(SchemaElement);
+			const int32 SubElementSize = Schema.GetObservationVectorSize(Parameters.Element);
+			const int32 SubElementDistributionSize = Schema.GetObservationDistributionVectorSize(Parameters.Element);
+
+			for (int32 ElementIdx = 0; ElementIdx < Parameters.Num; ElementIdx++)
+			{
+				FitObservationDistributionMinMax(
+					InOutObservationDistributionMin.Slice(ElementIdx * SubElementDistributionSize, SubElementDistributionSize),
+					InOutObservationDistributionMax.Slice(ElementIdx * SubElementDistributionSize, SubElementDistributionSize),
+					InObservationVector.Slice(ElementIdx * SubElementSize, SubElementSize),
+					Schema,
+					Parameters.Element);
+			}
+
+			return;
+		}
+
+		case EType::Set:
+		{
+			const FSchemaSetParameters Parameters = Schema.GetSet(SchemaElement);
+			const int32 SubElementSize = Schema.GetObservationVectorSize(Parameters.Element);
+
+			for (int32 SubElementIdx = 0; SubElementIdx < Parameters.MaxNum; SubElementIdx++)
+			{
+				check(
+					InObservationVector[SubElementSize * Parameters.MaxNum + SubElementIdx] == 0.0f ||
+					InObservationVector[SubElementSize * Parameters.MaxNum + SubElementIdx] == 1.0f);
 
-				if (InOutObservationVector[TotalSubElementSize + SubElementIdx] == 1.0f)
+				if (InObservationVector[SubElementSize * Parameters.MaxNum + SubElementIdx] == 0.0f)
 				{
-					AddGaussianNoiseToVector(
-						InOutRandomState,
+					break;
+				}
+
+				FitObservationDistributionMinMax(
+					InOutObservationDistributionMin,
+					InOutObservationDistributionMax,
+					InObservationVector.Slice(SubElementIdx * SubElementSize, SubElementSize),
+					Schema,
+					Parameters.Element);
+			}
+
+			return;
+		}
+
+		case EType::Encoding:
+		{
+			FitObservationDistributionMinMax(
+				InOutObservationDistributionMin,
+				InOutObservationDistributionMax,
+				InObservationVector,
+				Schema,
+				Schema.GetEncoding(SchemaElement).Element);
+
+			return;
+		}
+
+		case EType::Conv1d:
+		{
+			FitObservationDistributionMinMax(
+				InOutObservationDistributionMin,
+				InOutObservationDistributionMax,
+				InObservationVector,
+				Schema,
+				Schema.GetConv1d(SchemaElement).Element);
+
+			return;
+		}
+
+		case EType::Conv2d:
+		{
+			FitObservationDistributionMinMax(
+				InOutObservationDistributionMin,
+				InOutObservationDistributionMax,
+				InObservationVector,
+				Schema,
+				Schema.GetConv2d(SchemaElement).Element);
+
+			return;
+		}
+
+		default:
+		{
+			checkNoEntry();
+			return;
+		}
+		}
+	}
+
+	void ClampToObservationDistributionMinMax(
+		TLearningArrayView<1, float> InOutObservationVector,
+		const TLearningArrayView<1, const float> ObservationDistributionMin,
+		const TLearningArrayView<1, const float> ObservationDistributionMax,
+		const FSchema& Schema,
+		const FSchemaElement SchemaElement)
+	{
+		check(Schema.IsValid(SchemaElement));
+
+		const EType SchemaElementType = Schema.GetType(SchemaElement);
+
+		const int32 ObservationVectorSize = InOutObservationVector.Num();
+		const int32 ObservationDistributionVectorSize = ObservationDistributionMin.Num();
+		check(ObservationVectorSize == Schema.GetObservationVectorSize(SchemaElement));
+		check(ObservationDistributionVectorSize == Schema.GetObservationDistributionVectorSize(SchemaElement));
+
+		switch (SchemaElementType)
+		{
+
+		case EType::Null:
+		case EType::DiscreteExclusive:
+		case EType::DiscreteInclusive:
+		case EType::NamedDiscreteExclusive:
+		case EType::NamedDiscreteInclusive:
+		{
+			return;
+		}
+
+		case EType::Continuous:
+		{
+			check(ObservationVectorSize == Schema.GetContinuous(SchemaElement).Num);
+			check(ObservationVectorSize == ObservationDistributionVectorSize);
+			for (int32 Idx = 0; Idx < ObservationVectorSize; Idx++)
+			{
+				InOutObservationVector[Idx] = FMath::Clamp(
+					InOutObservationVector[Idx], 
+					ObservationDistributionMin[Idx], 
+					ObservationDistributionMax[Idx]);
+			}
+			return;
+		}
+
+		case EType::And:
+		{
+			const FSchemaAndParameters Parameters = Schema.GetAnd(SchemaElement);
+
+			int32 SubElementOffset = 0;
+			int32 SubElementDistributionOffset = 0;
+			for (int32 SchemaElementIdx = 0; SchemaElementIdx < Parameters.Elements.Num(); SchemaElementIdx++)
+			{
+				const int32 SubElementSize = Schema.GetObservationVectorSize(Parameters.Elements[SchemaElementIdx]);
+				const int32 SubElementDistributionSize = Schema.GetObservationDistributionVectorSize(Parameters.Elements[SchemaElementIdx]);
+
+				ClampToObservationDistributionMinMax(
+					InOutObservationVector.Slice(SubElementOffset, SubElementSize),
+					ObservationDistributionMin.Slice(SubElementDistributionOffset, SubElementDistributionSize),
+					ObservationDistributionMax.Slice(SubElementDistributionOffset, SubElementDistributionSize),
+					Schema,
+					Parameters.Elements[SchemaElementIdx]);
+
+				SubElementOffset += SubElementSize;
+				SubElementDistributionOffset += SubElementDistributionSize;
+			}
+			check(SubElementOffset == ObservationVectorSize);
+			check(SubElementDistributionOffset == ObservationDistributionVectorSize);
+
+			return;
+		}
+
+		case EType::OrExclusive:
+		{
+			const FSchemaOrExclusiveParameters Parameters = Schema.GetOrExclusive(SchemaElement);
+			const int32 MaxSubElementSize = Private::GetMaxObservationVectorSize(Schema, Parameters.Elements);
+
+			int32 SubElementDistributionOffset = 0;
+			for (int32 SchemaElementIdx = 0; SchemaElementIdx < Parameters.Elements.Num(); SchemaElementIdx++)
+			{
+				const int32 SubElementDistributionSize = Schema.GetObservationDistributionVectorSize(Parameters.Elements[SchemaElementIdx]);
+
+				check(InOutObservationVector[MaxSubElementSize + SchemaElementIdx] == 0.0f || InOutObservationVector[MaxSubElementSize + SchemaElementIdx] == 1.0f);
+				if (InOutObservationVector[MaxSubElementSize + SchemaElementIdx])
+				{
+					const int32 SubElementSize = Schema.GetObservationVectorSize(Parameters.Elements[SchemaElementIdx]);
+
+					ClampToObservationDistributionMinMax(
+						InOutObservationVector.Slice(0, SubElementSize),
+						ObservationDistributionMin.Slice(SubElementDistributionOffset, SubElementDistributionSize),
+						ObservationDistributionMax.Slice(SubElementDistributionOffset, SubElementDistributionSize),
+						Schema,
+						Parameters.Elements[SchemaElementIdx]);
+
+					break;
+				}
+
+				SubElementDistributionOffset += SubElementDistributionSize;
+			}
+
+			return;
+		}
+
+		case EType::OrInclusive:
+		{
+			const FSchemaOrInclusiveParameters Parameters = Schema.GetOrInclusive(SchemaElement);
+			const int32 TotalSubElementSize = Private::GetTotalObservationVectorSize(Schema, Parameters.Elements);
+
+			int32 SubElementOffset = 0;
+			int32 SubElementDistributionOffset = 0;
+			for (int32 SchemaElementIdx = 0; SchemaElementIdx < Parameters.Elements.Num(); SchemaElementIdx++)
+			{
+				const int32 SubElementSize = Schema.GetObservationVectorSize(Parameters.Elements[SchemaElementIdx]);
+				const int32 SubElementDistributionSize = Schema.GetObservationDistributionVectorSize(Parameters.Elements[SchemaElementIdx]);
+
+				check(
+					InOutObservationVector[TotalSubElementSize + SchemaElementIdx] == 0.0f ||
+					InOutObservationVector[TotalSubElementSize + SchemaElementIdx] == 1.0f);
+
+				if (InOutObservationVector[TotalSubElementSize + SchemaElementIdx] == 1.0f)
+				{
+					ClampToObservationDistributionMinMax(
 						InOutObservationVector.Slice(SubElementOffset, SubElementSize),
+						ObservationDistributionMin.Slice(SubElementDistributionOffset, SubElementDistributionSize),
+						ObservationDistributionMax.Slice(SubElementDistributionOffset, SubElementDistributionSize),
 						Schema,
-						Parameters.Elements[SubElementIdx],
-						NoiseScale);
+						Parameters.Elements[SchemaElementIdx]);
 				}
 
 				SubElementOffset += SubElementSize;
+				SubElementDistributionOffset += SubElementDistributionSize;
 			}
 			check(SubElementOffset + Parameters.Elements.Num() == ObservationVectorSize);
+			check(SubElementDistributionOffset + Parameters.Elements.Num() == ObservationDistributionVectorSize);
 
 			return;
 		}
@@ -2669,15 +3140,16 @@ namespace UE::Learning::Observation
 		{
 			const FSchemaArrayParameters Parameters = Schema.GetArray(SchemaElement);
 			const int32 SubElementSize = Schema.GetObservationVectorSize(Parameters.Element);
+			const int32 SubElementDistributionSize = Schema.GetObservationDistributionVectorSize(Parameters.Element);
 
 			for (int32 ElementIdx = 0; ElementIdx < Parameters.Num; ElementIdx++)
 			{
-				AddGaussianNoiseToVector(
-					InOutRandomState,
+				ClampToObservationDistributionMinMax(
 					InOutObservationVector.Slice(ElementIdx * SubElementSize, SubElementSize),
+					ObservationDistributionMin.Slice(ElementIdx * SubElementDistributionSize, SubElementDistributionSize),
+					ObservationDistributionMax.Slice(ElementIdx * SubElementDistributionSize, SubElementDistributionSize),
 					Schema,
-					Parameters.Element,
-					NoiseScale);
+					Parameters.Element);
 			}
 
 			return;
@@ -2699,12 +3171,12 @@ namespace UE::Learning::Observation
 					break;
 				}
 
-				AddGaussianNoiseToVector(
-					InOutRandomState,
+				ClampToObservationDistributionMinMax(
 					InOutObservationVector.Slice(SubElementIdx * SubElementSize, SubElementSize),
+					ObservationDistributionMin,
+					ObservationDistributionMax,
 					Schema,
-					Parameters.Element,
-					NoiseScale);
+					Parameters.Element);
 			}
 
 			return;
@@ -2712,36 +3184,36 @@ namespace UE::Learning::Observation
 
 		case EType::Encoding:
 		{
-			AddGaussianNoiseToVector(
-				InOutRandomState,
+			ClampToObservationDistributionMinMax(
 				InOutObservationVector,
+				ObservationDistributionMin,
+				ObservationDistributionMax,
 				Schema,
-				Schema.GetEncoding(SchemaElement).Element,
-				NoiseScale);
+				Schema.GetEncoding(SchemaElement).Element);
 
 			return;
 		}
 
 		case EType::Conv1d:
 		{
-			AddGaussianNoiseToVector(
-				InOutRandomState,
+			ClampToObservationDistributionMinMax(
 				InOutObservationVector,
+				ObservationDistributionMin,
+				ObservationDistributionMax,
 				Schema,
-				Schema.GetConv1d(SchemaElement).Element,
-				NoiseScale);
+				Schema.GetConv1d(SchemaElement).Element);
 
 			return;
 		}
 
 		case EType::Conv2d:
 		{
-			AddGaussianNoiseToVector(
-				InOutRandomState,
+			ClampToObservationDistributionMinMax(
 				InOutObservationVector,
+				ObservationDistributionMin,
+				ObservationDistributionMax,
 				Schema,
-				Schema.GetConv2d(SchemaElement).Element,
-				NoiseScale);
+				Schema.GetConv2d(SchemaElement).Element);
 
 			return;
 		}
diff --git a/Engine/Plugins/Experimental/LearningAgents/Source/Learning/Public/LearningObservation.h b/Engine/Plugins/Experimental/LearningAgents/Source/Learning/Public/LearningObservation.h
--- a/Engine/Plugins/Experimental/LearningAgents/Source/Learning/Public/LearningObservation.h
+++ b/Engine/Plugins/Experimental/LearningAgents/Source/Learning/Public/LearningObservation.h
@@ -317,6 +317,9 @@ namespace UE::Learning::Observation
 		// Get the encoded vector size of the given element
 		UE_API int32 GetEncodedVectorSize(const FSchemaElement Element) const;
 
+		// Get the observation distribution vector size of the given element
+		UE_API int32 GetObservationDistributionVectorSize(const FSchemaElement Element) const;
+
 		// Get the current generation
 		UE_API uint32 GetGeneration() const;
 
@@ -437,6 +440,7 @@ namespace UE::Learning::Observation
 		TArray<FName> Tags;
 		TArray<int32> ObservationVectorSizes;
 		TArray<int32> EncodedVectorSizes;
+		TArray<int32> ObservationDistributionVectorSizes;
 		TArray<int32> TypeDataIndices;
 
 		/** These are indexed based on the Type using TypeDataIndices */
@@ -751,20 +755,63 @@ namespace UE::Learning::Observation
 		const TLearningArrayView<1, const float> ObservationVector);
 
 	/**
-	 * Add Gaussian noise to all the continuous parts of an observation vector.
+	 * Outputs a vector mask which has 1s for entries in the observation vector which correspond to continuous observation values and 0s elsewhere. 
+	 * This can be useful if you want to apply some kind of numerical modifications to observation vectors such as adding noise or clamping, without 
+	 * breaking the parts of the vector which encode the structure.  
 	 *
-	 * @param InOutRandomState			Random State/Seed to use when sampling the noise
-	 * @param InOutObservationVector	Observation vector to add the noise to
+	 * @param OutObservationVectorMask	Observation vector mask
+	 * @param InObservationVector		Input observation vector
 	 * @param Schema					Observation Schema
 	 * @param SchemaElement				Observation Schema Element
-	 * @param NoiseScale				Scale of the noise to apply
 	 */
-	LEARNING_API void AddGaussianNoiseToVector(
-		uint32& InOutRandomState,
+	LEARNING_API void GetContinuousActionVectorMask(
+		TLearningArrayView<1, bool> OutObservationVectorMask,
+		const TLearningArrayView<1, const float> InObservationVector,
+		const FSchema& Schema,
+		const FSchemaElement SchemaElement);
+
+	/**
+	 * Initializes mins and maxs for observation distribution vectors. Should be called before any calls to FitObservationDistributionMinMax
+	 *
+	 * @param OutObservationDistributionMin		Minimum values for observations
+	 * @param OutObservationDistributionMax		Maximum values for observations
+	 */
+	LEARNING_API void InitObservationDistributionMinMax(
+		TLearningArrayView<1, float> OutObservationDistributionMin,
+		TLearningArrayView<1, float> OutObservationDistributionMax);
+
+	/**
+	 * Fits the mins and the maxs of each continuous observation in the observation vector. Typically you will want to call this function 
+	 * in a loop over multiple observation vectors to refine the fitted mins and maxs since this updates them in place.
+	 *
+	 * @param InOutObservationDistributionMin		Minimum values for observations
+	 * @param InOutObservationDistributionMax		Maximum values for observations
+	 * @param InObservationVector					Observation vector to fit the mins and maxs from
+	 * @param Schema								Observation Schema
+	 * @param SchemaElement							Observation Schema Element
+	 */
+	LEARNING_API void FitObservationDistributionMinMax(
+		TLearningArrayView<1, float> InOutObservationDistributionMin,
+		TLearningArrayView<1, float> InOutObservationDistributionMax,
+		const TLearningArrayView<1, const float> InObservationVector,
+		const FSchema& Schema,
+		const FSchemaElement SchemaElement);
+
+	/**
+	 * Clamps an observation vector using mins and maxs given by two observation distribution vectors.
+	 *
+	 * @param InOutObservationVector				Observation vector to clamp
+	 * @param ObservationDistributionMin			Minimum values for observations
+	 * @param ObservationDistributionMax			Maximum values for observations
+	 * @param Schema								Observation Schema
+	 * @param SchemaElement							Observation Schema Element
+	 */
+	LEARNING_API void ClampToObservationDistributionMinMax(
 		TLearningArrayView<1, float> InOutObservationVector,
+		const TLearningArrayView<1, const float> ObservationDistributionMin,
+		const TLearningArrayView<1, const float> ObservationDistributionMax,
 		const FSchema& Schema,
-		const FSchemaElement SchemaElement,
-		const float NoiseScale = 1.0f);
+		const FSchemaElement SchemaElement);
 
 }
 
