LearningAgents: fixed bug in frame attribute nonzero ranges computation

diff --git a/Engine/Plugins/Experimental/LearningAgents/Source/Learning/Private/LearningFrameAttribute.cpp b/Engine/Plugins/Experimental/LearningAgents/Source/Learning/Private/LearningFrameAttribute.cpp
--- a/Engine/Plugins/Experimental/LearningAgents/Source/Learning/Private/LearningFrameAttribute.cpp
+++ b/Engine/Plugins/Experimental/LearningAgents/Source/Learning/Private/LearningFrameAttribute.cpp
@@ -1290,7 +1290,7 @@ namespace UE::Learning
 						{
 							check(FrameIdx - StartFrameIndex > 0);
 							AddedRangeStarts.Add(StartFrame + StartFrameIndex);
-							AddedRangeLengths.Add(StartFrame + FrameIdx - StartFrameIndex);
+							AddedRangeLengths.Add(FrameIdx - StartFrameIndex);
 							StartFrameIndex = INDEX_NONE;
 						}
 					}
@@ -1299,7 +1299,7 @@ namespace UE::Learning
 					{
 						check(FrameNum - StartFrameIndex > 0);
 						AddedRangeStarts.Add(StartFrame + StartFrameIndex);
-						AddedRangeLengths.Add(StartFrame + FrameNum - StartFrameIndex);
+						AddedRangeLengths.Add(FrameNum - StartFrameIndex);
 					}
 				}
 
