[LearningAgents] Save snapshots at custom intervals

- also group snapshots into dirs on mlflow based on their iterations num

diff --git a/Engine/Plugins/Experimental/LearningAgents/Source/LearningAgentsTraining/Private/LearningAgentsImitationTrainer.cpp b/Engine/Plugins/Experimental/LearningAgents/Source/LearningAgentsTraining/Private/LearningAgentsImitationTrainer.cpp
--- a/Engine/Plugins/Experimental/LearningAgents/Source/LearningAgentsTraining/Private/LearningAgentsImitationTrainer.cpp
+++ b/Engine/Plugins/Experimental/LearningAgents/Source/LearningAgentsTraining/Private/LearningAgentsImitationTrainer.cpp
@@ -36,6 +36,7 @@ TSharedRef<FJsonObject> FLearningAgentsImitationTrainerTrainingSettings::AsJsonC
 	ConfigObject->SetStringField(TEXT("Device"), UE::Learning::Trainer::GetDeviceString(UE::Learning::Agents::GetTrainingDevice(Device)));
 	ConfigObject->SetBoolField(TEXT("UseTensorBoard"), bUseTensorboard);
 	ConfigObject->SetBoolField(TEXT("SaveSnapshots"), bSaveSnapshots);
+	ConfigObject->SetNumberField(TEXT("IterationsPerSnapshot"), IterationsPerSnapshot);
 	ConfigObject->SetBoolField(TEXT("UseMLflow"), bUseMLflow);
 	ConfigObject->SetStringField(TEXT("MLflowTrackingUri"), MLflowTrackingUri);
 
diff --git a/Engine/Plugins/Experimental/LearningAgents/Source/LearningAgentsTraining/Public/LearningAgentsImitationTrainer.h b/Engine/Plugins/Experimental/LearningAgents/Source/LearningAgentsTraining/Public/LearningAgentsImitationTrainer.h
--- a/Engine/Plugins/Experimental/LearningAgents/Source/LearningAgentsTraining/Public/LearningAgentsImitationTrainer.h
+++ b/Engine/Plugins/Experimental/LearningAgents/Source/LearningAgentsTraining/Public/LearningAgentsImitationTrainer.h
@@ -111,6 +111,10 @@ public:
 	UPROPERTY(EditAnywhere, Category = "LearningAgents")
 	bool bSaveSnapshots = false;
 
+	/** If bSaveSnapshots is true, the snapshots will be saved at an interval defined by the specified number of iterations. */
+	UPROPERTY(EditAnywhere, Category = "LearningAgents", meta = (ClampMin = "1", UIMin = "1"))
+	int32 IterationsPerSnapshot = 1000;
+
 	/** If true, MLflow will be used for experiment tracking. */
 	UPROPERTY(EditAnywhere, Category = "LearningAgents")
 	bool bUseMLflow = false;
