Add LR decay step size to Learning agents

#rnx
#changelist validated
#virtualized

diff --git a/Engine/Plugins/Experimental/LearningAgents/Source/LearningAgentsTraining/Private/LearningAgentsImitationTrainer.cpp b/Engine/Plugins/Experimental/LearningAgents/Source/LearningAgentsTraining/Private/LearningAgentsImitationTrainer.cpp
--- a/Engine/Plugins/Experimental/LearningAgents/Source/LearningAgentsTraining/Private/LearningAgentsImitationTrainer.cpp
+++ b/Engine/Plugins/Experimental/LearningAgents/Source/LearningAgentsTraining/Private/LearningAgentsImitationTrainer.cpp
@@ -27,6 +27,7 @@ TSharedRef<FJsonObject> FLearningAgentsImitationTrainerTrainingSettings::AsJsonC
 	ConfigObject->SetNumberField(TEXT("IterationNum"), NumberOfIterations);
 	ConfigObject->SetNumberField(TEXT("LearningRate"), LearningRate);
 	ConfigObject->SetNumberField(TEXT("LearningRateDecay"), LearningRateDecay);
+	ConfigObject->SetNumberField(TEXT("LearningRateDecayStepSize"), LearningRateDecayStepSize);
 	ConfigObject->SetNumberField(TEXT("WeightDecay"), WeightDecay);
 	ConfigObject->SetNumberField(TEXT("BatchSize"), BatchSize);
 	ConfigObject->SetNumberField(TEXT("Window"), Window);
diff --git a/Engine/Plugins/Experimental/LearningAgents/Source/LearningAgentsTraining/Public/LearningAgentsImitationTrainer.h b/Engine/Plugins/Experimental/LearningAgents/Source/LearningAgentsTraining/Public/LearningAgentsImitationTrainer.h
--- a/Engine/Plugins/Experimental/LearningAgents/Source/LearningAgentsTraining/Public/LearningAgentsImitationTrainer.h
+++ b/Engine/Plugins/Experimental/LearningAgents/Source/LearningAgentsTraining/Public/LearningAgentsImitationTrainer.h
@@ -51,10 +51,14 @@ public:
 	UPROPERTY(EditAnywhere, Category = "LearningAgents", meta = (ClampMin = "0.0", UIMin = "0.0", UIMax = "1.0"))
 	float LearningRate = 0.001f;
 
-	/** Amount by which to multiply the learning rate every 1000 iterations. */
+	/** Amount by which to multiply the learning rate every time it decays. */
 	UPROPERTY(EditAnywhere, Category = "LearningAgents", meta = (ClampMin = "0.0", ClampMax = "1.0", UIMin = "0.0", UIMax = "1.0"))
 	float LearningRateDecay = 1.0f;
 
+	/** The number of iterations to train before updating the learning rate. */
+	UPROPERTY(EditAnywhere, Category = "LearningAgents", meta = (ClampMin = "1", UIMin = "1"))
+	int32 LearningRateDecayStepSize = 1000;
+
 	/**
 	 * Amount of weight decay to apply to the network. Larger values encourage network weights to be smaller but too
 	 * large a value can cause the network weights to collapse to all zeros.
