[Learning Agents] Added a subsystem to record gameplay using a DemoNetDriver. Methods exposed to Blueprints.

#rnx

diff --git a/Engine/Plugins/Experimental/LearningAgents/LearningAgents.uplugin b/Engine/Plugins/Experimental/LearningAgents/LearningAgents.uplugin
--- a/Engine/Plugins/Experimental/LearningAgents/LearningAgents.uplugin
+++ b/Engine/Plugins/Experimental/LearningAgents/LearningAgents.uplugin
@@ -40,6 +40,11 @@
 			"Name": "LearningAgentsTrainingEditor",
 			"Type": "Editor",
 			"LoadingPhase": "Default"
+		},
+		{
+			"Name": "LearningAgentsReplay",
+			"Type": "Runtime",
+			"LoadingPhase": "Default"
 		}
 	],
 	"Plugins": [
diff --git a/Engine/Plugins/Experimental/LearningAgents/Source/LearningAgentsReplay/LearningAgentsReplay.Build.cs b/Engine/Plugins/Experimental/LearningAgents/Source/LearningAgentsReplay/LearningAgentsReplay.Build.cs
new file mode 100644
--- /dev/null
+++ b/Engine/Plugins/Experimental/LearningAgents/Source/LearningAgentsReplay/LearningAgentsReplay.Build.cs
@@ -0,0 +1,54 @@
+// Copyright Epic Games, Inc. All Rights Reserved.
+
+using UnrealBuildTool;
+
+public class LearningAgentsReplay : ModuleRules
+{
+	public LearningAgentsReplay(ReadOnlyTargetRules Target) : base(Target)
+	{
+		PCHUsage = ModuleRules.PCHUsageMode.UseExplicitOrSharedPCHs;
+		
+		PublicIncludePaths.AddRange(
+			new string[] {
+				// ... add public include paths required here ...
+			}
+			);
+				
+		
+		PrivateIncludePaths.AddRange(
+			new string[] {
+				// ... add other private include paths required here ...
+			}
+			);
+			
+		
+		PublicDependencyModuleNames.AddRange(
+			new string[]
+			{
+				"Core",
+				"NetworkReplayStreaming"
+				// ... add other public dependencies that you statically link with here ...
+			}
+			);
+			
+		
+		PrivateDependencyModuleNames.AddRange(
+			new string[]
+			{
+				"CoreUObject",
+				"Engine",
+				"Slate",
+				"SlateCore",
+				// ... add private dependencies that you statically link with here ...	
+			}
+			);
+		
+		
+		DynamicallyLoadedModuleNames.AddRange(
+			new string[]
+			{
+				// ... add any modules that your module loads dynamically here ...
+			}
+			);
+	}
+}
diff --git a/Engine/Plugins/Experimental/LearningAgents/Source/LearningAgentsReplay/Private/AsyncAction_LearningAgentsQueryReplays.cpp b/Engine/Plugins/Experimental/LearningAgents/Source/LearningAgentsReplay/Private/AsyncAction_LearningAgentsQueryReplays.cpp
new file mode 100644
--- /dev/null
+++ b/Engine/Plugins/Experimental/LearningAgents/Source/LearningAgentsReplay/Private/AsyncAction_LearningAgentsQueryReplays.cpp
@@ -0,0 +1,58 @@
+// Copyright Epic Games, Inc. All Rights Reserved.
+
+#include "AsyncAction_LearningAgentsQueryReplays.h"
+
+#include "GameFramework/PlayerController.h"
+#include "LearningAgentsReplaySubsystem.h"
+#include "Templates/Greater.h"
+
+UAsyncAction_LearningAgentsQueryReplays::UAsyncAction_LearningAgentsQueryReplays(const FObjectInitializer& ObjectInitializer)
+	: Super(ObjectInitializer)
+{
+}
+
+UAsyncAction_LearningAgentsQueryReplays* UAsyncAction_LearningAgentsQueryReplays::QueryLearningAgentsReplays(APlayerController* InPlayerController)
+{
+	UAsyncAction_LearningAgentsQueryReplays* Action = nullptr;
+
+	if (InPlayerController != nullptr)
+	{
+		Action = NewObject<UAsyncAction_LearningAgentsQueryReplays>();
+		Action->PlayerController = InPlayerController;
+	}
+
+	return Action;
+}
+
+void UAsyncAction_LearningAgentsQueryReplays::Activate()
+{
+	ReplayStreamer = FNetworkReplayStreaming::Get().GetFactory().CreateReplayStreamer();
+
+	ResultList = NewObject<ULearningAgentsReplayList>();
+	if (ReplayStreamer.IsValid())
+	{
+		FNetworkReplayVersion EnumerateStreamsVersion = FNetworkVersion::GetReplayVersion();
+
+		ReplayStreamer->EnumerateStreams(EnumerateStreamsVersion, INDEX_NONE, FString(), TArray<FString>(), FEnumerateStreamsCallback::CreateUObject(this, &ThisClass::OnEnumerateStreamsComplete));
+	}
+	else
+	{
+		QueryComplete.Broadcast(ResultList);
+	}
+}
+
+void UAsyncAction_LearningAgentsQueryReplays::OnEnumerateStreamsComplete(const FEnumerateStreamsResult& Result)
+{
+	for (const FNetworkReplayStreamInfo& StreamInfo : Result.FoundStreams)
+	{
+		ULearningAgentsReplayListEntry* NewReplayEntry = NewObject<ULearningAgentsReplayListEntry>(ResultList);
+		NewReplayEntry->StreamInfo = StreamInfo;
+		ResultList->Results.Add(NewReplayEntry);
+	}
+
+	// Sort demo names by date
+	Algo::SortBy(ResultList->Results, [](const TObjectPtr<ULearningAgentsReplayListEntry>& Data) { return Data->StreamInfo.Timestamp.GetTicks(); }, TGreater<>());
+
+	QueryComplete.Broadcast(ResultList);
+}
+
diff --git a/Engine/Plugins/Experimental/LearningAgents/Source/LearningAgentsReplay/Private/LearningAgentsReplay.cpp b/Engine/Plugins/Experimental/LearningAgents/Source/LearningAgentsReplay/Private/LearningAgentsReplay.cpp
new file mode 100644
--- /dev/null
+++ b/Engine/Plugins/Experimental/LearningAgents/Source/LearningAgentsReplay/Private/LearningAgentsReplay.cpp
@@ -0,0 +1,20 @@
+// Copyright Epic Games, Inc. All Rights Reserved.
+
+#include "LearningAgentsReplay.h"
+
+#define LOCTEXT_NAMESPACE "FLearningAgentsReplayModule"
+
+void FLearningAgentsReplayModule::StartupModule()
+{
+	// This code will execute after your module is loaded into memory; the exact timing is specified in the .uplugin file per-module
+}
+
+void FLearningAgentsReplayModule::ShutdownModule()
+{
+	// This function may be called during shutdown to clean up your module.  For modules that support dynamic reloading,
+	// we call this function before unloading the module.
+}
+
+#undef LOCTEXT_NAMESPACE
+	
+IMPLEMENT_MODULE(FLearningAgentsReplayModule, LearningAgentsReplay)
\ No newline at end of file
diff --git a/Engine/Plugins/Experimental/LearningAgents/Source/LearningAgentsReplay/Private/LearningAgentsReplaySubsystem.cpp b/Engine/Plugins/Experimental/LearningAgents/Source/LearningAgentsReplay/Private/LearningAgentsReplaySubsystem.cpp
new file mode 100644
--- /dev/null
+++ b/Engine/Plugins/Experimental/LearningAgents/Source/LearningAgentsReplay/Private/LearningAgentsReplaySubsystem.cpp
@@ -0,0 +1,112 @@
+// Copyright Epic Games, Inc. All Rights Reserved.
+
+#include "LearningAgentsReplaySubsystem.h"
+#include "Engine/GameInstance.h"
+#include "Engine/World.h"
+#include "Engine/DemoNetDriver.h"
+#include "Internationalization/Text.h"
+#include "Misc/DateTime.h"
+
+ULearningAgentsReplaySubsystem::ULearningAgentsReplaySubsystem()
+{
+}
+
+bool ULearningAgentsReplaySubsystem::ShouldCreateSubsystem(UObject* Outer) const
+{
+	return true;
+}
+
+void ULearningAgentsReplaySubsystem::Initialize(FSubsystemCollectionBase& Collection)
+{
+	Super::Initialize(Collection);
+	
+	FCoreUObjectDelegates::PostDemoPlay.AddUObject(this, &ThisClass::OnDemoPlayStarted);
+}
+
+void ULearningAgentsReplaySubsystem::Deinitialize()
+{
+	Super::Deinitialize();
+
+	FCoreUObjectDelegates::PostDemoPlay.RemoveAll(this);
+}
+
+void ULearningAgentsReplaySubsystem::OnDemoPlayStarted()
+{
+	
+}
+
+bool ULearningAgentsReplaySubsystem::DoesPlatformSupportReplays()
+{
+	return true;
+}
+
+void ULearningAgentsReplaySubsystem::PlayReplay(ULearningAgentsReplayListEntry* Replay)
+{
+	if (Replay != nullptr)
+	{
+		FString DemoName = Replay->StreamInfo.Name;
+		GetGameInstance()->PlayReplay(DemoName);
+	}
+}
+
+void ULearningAgentsReplaySubsystem::StopRecordingReplay()
+{
+	if (UGameInstance* GameInstance = GetGameInstance())
+	{
+		GameInstance->StopRecordingReplay();	
+	}
+}
+
+void ULearningAgentsReplaySubsystem::RecordClientReplay(APlayerController* PlayerController)
+{
+	UGameInstance* GameInstance = GetGameInstance();
+	if (ensure(DoesPlatformSupportReplays() && PlayerController && GameInstance))
+	{
+		FText FriendlyNameText = FText::Format(NSLOCTEXT("LearningAgents", "LearningAgentsReplayName_Format", "Client Replay {0}"), FText::AsDateTime(FDateTime::UtcNow(), EDateTimeStyle::Short, EDateTimeStyle::Short));
+		GameInstance->StartRecordingReplay(FString(), FriendlyNameText.ToString());
+	}
+}
+
+void ULearningAgentsReplaySubsystem::SeekInActiveReplay(float TimeInSeconds)
+{
+	if (UDemoNetDriver* DemoDriver = GetDemoDriver())
+	{
+		DemoDriver->GotoTimeInSeconds(TimeInSeconds);
+	}
+}
+
+float ULearningAgentsReplaySubsystem::GetReplayLengthInSeconds() const
+{
+	if (UDemoNetDriver* DemoDriver = GetDemoDriver())
+	{
+		return DemoDriver->GetDemoTotalTime();
+	}
+	return 0.0f;
+}
+
+float ULearningAgentsReplaySubsystem::GetReplayCurrentTime() const
+{
+	if (UDemoNetDriver* DemoDriver = GetDemoDriver())
+	{
+		return DemoDriver->GetDemoCurrentTime();
+	}
+	return 0.0f;
+}
+
+UDemoNetDriver* ULearningAgentsReplaySubsystem::GetDemoDriver() const
+{
+	UGameInstance* GameInstance = GetGameInstance();
+	if(GameInstance == nullptr)
+	{
+		return nullptr;
+	}
+	
+	if (UWorld* World = GameInstance->GetWorld())
+	{
+		return World->GetDemoNetDriver();
+	}
+	return nullptr;
+}
+
+
+
diff --git a/Engine/Plugins/Experimental/LearningAgents/Source/LearningAgentsReplay/Public/AsyncAction_LearningAgentsQueryReplays.h b/Engine/Plugins/Experimental/LearningAgents/Source/LearningAgentsReplay/Public/AsyncAction_LearningAgentsQueryReplays.h
new file mode 100644
--- /dev/null
+++ b/Engine/Plugins/Experimental/LearningAgents/Source/LearningAgentsReplay/Public/AsyncAction_LearningAgentsQueryReplays.h
@@ -0,0 +1,50 @@
+// Copyright Epic Games, Inc. All Rights Reserved.
+
+#pragma once
+
+#include "Kismet/BlueprintAsyncActionBase.h"
+
+#include "UObject/ObjectPtr.h"
+#include "AsyncAction_LearningAgentsQueryReplays.generated.h"
+
+class APlayerController;
+class INetworkReplayStreamer;
+class ULearningAgentsReplayList;
+class UObject;
+struct FEnumerateStreamsResult;
+struct FFrame;
+
+DECLARE_DYNAMIC_MULTICAST_DELEGATE_OneParam(FQueryReplayAsyncDelegate, ULearningAgentsReplayList*, Results);
+
+/**
+ * Returns a list of stored replays for a given player controller 
+ */
+UCLASS()
+class UAsyncAction_LearningAgentsQueryReplays : public UBlueprintAsyncActionBase
+{
+	GENERATED_BODY()
+
+public:
+	UAsyncAction_LearningAgentsQueryReplays(const FObjectInitializer& ObjectInitializer);
+
+	UFUNCTION(BlueprintCallable, meta = (BlueprintInternalUseOnly = "true"))
+	static UAsyncAction_LearningAgentsQueryReplays* QueryLearningAgentsReplays(APlayerController* PlayerController);
+
+	virtual void Activate() override;
+
+public:
+	// Called when the replay query completes
+	UPROPERTY(BlueprintAssignable)
+	FQueryReplayAsyncDelegate QueryComplete;
+
+private:
+	void OnEnumerateStreamsComplete(const FEnumerateStreamsResult& Result);
+
+private:
+	UPROPERTY()
+	TObjectPtr<ULearningAgentsReplayList> ResultList;
+
+	TWeakObjectPtr<APlayerController> PlayerController;
+
+	TSharedPtr<INetworkReplayStreamer> ReplayStreamer;
+};
diff --git a/Engine/Plugins/Experimental/LearningAgents/Source/LearningAgentsReplay/Public/LearningAgentsReplay.h b/Engine/Plugins/Experimental/LearningAgents/Source/LearningAgentsReplay/Public/LearningAgentsReplay.h
new file mode 100644
--- /dev/null
+++ b/Engine/Plugins/Experimental/LearningAgents/Source/LearningAgentsReplay/Public/LearningAgentsReplay.h
@@ -0,0 +1,14 @@
+// Copyright Epic Games, Inc. All Rights Reserved.
+
+#pragma once
+
+#include "Modules/ModuleManager.h"
+
+class FLearningAgentsReplayModule : public IModuleInterface
+{
+public:
+
+	/** IModuleInterface implementation */
+	virtual void StartupModule() override;
+	virtual void ShutdownModule() override;
+};
diff --git a/Engine/Plugins/Experimental/LearningAgents/Source/LearningAgentsReplay/Public/LearningAgentsReplaySubsystem.h b/Engine/Plugins/Experimental/LearningAgents/Source/LearningAgentsReplay/Public/LearningAgentsReplaySubsystem.h
new file mode 100644
--- /dev/null
+++ b/Engine/Plugins/Experimental/LearningAgents/Source/LearningAgentsReplay/Public/LearningAgentsReplaySubsystem.h
@@ -0,0 +1,106 @@
+// Copyright Epic Games, Inc. All Rights Reserved.
+
+#pragma once
+
+#include "NetworkReplayStreaming.h"
+#include "Subsystems/GameInstanceSubsystem.h"
+#include "LearningAgentsReplaySubsystem.generated.h"
+
+class UDemoNetDriver;
+class APlayerController;
+class ULocalPlayer;
+
+/** An available replay for display in the UI */
+UCLASS(BlueprintType)
+class LEARNINGAGENTSREPLAY_API ULearningAgentsReplayListEntry : public UObject
+{
+	GENERATED_BODY()
+
+public:
+	FNetworkReplayStreamInfo StreamInfo;
+
+	/** The UI friendly name of the stream */
+	UFUNCTION(BlueprintPure, Category=Replays)
+	FString GetFriendlyName() const { return StreamInfo.FriendlyName; }
+
+	/** The date and time the stream was recorded */
+	UFUNCTION(BlueprintPure, Category=Replays)
+	FDateTime GetTimestamp() const { return StreamInfo.Timestamp; }
+
+	/** The duration of the stream in MS */
+	UFUNCTION(BlueprintPure, Category=Replays)
+	FTimespan GetDuration() const { return FTimespan::FromMilliseconds(StreamInfo.LengthInMS); }
+
+	/** Number of viewers viewing this stream */
+	UFUNCTION(BlueprintPure, Category=Replays)
+	int32 GetNumViewers() const { return StreamInfo.NumViewers; }
+
+	/** True if the stream is live and the game hasn't completed yet */
+	UFUNCTION(BlueprintPure, Category=Replays)
+	bool GetIsLive() const { return StreamInfo.bIsLive; }
+};
+
+/** Results of querying for replays list of results for the UI */
+UCLASS(BlueprintType)
+class LEARNINGAGENTSREPLAY_API ULearningAgentsReplayList : public UObject
+{
+	GENERATED_BODY()
+
+public:
+	UPROPERTY(BlueprintReadWrite, Category=Replays)
+	TArray<TObjectPtr<ULearningAgentsReplayListEntry>> Results;
+};
+
+/** Subsystem to handle recording/loading replays. Exposes functionality to Blueprints */
+UCLASS()
+class LEARNINGAGENTSREPLAY_API ULearningAgentsReplaySubsystem : public UGameInstanceSubsystem
+{
+	GENERATED_BODY()
+
+public:
+	ULearningAgentsReplaySubsystem();
+
+	virtual bool ShouldCreateSubsystem(UObject* Outer) const override;
+	virtual void Initialize(FSubsystemCollectionBase& Collection) override;
+	virtual void Deinitialize() override;
+
+	UFUNCTION()
+	void OnDemoPlayStarted();
+
+	/** Returns true if this platform supports replays at all */
+	UFUNCTION(BlueprintCallable, Category = Replays, BlueprintPure = false)
+	static bool DoesPlatformSupportReplays();
+
+	/** Loads the appropriate map and plays a replay */
+	UFUNCTION(BlueprintCallable, Category=Replays)
+	void PlayReplay(ULearningAgentsReplayListEntry* Replay);
+	
+	/** Stops the current active recording */
+	UFUNCTION(BlueprintCallable, Category=Replays)
+	void StopRecordingReplay();
+
+	/** Starts recording a client replay, and handles any file cleanup needed */
+	UFUNCTION(BlueprintCallable, Category = Replays)
+	void RecordClientReplay(APlayerController* PlayerController);
+
+	/** Move forward or back in currently playing replay */
+	UFUNCTION(BlueprintCallable, Category=Replays)
+	void SeekInActiveReplay(float TimeInSeconds);
+
+	/** Gets length of current replay */
+	UFUNCTION(BlueprintCallable, Category = Replays, BlueprintPure = false)
+	float GetReplayLengthInSeconds() const;
+
+	/** Gets current playback time */
+	UFUNCTION(BlueprintCallable, Category=Replays, BlueprintPure=false)
+	float GetReplayCurrentTime() const;
+
+private:
+	TSharedPtr<INetworkReplayStreamer> CurrentReplayStreamer;
+
+	UPROPERTY()
+	TObjectPtr<ULocalPlayer> LocalPlayerDeletingReplays;
+
+	UDemoNetDriver* GetDemoDriver() const;
+};
+
