- Added functionality to load all records file from a folder

diff --git a/Engine/Plugins/Experimental/LearningAgents/Source/LearningAgentsTraining/Private/LearningAgentsRecording.cpp b/Engine/Plugins/Experimental/LearningAgents/Source/LearningAgentsTraining/Private/LearningAgentsRecording.cpp
--- a/Engine/Plugins/Experimental/LearningAgents/Source/LearningAgentsTraining/Private/LearningAgentsRecording.cpp
+++ b/Engine/Plugins/Experimental/LearningAgents/Source/LearningAgentsTraining/Private/LearningAgentsRecording.cpp
@@ -6,6 +6,7 @@
 
 #include "UObject/Package.h"
 #include "Misc/FileHelper.h"
+#include "Misc/Paths.h"
 #include "Serialization/JsonReader.h"
 #include "Serialization/JsonSerializer.h"
 #include "Serialization/JsonWriter.h"
@@ -127,6 +128,39 @@ void ULearningAgentsRecording::LoadRecordingFromFile(const FFilePath& File, cons
 	}
 }
 
+void ULearningAgentsRecording::AppendAllRecordingsFromFolder()
+{
+	if (RecordingFolder.Path.IsEmpty())
+	{
+		UE_LOG(LogLearning, Error, TEXT("%s: RecordingFolder path is empty."), *GetName());
+		return;
+	}
+
+	IFileManager& FileManager = IFileManager::Get();
+	TArray<FString> FileNames;
+	const FString SearchPath = FPaths::ConvertRelativePathToFull(RecordingFolder.Path) / TEXT("*.record");
+	FileManager.FindFiles(FileNames, *SearchPath, true, false); // only record files
+
+	if (FileNames.Num() == 0)
+	{
+		UE_LOG(LogLearning, Warning, TEXT("%s: No .record files found in folder: %s"), *GetName(), *RecordingFolder.Path);
+		return;
+	}
+
+	for (const FString& FileName : FileNames)
+	{
+		const FString FullPath = FPaths::Combine(RecordingFolder.Path, FileName);
+		ULearningAgentsRecording* TempRecording = NewObject<ULearningAgentsRecording>(this);
+		FFilePath FilePath;
+		FilePath.FilePath = FullPath;
+
+		TempRecording->LoadRecordingFromFile(FilePath, SchemaFile);
+		Records.Append(TempRecording->Records);
+	}
+
+	ForceMarkDirty();
+}
+
 void ULearningAgentsRecording::SaveRecordingToFile(const FFilePath& File) const
 {
 	TArray<uint8> RecordingData;
diff --git a/Engine/Plugins/Experimental/LearningAgents/Source/LearningAgentsTraining/Public/LearningAgentsRecording.h b/Engine/Plugins/Experimental/LearningAgents/Source/LearningAgentsTraining/Public/LearningAgentsRecording.h
--- a/Engine/Plugins/Experimental/LearningAgents/Source/LearningAgentsTraining/Public/LearningAgentsRecording.h
+++ b/Engine/Plugins/Experimental/LearningAgents/Source/LearningAgentsTraining/Public/LearningAgentsRecording.h
@@ -78,6 +78,10 @@ public:
 	UFUNCTION(CallInEditor, Category = "LearningAgents")
 	UE_API void AppendRecording();
 
+	/** Appends all .record files from the specified folder to this recording asset */
+	UFUNCTION(CallInEditor, Category = "LearningAgents", meta = (DisplayName = "Append All Recordings From Folder"))
+	UE_API void AppendAllRecordingsFromFolder();
+
 	/** Load this recording from a file. */
 	UFUNCTION(BlueprintCallable, Category = "LearningAgents", meta = (RelativePath))
 	UE_API void LoadRecordingFromFile(const FFilePath& File, const FFilePath& Schema);
@@ -147,6 +151,10 @@ public:
 	UPROPERTY(EditInstanceOnly, Category = "LearningAgents")
 	FFilePath SchemaFile;
 
+	/** A folder containing .record files to load. Used in combination with the "Append All Recordings From Folder" button in the editor.*/
+	UPROPERTY(EditInstanceOnly, Category = "LearningAgents")
+	FDirectoryPath RecordingFolder;
+
 	/** Set of records. */
 	UPROPERTY(EditInstanceOnly, Category = "LearningAgents")
 	TArray<FLearningAgentsRecord> Records;
