LearningAgents: Cleaned up some UE_API stuff in the Learning and LearningTraining modules

diff --git a/Engine/Plugins/Experimental/LearningAgents/Source/Learning/Public/LearningAction.h b/Engine/Plugins/Experimental/LearningAgents/Source/Learning/Public/LearningAction.h
--- a/Engine/Plugins/Experimental/LearningAgents/Source/Learning/Public/LearningAction.h
+++ b/Engine/Plugins/Experimental/LearningAgents/Source/Learning/Public/LearningAction.h
@@ -662,7 +662,7 @@ namespace UE::Learning::Action
 	 * @param SchemaElement				Action Schema Element
 	 * @param Salt						Hash salt
 	 */
-	LEARNING_API int32 GetSchemaObjectsCompatibilityHash(
+	UE_API int32 GetSchemaObjectsCompatibilityHash(
 		const FSchema& Schema,
 		const FSchemaElement SchemaElement,
 		const int32 Salt = 0x49710e77);
@@ -676,7 +676,7 @@ namespace UE::Learning::Action
 	 * @param SchemaElementB			Second Schema Element
 	 * @returns							true when the schemas are compatible
 	 */
-	LEARNING_API bool AreSchemaObjectsCompatible(
+	UE_API bool AreSchemaObjectsCompatible(
 		const FSchema& SchemaA,
 		const FSchemaElement SchemaElementA,
 		const FSchema& SchemaB,
@@ -711,7 +711,7 @@ namespace UE::Learning::Action
 	 * @param SchemaElement				Action Schema Element
 	 * @param NetworkSettings			Network Generation Settings
 	 */
-	LEARNING_API void MakeDecoderNetworkModelBuilderElementFromSchema(
+	UE_API void MakeDecoderNetworkModelBuilderElementFromSchema(
 		NNE::RuntimeBasic::FModelBuilderElement& OutElement,
 		NNE::RuntimeBasic::FModelBuilder& Builder,
 		const FSchema& Schema,
@@ -730,7 +730,7 @@ namespace UE::Learning::Action
 	 * @param NetworkSettings			Network Generation Settings
 	 * @param Seed						The random seed used in initializing the network weights
 	 */
-	LEARNING_API void GenerateDecoderNetworkFileDataFromSchema(
+	UE_API void GenerateDecoderNetworkFileDataFromSchema(
 		TArray<uint8>& OutFileData,
 		uint32& OutInputSize,
 		uint32& OutOutputSize,
@@ -750,7 +750,7 @@ namespace UE::Learning::Action
 	 * @param SchemaElement				Action Schema Element
 	 * @param ActionNoiseScale			Scale factor for the noise. Set this to zero to always choose the mean (expected) action.
 	 */
-	LEARNING_API void SampleVectorFromDistributionVector(
+	UE_API void SampleVectorFromDistributionVector(
 		uint32& InOutRandomState,
 		TLearningArrayView<1, float> OutActionVector,
 		const TLearningArrayView<1, const float> ActionDistributionVector,
@@ -768,7 +768,7 @@ namespace UE::Learning::Action
 	 * @param Object					Action Object
 	 * @param ObjectElement				Action Object Element
 	 */
-	LEARNING_API void SetVectorFromObject(
+	UE_API void SetVectorFromObject(
 		TLearningArrayView<1, float> OutActionVector,
 		const FSchema& Schema,
 		const FSchemaElement SchemaElement,
@@ -784,7 +784,7 @@ namespace UE::Learning::Action
 	 * @param SchemaElement				Action Schema Element
 	 * @param ActionVector				Input Action vector
 	 */
-	LEARNING_API void GetObjectFromVector(
+	UE_API void GetObjectFromVector(
 		FObject& OutObject,
 		FObjectElement& OutObjectElement,
 		const FSchema& Schema,
@@ -800,7 +800,7 @@ namespace UE::Learning::Action
 	 * @param Modifier					Action Modifier
 	 * @param ModifierElement			Action Modifier Element
 	 */
-	LEARNING_API void SetVectorFromModifier(
+	UE_API void SetVectorFromModifier(
 		TLearningArrayView<1, float> OutActionModifierVector,
 		const FSchema& Schema,
 		const FSchemaElement SchemaElement,
@@ -816,7 +816,7 @@ namespace UE::Learning::Action
 	 * @param SchemaElement				Action Schema Element
 	 * @param ActionModifierVector		Input Action Modifier vector
 	 */
-	LEARNING_API void GetModifierFromVector(
+	UE_API void GetModifierFromVector(
 		FModifier& OutModifier,
 		FModifierElement& OutModifierElement,
 		const FSchema& Schema,
diff --git a/Engine/Plugins/Experimental/LearningAgents/Source/Learning/Public/LearningArray.h b/Engine/Plugins/Experimental/LearningAgents/Source/Learning/Public/LearningArray.h
--- a/Engine/Plugins/Experimental/LearningAgents/Source/Learning/Public/LearningArray.h
+++ b/Engine/Plugins/Experimental/LearningAgents/Source/Learning/Public/LearningArray.h
@@ -207,7 +207,7 @@ namespace UE::Learning
 	* @param MinSliceElementNum		Minimum number of elements to include in each call to Body
 	* @param Body					Callback taking start and stop index slice of Num
 	*/
-	LEARNING_API void SlicedParallelFor(const int32 Num, const int32 MinSliceElementNum, const TFunctionRef<void(const int32 Start, const int32 Length)> Body);
+	UE_API void SlicedParallelFor(const int32 Num, const int32 MinSliceElementNum, const TFunctionRef<void(const int32 Start, const int32 Length)> Body);
 
 	/**
 	* Serialize an integer to bytes
diff --git a/Engine/Plugins/Experimental/LearningAgents/Source/Learning/Public/LearningCompletion.h b/Engine/Plugins/Experimental/LearningAgents/Source/Learning/Public/LearningCompletion.h
--- a/Engine/Plugins/Experimental/LearningAgents/Source/Learning/Public/LearningCompletion.h
+++ b/Engine/Plugins/Experimental/LearningAgents/Source/Learning/Public/LearningCompletion.h
@@ -48,14 +48,14 @@ namespace UE::Learning
 		/**
 		* Converts a ECompletionMode into a string.
 		*/
-		LEARNING_API const TCHAR* CompletionModeString(const ECompletionMode Completion);
+		UE_API const TCHAR* CompletionModeString(const ECompletionMode Completion);
 
 		/**
 		* Takes the logical Or of completions. More specifically, if either completion is `Terminated` that 
 		* will be the result of this operator, otherwise either completion being `Truncated` takes priority 
 		* over either being `Running`. Put simply: Terminated > Truncated > Running
 		*/
-		LEARNING_API ECompletionMode Or(const ECompletionMode Lhs, const ECompletionMode Rhs);
+		UE_API ECompletionMode Or(const ECompletionMode Lhs, const ECompletionMode Rhs);
 
 		/**
 		* Set completions for all instances whose episode has reached the max number of steps.
@@ -66,7 +66,7 @@ namespace UE::Learning
 		* @param Instances						Instances to process
 		*
 		*/
-		LEARNING_API void EvaluateEndOfEpisodeCompletions(
+		UE_API void EvaluateEndOfEpisodeCompletions(
 			TLearningArrayView<1, ECompletionMode> OutCompletions,
 			const TLearningArrayView<1, const int32> EpisodeStepNums,
 			const int32 EpisodeMaxStepNum,
diff --git a/Engine/Plugins/Experimental/LearningAgents/Source/Learning/Public/LearningFrameAttribute.h b/Engine/Plugins/Experimental/LearningAgents/Source/Learning/Public/LearningFrameAttribute.h
--- a/Engine/Plugins/Experimental/LearningAgents/Source/Learning/Public/LearningFrameAttribute.h
+++ b/Engine/Plugins/Experimental/LearningAgents/Source/Learning/Public/LearningFrameAttribute.h
@@ -132,32 +132,32 @@ namespace UE::Learning
 			const TLearningArrayView<1, const int32> RangeLengths)>;
 
 		/** Computes the intersection of a frame attribute and frame range set */
-		LEARNING_API void Intersection(FFrameAttribute& OutFrameAttribute, const FFrameAttribute& FrameAttribute, const FFrameRangeSet& FrameRangeSet);
+		UE_API void Intersection(FFrameAttribute& OutFrameAttribute, const FFrameAttribute& FrameAttribute, const FFrameRangeSet& FrameRangeSet);
 
 		/** Computes the frame range set where the given channel is non-zero in the frame attribute. */
-		LEARNING_API void NonZeroFrameRangeSet(FFrameRangeSet& OutFrameRangeSet, const FFrameAttribute& FrameAttribute, const int32 ChannelIdx);
+		UE_API void NonZeroFrameRangeSet(FFrameRangeSet& OutFrameRangeSet, const FFrameAttribute& FrameAttribute, const int32 ChannelIdx);
 
 		/** Computes a reduction on a frame attribute */
-		LEARNING_API void ReduceOp(
+		UE_API void ReduceOp(
 			const FFrameAttribute& In,
 			const ReduceOpFunction Op);
 
 		/** Creates a frame attribute from zero arguments and the given op */
-		LEARNING_API void NullaryOp(
+		UE_API void NullaryOp(
 			FFrameAttribute& Out,
 			const int32 OutChannelNum,
 			const FFrameRangeSet& FrameRangeSet,
 			const NullaryOpFunction Op);
 
 		/** Creates a frame attribute from another and the given op */
-		LEARNING_API void UnaryOp(
+		UE_API void UnaryOp(
 			FFrameAttribute& Out,
 			const int32 OutChannelNum,
 			const FFrameAttribute& In,
 			const UnaryOpFunction Op);
 
 		/** Creates a frame attribute from two others and the given op. Performs an intersection of the Lhs and Rhs if they do not match. */
-		LEARNING_API void BinaryOp(
+		UE_API void BinaryOp(
 			FFrameAttribute& Out,
 			const int32 OutChannelNum,
 			const FFrameAttribute& Lhs,
@@ -165,256 +165,256 @@ namespace UE::Learning
 			const BinaryOpFunction Op);
 
 		/** Creates a frame attribute from multiple others and the given op. */
-		LEARNING_API void NaryOp(
+		UE_API void NaryOp(
 			FFrameAttribute& Out,
 			const int32 OutChannelNum,
 			const TArrayView<const ConstFrameAttributePtr> Inputs,
 			const NaryOpFunction Op);
 
 		/** Find the channel and frame with the smallest value */
-		LEARNING_API bool FindMinimum(
+		UE_API bool FindMinimum(
 			int32& OutChannelIdx,
 			int32& OutFrameIdx,
 			float& OutMinimumValue, 
 			const FFrameAttribute& In);
 
 		/** Find the channel and frame with the largest value */
-		LEARNING_API bool FindMaximum(
+		UE_API bool FindMaximum(
 			int32& OutChannelIdx,
 			int32& OutFrameIdx,
 			float& OutMaximumValue, 
 			const FFrameAttribute& In);
 
 		/** Create a frame attribute made up of zeros */
-		LEARNING_API void Zeros(FFrameAttribute& Out, const FFrameRangeSet& FrameRangeSet, const int32 ChannelNum);
+		UE_API void Zeros(FFrameAttribute& Out, const FFrameRangeSet& FrameRangeSet, const int32 ChannelNum);
 
 		/** Create a frame attribute made up of ones */
-		LEARNING_API void Ones(FFrameAttribute& Out, const FFrameRangeSet& FrameRangeSet, const int32 ChannelNum);
+		UE_API void Ones(FFrameAttribute& Out, const FFrameRangeSet& FrameRangeSet, const int32 ChannelNum);
 
 		/** Create a frame attribute filled with the given value at each frame */
-		LEARNING_API void Fill(FFrameAttribute& Out, const FFrameRangeSet& FrameRangeSet, const TLearningArrayView<1, const float> Values);
+		UE_API void Fill(FFrameAttribute& Out, const FFrameRangeSet& FrameRangeSet, const TLearningArrayView<1, const float> Values);
 
 		/** Add two frame attributes. Channel numbers must match. */
-		LEARNING_API void Add(FFrameAttribute& Out, const FFrameAttribute& Lhs, const FFrameAttribute& Rhs);
+		UE_API void Add(FFrameAttribute& Out, const FFrameAttribute& Lhs, const FFrameAttribute& Rhs);
 
 		/** Subtract two frame attributes. Channel numbers must match. */
-		LEARNING_API void Sub(FFrameAttribute& Out, const FFrameAttribute& Lhs, const FFrameAttribute& Rhs);
+		UE_API void Sub(FFrameAttribute& Out, const FFrameAttribute& Lhs, const FFrameAttribute& Rhs);
 
 		/** Multiply two frame attributes. Channel numbers must match. */
-		LEARNING_API void Mul(FFrameAttribute& Out, const FFrameAttribute& Lhs, const FFrameAttribute& Rhs);
+		UE_API void Mul(FFrameAttribute& Out, const FFrameAttribute& Lhs, const FFrameAttribute& Rhs);
 
 		/** Divide two frame attributes. Channel numbers must match. */
-		LEARNING_API void Div(FFrameAttribute& Out, const FFrameAttribute& Lhs, const FFrameAttribute& Rhs);
+		UE_API void Div(FFrameAttribute& Out, const FFrameAttribute& Lhs, const FFrameAttribute& Rhs);
 
 		/** Compute the dot product of two frame attributes. Channel numbers must match. */
-		LEARNING_API void Dot(FFrameAttribute& Out, const FFrameAttribute& Lhs, const FFrameAttribute& Rhs);
+		UE_API void Dot(FFrameAttribute& Out, const FFrameAttribute& Lhs, const FFrameAttribute& Rhs);
 
 		/** Negate a frame attribute. */
-		LEARNING_API void Neg(FFrameAttribute& Out, const FFrameAttribute& In);
+		UE_API void Neg(FFrameAttribute& Out, const FFrameAttribute& In);
 
 		/** Invert a frame attribute (compute 1/x). */
-		LEARNING_API void Inv(FFrameAttribute& Out, const FFrameAttribute& In);
+		UE_API void Inv(FFrameAttribute& Out, const FFrameAttribute& In);
 
 		/** Compute the absolute value of a frame attribute. */
-		LEARNING_API void Abs(FFrameAttribute& Out, const FFrameAttribute& In);
+		UE_API void Abs(FFrameAttribute& Out, const FFrameAttribute& In);
 
 		/** Compute the log of a frame attribute. */
-		LEARNING_API void Log(FFrameAttribute& Out, const FFrameAttribute& In);
+		UE_API void Log(FFrameAttribute& Out, const FFrameAttribute& In);
 
 		/** Compute the exp of a frame attribute. */
-		LEARNING_API void Exp(FFrameAttribute& Out, const FFrameAttribute& In);
+		UE_API void Exp(FFrameAttribute& Out, const FFrameAttribute& In);
 
 		/** Compute the length of a frame attribute over the channels dimension. */
-		LEARNING_API void Length(FFrameAttribute& Out, const FFrameAttribute& In);
+		UE_API void Length(FFrameAttribute& Out, const FFrameAttribute& In);
 
 		/** Normalize a frame attribute over the channels dimension. */
-		LEARNING_API void Normalize(FFrameAttribute& Out, const FFrameAttribute& In);
+		UE_API void Normalize(FFrameAttribute& Out, const FFrameAttribute& In);
 
 		/** Extract a single channel from frame attribute. */
-		LEARNING_API void Index(FFrameAttribute& Out, const FFrameAttribute& In, const int32 ChannelIdx);
+		UE_API void Index(FFrameAttribute& Out, const FFrameAttribute& In, const int32 ChannelIdx);
 
 		/** Add a constant value to a frame attribute. Channel number must match the size of Rhs. */
-		LEARNING_API void AddConstant(FFrameAttribute& Out, const FFrameAttribute& Lhs, const TLearningArrayView<1, const float> Rhs);
+		UE_API void AddConstant(FFrameAttribute& Out, const FFrameAttribute& Lhs, const TLearningArrayView<1, const float> Rhs);
 
 		/** Subtract a constant value to a frame attribute. Channel number must match the size of Rhs. */
-		LEARNING_API void SubConstant(FFrameAttribute& Out, const FFrameAttribute& Lhs, const TLearningArrayView<1, const float> Rhs);
+		UE_API void SubConstant(FFrameAttribute& Out, const FFrameAttribute& Lhs, const TLearningArrayView<1, const float> Rhs);
 
 		/** Multiply a frame attribute by a constant value. Channel number must match the size of Rhs. */
-		LEARNING_API void MulConstant(FFrameAttribute& Out, const FFrameAttribute& Lhs, const TLearningArrayView<1, const float> Rhs);
+		UE_API void MulConstant(FFrameAttribute& Out, const FFrameAttribute& Lhs, const TLearningArrayView<1, const float> Rhs);
 
 		/** Divide a frame attribute by a constant value. Channel number must match the size of Rhs. */
-		LEARNING_API void DivConstant(FFrameAttribute& Out, const FFrameAttribute& Lhs, const TLearningArrayView<1, const float> Rhs);
+		UE_API void DivConstant(FFrameAttribute& Out, const FFrameAttribute& Lhs, const TLearningArrayView<1, const float> Rhs);
 
 		/** Add a constant value to a frame attribute. Channel number must match the size of Lhs. */
-		LEARNING_API void ConstantAdd(FFrameAttribute& Out, const TLearningArrayView<1, const float> Lhs, const FFrameAttribute& Rhs);
+		UE_API void ConstantAdd(FFrameAttribute& Out, const TLearningArrayView<1, const float> Lhs, const FFrameAttribute& Rhs);
 
 		/** Subtract a constant value to a frame attribute. Channel number must match the size of Lhs. */
-		LEARNING_API void ConstantSub(FFrameAttribute& Out, const TLearningArrayView<1, const float> Lhs, const FFrameAttribute& Rhs);
+		UE_API void ConstantSub(FFrameAttribute& Out, const TLearningArrayView<1, const float> Lhs, const FFrameAttribute& Rhs);
 
 		/** Multiply a frame attribute by a constant value. Channel number must match the size of Lhs. */
-		LEARNING_API void ConstantMul(FFrameAttribute& Out, const TLearningArrayView<1, const float> Lhs, const FFrameAttribute& Rhs);
+		UE_API void ConstantMul(FFrameAttribute& Out, const TLearningArrayView<1, const float> Lhs, const FFrameAttribute& Rhs);
 
 		/** Divide a frame attribute by a constant value. Channel number must match the size of Lhs. */
-		LEARNING_API void ConstantDiv(FFrameAttribute& Out, const TLearningArrayView<1, const float> Lhs, const FFrameAttribute& Rhs);
+		UE_API void ConstantDiv(FFrameAttribute& Out, const TLearningArrayView<1, const float> Lhs, const FFrameAttribute& Rhs);
 
 		/** Compute the sum of an array of frame attributes. Channel numbers must match. */
-		LEARNING_API void Sum(FFrameAttribute& Out, const TArrayView<const ConstFrameAttributePtr> Inputs);
+		UE_API void Sum(FFrameAttribute& Out, const TArrayView<const ConstFrameAttributePtr> Inputs);
 
 		/** Compute the product of an array of frame attributes. Channel numbers must match. */
-		LEARNING_API void Prod(FFrameAttribute& Out, const TArrayView<const ConstFrameAttributePtr> Inputs);
+		UE_API void Prod(FFrameAttribute& Out, const TArrayView<const ConstFrameAttributePtr> Inputs);
 
 		/** Computes the logical and of two frame attributes. Channel numbers must match. */
-		LEARNING_API void LogicalAnd(FFrameAttribute& Out, const FFrameAttribute& Lhs, const FFrameAttribute& Rhs);
+		UE_API void LogicalAnd(FFrameAttribute& Out, const FFrameAttribute& Lhs, const FFrameAttribute& Rhs);
 
 		/** Computes the logical or of two frame attributes. Channel numbers must match. */
-		LEARNING_API void LogicalOr(FFrameAttribute& Out, const FFrameAttribute& Lhs, const FFrameAttribute& Rhs);
+		UE_API void LogicalOr(FFrameAttribute& Out, const FFrameAttribute& Lhs, const FFrameAttribute& Rhs);
 
 		/** Computes the logical not of a frame attribute. */
-		LEARNING_API void LogicalNot(FFrameAttribute& Out, const FFrameAttribute& In);
+		UE_API void LogicalNot(FFrameAttribute& Out, const FFrameAttribute& In);
 
 		/** Computes if the values in one frame attribute are greater than another. Channel numbers must match. */
-		LEARNING_API void Gt(FFrameAttribute& Out, const FFrameAttribute& Lhs, const FFrameAttribute& Rhs);
+		UE_API void Gt(FFrameAttribute& Out, const FFrameAttribute& Lhs, const FFrameAttribute& Rhs);
 
 		/** Computes if the values in one frame attribute are greater than or equal to another. Channel numbers must match. */
-		LEARNING_API void Ge(FFrameAttribute& Out, const FFrameAttribute& Lhs, const FFrameAttribute& Rhs);
+		UE_API void Ge(FFrameAttribute& Out, const FFrameAttribute& Lhs, const FFrameAttribute& Rhs);
 
 		/** Computes if the values in one frame attribute are less than another. Channel numbers must match. */
-		LEARNING_API void Lt(FFrameAttribute& Out, const FFrameAttribute& Lhs, const FFrameAttribute& Rhs);
+		UE_API void Lt(FFrameAttribute& Out, const FFrameAttribute& Lhs, const FFrameAttribute& Rhs);
 
 		/** Computes if the values in one frame attribute are less than or equal to another. Channel numbers must match. */
-		LEARNING_API void Le(FFrameAttribute& Out, const FFrameAttribute& Lhs, const FFrameAttribute& Rhs);
+		UE_API void Le(FFrameAttribute& Out, const FFrameAttribute& Lhs, const FFrameAttribute& Rhs);
 
 		/** Computes if the values in one frame attribute are equal to another. Channel numbers must match. */
-		LEARNING_API void Eq(FFrameAttribute& Out, const FFrameAttribute& Lhs, const FFrameAttribute& Rhs);
+		UE_API void Eq(FFrameAttribute& Out, const FFrameAttribute& Lhs, const FFrameAttribute& Rhs);
 
 		/** Computes if the values in one frame attribute are not equal to another. Channel numbers must match. */
-		LEARNING_API void Neq(FFrameAttribute& Out, const FFrameAttribute& Lhs, const FFrameAttribute& Rhs);
+		UE_API void Neq(FFrameAttribute& Out, const FFrameAttribute& Lhs, const FFrameAttribute& Rhs);
 
 		/** Computes if the values in a frame attribute are greater than a constant. Channel number must match the size of Rhs. */
-		LEARNING_API void GtConstant(FFrameAttribute& Out, const FFrameAttribute& Lhs, const TLearningArrayView<1, const float> Rhs);
+		UE_API void GtConstant(FFrameAttribute& Out, const FFrameAttribute& Lhs, const TLearningArrayView<1, const float> Rhs);
 
 		/** Computes if the values in a frame attribute are greater than or equal to a constant. Channel number must match the size of Rhs. */
-		LEARNING_API void GeConstant(FFrameAttribute& Out, const FFrameAttribute& Lhs, const TLearningArrayView<1, const float> Rhs);
+		UE_API void GeConstant(FFrameAttribute& Out, const FFrameAttribute& Lhs, const TLearningArrayView<1, const float> Rhs);
 
 		/** Computes if the values in a frame attribute are less than a constant. Channel number must match the size of Rhs. */
-		LEARNING_API void LtConstant(FFrameAttribute& Out, const FFrameAttribute& Lhs, const TLearningArrayView<1, const float> Rhs);
+		UE_API void LtConstant(FFrameAttribute& Out, const FFrameAttribute& Lhs, const TLearningArrayView<1, const float> Rhs);
 
 		/** Computes if the values in a frame attribute are greater than or equal to a constant. Channel number must match the size of Rhs. */
-		LEARNING_API void LeConstant(FFrameAttribute& Out, const FFrameAttribute& Lhs, const TLearningArrayView<1, const float> Rhs);
+		UE_API void LeConstant(FFrameAttribute& Out, const FFrameAttribute& Lhs, const TLearningArrayView<1, const float> Rhs);
 		
 		/** Computes if the values in a frame attribute are equal to a constant. Channel number must match the size of Rhs. */
-		LEARNING_API void EqConstant(FFrameAttribute& Out, const FFrameAttribute& Lhs, const TLearningArrayView<1, const float> Rhs);
+		UE_API void EqConstant(FFrameAttribute& Out, const FFrameAttribute& Lhs, const TLearningArrayView<1, const float> Rhs);
 
 		/** Computes if the values in a frame attribute are not equal to a constant. Channel number must match the size of Rhs. */
-		LEARNING_API void NeqConstant(FFrameAttribute& Out, const FFrameAttribute& Lhs, const TLearningArrayView<1, const float> Rhs);
+		UE_API void NeqConstant(FFrameAttribute& Out, const FFrameAttribute& Lhs, const TLearningArrayView<1, const float> Rhs);
 
 		/** Computes if the values in a frame attribute are greater than a constant. Channel number must match the size of Lhs. */
-		LEARNING_API void ConstantGt(FFrameAttribute& Out, const TLearningArrayView<1, const float> Lhs, const FFrameAttribute& Rhs);
+		UE_API void ConstantGt(FFrameAttribute& Out, const TLearningArrayView<1, const float> Lhs, const FFrameAttribute& Rhs);
 		
 		/** Computes if the values in a frame attribute are greater than or equal to a constant. Channel number must match the size of Lhs. */
-		LEARNING_API void ConstantGe(FFrameAttribute& Out, const TLearningArrayView<1, const float> Lhs, const FFrameAttribute& Rhs);
+		UE_API void ConstantGe(FFrameAttribute& Out, const TLearningArrayView<1, const float> Lhs, const FFrameAttribute& Rhs);
 		
 		/** Computes if the values in a frame attribute are less than a constant. Channel number must match the size of Lhs. */
-		LEARNING_API void ConstantLt(FFrameAttribute& Out, const TLearningArrayView<1, const float> Lhs, const FFrameAttribute& Rhs);
+		UE_API void ConstantLt(FFrameAttribute& Out, const TLearningArrayView<1, const float> Lhs, const FFrameAttribute& Rhs);
 
 		/** Computes if the values in a frame attribute are greater than or equal to a constant. Channel number must match the size of Lhs. */
-		LEARNING_API void ConstantLe(FFrameAttribute& Out, const TLearningArrayView<1, const float> Lhs, const FFrameAttribute& Rhs);
+		UE_API void ConstantLe(FFrameAttribute& Out, const TLearningArrayView<1, const float> Lhs, const FFrameAttribute& Rhs);
 
 		/** Computes if the values in a frame attribute are equal to a constant. Channel number must match the size of Lhs. */
-		LEARNING_API void ConstantEq(FFrameAttribute& Out, const TLearningArrayView<1, const float> Lhs, const FFrameAttribute& Rhs);
+		UE_API void ConstantEq(FFrameAttribute& Out, const TLearningArrayView<1, const float> Lhs, const FFrameAttribute& Rhs);
 		
 		/** Computes if the values in a frame attribute are not equal to a constant. Channel number must match the size of Lhs. */
-		LEARNING_API void ConstantNeq(FFrameAttribute& Out, const TLearningArrayView<1, const float> Lhs, const FFrameAttribute& Rhs);
+		UE_API void ConstantNeq(FFrameAttribute& Out, const TLearningArrayView<1, const float> Lhs, const FFrameAttribute& Rhs);
 
 		/** Applies a Gaussian smoothing filter to the ranges of values in the frame attribute. */
-		LEARNING_API void FilterGaussian(FFrameAttribute& Out, const FFrameAttribute& In, const float StdInFrames);
+		UE_API void FilterGaussian(FFrameAttribute& Out, const FFrameAttribute& In, const float StdInFrames);
 
 		/** Applies a Majority Vote filter to the ranges of values in the frame attribute. */
-		LEARNING_API void FilterMajorityVote(FFrameAttribute& Out, const FFrameAttribute& In, const int32 FilterWidthFrames);
+		UE_API void FilterMajorityVote(FFrameAttribute& Out, const FFrameAttribute& In, const int32 FilterWidthFrames);
 
 		/** Computes the mean and std across all frames. OutMean and OutStd should be size of the number of channels */
-		LEARNING_API void MeanStd(TLearningArrayView<1, float> OutMean, TLearningArrayView<1, float> OutStd, const FFrameAttribute& In);
+		UE_API void MeanStd(TLearningArrayView<1, float> OutMean, TLearningArrayView<1, float> OutStd, const FFrameAttribute& In);
 
 		/** Computes the mean and log std across all frames in the log space. OutMean and OutLogStd should be size of the number of channels */
-		LEARNING_API void LogMeanStd(TLearningArrayView<1, float> OutMean, TLearningArrayView<1, float> OutLogStd, const FFrameAttribute& In);
+		UE_API void LogMeanStd(TLearningArrayView<1, float> OutMean, TLearningArrayView<1, float> OutLogStd, const FFrameAttribute& In);
 
 		/** Computes the quaternion multiplication of two frame attributes. Both arguments are assumed to have 4 channels. */
-		LEARNING_API void QuatMul(FFrameAttribute& Out, const FFrameAttribute& Lhs, const FFrameAttribute& Rhs);
+		UE_API void QuatMul(FFrameAttribute& Out, const FFrameAttribute& Lhs, const FFrameAttribute& Rhs);
 
 		/** Computes the quaternion division of two frame attributes. Both arguments are assumed to have 4 channels. */
-		LEARNING_API void QuatDiv(FFrameAttribute& Out, const FFrameAttribute& Lhs, const FFrameAttribute& Rhs);
+		UE_API void QuatDiv(FFrameAttribute& Out, const FFrameAttribute& Lhs, const FFrameAttribute& Rhs);
 
 		/** Compute the quaternion inverse of a frame attribute. Argument is assumed to have 4 channels. */
-		LEARNING_API void QuatInv(FFrameAttribute& Out, const FFrameAttribute& In);
+		UE_API void QuatInv(FFrameAttribute& Out, const FFrameAttribute& In);
 
 		/** Compute the quaternion closest to the identity quaternion for a frame attribute. Argument is assumed to have 4 channels. */
-		LEARNING_API void QuatAbs(FFrameAttribute& Out, const FFrameAttribute& In);
+		UE_API void QuatAbs(FFrameAttribute& Out, const FFrameAttribute& In);
 
 		/** Compute a rotation vector from a quaternion. Argument is assumed to have 4 channels. */
-		LEARNING_API void QuatToRotationVector(FFrameAttribute& Out, const FFrameAttribute& In);
+		UE_API void QuatToRotationVector(FFrameAttribute& Out, const FFrameAttribute& In);
 
 		/** Compute a quaternion from a rotation vector. Argument is assumed to have 3 channels. */
-		LEARNING_API void QuatFromRotationVector(FFrameAttribute& Out, const FFrameAttribute& In);
+		UE_API void QuatFromRotationVector(FFrameAttribute& Out, const FFrameAttribute& In);
 
 		/** Computes the quaternion inverse multiplication of two frame attributes. Both arguments are assumed to have 4 channels. */
-		LEARNING_API void QuatInvMul(FFrameAttribute& Out, const FFrameAttribute& Lhs, const FFrameAttribute& Rhs);
+		UE_API void QuatInvMul(FFrameAttribute& Out, const FFrameAttribute& Lhs, const FFrameAttribute& Rhs);
 
 		/** Computes the quaternion multiplication inverse of two frame attributes. Both arguments are assumed to have 4 channels. */
-		LEARNING_API void QuatMulInv(FFrameAttribute& Out, const FFrameAttribute& Lhs, const FFrameAttribute& Rhs);
+		UE_API void QuatMulInv(FFrameAttribute& Out, const FFrameAttribute& Lhs, const FFrameAttribute& Rhs);
 
 		/** Computes the rotation of the Rhs by the Lhs quaternion. Lhs is assumed to have 4 channels, and Rhs is assumed to have 3. */
-		LEARNING_API void QuatRotate(FFrameAttribute& Out, const FFrameAttribute& Lhs, const FFrameAttribute& Rhs);
+		UE_API void QuatRotate(FFrameAttribute& Out, const FFrameAttribute& Lhs, const FFrameAttribute& Rhs);
 
 		/** Computes the inverse rotation of the Rhs by the Lhs quaternion. Lhs is assumed to have 4 channels, and Rhs is assumed to have 3. */
-		LEARNING_API void QuatUnrotate(FFrameAttribute& Out, const FFrameAttribute& Lhs, const FFrameAttribute& Rhs);
+		UE_API void QuatUnrotate(FFrameAttribute& Out, const FFrameAttribute& Lhs, const FFrameAttribute& Rhs);
 
 		/** Computes the shortest rotation between two vectors. Both arguments are assumed to have 3 channels. */
-		LEARNING_API void QuatBetween(FFrameAttribute& Out, const FFrameAttribute& Lhs, const FFrameAttribute& Rhs);
+		UE_API void QuatBetween(FFrameAttribute& Out, const FFrameAttribute& Lhs, const FFrameAttribute& Rhs);
 
 		/** Computes the quaternion multiplication of a frame attribute by the given quaternion. Lhs is assumed to have 4 channels. */
-		LEARNING_API void QuatMulConstant(FFrameAttribute& Out, const FFrameAttribute& Lhs, const FQuat4f Rhs);
+		UE_API void QuatMulConstant(FFrameAttribute& Out, const FFrameAttribute& Lhs, const FQuat4f Rhs);
 
 		/** Computes the quaternion division of a frame attribute by the given quaternion. Lhs is assumed to have 4 channels. */
-		LEARNING_API void QuatDivConstant(FFrameAttribute& Out, const FFrameAttribute& Lhs, const FQuat4f Rhs);
+		UE_API void QuatDivConstant(FFrameAttribute& Out, const FFrameAttribute& Lhs, const FQuat4f Rhs);
 
 		/** Computes the quaternion inverse multiplication of a frame attribute by the given quaternion. Lhs is assumed to have 4 channels. */
-		LEARNING_API void QuatInvMulConstant(FFrameAttribute& Out, const FFrameAttribute& Lhs, const FQuat4f Rhs);
+		UE_API void QuatInvMulConstant(FFrameAttribute& Out, const FFrameAttribute& Lhs, const FQuat4f Rhs);
 
 		/** Computes the quaternion multiplication inverse of a frame attribute by the given quaternion. Lhs is assumed to have 4 channels. */
-		LEARNING_API void QuatMulInvConstant(FFrameAttribute& Out, const FFrameAttribute& Lhs, const FQuat4f Rhs);
+		UE_API void QuatMulInvConstant(FFrameAttribute& Out, const FFrameAttribute& Lhs, const FQuat4f Rhs);
 
 		/** Computes the quaternion rotation of the given vector by a frame attribute. Lhs is assumed to have 4 channels. */
-		LEARNING_API void QuatRotateConstant(FFrameAttribute& Out, const FFrameAttribute& Lhs, const FVector3f Rhs);
+		UE_API void QuatRotateConstant(FFrameAttribute& Out, const FFrameAttribute& Lhs, const FVector3f Rhs);
 
 		/** Computes the quaternion inverse rotation of the given vector by a frame attribute. Lhs is assumed to have 4 channels. */
-		LEARNING_API void QuatUnrotateConstant(FFrameAttribute& Out, const FFrameAttribute& Lhs, const FVector3f Rhs);
+		UE_API void QuatUnrotateConstant(FFrameAttribute& Out, const FFrameAttribute& Lhs, const FVector3f Rhs);
 
 		/** Computes the quaternion between the given frame attribute and vector. Lhs is assumed to have 3 channels. */
-		LEARNING_API void QuatBetweenConstant(FFrameAttribute& Out, const FFrameAttribute& Lhs, const FVector3f Rhs);
+		UE_API void QuatBetweenConstant(FFrameAttribute& Out, const FFrameAttribute& Lhs, const FVector3f Rhs);
 
 		/** Computes the quaternion multiplication of a frame attribute by the given quaternion. Rhs is assumed to have 4 channels. */
-		LEARNING_API void QuatConstantMul(FFrameAttribute& Out, const FQuat4f Lhs, const FFrameAttribute& Rhs);
+		UE_API void QuatConstantMul(FFrameAttribute& Out, const FQuat4f Lhs, const FFrameAttribute& Rhs);
 
 		/** Computes the quaternion division of a frame attribute by the given quaternion. Rhs is assumed to have 4 channels. */
-		LEARNING_API void QuatConstantDiv(FFrameAttribute& Out, const FQuat4f Lhs, const FFrameAttribute& Rhs);
+		UE_API void QuatConstantDiv(FFrameAttribute& Out, const FQuat4f Lhs, const FFrameAttribute& Rhs);
 		
 		/** Computes the quaternion inverse multiplication of a frame attribute by the given quaternion. Rhs is assumed to have 4 channels. */
-		LEARNING_API void QuatConstantInvMul(FFrameAttribute& Out, const FQuat4f Lhs, const FFrameAttribute& Rhs);
+		UE_API void QuatConstantInvMul(FFrameAttribute& Out, const FQuat4f Lhs, const FFrameAttribute& Rhs);
 		
 		/** Computes the quaternion multiplication inverse of a frame attribute by the given quaternion. Rhs is assumed to have 4 channels. */
-		LEARNING_API void QuatConstantMulInv(FFrameAttribute& Out, const FQuat4f Lhs, const FFrameAttribute& Rhs);
+		UE_API void QuatConstantMulInv(FFrameAttribute& Out, const FQuat4f Lhs, const FFrameAttribute& Rhs);
 		
 		/** Computes the quaternion rotation of the given frame attribute. Rhs is assumed to have 3 channels. */
-		LEARNING_API void QuatConstantRotate(FFrameAttribute& Out, const FQuat4f Lhs, const FFrameAttribute& Rhs);
+		UE_API void QuatConstantRotate(FFrameAttribute& Out, const FQuat4f Lhs, const FFrameAttribute& Rhs);
 
 		/** Computes the quaternion inverse rotation of the given frame attribute. Rhs is assumed to have 3 channels. */
-		LEARNING_API void QuatConstantUnrotate(FFrameAttribute& Out, const FQuat4f Lhs, const FFrameAttribute& Rhs);
+		UE_API void QuatConstantUnrotate(FFrameAttribute& Out, const FQuat4f Lhs, const FFrameAttribute& Rhs);
 
 		/** Computes the quaternion between the given frame attribute and vector. Rhs is assumed to have 3 channels. */
-		LEARNING_API void QuatConstantBetween(FFrameAttribute& Out, const FVector3f Lhs, const FFrameAttribute& Rhs);
+		UE_API void QuatConstantBetween(FFrameAttribute& Out, const FVector3f Lhs, const FFrameAttribute& Rhs);
 
 		/** Computes the quaternion mean and std across all frames. */
-		LEARNING_API void QuatMeanStd(FQuat4f& OutMean, FVector3f& OutStd, const FFrameAttribute& In);
+		UE_API void QuatMeanStd(FQuat4f& OutMean, FVector3f& OutStd, const FFrameAttribute& In);
 	}
 }
 
diff --git a/Engine/Plugins/Experimental/LearningAgents/Source/Learning/Public/LearningFrameRangeSet.h b/Engine/Plugins/Experimental/LearningAgents/Source/Learning/Public/LearningFrameRangeSet.h
--- a/Engine/Plugins/Experimental/LearningAgents/Source/Learning/Public/LearningFrameRangeSet.h
+++ b/Engine/Plugins/Experimental/LearningAgents/Source/Learning/Public/LearningFrameRangeSet.h
@@ -159,35 +159,35 @@ namespace UE::Learning
 	namespace FrameRangeSet
 	{
 		/** Checks if two frame range sets are equal */
-		LEARNING_API bool Equal(const FFrameRangeSet& Lhs, const FFrameRangeSet& Rhs);
+		UE_API bool Equal(const FFrameRangeSet& Lhs, const FFrameRangeSet& Rhs);
 
 		/** Computes the union of a frame set and frame range set */
-		LEARNING_API void Union(FFrameRangeSet& OutFrameRangeSet, const FFrameSet& FrameSet, const FFrameRangeSet& FrameRangeSet);
+		UE_API void Union(FFrameRangeSet& OutFrameRangeSet, const FFrameSet& FrameSet, const FFrameRangeSet& FrameRangeSet);
 
 		/** Computes the intersection of a frame set and frame range set */
-		LEARNING_API void Intersection(FFrameSet& OutFrameSet, const FFrameSet& FrameSet, const FFrameRangeSet& FrameRangeSet);
+		UE_API void Intersection(FFrameSet& OutFrameSet, const FFrameSet& FrameSet, const FFrameRangeSet& FrameRangeSet);
 
 		/** Computes the difference of a frame set and frame range set */
-		LEARNING_API void Difference(FFrameSet& OutFrameSet, const FFrameSet& FrameSet, const FFrameRangeSet& FrameRangeSet);
+		UE_API void Difference(FFrameSet& OutFrameSet, const FFrameSet& FrameSet, const FFrameRangeSet& FrameRangeSet);
 
 		/** Computes the difference of a frame range set and a frame set */
-		LEARNING_API void Difference(FFrameRangeSet& OutFrameRangeSet, const FFrameRangeSet& FrameRangeSet, const FFrameSet& FrameSet);
+		UE_API void Difference(FFrameRangeSet& OutFrameRangeSet, const FFrameRangeSet& FrameRangeSet, const FFrameSet& FrameSet);
 
 		/** Computes the union of two frame range sets */
-		LEARNING_API void Union(FFrameRangeSet& Out, const FFrameRangeSet& Lhs, const FFrameRangeSet& Rhs);
+		UE_API void Union(FFrameRangeSet& Out, const FFrameRangeSet& Lhs, const FFrameRangeSet& Rhs);
 
 		/** Computes the intersection of two frame range sets */
-		LEARNING_API void Intersection(FFrameRangeSet& Out, const FFrameRangeSet& Lhs, const FFrameRangeSet& Rhs);
+		UE_API void Intersection(FFrameRangeSet& Out, const FFrameRangeSet& Lhs, const FFrameRangeSet& Rhs);
 
 		/** Computes the difference of two frame range sets */
-		LEARNING_API void Difference(FFrameRangeSet& Out, const FFrameRangeSet& Lhs, const FFrameRangeSet& Rhs);
+		UE_API void Difference(FFrameRangeSet& Out, const FFrameRangeSet& Lhs, const FFrameRangeSet& Rhs);
 
 		/**
 		 * Computes the intersection of two frame range sets while recording the flat frame offsets in the lhs and rhs, associated with each range
 		 * that gets added to the output. The OutLhsOffsets and OutRhsOffsets outputs should be pre-allocated to at least the maximum of the
 		 * number of ranges in Lhs and Rhs. Returns the number of ranges in the added to the output frame range set.
 		 */
-		LEARNING_API int32 IntersectionWithOffsets(
+		UE_API int32 IntersectionWithOffsets(
 			FFrameRangeSet& Out, 
 			TLearningArrayView<1, int32> OutLhsOffsets,
 			TLearningArrayView<1, int32> OutRhsOffsets,
@@ -195,76 +195,76 @@ namespace UE::Learning
 			const FFrameRangeSet& Rhs);
 
 		/** Creates a new frame range set with the start of all ranges trimmed by the given number of frames */
-		LEARNING_API void TrimStart(FFrameRangeSet& Out, const FFrameRangeSet& FrameRangeSet, const int32 TrimFrameNum);
+		UE_API void TrimStart(FFrameRangeSet& Out, const FFrameRangeSet& FrameRangeSet, const int32 TrimFrameNum);
 
 		/** Creates a new frame range set with the end of all ranges trimmed by the given number of frames */
-		LEARNING_API void TrimEnd(FFrameRangeSet& Out, const FFrameRangeSet& FrameRangeSet, const int32 TrimFrameNum);
+		UE_API void TrimEnd(FFrameRangeSet& Out, const FFrameRangeSet& FrameRangeSet, const int32 TrimFrameNum);
 
 		/** Creates a new frame range set with both the start and the end of all ranges trimmed by the given numbers of frames */
-		LEARNING_API void Trim(FFrameRangeSet& Out, const FFrameRangeSet& FrameRangeSet, const int32 TrimStartFrameNum, const int32 TrimEndFrameNum);
+		UE_API void Trim(FFrameRangeSet& Out, const FFrameRangeSet& FrameRangeSet, const int32 TrimStartFrameNum, const int32 TrimEndFrameNum);
 
 		/** Creates a new frame range set with the start of all ranges padded by the given number of frames */
-		LEARNING_API void PadStart(FFrameRangeSet& Out, const FFrameRangeSet& FrameRangeSet, const int32 PadFrameNum);
+		UE_API void PadStart(FFrameRangeSet& Out, const FFrameRangeSet& FrameRangeSet, const int32 PadFrameNum);
 
 		/** Creates a new frame range set with the end of all ranges padded by the given number of frames */
-		LEARNING_API void PadEnd(FFrameRangeSet& Out, const FFrameRangeSet& FrameRangeSet, const int32 PadFrameNum);
+		UE_API void PadEnd(FFrameRangeSet& Out, const FFrameRangeSet& FrameRangeSet, const int32 PadFrameNum);
 
 		/** Creates a new frame range set with both the start and the end of all ranges padded by the given numbers of frames */
-		LEARNING_API void Pad(FFrameRangeSet& Out, const FFrameRangeSet& FrameRangeSet, const int32 PadStartFrameNum, int32 PadEndFrameNum);
+		UE_API void Pad(FFrameRangeSet& Out, const FFrameRangeSet& FrameRangeSet, const int32 PadStartFrameNum, int32 PadEndFrameNum);
 
 		/** Creates a new frame range set with just the ranges given by the provided indices */
-		LEARNING_API void GatherRanges(FFrameRangeSet& Out, const FFrameRangeSet& FrameRangeSet, const FIndexSet Indices);
+		UE_API void GatherRanges(FFrameRangeSet& Out, const FFrameRangeSet& FrameRangeSet, const FIndexSet Indices);
 
 		/** Makes a Frame Range Set consisting of single-frame ranges from a frame set */
-		LEARNING_API void MakeFromFrameSet(FFrameRangeSet& OutFrameRangeSet, const FFrameSet& FrameSet);
+		UE_API void MakeFromFrameSet(FFrameRangeSet& OutFrameRangeSet, const FFrameSet& FrameSet);
 
 		/** Create a frame set from all of the range starts of a given frame range set */
-		LEARNING_API void MakeFrameSetFromRangeStarts(FFrameSet& OutFrameSet, const FFrameRangeSet& FrameRangeSet);
+		UE_API void MakeFrameSetFromRangeStarts(FFrameSet& OutFrameSet, const FFrameRangeSet& FrameRangeSet);
 
 		/** Create a frame set from all of the range ends of a given frame range set */
-		LEARNING_API void MakeFrameSetFromRangeEnds(FFrameSet& OutFrameSet, const FFrameRangeSet& FrameRangeSet);
+		UE_API void MakeFrameSetFromRangeEnds(FFrameSet& OutFrameSet, const FFrameRangeSet& FrameRangeSet);
 
 		/** Trim a frame range set to just the period of time before the first occurrence of a frame from a given frame set */
-		LEARNING_API void RangesBeforeFrameSet(FFrameRangeSet& OutFrameRangeSet, const FFrameRangeSet& FrameRangeSet, const FFrameSet& FrameSet);
+		UE_API void RangesBeforeFrameSet(FFrameRangeSet& OutFrameRangeSet, const FFrameRangeSet& FrameRangeSet, const FFrameSet& FrameSet);
 
 		/** Trim a frame range set to just the period of time after the first occurrence of a frame from a given frame set */
-		LEARNING_API void RangesAfterFrameSet(FFrameRangeSet& OutFrameRangeSet, const FFrameRangeSet& FrameRangeSet, const FFrameSet& FrameSet);
+		UE_API void RangesAfterFrameSet(FFrameRangeSet& OutFrameRangeSet, const FFrameRangeSet& FrameRangeSet, const FFrameSet& FrameSet);
 
 		/** Computes the entry indices associated with every range in the set */
-		LEARNING_API void AllRangeEntries(
+		UE_API void AllRangeEntries(
 			TLearningArrayView<1, int32> OutRangeEntries,
 			const FFrameRangeSet& FrameRangeSet);
 
 		/** Computes the range indices associated with every range in the set */
-		LEARNING_API void AllRangeIndices(
+		UE_API void AllRangeIndices(
 			TLearningArrayView<1, int32> OutRangeIndices,
 			const FFrameRangeSet& FrameRangeSet);
 
 		/** Computes the sequences associated with every range in the set */
-		LEARNING_API void AllRangeSequences(
+		UE_API void AllRangeSequences(
 			TLearningArrayView<1, int32> OutRangeSequences,
 			const FFrameRangeSet& FrameRangeSet);
 
 		/** Computes the start times associated with every range in the set */
-		LEARNING_API void AllRangeStartTimes(
+		UE_API void AllRangeStartTimes(
 			TLearningArrayView<1, float> OutRangeStartTimes,
 			const FFrameRangeSet& FrameRangeSet,
 			const float FrameDeltaTime);
 
 		/** Computes the end times associated with every range in the set */
-		LEARNING_API void AllRangeEndTimes(
+		UE_API void AllRangeEndTimes(
 			TLearningArrayView<1, float> OutRangeEndTimes,
 			const FFrameRangeSet& FrameRangeSet,
 			const float FrameDeltaTime);
 
 		/** Computes the durations associated with every range in the set */
-		LEARNING_API void AllRangeDurations(
+		UE_API void AllRangeDurations(
 			TLearningArrayView<1, float> OutRangeDurations,
 			const FFrameRangeSet& FrameRangeSet,
 			const float FrameDeltaTime);
 
 		/** Iterates over every range in the set and calls the provided callback */
-		LEARNING_API void ForEachRange(
+		UE_API void ForEachRange(
 			const FFrameRangeSet& FrameRangeSet,
 			const TFunctionRef<void(
 				const int32 TotalRangeIdx,
@@ -272,7 +272,7 @@ namespace UE::Learning
 				const int32 RangeIdx)> Body);
 
 		/** Iterates over every range in the set in parallel and calls the provided callback */
-		LEARNING_API void ParallelForEachRange(
+		UE_API void ParallelForEachRange(
 			const FFrameRangeSet& FrameRangeSet,
 			const TFunctionRef<void(
 				const int32 TotalRangeIdx,
diff --git a/Engine/Plugins/Experimental/LearningAgents/Source/Learning/Public/LearningFrameSet.h b/Engine/Plugins/Experimental/LearningAgents/Source/Learning/Public/LearningFrameSet.h
--- a/Engine/Plugins/Experimental/LearningAgents/Source/Learning/Public/LearningFrameSet.h
+++ b/Engine/Plugins/Experimental/LearningAgents/Source/Learning/Public/LearningFrameSet.h
@@ -121,40 +121,40 @@ namespace UE::Learning
 	namespace FrameSet
 	{
 		/** Checks if two frame sets are equal */
-		LEARNING_API bool Equal(const FFrameSet& Lhs, const FFrameSet& Rhs);
+		UE_API bool Equal(const FFrameSet& Lhs, const FFrameSet& Rhs);
 
 		/** Computes the union of two frame sets */
-		LEARNING_API void Union(FFrameSet& Out, const FFrameSet& Lhs, const FFrameSet& Rhs);
+		UE_API void Union(FFrameSet& Out, const FFrameSet& Lhs, const FFrameSet& Rhs);
 
 		/** Computes the intersection of two frame sets */
-		LEARNING_API void Intersection(FFrameSet& Out, const FFrameSet& Lhs, const FFrameSet& Rhs);
+		UE_API void Intersection(FFrameSet& Out, const FFrameSet& Lhs, const FFrameSet& Rhs);
 
 		/** Computes the difference of two frame sets */
-		LEARNING_API void Difference(FFrameSet& Out, const FFrameSet& Lhs, const FFrameSet& Rhs);
+		UE_API void Difference(FFrameSet& Out, const FFrameSet& Lhs, const FFrameSet& Rhs);
 
 		/** Computes the entry indices associated with every frame in the set */
-		LEARNING_API void AllFrameEntries(
+		UE_API void AllFrameEntries(
 			TLearningArrayView<1, int32> OutFrameEntries,
 			const FFrameSet& FrameSet);
 
 		/** Computes the frame indices associated with every frame in the set */
-		LEARNING_API void AllFrameIndices(
+		UE_API void AllFrameIndices(
 			TLearningArrayView<1, int32> OutFrameIndices,
 			const FFrameSet& FrameSet);
 
 		/** Computes the sequences associated with every frame in the set */
-		LEARNING_API void AllFrameSequences(
+		UE_API void AllFrameSequences(
 			TLearningArrayView<1, int32> OutFrameSequences,
 			const FFrameSet& FrameSet);
 
 		/** Computes the times associated with every frame in the set */
-		LEARNING_API void AllFrameTimes(
+		UE_API void AllFrameTimes(
 			TLearningArrayView<1, float> OutFrameTimes,
 			const FFrameSet& FrameSet,
 			const float FrameDeltaTime);
 
 		/** Iterates over every frame in the set and calls the provided callback */
-		LEARNING_API void ForEachFrame(
+		UE_API void ForEachFrame(
 			const FFrameSet& FrameSet,
 			const TFunctionRef<void(
 				const int32 TotalFrameIdx,
@@ -162,7 +162,7 @@ namespace UE::Learning
 				const int32 FrameIdx)> Body);
 
 		/** Iterates over every frame in the set in parallel and calls the provided callback */
-		LEARNING_API void ParallelForEachFrame(
+		UE_API void ParallelForEachFrame(
 			const FFrameSet& FrameSet,
 			const TFunctionRef<void(
 				const int32 TotalFrameIdx,
diff --git a/Engine/Plugins/Experimental/LearningAgents/Source/Learning/Public/LearningKMeans.h b/Engine/Plugins/Experimental/LearningAgents/Source/Learning/Public/LearningKMeans.h
--- a/Engine/Plugins/Experimental/LearningAgents/Source/Learning/Public/LearningKMeans.h
+++ b/Engine/Plugins/Experimental/LearningAgents/Source/Learning/Public/LearningKMeans.h
@@ -4,6 +4,8 @@
 
 #include "LearningArray.h"
 
+#define UE_API LEARNING_API
+
 namespace UE::Learning::KMeans
 {
 
@@ -15,7 +17,7 @@ namespace UE::Learning::KMeans
 	* @param Samples			Samples to be clustered of shape (SamplesNum, DimensionNum)
 	* @param Seed				Random Seed
 	*/
-	LEARNING_API void InitCenters(
+	UE_API void InitCenters(
 		TLearningArrayView<2, float> OutCenters,
 		const TLearningArrayView<2, const float> Samples,
 		const uint32 Seed);
@@ -29,7 +31,7 @@ namespace UE::Learning::KMeans
 	* @param bParallelEvaluation	If to allow for multi-threaded evaluation
 	* @param MinParallelBatchSize	Minimum batch size to use for multi-threaded evaluation
 	*/
-	LEARNING_API void UpdateAssignmentsFromCenters(
+	UE_API void UpdateAssignmentsFromCenters(
 		TLearningArrayView<1, int32> OutAssignments,
 		const TLearningArrayView<2, const float> Centers,
 		const TLearningArrayView<2, const float> Samples,
@@ -42,7 +44,7 @@ namespace UE::Learning::KMeans
 	* @param OutAssignmentCounts	Number of samples assigned to each cluster - shape (ClusterNum)
 	* @param Assignments			Index of the cluster assigned to each sample - shape (SamplesNum
 	*/
-	LEARNING_API void CountClusterAssignments(
+	UE_API void CountClusterAssignments(
 		TLearningArrayView<1, int32> OutAssignmentCounts,
 		const TLearningArrayView<1, const int32> Assignments);
 
@@ -54,7 +56,7 @@ namespace UE::Learning::KMeans
 	* @param AssignmentCounts	Number of samples assigned to each cluster - shape (ClusterNum)
 	* @param Samples			Samples to be clustered of shape (SamplesNum, DimensionNum)
 	*/
-	LEARNING_API void UpdateCenters(
+	UE_API void UpdateCenters(
 		TLearningArrayView<2, float> OutCenters,
 		const TLearningArrayView<1, const int32> Assignments,
 		const TLearningArrayView<1, const int32> AssignmentCounts,
@@ -69,7 +71,7 @@ namespace UE::Learning::KMeans
 	* @param Samples			Samples to be clustered of shape (SamplesNum, DimensionNum)
 	* @param Seed				Random Seed
 	*/
-	LEARNING_API void InitBounds(
+	UE_API void InitBounds(
 		TLearningArrayView<2, float> OutMins,
 		TLearningArrayView<2, float> OutMaxs,
 		const TLearningArrayView<2, const float> Samples,
@@ -85,7 +87,7 @@ namespace UE::Learning::KMeans
 	* @param bParallelEvaluation	If to allow for multi-threaded evaluation
 	* @param MinParallelBatchSize	Minimum batch size to use for multi-threaded evaluation
 	*/
-	LEARNING_API void UpdateAssignmentsFromBounds(
+	UE_API void UpdateAssignmentsFromBounds(
 		TLearningArrayView<1, int32> OutAssignments,
 		const TLearningArrayView<2, const float> Mins,
 		const TLearningArrayView<2, const float> Maxs,
@@ -102,7 +104,7 @@ namespace UE::Learning::KMeans
 	* @param AssignmentCounts	Number of samples assigned to each cluster - shape (ClusterNum)
 	* @param Samples			Samples to be clustered of shape (SamplesNum, DimensionNum)
 	*/
-	LEARNING_API void UpdateBounds(
+	UE_API void UpdateBounds(
 		TLearningArrayView<2, float> OutMins,
 		TLearningArrayView<2, float> OutMaxs,
 		const TLearningArrayView<1, const int32> Assignments,
@@ -120,7 +122,7 @@ namespace UE::Learning::KMeans
 	* @param AssignmentCounts			Number of samples assigned to each cluster - shape (ClusterNum)
 	* @param Samples					Samples to be clustered of shape (SamplesNum, DimensionNum)
 	*/
-	LEARNING_API void ComputeClusteredIndex(
+	UE_API void ComputeClusteredIndex(
 		TLearningArrayView<2, float> OutClusteredSamples,
 		TLearningArrayView<1, int32> OutClusterStarts,
 		TLearningArrayView<1, int32> OutClusterLengths,
@@ -129,6 +131,6 @@ namespace UE::Learning::KMeans
 		const TLearningArrayView<1, const int32> Assignments,
 		const TLearningArrayView<1, const int32> AssignmentCounts,
 		const TLearningArrayView<2, const float> Samples);
+}
 
-
-}
\ No newline at end of file
+#undef UE_API
\ No newline at end of file
diff --git a/Engine/Plugins/Experimental/LearningAgents/Source/Learning/Public/LearningLog.h b/Engine/Plugins/Experimental/LearningAgents/Source/Learning/Public/LearningLog.h
--- a/Engine/Plugins/Experimental/LearningAgents/Source/Learning/Public/LearningLog.h
+++ b/Engine/Plugins/Experimental/LearningAgents/Source/Learning/Public/LearningLog.h
@@ -4,6 +4,8 @@
 
 #include "Logging/LogMacros.h"
 
+#define UE_API LEARNING_API
+
 namespace UE::Learning
 {
 	/**
@@ -19,4 +21,6 @@ namespace UE::Learning
 	};
 }
 
-LEARNING_API DECLARE_LOG_CATEGORY_EXTERN(LogLearning, Log, All);
+UE_API DECLARE_LOG_CATEGORY_EXTERN(LogLearning, Log, All);
+
+#undef UE_API
\ No newline at end of file
diff --git a/Engine/Plugins/Experimental/LearningAgents/Source/Learning/Public/LearningObservation.h b/Engine/Plugins/Experimental/LearningAgents/Source/Learning/Public/LearningObservation.h
--- a/Engine/Plugins/Experimental/LearningAgents/Source/Learning/Public/LearningObservation.h
+++ b/Engine/Plugins/Experimental/LearningAgents/Source/Learning/Public/LearningObservation.h
@@ -645,7 +645,7 @@ namespace UE::Learning::Observation
 	 * @param SchemaElement				Observation Schema Element
 	 * @param Salt						Hash salt
 	 */
-	LEARNING_API int32 GetSchemaObjectsCompatibilityHash(
+	UE_API int32 GetSchemaObjectsCompatibilityHash(
 		const FSchema& Schema,
 		const FSchemaElement SchemaElement,
 		const int32 Salt = 0x62625165);
@@ -659,7 +659,7 @@ namespace UE::Learning::Observation
 	 * @param SchemaElementB			Second Schema Element
 	 * @returns							true when the schemas are compatible
 	 */
-	LEARNING_API bool AreSchemaObjectsCompatible(
+	UE_API bool AreSchemaObjectsCompatible(
 		const FSchema& SchemaA,
 		const FSchemaElement SchemaElementA,
 		const FSchema& SchemaB,
@@ -694,7 +694,7 @@ namespace UE::Learning::Observation
 	 * @param SchemaElement				Observation Schema Element
 	 * @param NetworkSettings			Network Generation Settings
 	 */
-	LEARNING_API void MakeEncoderNetworkModelBuilderElementFromSchema(
+	UE_API void MakeEncoderNetworkModelBuilderElementFromSchema(
 		NNE::RuntimeBasic::FModelBuilderElement& OutElement,
 		NNE::RuntimeBasic::FModelBuilder& Builder,
 		const FSchema& Schema,
@@ -713,7 +713,7 @@ namespace UE::Learning::Observation
 	 * @param NetworkSettings			Network Generation Settings
 	 * @param Seed						The random seed used in initializing the network weights
 	 */
-	LEARNING_API void GenerateEncoderNetworkFileDataFromSchema(
+	UE_API void GenerateEncoderNetworkFileDataFromSchema(
 		TArray<uint8>& OutFileData,
 		uint32& OutInputSize,
 		uint32& OutOutputSize,
@@ -731,7 +731,7 @@ namespace UE::Learning::Observation
 	 * @param Object					Observation Object
 	 * @param ObjectElement				Observation Object Element
 	 */
-	LEARNING_API void SetVectorFromObject(
+	UE_API void SetVectorFromObject(
 		TLearningArrayView<1, float> OutObservationVector,
 		const FSchema& Schema,
 		const FSchemaElement SchemaElement,
@@ -747,7 +747,7 @@ namespace UE::Learning::Observation
 	 * @param SchemaElement				Observation Schema Element
 	 * @param ObservationVector			Input Observation vector
 	 */
-	LEARNING_API void GetObjectFromVector(
+	UE_API void GetObjectFromVector(
 		FObject& OutObject,
 		FObjectElement& OutObjectElement,
 		const FSchema& Schema,
@@ -764,7 +764,7 @@ namespace UE::Learning::Observation
 	 * @param Schema					Observation Schema
 	 * @param SchemaElement				Observation Schema Element
 	 */
-	LEARNING_API void GetContinuousActionVectorMask(
+	UE_API void GetContinuousActionVectorMask(
 		TLearningArrayView<1, bool> OutObservationVectorMask,
 		const TLearningArrayView<1, const float> InObservationVector,
 		const FSchema& Schema,
@@ -776,7 +776,7 @@ namespace UE::Learning::Observation
 	 * @param OutObservationDistributionMin		Minimum values for observations
 	 * @param OutObservationDistributionMax		Maximum values for observations
 	 */
-	LEARNING_API void InitObservationDistributionMinMax(
+	UE_API void InitObservationDistributionMinMax(
 		TLearningArrayView<1, float> OutObservationDistributionMin,
 		TLearningArrayView<1, float> OutObservationDistributionMax);
 
@@ -790,7 +790,7 @@ namespace UE::Learning::Observation
 	 * @param Schema								Observation Schema
 	 * @param SchemaElement							Observation Schema Element
 	 */
-	LEARNING_API void FitObservationDistributionMinMax(
+	UE_API void FitObservationDistributionMinMax(
 		TLearningArrayView<1, float> InOutObservationDistributionMin,
 		TLearningArrayView<1, float> InOutObservationDistributionMax,
 		const TLearningArrayView<1, const float> InObservationVector,
@@ -806,7 +806,7 @@ namespace UE::Learning::Observation
 	 * @param Schema								Observation Schema
 	 * @param SchemaElement							Observation Schema Element
 	 */
-	LEARNING_API void ClampToObservationDistributionMinMax(
+	UE_API void ClampToObservationDistributionMinMax(
 		TLearningArrayView<1, float> InOutObservationVector,
 		const TLearningArrayView<1, const float> ObservationDistributionMin,
 		const TLearningArrayView<1, const float> ObservationDistributionMax,
diff --git a/Engine/Plugins/Experimental/LearningAgents/Source/Learning/Public/LearningOptimizationPlanner.h b/Engine/Plugins/Experimental/LearningAgents/Source/Learning/Public/LearningOptimizationPlanner.h
--- a/Engine/Plugins/Experimental/LearningAgents/Source/Learning/Public/LearningOptimizationPlanner.h
+++ b/Engine/Plugins/Experimental/LearningAgents/Source/Learning/Public/LearningOptimizationPlanner.h
@@ -39,7 +39,7 @@ namespace UE::Learning
 		* @param ActionVectors					Action Vectors of shape (StepNum, InstanceNum, ActionVectorDimNum)
 		* @param Instances						Set of instances to run action vectors for
 		*/
-		LEARNING_API void RunPlan(
+		UE_API void RunPlan(
 			TLearningArrayView<2, float> ActionVectorBuffer,
 			const TFunctionRef<void(const FIndexSet Instances)> ResetFunction,
 			const TFunctionRef<void(const FIndexSet Instances)> ActionFunction,
@@ -66,7 +66,7 @@ namespace UE::Learning
 		* @param ActionVectorsLock				Optional lock to use when updating InOutActionVectors
 		* @param bActionVectorsUpdatedFlag		Optional signal to set when updating InOutActionVectors
 		*/
-		LEARNING_API void Plan(
+		UE_API void Plan(
 			TLearningArrayView<2, float> InOutActionVectors,
 			FOptimizationPlannerBuffer& OptimizationPlannerBuffer,
 			IOptimizer& Optimizer,
@@ -110,7 +110,7 @@ namespace UE::Learning
 		* @param ActionVectorsLock				Optional lock to use when updating InOutActionVectors
 		* @param bActionVectorsUpdatedFlag		Optional signal to set when updating InOutActionVectors
 		*/
-		LEARNING_API void PlanWindowed(
+		UE_API void PlanWindowed(
 			TLearningArrayView<2, float> InOutActionVectors,
 			FOptimizationPlannerBuffer& OptimizationPlannerBuffer,
 			IOptimizer& Optimizer,
diff --git a/Engine/Plugins/Experimental/LearningAgents/Source/Learning/Public/LearningRandom.h b/Engine/Plugins/Experimental/LearningAgents/Source/Learning/Public/LearningRandom.h
--- a/Engine/Plugins/Experimental/LearningAgents/Source/Learning/Public/LearningRandom.h
+++ b/Engine/Plugins/Experimental/LearningAgents/Source/Learning/Public/LearningRandom.h
@@ -7,6 +7,8 @@
 #include "Math/Vector.h"
 #include "Math/Quat.h"
 
+#define UE_API LEARNING_API
+
 namespace UE::Learning
 {
 	struct FFrameSet;
@@ -61,48 +63,48 @@ namespace UE::Learning::Random
 	// 
 	/////////////////////////////////////////////////////////
 
-	LEARNING_API uint32 Int(const uint32 State);
+	UE_API uint32 Int(const uint32 State);
 
-	LEARNING_API float Float(const uint32 State);
+	UE_API float Float(const uint32 State);
 
-	LEARNING_API int32 IntInRange(const uint32 State, const int32 Min, const int32 Max);
+	UE_API int32 IntInRange(const uint32 State, const int32 Min, const int32 Max);
 
-	LEARNING_API float Uniform(
+	UE_API float Uniform(
 		const uint32 State,
 		const float Min = 0.0f,
 		const float Max = 1.0f);
 
-	LEARNING_API float Gaussian(
+	UE_API float Gaussian(
 		const uint32 State,
 		const float Mean = 0.0f,
 		const float Std = 1.0f);
 
-	LEARNING_API FVector VectorGaussian(
+	UE_API FVector VectorGaussian(
 		const uint32 State,
 		const FVector Mean = FVector::ZeroVector,
 		const FVector Std = FVector::OneVector);
 
-	LEARNING_API float ClippedGaussian(
+	UE_API float ClippedGaussian(
 		const uint32 State,
 		const float Mean = 0.0f,
 		const float Std = 1.0f,
 		const float Clip = 10.0f);
 
-	LEARNING_API FVector PlanarUniform(
+	UE_API FVector PlanarUniform(
 		const uint32 State,
 		const float Min = 0.0f,
 		const float Max = 1.0f,
 		const FVector Axis0 = FVector::ForwardVector,
 		const FVector Axis1 = FVector::RightVector);
 
-	LEARNING_API FVector PlanarGaussian(
+	UE_API FVector PlanarGaussian(
 		const uint32 State,
 		const float Mean = 0.0f,
 		const float Std = 1.0f,
 		const FVector Axis0 = FVector::ForwardVector,
 		const FVector Axis1 = FVector::RightVector);
 
-	LEARNING_API FVector PlanarClippedGaussian(
+	UE_API FVector PlanarClippedGaussian(
 		const uint32 State,
 		const float Mean = 0.0f,
 		const float Std = 1.0f,
@@ -110,46 +112,46 @@ namespace UE::Learning::Random
 		const FVector Axis0 = FVector::ForwardVector,
 		const FVector Axis1 = FVector::RightVector);
 
-	LEARNING_API FVector PlanarDirection(
+	UE_API FVector PlanarDirection(
 		const uint32 State,
 		const FVector Axis0 = FVector::ForwardVector,
 		const FVector Axis1 = FVector::RightVector);
 
-	LEARNING_API FQuat Rotation(const uint32 Input);
+	UE_API FQuat Rotation(const uint32 Input);
 
-	LEARNING_API void IntArray(
+	UE_API void IntArray(
 		TLearningArrayView<1, uint32> Output,
 		const uint32 State);
 
-	LEARNING_API void IntArray(
+	UE_API void IntArray(
 		TLearningArrayView<1, uint32> Output,
 		const uint32 State,
 		const FIndexSet Indices);
 
-	LEARNING_API void FloatArray(
+	UE_API void FloatArray(
 		TLearningArrayView<1, float> Output,
 		const uint32 State);
 
-	LEARNING_API void UniformArray(
+	UE_API void UniformArray(
 		TLearningArrayView<1, float> Output,
 		const uint32 State,
 		const float Min = 0.0f,
 		const float Max = 1.0f);
 
-	LEARNING_API void GaussianArray(
+	UE_API void GaussianArray(
 		TLearningArrayView<1, float> Output,
 		const uint32 State,
 		const float Mean = 0.0f,
 		const float Std = 1.0f);
 
-	LEARNING_API void ClippedGaussianArray(
+	UE_API void ClippedGaussianArray(
 		TLearningArrayView<1, float> Output,
 		const uint32 State,
 		const float Mean = 0.0f,
 		const float Std = 1.0f,
 		const float Clip = 10.0f);
 
-	LEARNING_API void PlanarClippedGaussianArray(
+	UE_API void PlanarClippedGaussianArray(
 		TLearningArrayView<1, FVector> Output,
 		const uint32 State,
 		const float Mean = 0.0f,
@@ -158,20 +160,20 @@ namespace UE::Learning::Random
 		const FVector Axis0 = FVector::ForwardVector,
 		const FVector Axis1 = FVector::RightVector);
 
-	LEARNING_API void PlanarDirectionArray(
+	UE_API void PlanarDirectionArray(
 		TLearningArrayView<1, FVector> Output,
 		const uint32 State,
 		const FVector Axis0 = FVector::ForwardVector,
 		const FVector Axis1 = FVector::RightVector);
 
-	LEARNING_API void DistributionIndependantNormal(
+	UE_API void DistributionIndependantNormal(
 		TLearningArrayView<1, float> Output,
 		const uint32 State,
 		const TLearningArrayView<1, const float> Mean,
 		const TLearningArrayView<1, const float> LogStd,
 		const float Scale);
 
-	LEARNING_API void DistributionIndependantNormalMasked(
+	UE_API void DistributionIndependantNormalMasked(
 		TLearningArrayView<1, float> Output,
 		const uint32 State,
 		const TLearningArrayView<1, const float> Mean,
@@ -180,39 +182,39 @@ namespace UE::Learning::Random
 		const TLearningArrayView<1, const float> MaskedValues,
 		const float Scale);
 
-	LEARNING_API void DistributionMultinoulli(
+	UE_API void DistributionMultinoulli(
 		TLearningArrayView<1, float> Output,
 		const uint32 State,
 		const TLearningArrayView<1, const float> Logits,
 		const float Scale);
 
-	LEARNING_API void DistributionMultinoulliMasked(
+	UE_API void DistributionMultinoulliMasked(
 		TLearningArrayView<1, float> Output,
 		const uint32 State,
 		const TLearningArrayView<1, const float> Logits,
 		const TLearningArrayView<1, const bool> Masked,
 		const float Scale);
 
-	LEARNING_API void DistributionBernoulli(
+	UE_API void DistributionBernoulli(
 		TLearningArrayView<1, float> Output,
 		const uint32 State,
 		const TLearningArrayView<1, const float> Logits,
 		const float Scale);
 
-	LEARNING_API void DistributionBernoulliMasked(
+	UE_API void DistributionBernoulliMasked(
 		TLearningArrayView<1, float> Output,
 		const uint32 State,
 		const TLearningArrayView<1, const float> Logits,
 		const TLearningArrayView<1, const bool> Masked,
 		const float Scale);
 
-	LEARNING_API void FrameInFrameSet(
+	UE_API void FrameInFrameSet(
 		int32& OutEntryIdx,
 		int32& OutFrameIdx,
 		const FFrameSet& FrameSet,
 		const uint32 State);
 
-	LEARNING_API void FrameInFrameRangeSet(
+	UE_API void FrameInFrameRangeSet(
 		int32& OutEntryIdx,
 		int32& OutRangeIdx,
 		int32& OutRangeFrame,
@@ -226,48 +228,48 @@ namespace UE::Learning::Random
 	// 
 	/////////////////////////////////////////////////////////
 
-	LEARNING_API uint32 SampleInt(uint32& State);
+	UE_API uint32 SampleInt(uint32& State);
 
-	LEARNING_API float SampleFloat(uint32& State);
+	UE_API float SampleFloat(uint32& State);
 
-	LEARNING_API int32 SampleIntInRange(uint32& State, const int32 Min, const int32 Max);
+	UE_API int32 SampleIntInRange(uint32& State, const int32 Min, const int32 Max);
 
-	LEARNING_API float SampleUniform(
+	UE_API float SampleUniform(
 		uint32& State,
 		const float Min = 0.0f,
 		const float Max = 1.0f);
 
-	LEARNING_API float SampleGaussian(
+	UE_API float SampleGaussian(
 		uint32& State,
 		const float Mean = 0.0f,
 		const float Std = 1.0f);
 
-	LEARNING_API FVector SampleVectorGaussian(
+	UE_API FVector SampleVectorGaussian(
 		uint32& State,
 		const FVector Mean = FVector::ZeroVector,
 		const FVector Std = FVector::OneVector);
 
-	LEARNING_API float SampleClippedGaussian(
+	UE_API float SampleClippedGaussian(
 		uint32& State,
 		const float Mean = 0.0f,
 		const float Std = 1.0f,
 		const float Clip = 10.0f);
 
-	LEARNING_API FVector SamplePlanarUniform(
+	UE_API FVector SamplePlanarUniform(
 		uint32& State,
 		const float Min = 0.0f,
 		const float Max = 1.0f,
 		const FVector Axis0 = FVector::ForwardVector,
 		const FVector Axis1 = FVector::RightVector);
 
-	LEARNING_API FVector SamplePlanarGaussian(
+	UE_API FVector SamplePlanarGaussian(
 		uint32& State,
 		const float Mean = 0.0f,
 		const float Std = 1.0f,
 		const FVector Axis0 = FVector::ForwardVector,
 		const FVector Axis1 = FVector::RightVector);
 
-	LEARNING_API FVector SamplePlanarClippedGaussian(
+	UE_API FVector SamplePlanarClippedGaussian(
 		uint32& State,
 		const float Mean = 0.0f,
 		const float Std = 1.0f,
@@ -275,46 +277,46 @@ namespace UE::Learning::Random
 		const FVector Axis0 = FVector::ForwardVector,
 		const FVector Axis1 = FVector::RightVector);
 
-	LEARNING_API FVector SamplePlanarDirection(
+	UE_API FVector SamplePlanarDirection(
 		uint32& State,
 		const FVector Axis0 = FVector::ForwardVector,
 		const FVector Axis1 = FVector::RightVector);
 
-	LEARNING_API FQuat SampleRotation(uint32& State);
+	UE_API FQuat SampleRotation(uint32& State);
 
-	LEARNING_API void SampleIntArray(
+	UE_API void SampleIntArray(
 		TLearningArrayView<1, uint32> Output,
 		uint32& State);
 
-	LEARNING_API void SampleIntArray(
+	UE_API void SampleIntArray(
 		TLearningArrayView<1, uint32> Output,
 		uint32& State,
 		const FIndexSet Indices);
 
-	LEARNING_API void SampleFloatArray(
+	UE_API void SampleFloatArray(
 		TLearningArrayView<1, float> Output,
 		uint32& State);
 
-	LEARNING_API void SampleUniformArray(
+	UE_API void SampleUniformArray(
 		TLearningArrayView<1, float> Output,
 		uint32& State,
 		const float Min = 0.0f,
 		const float Max = 1.0f);
 
-	LEARNING_API void SampleGaussianArray(
+	UE_API void SampleGaussianArray(
 		TLearningArrayView<1, float> Output,
 		uint32& State,
 		const float Mean = 0.0f,
 		const float Std = 1.0f);
 	
-	LEARNING_API void SampleClippedGaussianArray(
+	UE_API void SampleClippedGaussianArray(
 		TLearningArrayView<1, float> Output,
 		uint32& State,
 		const float Mean = 0.0f,
 		const float Std = 1.0f,
 		const float Clip = 10.0f);
 
-	LEARNING_API void SamplePlanarClippedGaussianArray(
+	UE_API void SamplePlanarClippedGaussianArray(
 		TLearningArrayView<1, FVector> Output,
 		uint32& State,
 		const float Mean = 0.0f,
@@ -323,20 +325,20 @@ namespace UE::Learning::Random
 		const FVector Axis0 = FVector::ForwardVector,
 		const FVector Axis1 = FVector::RightVector);
 
-	LEARNING_API void SamplePlanarDirectionArray(
+	UE_API void SamplePlanarDirectionArray(
 		TLearningArrayView<1, FVector> Output,
 		uint32& State,
 		const FVector Axis0 = FVector::ForwardVector,
 		const FVector Axis1 = FVector::RightVector);
 
-	LEARNING_API void SampleDistributionIndependantNormal(
+	UE_API void SampleDistributionIndependantNormal(
 		TLearningArrayView<1, float> Output,
 		uint32& State,
 		const TLearningArrayView<1, const float> Mean,
 		const TLearningArrayView<1, const float> LogStd,
 		const float Scale);
 
-	LEARNING_API void SampleDistributionIndependantNormalMasked(
+	UE_API void SampleDistributionIndependantNormalMasked(
 		TLearningArrayView<1, float> Output,
 		uint32& State,
 		const TLearningArrayView<1, const float> Mean,
@@ -345,42 +347,44 @@ namespace UE::Learning::Random
 		const TLearningArrayView<1, const float> MaskedValues,
 		const float Scale);
 
-	LEARNING_API void SampleDistributionMultinoulli(
+	UE_API void SampleDistributionMultinoulli(
 		TLearningArrayView<1, float> Output,
 		uint32& State,
 		const TLearningArrayView<1, const float> Logits,
 		const float Scale);
 
-	LEARNING_API void SampleDistributionMultinoulliMasked(
+	UE_API void SampleDistributionMultinoulliMasked(
 		TLearningArrayView<1, float> Output,
 		uint32& State,
 		const TLearningArrayView<1, const float> Logits,
 		const TLearningArrayView<1, const bool> Masked,
 		const float Scale);
 
-	LEARNING_API void SampleDistributionBernoulli(
+	UE_API void SampleDistributionBernoulli(
 		TLearningArrayView<1, float> Output,
 		uint32& State,
 		const TLearningArrayView<1, const float> Logits,
 		const float Scale);
 
-	LEARNING_API void SampleDistributionBernoulliMasked(
+	UE_API void SampleDistributionBernoulliMasked(
 		TLearningArrayView<1, float> Output,
 		uint32& State,
 		const TLearningArrayView<1, const float> Logits,
 		const TLearningArrayView<1, const bool> Masked,
 		const float Scale);
 
-	LEARNING_API void SampleFrameInFrameSet(
+	UE_API void SampleFrameInFrameSet(
 		int32& OutEntryIdx,
 		int32& OutFrameIdx,
 		uint32& State,
 		const FFrameSet& FrameSet);
 
-	LEARNING_API void SampleFrameInFrameRangeSet(
+	UE_API void SampleFrameInFrameRangeSet(
 		int32& OutEntryIdx,
 		int32& OutRangeIdx,
 		int32& OutRangeFrame,
 		uint32& State,
 		const FFrameRangeSet& FrameRangeSet);
-}
\ No newline at end of file
+}
+
+#undef UE_API
\ No newline at end of file
diff --git a/Engine/Plugins/Experimental/LearningAgents/Source/LearningTraining/Public/LearningExperience.h b/Engine/Plugins/Experimental/LearningAgents/Source/LearningTraining/Public/LearningExperience.h
--- a/Engine/Plugins/Experimental/LearningAgents/Source/LearningTraining/Public/LearningExperience.h
+++ b/Engine/Plugins/Experimental/LearningAgents/Source/LearningTraining/Public/LearningExperience.h
@@ -293,7 +293,7 @@ namespace UE::Learning
 		* @param CompletionFunction						Function to run for evaluating completions
 		* @param Instances								Set of instances to gather experience for
 		*/
-		LEARNINGTRAINING_API void GatherExperienceUntilReplayBufferFull(
+		UE_API void GatherExperienceUntilReplayBufferFull(
 			FReplayBuffer& ReplayBuffer,
 			FEpisodeBuffer& EpisodeBuffer,
 			FResetInstanceBuffer& ResetBuffer,
diff --git a/Engine/Plugins/Experimental/LearningAgents/Source/LearningTraining/Public/LearningPPOTrainer.h b/Engine/Plugins/Experimental/LearningAgents/Source/LearningTraining/Public/LearningPPOTrainer.h
--- a/Engine/Plugins/Experimental/LearningAgents/Source/LearningTraining/Public/LearningPPOTrainer.h
+++ b/Engine/Plugins/Experimental/LearningAgents/Source/LearningTraining/Public/LearningPPOTrainer.h
@@ -190,7 +190,7 @@ namespace UE::Learning
 		* @param LogSettings							Logging settings
 		* @returns										Trainer response in case of errors during communication otherwise Success
 		*/
-		LEARNINGTRAINING_API ETrainerResponse Train(
+		UE_API ETrainerResponse Train(
 			IExternalTrainer* ExternalTrainer,
 			FReplayBuffer& ReplayBuffer,
 			FEpisodeBuffer& EpisodeBuffer,
diff --git a/Engine/Plugins/Experimental/LearningAgents/Source/LearningTraining/Public/LearningTrainer.h b/Engine/Plugins/Experimental/LearningAgents/Source/LearningTraining/Public/LearningTrainer.h
--- a/Engine/Plugins/Experimental/LearningAgents/Source/LearningTraining/Public/LearningTrainer.h
+++ b/Engine/Plugins/Experimental/LearningAgents/Source/LearningTraining/Public/LearningTrainer.h
@@ -150,12 +150,12 @@ namespace UE::Learning
 		/**
 		* Converts a ETrainerDevice into a string.
 		*/
-		LEARNINGTRAINING_API const TCHAR* GetDeviceString(const ETrainerDevice Device);
+		UE_API const TCHAR* GetDeviceString(const ETrainerDevice Device);
 
 		/**
 		* Converts a ETrainerResponse into a string for use in logging and error messages.
 		*/
-		LEARNINGTRAINING_API const TCHAR* GetResponseString(const ETrainerResponse Response);
+		UE_API const TCHAR* GetResponseString(const ETrainerResponse Response);
 
 		/**
 		* Compute the discount factor that corresponds to a particular HalfLife and DeltaTime.
@@ -164,7 +164,7 @@ namespace UE::Learning
 		* @param DeltaTime		DeltaTime taken upon each step of the environment
 		* @returns				Corresponding discount factor
 		*/
-		LEARNINGTRAINING_API float DiscountFactorFromHalfLife(const float HalfLife, const float DeltaTime);
+		UE_API float DiscountFactorFromHalfLife(const float HalfLife, const float DeltaTime);
 
 		/**
 		* Compute the discount factor that corresponds to a particular HalfLife provided in terms of number of steps
@@ -172,67 +172,67 @@ namespace UE::Learning
 		* @param HalfLifeSteps	Number of steps taken at which the reward should be discounted by half
 		* @returns				Corresponding discount factor
 		*/
-		LEARNINGTRAINING_API float DiscountFactorFromHalfLifeSteps(const int32 HalfLifeSteps);
+		UE_API float DiscountFactorFromHalfLifeSteps(const int32 HalfLifeSteps);
 
 		/**
 		* Gets the python executable path from the engine directory.
 		*/
-		LEARNINGTRAINING_API FString GetPythonExecutablePath(const FString& EngineDir);
+		UE_API FString GetPythonExecutablePath(const FString& EngineDir);
 
 		/**
 		* Gets the PythonFoundationPackages site-packages path from the engine directory.
 		*/
-		LEARNINGTRAINING_API FString GetSitePackagesPath(const FString& EngineDir);
+		UE_API FString GetSitePackagesPath(const FString& EngineDir);
 
 		/**
 		* Gets the LearningAgents Content path from the engine directory.
 		*/
-		LEARNINGTRAINING_API FString GetPythonContentPath(const FString& EngineDir);
+		UE_API FString GetPythonContentPath(const FString& EngineDir);
 
 		/**
 		* Gets the project's Python Content path.
 		*/
-		LEARNINGTRAINING_API FString GetProjectPythonContentPath();
+		UE_API FString GetProjectPythonContentPath();
 
 		/**
 		* Gets the LearningAgents Intermediate path from the intermediate directory.
 		*/
-		LEARNINGTRAINING_API FString GetIntermediatePath(const FString& IntermediateDir);
+		UE_API FString GetIntermediatePath(const FString& IntermediateDir);
 
 		/**
 		* Converts an observation schema element into a JSON representation.
 		*/
-		LEARNINGTRAINING_API TSharedPtr<FJsonObject> ConvertObservationSchemaToJSON(
+		UE_API TSharedPtr<FJsonObject> ConvertObservationSchemaToJSON(
 			const Observation::FSchema& ObservationSchema,
 			const Observation::FSchemaElement& ObservationSchemaElement);
 
 		/**
 		* Converts an action schema element into a JSON representation.
 		*/
-		LEARNINGTRAINING_API TSharedPtr<FJsonObject> ConvertActionSchemaToJSON(
+		UE_API TSharedPtr<FJsonObject> ConvertActionSchemaToJSON(
 			const Action::FSchema& ActionSchema,
 			const Action::FSchemaElement& ActionSchemaElement);
 
 		/* Returns true if the schema and element are a subset of the provided JSON. */
-		LEARNINGTRAINING_API bool IsObservationSchemaSubsetCompatible(
+		UE_API bool IsObservationSchemaSubsetCompatible(
 			const FString& SourceJsonString,
 			const Observation::FSchema& ObservationSchema,
 			const Observation::FSchemaElement& ObservationSchemaElement);
 
 		/* Returns true if the schema and element are a subset of the provided JSON. */
-		LEARNINGTRAINING_API bool IsObservationSchemaSubsetCompatible(
+		UE_API bool IsObservationSchemaSubsetCompatible(
 			const TSharedRef<FJsonObject> Object,
 			const Observation::FSchema& ObservationSchema,
 			const Observation::FSchemaElement& ObservationSchemaElement);
 
 		/* Computes the indices to copy from a subset compatible JSON. */
-		LEARNINGTRAINING_API TArray<int32> ComputeObservationSchemaSubsetIndices(
+		UE_API TArray<int32> ComputeObservationSchemaSubsetIndices(
 			const FString& SourceJsonString,
 			const Observation::FSchema& ObservationSchema,
 			const Observation::FSchemaElement& ObservationSchemaElement);
 
 		/* Computes the indices to copy from a subset compatible JSON. */
-		LEARNINGTRAINING_API void ComputeObservationSchemaSubsetIndices(
+		UE_API void ComputeObservationSchemaSubsetIndices(
 			int32& NextIndex,
 			TArray<int32>& Indices,
 			const TSharedRef<FJsonObject> Object,
@@ -240,25 +240,25 @@ namespace UE::Learning
 			const Observation::FSchemaElement& ObservationSchemaElement);
 
 		/* Returns true if the schema and element are a subset of the provided JSON. */
-		LEARNINGTRAINING_API bool IsActionSchemaSubsetCompatible(
+		UE_API bool IsActionSchemaSubsetCompatible(
 			const FString& SourceJsonString,
 			const Action::FSchema& ActionSchema,
 			const Action::FSchemaElement& ActionSchemaElement);
 
 		/* Returns true if the schema and element are a subset of the provided JSON. */
-		LEARNINGTRAINING_API bool IsActionSchemaSubsetCompatible(
+		UE_API bool IsActionSchemaSubsetCompatible(
 			const TSharedRef<FJsonObject> Object,
 			const Action::FSchema& ActionSchema,
 			const Action::FSchemaElement& ActionSchemaElement);
 
 		/* Computes the indices to copy from a subset compatible JSON. */
-		LEARNINGTRAINING_API TArray<int32> ComputeActionSchemaSubsetIndices(
+		UE_API TArray<int32> ComputeActionSchemaSubsetIndices(
 			const FString& SourceJsonString,
 			const Action::FSchema& ActionSchema,
 			const Action::FSchemaElement& ActionSchemaElement);
 
 		/* Computes the indices to copy from a subset compatible JSON. */
-		LEARNINGTRAINING_API void ComputeActionSchemaSubsetIndices(
+		UE_API void ComputeActionSchemaSubsetIndices(
 			int32& NextIndex,
 			TArray<int32>& Indices,
 			const TSharedRef<FJsonObject> Object,
