[LearningAgents] File Communicator

A communicator that just serializes 4 things to file:
1. data-config.json
2. trainer-config.json
3. networks binaries
4. replay binaries

usage ex: upload files generated from c++ FileCommunicator to s3, and use the python FileCommunicator for training
Additional:

- finally moved socket communicator cleanups from github to p4
- added action noise as a parameter
- more cleanups with device settings

diff --git a/Engine/Plugins/Experimental/LearningAgents/LearningAgents.uplugin b/Engine/Plugins/Experimental/LearningAgents/LearningAgents.uplugin
--- a/Engine/Plugins/Experimental/LearningAgents/LearningAgents.uplugin
+++ b/Engine/Plugins/Experimental/LearningAgents/LearningAgents.uplugin
@@ -56,5 +56,19 @@
 			"Name": "NNERuntimeBasicCpu",
 			"Enabled": true
 		}
+	],
+	"PythonRequirements": [
+		{
+			"Platform": "All",
+			"Requirements": [
+				"boto3==1.40.11 --hash=sha256:9d2d211d9cb3efc9a2b2ceec3c510b4e62e389618fd5c871e74d2cbca4561ff5",
+				"botocore==1.40.11 --hash=sha256:4beca0c5f92201da1bf1bc0a55038538ad2defded32ab0638cb68f5631dcc665",
+				"jmespath==1.0.1 --hash=sha256:02e2e4cc71b5bcab88332eebf907519190dd9e6e82107fa7f83b1003a6252980",
+				"python_dateutil==2.9.0 --hash=sha256:cbf2f1da5e6083ac2fbfd4da39a25f34312230110440f424a14c7558bb85d82e",
+				"s3transfer==0.13.1 --hash=sha256:a981aa7429be23fe6dfc13e80e4020057cbab622b08c0315288758d67cabc724",
+				"six==1.16.0 --hash=sha256:8abb2f1d86890a2dfb989f9a77cfcfd3e47c2a354b01111771326f8aa26e0254",
+				"smart_open==7.3.0 --hash=sha256:0fc37c7144868e2619a088537094aa3cb0fe5fd8c17162149fa6c983fa8efe6a"
+			]
+		}
 	]
 }
\ No newline at end of file
diff --git a/Engine/Plugins/Experimental/LearningAgents/Source/LearningAgentsTraining/Private/LearningAgentsImitationTrainer.cpp b/Engine/Plugins/Experimental/LearningAgents/Source/LearningAgentsTraining/Private/LearningAgentsImitationTrainer.cpp
--- a/Engine/Plugins/Experimental/LearningAgents/Source/LearningAgentsTraining/Private/LearningAgentsImitationTrainer.cpp
+++ b/Engine/Plugins/Experimental/LearningAgents/Source/LearningAgentsTraining/Private/LearningAgentsImitationTrainer.cpp
@@ -33,13 +33,10 @@ TSharedRef<FJsonObject> FLearningAgentsImitationTrainerTrainingSettings::AsJsonC
 	ConfigObject->SetNumberField(TEXT("Window"), Window);
 	ConfigObject->SetNumberField(TEXT("ActionRegularizationWeight"), ActionRegularizationWeight);
 	ConfigObject->SetNumberField(TEXT("ActionEntropyWeight"), ActionEntropyWeight);
+	ConfigObject->SetNumberField(TEXT("ObservationNoiseScale"), ObservationNoiseScale);
 	ConfigObject->SetNumberField(TEXT("Seed"), RandomSeed);
-	ConfigObject->SetStringField(TEXT("Device"), UE::Learning::Trainer::GetDeviceString(UE::Learning::Agents::GetTrainingDevice(Device)));
-	ConfigObject->SetBoolField(TEXT("UseTensorBoard"), bUseTensorboard);
 	ConfigObject->SetBoolField(TEXT("SaveSnapshots"), bSaveSnapshots);
 	ConfigObject->SetNumberField(TEXT("IterationsPerSnapshot"), IterationsPerSnapshot);
-	ConfigObject->SetBoolField(TEXT("UseMLflow"), bUseMLflow);
-	ConfigObject->SetStringField(TEXT("MLflowTrackingUri"), MLflowTrackingUri);
 
 	return ConfigObject;
 }
@@ -455,6 +452,12 @@ TSharedRef<FJsonObject> ULearningAgentsImitationTrainer::CreateTrainerConfig(con
 	ConfigObject->SetObjectField(TEXT("BehaviorCloningSettings"), TrainingSettings.AsJsonConfig());
 	ConfigObject->SetNumberField(TEXT("MemoryStateNum"), Policy->GetMemoryStateSize());
 
+	// Add Training Workflow Config Entries
+	ConfigObject->SetStringField(TEXT("Device"), UE::Learning::Trainer::GetDeviceString(UE::Learning::Agents::GetTrainingDevice(TrainingSettings.Device)));
+	ConfigObject->SetBoolField(TEXT("UseTensorBoard"), TrainingSettings.bUseTensorboard);
+	ConfigObject->SetBoolField(TEXT("UseMLflow"), TrainingSettings.bUseMLflow);
+	ConfigObject->SetStringField(TEXT("MLflowTrackingUri"), TrainingSettings.MLflowTrackingUri);
+
 	return ConfigObject;
 }
 
diff --git a/Engine/Plugins/Experimental/LearningAgents/Source/LearningAgentsTraining/Private/LearningAgentsPPOTrainer.cpp b/Engine/Plugins/Experimental/LearningAgents/Source/LearningAgentsTraining/Private/LearningAgentsPPOTrainer.cpp
--- a/Engine/Plugins/Experimental/LearningAgents/Source/LearningAgentsTraining/Private/LearningAgentsPPOTrainer.cpp
+++ b/Engine/Plugins/Experimental/LearningAgents/Source/LearningAgentsTraining/Private/LearningAgentsPPOTrainer.cpp
@@ -53,12 +53,8 @@ TSharedRef<FJsonObject> FLearningAgentsPPOTrainingSettings::AsJsonConfig() const
 	ConfigObject->SetNumberField(TEXT("TrimEpisodeEndStepNum"), NumberOfStepsToTrimAtEndOfEpisode);
 	ConfigObject->SetNumberField(TEXT("Seed"), RandomSeed);
 	ConfigObject->SetNumberField(TEXT("DiscountFactor"), DiscountFactor);
-	ConfigObject->SetStringField(TEXT("Device"), UE::Learning::Trainer::GetDeviceString(UE::Learning::Agents::GetTrainingDevice(Device)));
-	ConfigObject->SetBoolField(TEXT("UseTensorBoard"), bUseTensorboard);
 	ConfigObject->SetBoolField(TEXT("SaveSnapshots"), bSaveSnapshots);
 	ConfigObject->SetNumberField(TEXT("IterationsPerSnapshot"), IterationsPerSnapshot);
-	ConfigObject->SetBoolField(TEXT("UseMLflow"), bUseMLflow);
-	ConfigObject->SetStringField(TEXT("MLflowTrackingUri"), MLflowTrackingUri);
 
 	return ConfigObject;
 }
@@ -460,6 +456,12 @@ TSharedRef<FJsonObject> ULearningAgentsPPOTrainer::CreateTrainerConfig(const FLe
 	// Add PPO-specific Config Entries
 	ConfigObject->SetObjectField(TEXT("PPOSettings"), TrainingSettings.AsJsonConfig());
 
+	// Add Training Workflow Config Entries
+	ConfigObject->SetStringField(TEXT("Device"), UE::Learning::Trainer::GetDeviceString(UE::Learning::Agents::GetTrainingDevice(TrainingSettings.Device)));
+	ConfigObject->SetBoolField(TEXT("UseTensorBoard"), TrainingSettings.bUseTensorboard);
+	ConfigObject->SetBoolField(TEXT("UseMLflow"), TrainingSettings.bUseMLflow);
+	ConfigObject->SetStringField(TEXT("MLflowTrackingUri"), TrainingSettings.MLflowTrackingUri);
+
 	return ConfigObject;
 }
 
diff --git a/Engine/Plugins/Experimental/LearningAgents/Source/LearningAgentsTraining/Public/LearningAgentsImitationTrainer.h b/Engine/Plugins/Experimental/LearningAgents/Source/LearningAgentsTraining/Public/LearningAgentsImitationTrainer.h
--- a/Engine/Plugins/Experimental/LearningAgents/Source/LearningAgentsTraining/Public/LearningAgentsImitationTrainer.h
+++ b/Engine/Plugins/Experimental/LearningAgents/Source/LearningAgentsTraining/Public/LearningAgentsImitationTrainer.h
@@ -33,6 +33,7 @@ public:
 	/** Time in seconds to wait for the training process before timing out. */
 	UPROPERTY(EditAnywhere, Category = "LearningAgents", meta = (ClampMin = "0.0", UIMin = "0.0"))
 	float TrainerCommunicationTimeout = 10.0f;
+
 };
 
 /** The configurable settings for the training process. */
@@ -94,6 +95,12 @@ public:
 	UPROPERTY(EditAnywhere, Category = "LearningAgents", meta = (ClampMin = "0.0", UIMin = "0.0"))
 	float ActionEntropyWeight = 0.0f;
 
+	/**
+	 * A multiplicative scaling factor that controls the observation noise that increases the perturbations added to observations.
+	 */
+	UPROPERTY(EditAnywhere, Category = "LearningAgents", meta = (ClampMin = "0.0", UIMin = "0.0"))
+	float ObservationNoiseScale = 0.0f;
+
 	/** The seed used for any random sampling the trainer will perform, e.g. for weight initialization. */
 	UPROPERTY(EditAnywhere, Category = "LearningAgents", meta = (ClampMin = "0", UIMin = "0"))
 	int32 RandomSeed = 1234;
diff --git a/Engine/Plugins/Experimental/LearningAgents/Source/LearningAgentsTrainingEditor/Private/LearningAgentsImitationTrainerEditor.cpp b/Engine/Plugins/Experimental/LearningAgents/Source/LearningAgentsTrainingEditor/Private/LearningAgentsImitationTrainerEditor.cpp
--- a/Engine/Plugins/Experimental/LearningAgents/Source/LearningAgentsTrainingEditor/Private/LearningAgentsImitationTrainerEditor.cpp
+++ b/Engine/Plugins/Experimental/LearningAgents/Source/LearningAgentsTrainingEditor/Private/LearningAgentsImitationTrainerEditor.cpp
@@ -1,6 +1,7 @@
 // Copyright Epic Games, Inc. All Rights Reserved.
 
 #include "LearningAgentsImitationTrainerEditor.h"
+#include "LearningExternalTrainer.h"
 #include "LearningLog.h"
 #include "Containers/Ticker.h"
 
@@ -61,6 +62,13 @@ void ALearningAgentsImitationTrainerEditor::StopTraining()
 	LearningAgentsImitationTrainer->EndTraining();
 }
 
+FLearningAgentsCommunicator ALearningAgentsImitationTrainerEditor::MakeFileCommunicator(FDirectoryPath EditorIntermediateRelativePath)
+{
+	FLearningAgentsCommunicator Communicator;
+	Communicator.Trainer = MakeShared<UE::Learning::FFileTrainer>(UE::Learning::Trainer::GetIntermediatePath(EditorIntermediateRelativePath.Path));
+	return Communicator;
+}
+
 bool ALearningAgentsImitationTrainerEditor::IsTraining() const
 {
 	if (LearningAgentsImitationTrainer)
diff --git a/Engine/Plugins/Experimental/LearningAgents/Source/LearningAgentsTrainingEditor/Public/LearningAgentsImitationTrainerEditor.h b/Engine/Plugins/Experimental/LearningAgents/Source/LearningAgentsTrainingEditor/Public/LearningAgentsImitationTrainerEditor.h
--- a/Engine/Plugins/Experimental/LearningAgents/Source/LearningAgentsTrainingEditor/Public/LearningAgentsImitationTrainerEditor.h
+++ b/Engine/Plugins/Experimental/LearningAgents/Source/LearningAgentsTrainingEditor/Public/LearningAgentsImitationTrainerEditor.h
@@ -5,8 +5,8 @@
 #include "CoreMinimal.h"
 #include "GameFramework/Actor.h"
 #include "LearningAgentsManager.h"
+#include "LearningAgentsCommunicator.h"
 #include "LearningAgentsImitationTrainer.h"
-
 #include "LearningAgentsImitationTrainerEditor.generated.h"
 
 #define UE_API LEARNINGAGENTSTRAININGEDITOR_API
@@ -32,6 +32,10 @@ public:
 	UFUNCTION(BlueprintCallable, Category = "LearningAgents")
 	UE_API void StopTraining();
 
+	/** Make File communicator to file training materials for external training. */
+	UFUNCTION(BlueprintCallable, Category = "LearningAgents")
+	UE_API FLearningAgentsCommunicator MakeFileCommunicator(FDirectoryPath EditorIntermediateRelativePath);
+
 	UE_API bool IsTraining() const;
 
 	UPROPERTY(EditDefaultsOnly, Category = "LearningAgents")
diff --git a/Engine/Plugins/Experimental/LearningAgents/Source/LearningTraining/Private/LearningExternalTrainer.cpp b/Engine/Plugins/Experimental/LearningAgents/Source/LearningTraining/Private/LearningExternalTrainer.cpp
--- a/Engine/Plugins/Experimental/LearningAgents/Source/LearningTraining/Private/LearningExternalTrainer.cpp
+++ b/Engine/Plugins/Experimental/LearningAgents/Source/LearningTraining/Private/LearningExternalTrainer.cpp
@@ -6,6 +6,7 @@
 #include "LearningNeuralNetwork.h"
 #include "LearningSharedMemoryTraining.h"
 #include "LearningSocketTraining.h"
+#include "LearningProgress.h"
 #include "GenericPlatform/GenericPlatformFile.h"
 #include "HAL/PlatformFileManager.h"
 
@@ -1048,4 +1049,156 @@ namespace UE::Learning
 			Timeout,
 			LogSettings);
 	}
+
+	FFileTrainer::FFileTrainer(const FString& InIntermediatePath)
+	{
+		IntermediatePath = InIntermediatePath;
+		FileSaveDirectory = IntermediatePath / TEXT("TrainingMaterials") / FDateTime::Now().ToString(TEXT("%Y%m%d%H%M%S"));
+	}
+
+	FFileTrainer::~FFileTrainer(){}
+
+	ETrainerResponse FFileTrainer::SendConfigs(
+		const TSharedRef<FJsonObject>& DataConfigObject,
+		const TSharedRef<FJsonObject>& TrainerConfigObject,
+		const ELogSetting LogSettings)
+	{
+		IFileManager& FileManager = IFileManager::Get();
+
+		// Add intermediate path as a temp directory for the artifacts logging on python side
+		TrainerConfigObject->SetStringField(TEXT("TempDirectory"), *FileManager.ConvertToAbsolutePathForExternalAppForRead(*IntermediatePath));
+
+		FString DataConfigString;
+		TSharedRef<TJsonWriter<>> DataJsonWriter = TJsonWriterFactory<>::Create(&DataConfigString, 0);
+		FJsonSerializer::Serialize(DataConfigObject, DataJsonWriter, true);
+		FFileHelper::SaveStringToFile(DataConfigString, *(FileSaveDirectory / "data-config.json"));
+
+		FString TrainerConfigString;
+		TSharedRef<TJsonWriter<>> TrainerJsonWriter = TJsonWriterFactory<>::Create(&TrainerConfigString, 0);
+		FJsonSerializer::Serialize(TrainerConfigObject, TrainerJsonWriter, true);
+		FFileHelper::SaveStringToFile(TrainerConfigString, *(FileSaveDirectory / "trainer-config.json"));
+
+		UE_LOG(LogLearning, Display, TEXT("FileTrainer: wrote config files to %s."), *IntermediatePath);
+
+		return ETrainerResponse::Success;
+	}
+
+	int32 FFileTrainer::AddNetwork(const ULearningNeuralNetworkData& Network)
+	{
+		const int32 Id = NetworksCount;
+		NetworksCount++;
+		return Id;
+	}
+
+	ETrainerResponse FFileTrainer::SendNetwork(
+		const int32 NetworkId,
+		const ULearningNeuralNetworkData& Network,
+		FRWLock* NetworkLock,
+		const ELogSetting LogSettings)
+	{
+		if (!ensureMsgf(NetworksCount > NetworkId, TEXT("Network %d has not been added. Call AddNetwork prior to SendNetwork."), NetworkId))
+		{
+			return ETrainerResponse::Unexpected;
+		}
+
+		IPlatformFile& FileManager = FPlatformFileManager::Get().GetPlatformFile();
+		const FString NetworksDir = FileSaveDirectory / TEXT("Networks");
+		FileManager.CreateDirectoryTree(*NetworksDir);
+
+		const int32 ByteNum = Network.GetSnapshotByteNum();
+		TLearningArray<1, uint8> Snapshot;
+		Snapshot.SetNumUninitialized({ ByteNum });
+
+		FScopeNullableReadLock ScopeLock(NetworkLock);
+		Network.SaveToSnapshot(MakeArrayView(Snapshot.GetData(), Snapshot.Num()));
+
+		const FString Base = FString::Printf(TEXT("network_%d"), NetworkId);
+		const FString BinPath = NetworksDir / (Base + TEXT(".bin"));
+		const FString TmpPath = BinPath + TEXT(".tmp");
+
+		TUniquePtr<FArchive> FileWriter(IFileManager::Get().CreateFileWriter(*TmpPath));
+		if (!FileWriter) 
+		{
+			return ETrainerResponse::Unexpected; 
+		}
+
+		FileWriter->Serialize(Snapshot.GetData(), Snapshot.Num());
+		FileWriter->Close();
+
+		if (!FileManager.MoveFile(*BinPath, *TmpPath))
+		{
+			FileManager.DeleteFile(*TmpPath);
+			return ETrainerResponse::Unexpected;
+		}
+
+		if (LogSettings != ELogSetting::Silent)
+		{
+			UE_LOG(LogLearning, Display, TEXT("FileTrainer: saved network %d snapshot to %s"), NetworkId, *BinPath);
+		}
+
+		return ETrainerResponse::Success;
+	}
+
+	int32 FFileTrainer::AddReplayBuffer(const FReplayBuffer& ReplayBuffer)
+	{
+		const int32 Id = ReplayBufferCount++;
+		return Id;
+	}
+
+	ETrainerResponse FFileTrainer::SendReplayBuffer(const int32 ReplayBufferId, const FReplayBuffer& ReplayBuffer, const ELogSetting LogSettings)
+	{
+		if (ReplayBufferId < 0)
+		{
+			return ETrainerResponse::Unexpected;
+		}
+
+		IPlatformFile& FileManager = FPlatformFileManager::Get().GetPlatformFile();
+		FileManager.CreateDirectoryTree(*FileSaveDirectory);
+
+		const FString BinPath = FileSaveDirectory / TEXT("ReplayBuffers") / FString::Printf(TEXT("replaybuffer_%d.bin"), ReplayBufferId);
+		const FString TempPath = BinPath + TEXT(".tmp");
+
+		TUniquePtr<FArchive> FileWriter(IFileManager::Get().CreateFileWriter(*TempPath));
+		if (!FileWriter)
+		{
+			return ETrainerResponse::Unexpected;
+		}
+
+		FileWriter->Serialize((void*)&ReplayBufferId, sizeof(int32));
+
+		const int32 EpisodeNum = ReplayBuffer.GetEpisodeNum();
+		FileWriter->Serialize((void*)&EpisodeNum, sizeof(int32));
+
+		const int32 StepNum = ReplayBuffer.GetStepNum();
+		FileWriter->Serialize((void*)&StepNum, sizeof(int32));
+
+		FileWriter->Serialize((void*)ReplayBuffer.GetEpisodeStarts().GetData(), sizeof(int32) * ReplayBuffer.GetEpisodeStarts().Num());
+		FileWriter->Serialize((void*)ReplayBuffer.GetEpisodeLengths().GetData(), sizeof(int32) * ReplayBuffer.GetEpisodeLengths().Num());
+	
+		for (int32 Index = 0; Index < ReplayBuffer.GetObservationsNum(); Index++)
+		{
+			FileWriter->Serialize((void*)ReplayBuffer.GetObservations(Index).GetData(), ReplayBuffer.GetObservations(Index).Num() * sizeof(float));
+		}
+
+		for (int32 Index = 0; Index < ReplayBuffer.GetActionsNum(); Index++)
+		{
+			FileWriter->Serialize((void*)ReplayBuffer.GetActions(Index).GetData(), ReplayBuffer.GetActions(Index).Num() * sizeof(float));
+		}
+
+		FileWriter->Close();
+
+		if (!FileManager.MoveFile(*BinPath, *TempPath))
+		{
+			FileManager.DeleteFile(*TempPath);
+			return ETrainerResponse::Unexpected;
+		}
+
+		if (LogSettings != ELogSetting::Silent)
+		{
+			UE_LOG(LogLearning, Display, TEXT("FileTrainer: wrote replay buffer (%d eps, %d steps) to %s"), ReplayBuffer.GetEpisodeNum(), ReplayBuffer.GetStepNum(), *BinPath);
+		}
+
+		return ETrainerResponse::Success;
+	}
 }
+
diff --git a/Engine/Plugins/Experimental/LearningAgents/Source/LearningTraining/Public/LearningExternalTrainer.h b/Engine/Plugins/Experimental/LearningAgents/Source/LearningTraining/Public/LearningExternalTrainer.h
--- a/Engine/Plugins/Experimental/LearningAgents/Source/LearningTraining/Public/LearningExternalTrainer.h
+++ b/Engine/Plugins/Experimental/LearningAgents/Source/LearningTraining/Public/LearningExternalTrainer.h
@@ -10,6 +10,7 @@
 #include "Misc/Guid.h"
 #include "Templates/SharedPointer.h"
 #include "UObject/NameTypes.h"
+#include "UObject/ObjectPtr.h"  
 
 #define UE_API LEARNINGTRAINING_API
 
@@ -475,6 +476,68 @@ namespace UE::Learning
 
 		int32 NetworksVersion = -1;
 	};
+	
+	struct FFileTrainer : public IExternalTrainer
+	{
+		/**
+		* Creates a new trainer that dumps training materials to file.
+		*
+		* @param InIntermediatePath				The Intermediate Path - for writing temporary files.
+		*/
+		UE_API FFileTrainer(const FString& InIntermediatePath = FString());
+
+		UE_API ~FFileTrainer();
+
+		UE_API virtual bool IsValid() override final { return true; }
+
+		UE_API virtual void Terminate() override final {}
+
+		UE_API virtual ETrainerResponse SendStop() override final { return ETrainerResponse::Success; }
+
+		UE_API virtual ETrainerResponse Wait() override final { return ETrainerResponse::Success; }
+
+		UE_API virtual bool HasNetworkOrCompleted() override final { return true; }
+
+		UE_API virtual ETrainerResponse SendConfigs(
+			const TSharedRef<FJsonObject>& DataConfigObject,
+			const TSharedRef<FJsonObject>& TrainerConfigObject,
+			const ELogSetting LogSettings = Trainer::DefaultLogSettings) override final;
+
+		UE_API virtual int32 AddNetwork(const ULearningNeuralNetworkData& Network) override final;
+
+		UE_API virtual ETrainerResponse ReceiveNetwork(
+			const int32 NetworkId,
+			ULearningNeuralNetworkData& OutNetwork,
+			FRWLock* NetworkLock = nullptr,
+			const ELogSetting LogSettings = Trainer::DefaultLogSettings) override final { return ETrainerResponse::Completed; }
+
+		UE_API virtual TArray<ETrainerResponse> ReceiveNetworks(
+			const TArray<int32>& NetworkIds,
+			TArray<TObjectPtr<ULearningNeuralNetworkData>> Networks,
+			TArray<FRWLock*> NetworkLocks = TArray<FRWLock*>(),
+			const ELogSetting LogSettings = Trainer::DefaultLogSettings) override final { return TArray<ETrainerResponse>(); }
+
+		UE_API virtual ETrainerResponse SendNetwork(
+			const int32 NetworkId,
+			const ULearningNeuralNetworkData& Network,
+			FRWLock* NetworkLock = nullptr,
+			const ELogSetting LogSettings = Trainer::DefaultLogSettings) override final;
+
+		UE_API virtual int32 AddReplayBuffer(const FReplayBuffer& ReplayBuffer) override final;
+
+		UE_API virtual ETrainerResponse SendReplayBuffer(
+			const int32 ReplayBufferId,
+			const FReplayBuffer& ReplayBuffer,
+			const ELogSetting LogSettings = Trainer::DefaultLogSettings) override final;
+
+	private:
+		FString IntermediatePath;
+		FString FileSaveDirectory;
+		int32 ReplayBufferCount = 0;
+		int32 NetworksCount = 0;
+
+		TArray<TSharedMemoryArrayView<1, uint8>> NeuralNetworkSharedMemoryArrayViews;
+	};
 }
 
 #undef UE_API
