Separated LearningAgents plugin into two: LearningCore and LearningAgents

diff --git a/Engine/Plugins/Experimental/LearningAgents/LearningAgents.uplugin b/Engine/Plugins/Experimental/LearningAgents/LearningAgents.uplugin
--- a/Engine/Plugins/Experimental/LearningAgents/LearningAgents.uplugin
+++ b/Engine/Plugins/Experimental/LearningAgents/LearningAgents.uplugin
@@ -16,16 +16,6 @@
 	"IsExperimentalVersion": true,
 	"Installed": false,
 	"Modules": [
-		{
-			"Name": "Learning",
-			"Type": "Runtime",
-			"LoadingPhase": "Default"
-		},
-		{
-			"Name": "LearningTraining",
-			"Type": "Runtime",
-			"LoadingPhase": "Default"
-		},
 		{
 			"Name": "LearningAgents",
 			"Type": "Runtime",
@@ -48,6 +38,10 @@
 		}
 	],
 	"Plugins": [
+		{
+			"Name": "LearningCore",
+			"Enabled": true
+		},
 		{
 			"Name": "PythonMLPackages",
 			"Enabled": true
@@ -56,19 +50,5 @@
 			"Name": "NNERuntimeBasicCpu",
 			"Enabled": true
 		}
-	],
-	"PythonRequirements": [
-		{
-			"Platform": "All",
-			"Requirements": [
-				"boto3==1.40.11 --hash=sha256:9d2d211d9cb3efc9a2b2ceec3c510b4e62e389618fd5c871e74d2cbca4561ff5",
-				"botocore==1.40.11 --hash=sha256:4beca0c5f92201da1bf1bc0a55038538ad2defded32ab0638cb68f5631dcc665",
-				"jmespath==1.0.1 --hash=sha256:02e2e4cc71b5bcab88332eebf907519190dd9e6e82107fa7f83b1003a6252980",
-				"python_dateutil==2.9.0 --hash=sha256:cbf2f1da5e6083ac2fbfd4da39a25f34312230110440f424a14c7558bb85d82e",
-				"s3transfer==0.13.1 --hash=sha256:a981aa7429be23fe6dfc13e80e4020057cbab622b08c0315288758d67cabc724",
-				"six==1.16.0 --hash=sha256:8abb2f1d86890a2dfb989f9a77cfcfd3e47c2a354b01111771326f8aa26e0254",
-				"smart_open==7.3.0 --hash=sha256:0fc37c7144868e2619a088537094aa3cb0fe5fd8c17162149fa6c983fa8efe6a"
-			]
-		}
 	]
 }
\ No newline at end of file
diff --git a/Engine/Plugins/Experimental/LearningAgents/Source/LearningAgentsTraining/Private/LearningAgentsCommunicator.cpp b/Engine/Plugins/Experimental/LearningAgents/Source/LearningAgentsTraining/Private/LearningAgentsCommunicator.cpp
--- a/Engine/Plugins/Experimental/LearningAgents/Source/LearningAgentsTraining/Private/LearningAgentsCommunicator.cpp
+++ b/Engine/Plugins/Experimental/LearningAgents/Source/LearningAgentsTraining/Private/LearningAgentsCommunicator.cpp
@@ -38,14 +38,14 @@ FLearningAgentsSharedMemoryTrainerProcess ULearningAgentsCommunicatorLibrary::Sp
 		return TrainerProcess;
 	}
 
-	const FString PythonContentPath = UE::Learning::Trainer::GetPythonContentPath(TrainerProcessSettings.GetEditorEnginePath());
+	const FString PythonContentPath = UE::Learning::Agents::Trainer::GetPythonContentPath(TrainerProcessSettings.GetEditorEnginePath());
 	if (!FPaths::DirectoryExists(PythonContentPath))
 	{
 		UE_LOG(LogLearning, Error, TEXT("SpawnSharedMemoryTrainingProcess: Can't find LearningAgents plugin Content \"%s\"."), *PythonContentPath);
 		return TrainerProcess;
 	}
 
-	const FString IntermediatePath = UE::Learning::Trainer::GetIntermediatePath(TrainerProcessSettings.GetIntermediatePath());
+	const FString IntermediatePath = UE::Learning::Agents::Trainer::GetIntermediatePath(TrainerProcessSettings.GetIntermediatePath());
 
 	const FString CustomTrainerModulePath = TrainerProcessSettings.GetCustomTrainerModulePath();
 	if (!CustomTrainerModulePath.IsEmpty() and !FPaths::DirectoryExists(CustomTrainerModulePath))
@@ -145,14 +145,14 @@ FLearningAgentsSocketTrainerProcess ULearningAgentsCommunicatorLibrary::SpawnSoc
 		return TrainerProcess;
 	}
 
-	const FString PythonContentPath = UE::Learning::Trainer::GetPythonContentPath(TrainerProcessSettings.GetEditorEnginePath());
+	const FString PythonContentPath = UE::Learning::Agents::Trainer::GetPythonContentPath(TrainerProcessSettings.GetEditorEnginePath());
 	if (!FPaths::DirectoryExists(PythonContentPath))
 	{
 		UE_LOG(LogLearning, Error, TEXT("SpawnSocketTrainingProcess: Can't find LearningAgents plugin Content \"%s\"."), *PythonContentPath);
 		return TrainerProcess;
 	}
 
-	const FString IntermediatePath = UE::Learning::Trainer::GetIntermediatePath(TrainerProcessSettings.GetIntermediatePath());
+	const FString IntermediatePath = UE::Learning::Agents::Trainer::GetIntermediatePath(TrainerProcessSettings.GetIntermediatePath());
 
 	const FString CustomTrainerModulePath = TrainerProcessSettings.GetCustomTrainerModulePath();
 	if (!CustomTrainerModulePath.IsEmpty() and !FPaths::DirectoryExists(CustomTrainerModulePath))
diff --git a/Engine/Plugins/Experimental/LearningAgents/Source/LearningAgentsTraining/Private/LearningAgentsTrainer.cpp b/Engine/Plugins/Experimental/LearningAgents/Source/LearningAgentsTraining/Private/LearningAgentsTrainer.cpp
--- a/Engine/Plugins/Experimental/LearningAgents/Source/LearningAgentsTraining/Private/LearningAgentsTrainer.cpp
+++ b/Engine/Plugins/Experimental/LearningAgents/Source/LearningAgentsTraining/Private/LearningAgentsTrainer.cpp
@@ -164,6 +164,19 @@ namespace UE::Learning::Agents
 		}
 #endif
 	}
+
+	namespace Trainer
+	{
+		FString GetIntermediatePath(const FString& IntermediateDir)
+		{
+			return IntermediateDir / TEXT("LearningAgents");
+		}
+
+		FString GetPythonContentPath(const FString& EngineDir)
+		{
+			return EngineDir / TEXT("Plugins/Experimental/LearningAgents/Content/Python/");
+		}
+	}
 }
 
 FLearningAgentsTrainerProcessSettings::FLearningAgentsTrainerProcessSettings()
diff --git a/Engine/Plugins/Experimental/LearningAgents/Source/LearningAgentsTraining/Private/LearningAgentsTraining.cpp b/Engine/Plugins/Experimental/LearningAgents/Source/LearningAgentsTraining/Private/LearningAgentsTraining.cpp
--- a/Engine/Plugins/Experimental/LearningAgents/Source/LearningAgentsTraining/Private/LearningAgentsTraining.cpp
+++ b/Engine/Plugins/Experimental/LearningAgents/Source/LearningAgentsTraining/Private/LearningAgentsTraining.cpp
@@ -1,6 +1,8 @@
 // Copyright Epic Games, Inc. All Rights Reserved.
 
 #include "LearningAgentsTraining.h"
+#include "Misc/FileHelper.h"
+#include "Misc/Paths.h"
 
 void FLearningAgentsTrainingModule::StartupModule()
 {
@@ -12,5 +14,17 @@ void FLearningAgentsTrainingModule::ShutdownModule()
 	// This function may be called during shutdown to clean up your module.  For modules that support dynamic reloading,
 	// we call this function before unloading the module.
 }
+
+void FLearningAgentsTrainingModule::RegisterPythonPackages()
+{
+	const FString ExperimentalPluginsPath = FPaths::Combine(FPaths::EnginePluginsDir(), TEXT("Experimental"));
+	const FString SitePackagesPath = FPaths::ConvertRelativePathToFull(FPaths::ProjectIntermediateDir() / TEXT("PipInstall") / TEXT("Lib") / TEXT("site-packages"));
+	const FString PyPluginsSitePackageFile = FPaths::ConvertRelativePathToFull(SitePackagesPath / TEXT("learning_agents.pth"));
+
+	TArray<FString> PluginSitePackagePaths;
+	PluginSitePackagePaths.Add(FPaths::ConvertRelativePathToFull(FPaths::Combine(ExperimentalPluginsPath, TEXT("LearningAgents/Content/Python"))));
+
+	FFileHelper::SaveStringArrayToFile(PluginSitePackagePaths, *PyPluginsSitePackageFile, FFileHelper::EEncodingOptions::ForceUTF8WithoutBOM);
+}
 	
 IMPLEMENT_MODULE(FLearningAgentsTrainingModule, LearningAgentsTraining)
diff --git a/Engine/Plugins/Experimental/LearningAgents/Source/LearningAgentsTraining/Public/LearningAgentsTrainer.h b/Engine/Plugins/Experimental/LearningAgents/Source/LearningAgentsTraining/Public/LearningAgentsTrainer.h
--- a/Engine/Plugins/Experimental/LearningAgents/Source/LearningAgentsTraining/Public/LearningAgentsTrainer.h
+++ b/Engine/Plugins/Experimental/LearningAgents/Source/LearningAgentsTraining/Public/LearningAgentsTrainer.h
@@ -89,14 +89,23 @@ namespace UE::Learning::Agents
 	};
 
 	/** Get the learning agents trainer device from the UE::Learning trainer device. */
-	LEARNINGAGENTSTRAINING_API ELearningAgentsTrainingDevice GetLearningAgentsTrainingDevice(const ETrainerDevice Device);
+	UE_API ELearningAgentsTrainingDevice GetLearningAgentsTrainingDevice(const ETrainerDevice Device);
 
 	/** Get the UE::Learning trainer device from the learning agents trainer device. */
-	LEARNINGAGENTSTRAINING_API ETrainerDevice GetTrainingDevice(const ELearningAgentsTrainingDevice Device);
+	UE_API ETrainerDevice GetTrainingDevice(const ELearningAgentsTrainingDevice Device);
 
-	LEARNINGAGENTSTRAINING_API void ApplyGameSettings(const FLearningAgentsTrainingGameSettings& Settings, const UWorld* World, FGameSettingsState& OutGameSettingsState);
+	UE_API void ApplyGameSettings(const FLearningAgentsTrainingGameSettings& Settings, const UWorld* World, FGameSettingsState& OutGameSettingsState);
 
-	LEARNINGAGENTSTRAINING_API void RevertGameSettings(const FGameSettingsState& Settings, const UWorld* World);
+	UE_API void RevertGameSettings(const FGameSettingsState& Settings, const UWorld* World);
+
+	namespace Trainer
+	{
+		/** Gets the LearningAgents Intermediate path from the intermediate directory. */
+		UE_API FString GetIntermediatePath(const FString& IntermediateDir);
+
+		/** Gets the LearningAgents Content path from the engine directory. */
+		UE_API FString GetPythonContentPath(const FString& EngineDir);
+	}
 }
 
 /** The path settings for the trainer. */
@@ -156,7 +165,7 @@ public:
 
 	/** Trainer file name. The name of the python file to use for training. Do NOT include the '.py' file extension. */
 	UPROPERTY(EditAnywhere, Category = "LearningAgents")
-	FString TrainerFileName = TEXT("learning_agents.train_ppo");
+	FString TrainerFileName = TEXT("learning_core.train_ppo");
 
 public:
 
diff --git a/Engine/Plugins/Experimental/LearningAgents/Source/LearningAgentsTraining/Public/LearningAgentsTraining.h b/Engine/Plugins/Experimental/LearningAgents/Source/LearningAgentsTraining/Public/LearningAgentsTraining.h
--- a/Engine/Plugins/Experimental/LearningAgents/Source/LearningAgentsTraining/Public/LearningAgentsTraining.h
+++ b/Engine/Plugins/Experimental/LearningAgents/Source/LearningAgentsTraining/Public/LearningAgentsTraining.h
@@ -10,4 +10,7 @@ public:
 	/** IModuleInterface implementation */
 	virtual void StartupModule() override;
 	virtual void ShutdownModule() override;
+
+	/** Register any internal Python dependencies to make them available in Python sys.path */
+	virtual void RegisterPythonPackages();
 };
diff --git a/Engine/Plugins/Experimental/LearningCore/Source/LearningTraining/Private/LearningTraining.cpp b/Engine/Plugins/Experimental/LearningCore/Source/LearningTraining/Private/LearningTraining.cpp
new file mode 100644
--- /dev/null
+++ b/Engine/Plugins/Experimental/LearningCore/Source/LearningTraining/Private/LearningTraining.cpp
@@ -0,0 +1,32 @@
+// Copyright Epic Games, Inc. All Rights Reserved.
+
+#include "LearningTraining.h"
+#include "Misc/FileHelper.h"
+#include "Misc/Paths.h"
+
+void FLearningTrainingModule::StartupModule()
+{
+	// This code will execute after your module is loaded into memory; the exact timing is specified in the .uplugin file per-module
+	RegisterPythonPackages();
+}
+
+void FLearningTrainingModule::ShutdownModule()
+{
+	// This function may be called during shutdown to clean up your module.  For modules that support dynamic reloading,
+	// we call this function before unloading the module.
+}
+
+void FLearningTrainingModule::RegisterPythonPackages()
+{
+	const FString ExperimentalPluginsPath = FPaths::Combine(FPaths::EnginePluginsDir(), TEXT("Experimental"));
+	const FString SitePackagesPath = FPaths::ConvertRelativePathToFull(FPaths::ProjectIntermediateDir() / TEXT("PipInstall") / TEXT("Lib") / TEXT("site-packages"));
+	const FString PyPluginsSitePackageFile = FPaths::ConvertRelativePathToFull(SitePackagesPath / TEXT("learning_core.pth"));
+	
+	TArray<FString> PluginSitePackagePaths;
+	PluginSitePackagePaths.Add(FPaths::ConvertRelativePathToFull(FPaths::Combine(ExperimentalPluginsPath, TEXT("LearningCore/Content/Python"))));
+	PluginSitePackagePaths.Add(FPaths::ConvertRelativePathToFull(FPaths::Combine(ExperimentalPluginsPath, TEXT("NNERuntimeBasicCpu/Content/Python"))));
+	
+	FFileHelper::SaveStringArrayToFile(PluginSitePackagePaths, *PyPluginsSitePackageFile, FFileHelper::EEncodingOptions::ForceUTF8WithoutBOM);
+}
+
+IMPLEMENT_MODULE(FLearningTrainingModule, LearningTraining)
diff --git a/Engine/Plugins/Experimental/LearningCore/Source/LearningTraining/Public/LearningTraining.h b/Engine/Plugins/Experimental/LearningCore/Source/LearningTraining/Public/LearningTraining.h
new file mode 100644
--- /dev/null
+++ b/Engine/Plugins/Experimental/LearningCore/Source/LearningTraining/Public/LearningTraining.h
@@ -0,0 +1,20 @@
+// Copyright Epic Games, Inc. All Rights Reserved.
+
+#pragma once
+
+#include "Modules/ModuleManager.h"
+
+class FLearningTrainingModule : public IModuleInterface
+{
+public:
+
+	/** IModuleInterface implementation */
+	virtual void StartupModule() override;
+	virtual void ShutdownModule() override;
+
+	/** Register any internal Python dependencies to make them available in Python sys.path */
+	virtual void RegisterPythonPackages();
+};
+
+
+
