[Backout] - CL45078039
#rnx
Original CL Desc
-----------------------------------------------------------------
Separated LearningAgents plugin into two: LearningCore and LearningAgents

diff --git a/Engine/Plugins/Experimental/LearningAgents/Source/LearningAgentsTraining/Private/LearningAgentsCommunicator.cpp b/Engine/Plugins/Experimental/LearningAgents/Source/LearningAgentsTraining/Private/LearningAgentsCommunicator.cpp
--- a/Engine/Plugins/Experimental/LearningAgents/Source/LearningAgentsTraining/Private/LearningAgentsCommunicator.cpp
+++ b/Engine/Plugins/Experimental/LearningAgents/Source/LearningAgentsTraining/Private/LearningAgentsCommunicator.cpp
@@ -38,14 +38,14 @@ FLearningAgentsSharedMemoryTrainerProcess ULearningAgentsCommunicatorLibrary::Sp
 		return TrainerProcess;
 	}
 
-	const FString PythonContentPath = UE::Learning::Agents::Trainer::GetPythonContentPath(TrainerProcessSettings.GetEditorEnginePath());
+	const FString PythonContentPath = UE::Learning::Trainer::GetPythonContentPath(TrainerProcessSettings.GetEditorEnginePath());
 	if (!FPaths::DirectoryExists(PythonContentPath))
 	{
 		UE_LOG(LogLearning, Error, TEXT("SpawnSharedMemoryTrainingProcess: Can't find LearningAgents plugin Content \"%s\"."), *PythonContentPath);
 		return TrainerProcess;
 	}
 
-	const FString IntermediatePath = UE::Learning::Agents::Trainer::GetIntermediatePath(TrainerProcessSettings.GetIntermediatePath());
+	const FString IntermediatePath = UE::Learning::Trainer::GetIntermediatePath(TrainerProcessSettings.GetIntermediatePath());
 
 	const FString CustomTrainerModulePath = TrainerProcessSettings.GetCustomTrainerModulePath();
 	if (!CustomTrainerModulePath.IsEmpty() and !FPaths::DirectoryExists(CustomTrainerModulePath))
@@ -145,14 +145,14 @@ FLearningAgentsSocketTrainerProcess ULearningAgentsCommunicatorLibrary::SpawnSoc
 		return TrainerProcess;
 	}
 
-	const FString PythonContentPath = UE::Learning::Agents::Trainer::GetPythonContentPath(TrainerProcessSettings.GetEditorEnginePath());
+	const FString PythonContentPath = UE::Learning::Trainer::GetPythonContentPath(TrainerProcessSettings.GetEditorEnginePath());
 	if (!FPaths::DirectoryExists(PythonContentPath))
 	{
 		UE_LOG(LogLearning, Error, TEXT("SpawnSocketTrainingProcess: Can't find LearningAgents plugin Content \"%s\"."), *PythonContentPath);
 		return TrainerProcess;
 	}
 
-	const FString IntermediatePath = UE::Learning::Agents::Trainer::GetIntermediatePath(TrainerProcessSettings.GetIntermediatePath());
+	const FString IntermediatePath = UE::Learning::Trainer::GetIntermediatePath(TrainerProcessSettings.GetIntermediatePath());
 
 	const FString CustomTrainerModulePath = TrainerProcessSettings.GetCustomTrainerModulePath();
 	if (!CustomTrainerModulePath.IsEmpty() and !FPaths::DirectoryExists(CustomTrainerModulePath))
diff --git a/Engine/Plugins/Experimental/LearningAgents/Source/LearningAgentsTraining/Private/LearningAgentsTrainer.cpp b/Engine/Plugins/Experimental/LearningAgents/Source/LearningAgentsTraining/Private/LearningAgentsTrainer.cpp
--- a/Engine/Plugins/Experimental/LearningAgents/Source/LearningAgentsTraining/Private/LearningAgentsTrainer.cpp
+++ b/Engine/Plugins/Experimental/LearningAgents/Source/LearningAgentsTraining/Private/LearningAgentsTrainer.cpp
@@ -164,19 +164,6 @@ namespace UE::Learning::Agents
 		}
 #endif
 	}
-
-	namespace Trainer
-	{
-		FString GetIntermediatePath(const FString& IntermediateDir)
-		{
-			return IntermediateDir / TEXT("LearningAgents");
-		}
-
-		FString GetPythonContentPath(const FString& EngineDir)
-		{
-			return EngineDir / TEXT("Plugins/Experimental/LearningAgents/Content/Python/");
-		}
-	}
 }
 
 FLearningAgentsTrainerProcessSettings::FLearningAgentsTrainerProcessSettings()
diff --git a/Engine/Plugins/Experimental/LearningAgents/Source/LearningAgentsTraining/Private/LearningAgentsTraining.cpp b/Engine/Plugins/Experimental/LearningAgents/Source/LearningAgentsTraining/Private/LearningAgentsTraining.cpp
--- a/Engine/Plugins/Experimental/LearningAgents/Source/LearningAgentsTraining/Private/LearningAgentsTraining.cpp
+++ b/Engine/Plugins/Experimental/LearningAgents/Source/LearningAgentsTraining/Private/LearningAgentsTraining.cpp
@@ -1,8 +1,6 @@
 // Copyright Epic Games, Inc. All Rights Reserved.
 
 #include "LearningAgentsTraining.h"
-#include "Misc/FileHelper.h"
-#include "Misc/Paths.h"
 
 void FLearningAgentsTrainingModule::StartupModule()
 {
@@ -14,17 +12,5 @@ void FLearningAgentsTrainingModule::ShutdownModule()
 	// This function may be called during shutdown to clean up your module.  For modules that support dynamic reloading,
 	// we call this function before unloading the module.
 }
-
-void FLearningAgentsTrainingModule::RegisterPythonPackages()
-{
-	const FString ExperimentalPluginsPath = FPaths::Combine(FPaths::EnginePluginsDir(), TEXT("Experimental"));
-	const FString SitePackagesPath = FPaths::ConvertRelativePathToFull(FPaths::ProjectIntermediateDir() / TEXT("PipInstall") / TEXT("Lib") / TEXT("site-packages"));
-	const FString PyPluginsSitePackageFile = FPaths::ConvertRelativePathToFull(SitePackagesPath / TEXT("learning_agents.pth"));
-
-	TArray<FString> PluginSitePackagePaths;
-	PluginSitePackagePaths.Add(FPaths::ConvertRelativePathToFull(FPaths::Combine(ExperimentalPluginsPath, TEXT("LearningAgents/Content/Python"))));
-
-	FFileHelper::SaveStringArrayToFile(PluginSitePackagePaths, *PyPluginsSitePackageFile, FFileHelper::EEncodingOptions::ForceUTF8WithoutBOM);
-}
 	
 IMPLEMENT_MODULE(FLearningAgentsTrainingModule, LearningAgentsTraining)
diff --git a/Engine/Plugins/Experimental/LearningAgents/Source/LearningAgentsTraining/Public/LearningAgentsTrainer.h b/Engine/Plugins/Experimental/LearningAgents/Source/LearningAgentsTraining/Public/LearningAgentsTrainer.h
--- a/Engine/Plugins/Experimental/LearningAgents/Source/LearningAgentsTraining/Public/LearningAgentsTrainer.h
+++ b/Engine/Plugins/Experimental/LearningAgents/Source/LearningAgentsTraining/Public/LearningAgentsTrainer.h
@@ -89,23 +89,14 @@ namespace UE::Learning::Agents
 	};
 
 	/** Get the learning agents trainer device from the UE::Learning trainer device. */
-	UE_API ELearningAgentsTrainingDevice GetLearningAgentsTrainingDevice(const ETrainerDevice Device);
+	LEARNINGAGENTSTRAINING_API ELearningAgentsTrainingDevice GetLearningAgentsTrainingDevice(const ETrainerDevice Device);
 
 	/** Get the UE::Learning trainer device from the learning agents trainer device. */
-	UE_API ETrainerDevice GetTrainingDevice(const ELearningAgentsTrainingDevice Device);
+	LEARNINGAGENTSTRAINING_API ETrainerDevice GetTrainingDevice(const ELearningAgentsTrainingDevice Device);
 
-	UE_API void ApplyGameSettings(const FLearningAgentsTrainingGameSettings& Settings, const UWorld* World, FGameSettingsState& OutGameSettingsState);
+	LEARNINGAGENTSTRAINING_API void ApplyGameSettings(const FLearningAgentsTrainingGameSettings& Settings, const UWorld* World, FGameSettingsState& OutGameSettingsState);
 
-	UE_API void RevertGameSettings(const FGameSettingsState& Settings, const UWorld* World);
-
-	namespace Trainer
-	{
-		/** Gets the LearningAgents Intermediate path from the intermediate directory. */
-		UE_API FString GetIntermediatePath(const FString& IntermediateDir);
-
-		/** Gets the LearningAgents Content path from the engine directory. */
-		UE_API FString GetPythonContentPath(const FString& EngineDir);
-	}
+	LEARNINGAGENTSTRAINING_API void RevertGameSettings(const FGameSettingsState& Settings, const UWorld* World);
 }
 
 /** The path settings for the trainer. */
@@ -165,7 +156,7 @@ public:
 
 	/** Trainer file name. The name of the python file to use for training. Do NOT include the '.py' file extension. */
 	UPROPERTY(EditAnywhere, Category = "LearningAgents")
-	FString TrainerFileName = TEXT("learning_core.train_ppo");
+	FString TrainerFileName = TEXT("learning_agents.train_ppo");
 
 public:
 
diff --git a/Engine/Plugins/Experimental/LearningAgents/Source/LearningAgentsTraining/Public/LearningAgentsTraining.h b/Engine/Plugins/Experimental/LearningAgents/Source/LearningAgentsTraining/Public/LearningAgentsTraining.h
--- a/Engine/Plugins/Experimental/LearningAgents/Source/LearningAgentsTraining/Public/LearningAgentsTraining.h
+++ b/Engine/Plugins/Experimental/LearningAgents/Source/LearningAgentsTraining/Public/LearningAgentsTraining.h
@@ -10,7 +10,4 @@ public:
 	/** IModuleInterface implementation */
 	virtual void StartupModule() override;
 	virtual void ShutdownModule() override;
-
-	/** Register any internal Python dependencies to make them available in Python sys.path */
-	virtual void RegisterPythonPackages();
 };
diff --git a/Engine/Plugins/Experimental/LearningAgents/Source/LearningTraining/Private/LearningTraining.cpp b/Engine/Plugins/Experimental/LearningAgents/Source/LearningTraining/Private/LearningTraining.cpp
new file mode 100644
--- /dev/null
+++ b/Engine/Plugins/Experimental/LearningAgents/Source/LearningTraining/Private/LearningTraining.cpp
@@ -0,0 +1,32 @@
+// Copyright Epic Games, Inc. All Rights Reserved.
+
+#include "LearningTraining.h"
+#include "Misc/FileHelper.h"
+#include "Misc/Paths.h"
+
+void FLearningTrainingModule::StartupModule()
+{
+	// This code will execute after your module is loaded into memory; the exact timing is specified in the .uplugin file per-module
+	RegisterPythonPackages();
+}
+
+void FLearningTrainingModule::ShutdownModule()
+{
+	// This function may be called during shutdown to clean up your module.  For modules that support dynamic reloading,
+	// we call this function before unloading the module.
+}
+
+void FLearningTrainingModule::RegisterPythonPackages()
+{
+	const FString ExperimentalPluginsPath = FPaths::Combine(FPaths::EnginePluginsDir(), TEXT("Experimental"));
+	const FString SitePackagesPath = FPaths::ConvertRelativePathToFull(FPaths::ProjectIntermediateDir() / TEXT("PipInstall") / TEXT("Lib") / TEXT("site-packages"));
+	const FString PyPluginsSitePackageFile = FPaths::ConvertRelativePathToFull(SitePackagesPath / TEXT("learning_agents.pth"));
+	
+	TArray<FString> PluginSitePackagePaths;
+	PluginSitePackagePaths.Add(FPaths::ConvertRelativePathToFull(FPaths::Combine(ExperimentalPluginsPath, TEXT("LearningAgents/Content/Python"))));
+	PluginSitePackagePaths.Add(FPaths::ConvertRelativePathToFull(FPaths::Combine(ExperimentalPluginsPath, TEXT("NNERuntimeBasicCpu/Content/Python"))));
+	
+	FFileHelper::SaveStringArrayToFile(PluginSitePackagePaths, *PyPluginsSitePackageFile, FFileHelper::EEncodingOptions::ForceUTF8WithoutBOM);
+}
+
+IMPLEMENT_MODULE(FLearningTrainingModule, LearningTraining)
diff --git a/Engine/Plugins/Experimental/LearningAgents/Source/LearningTraining/Public/LearningTraining.h b/Engine/Plugins/Experimental/LearningAgents/Source/LearningTraining/Public/LearningTraining.h
new file mode 100644
--- /dev/null
+++ b/Engine/Plugins/Experimental/LearningAgents/Source/LearningTraining/Public/LearningTraining.h
@@ -0,0 +1,20 @@
+// Copyright Epic Games, Inc. All Rights Reserved.
+
+#pragma once
+
+#include "Modules/ModuleManager.h"
+
+class FLearningTrainingModule : public IModuleInterface
+{
+public:
+
+	/** IModuleInterface implementation */
+	virtual void StartupModule() override;
+	virtual void ShutdownModule() override;
+
+	/** Register any internal Python dependencies to make them available in Python sys.path */
+	virtual void RegisterPythonPackages();
+};
+
+
+
