[Learning Agents] Added a getter for an interactor to access configs from a recorder

#rnx

diff --git a/Engine/Plugins/Experimental/LearningAgents/Source/LearningAgentsTraining/Private/LearningAgentsRecorder.cpp b/Engine/Plugins/Experimental/LearningAgents/Source/LearningAgentsTraining/Private/LearningAgentsRecorder.cpp
--- a/Engine/Plugins/Experimental/LearningAgents/Source/LearningAgentsTraining/Private/LearningAgentsRecorder.cpp
+++ b/Engine/Plugins/Experimental/LearningAgents/Source/LearningAgentsTraining/Private/LearningAgentsRecorder.cpp
@@ -406,6 +406,11 @@ const ULearningAgentsRecording* ULearningAgentsRecorder::GetRecordingAsset() con
 	return Recording;
 }
 
+ULearningAgentsInteractor* ULearningAgentsRecorder::GetInteractor() const
+{
+	return Interactor;
+}
+
 void ULearningAgentsRecorder::SaveSchemaConfig()
 {
 	TSharedRef<FJsonObject> ConfigObject = MakeShared<FJsonObject>();
diff --git a/Engine/Plugins/Experimental/LearningAgents/Source/LearningAgentsTraining/Public/LearningAgentsRecorder.h b/Engine/Plugins/Experimental/LearningAgents/Source/LearningAgentsTraining/Public/LearningAgentsRecorder.h
--- a/Engine/Plugins/Experimental/LearningAgents/Source/LearningAgentsTraining/Public/LearningAgentsRecorder.h
+++ b/Engine/Plugins/Experimental/LearningAgents/Source/LearningAgentsTraining/Public/LearningAgentsRecorder.h
@@ -132,6 +132,9 @@ public:
 	UFUNCTION(BlueprintPure, Category = "LearningAgents")
 	UE_API const ULearningAgentsRecording* GetRecordingAsset() const;
 
+protected:
+	UFUNCTION(BlueprintPure, Category="LearningAgents")
+	UE_API ULearningAgentsInteractor* GetInteractor() const;
 
 // ----- Private Data ----- 
 private:
