[LearningAgents] Added Option to use standard GRU

diff --git a/Engine/Plugins/Experimental/LearningAgents/Source/LearningAgents/LearningAgents.Build.cs b/Engine/Plugins/Experimental/LearningAgents/Source/LearningAgents/LearningAgents.Build.cs
--- a/Engine/Plugins/Experimental/LearningAgents/Source/LearningAgents/LearningAgents.Build.cs
+++ b/Engine/Plugins/Experimental/LearningAgents/Source/LearningAgents/LearningAgents.Build.cs
@@ -22,6 +22,8 @@ public class LearningAgents : ModuleRules
 			{
 				"NNE",
 				"NNERuntimeBasicCpu",
+				"Json",
+				"JsonUtilities",
 			});
 	}
 }
diff --git a/Engine/Plugins/Experimental/LearningAgents/Source/LearningAgents/Private/LearningAgentsPolicy.cpp b/Engine/Plugins/Experimental/LearningAgents/Source/LearningAgents/Private/LearningAgentsPolicy.cpp
--- a/Engine/Plugins/Experimental/LearningAgents/Source/LearningAgents/Private/LearningAgentsPolicy.cpp
+++ b/Engine/Plugins/Experimental/LearningAgents/Source/LearningAgents/Private/LearningAgentsPolicy.cpp
@@ -13,6 +13,7 @@
 #include "Misc/FileHelper.h"
 #include "Misc/Paths.h"
 #include "GameFramework/Actor.h"
+#include "Dom/JsonObject.h"
 
 #include "NNERuntimeBasicCpuBuilder.h"
 
@@ -27,6 +28,7 @@ namespace UE::Learning::Agents::Policy::Private
 		case ELearningAgentsActivationFunction::ReLU: return UE::NNE::RuntimeBasic::FModelBuilder::EActivationFunction::ReLU;
 		case ELearningAgentsActivationFunction::ELU: return UE::NNE::RuntimeBasic::FModelBuilder::EActivationFunction::ELU;
 		case ELearningAgentsActivationFunction::TanH: return UE::NNE::RuntimeBasic::FModelBuilder::EActivationFunction::TanH;
+		case ELearningAgentsActivationFunction::GELU: return UE::NNE::RuntimeBasic::FModelBuilder::EActivationFunction::GELU;
 		default:
 			UE_LOG(LogLearning, Error, TEXT("Unknown Activation Function"));
 			return UE::NNE::RuntimeBasic::FModelBuilder::EActivationFunction::ReLU;
@@ -93,7 +95,7 @@ void ULearningAgentsPolicy::SetupPolicy(
 	const bool bReinitializeEncoderNetwork,
 	const bool bReinitializePolicyNetwork,
 	const bool bReinitializeDecoderNetwork,
-	const FLearningAgentsPolicySettings& PolicySettings,
+	const FLearningAgentsPolicySettings& InPolicySettings,
 	const int32 Seed)
 {
 	if (IsSetup())
@@ -142,10 +144,12 @@ void ULearningAgentsPolicy::SetupPolicy(
 
 	Manager = InManager;
 	Interactor = InInteractor;
+	PolicySettings = InPolicySettings;
 
 	// Compatibility Hashes and Sizes
 
-	const int32 MemoryStateSize = PolicySettings.bUseMemory ? PolicySettings.MemoryStateSize : 0;
+	const int32 MemoryStateSize = PolicySettings.MemoryCell == ELearningAgentsMemoryCell::NoMemoryCell ? 0: PolicySettings.MemoryStateSize;
+
 	const int32 ObservationVectorSize = Interactor->GetObservationVectorSize();
 	const int32 ObservationEncodedVectorSize = Interactor->GetObservationEncodedVectorSize();
 	const int32 ActionEncodedVectorSize = Interactor->GetActionEncodedVectorSize();
@@ -267,35 +271,86 @@ void ULearningAgentsPolicy::SetupPolicy(
 		MemoryCellWeightSettings.WeightInitializationSettings.Scale = 0.1f;
 
 		TArray<uint8> FileData;
-		uint32 PolicyInputSize, PolicyOutputSize;
-		Builder.WriteFileDataAndReset(FileData, PolicyInputSize, PolicyOutputSize,
-			Builder.MakeMemoryBackbone(
-				MemoryStateSize,
-				Builder.MakeMLP(
-					ObservationEncodedVectorSize,
-					PolicyHiddenLayerSize,
-					PolicyHiddenLayerSize,
-					PolicyHiddenLayerNum / 2 + 2, // Add 2 to account for input and output layers
-					UE::Learning::Agents::Policy::Private::GetBuilderActivationFunction(PolicyActivationFunction),
-					true),
-				Builder.MakeMemoryCellLayer(
-					PolicyHiddenLayerSize,
-					PolicyHiddenLayerSize,
-					MemoryStateSize,
-					MemoryCellWeightSettings),
+		uint32 PolicyInputSize = 0, PolicyOutputSize = 0;
+
+		switch (PolicySettings.MemoryCell)
+		{
+		case ELearningAgentsMemoryCell::NoMemoryCell:
+			Builder.WriteFileDataAndReset(FileData, PolicyInputSize, PolicyOutputSize,
 				Builder.MakeSequence({
 					Builder.MakeMLP(
-						PolicyHiddenLayerSize,
+						ObservationEncodedVectorSize,
 						ActionEncodedVectorSize,
 						PolicyHiddenLayerSize,
-						PolicyHiddenLayerNum / 2 + 2, // Add 2 to account for input and output layers
+						PolicyHiddenLayerNum + 2, // Add 2 to account for input and output layers
 						UE::Learning::Agents::Policy::Private::GetBuilderActivationFunction(PolicyActivationFunction)),
 					Builder.MakeDenormalize(
 						ActionEncodedVectorSize,
 						Builder.MakeValuesZero(ActionEncodedVectorSize),
-						Builder.MakeValuesConstant(ActionEncodedVectorSize, PolicyInitialEncodedActionScale))
-					})
-			));
+						Builder.MakeValuesConstant(ActionEncodedVectorSize, PolicyInitialEncodedActionScale)) }));
+			break;
+		case ELearningAgentsMemoryCell::LearningAgentsGRU:
+			Builder.WriteFileDataAndReset(FileData, PolicyInputSize, PolicyOutputSize,
+				Builder.MakeMemoryBackbone(
+					MemoryStateSize,
+					Builder.MakeMLP(
+						ObservationEncodedVectorSize,
+						PolicyHiddenLayerSize,
+						PolicyHiddenLayerSize,
+						PolicyHiddenLayerNum / 2 + 2, // Add 2 to account for input and output layers
+						UE::Learning::Agents::Policy::Private::GetBuilderActivationFunction(PolicyActivationFunction),
+						true),
+					Builder.MakeMemoryCellLayer(
+						PolicyHiddenLayerSize,
+						PolicyHiddenLayerSize,
+						MemoryStateSize,
+						MemoryCellWeightSettings),
+					Builder.MakeSequence({
+						Builder.MakeMLP(
+							PolicyHiddenLayerSize,
+							ActionEncodedVectorSize,
+							PolicyHiddenLayerSize,
+							PolicyHiddenLayerNum / 2 + 2, // Add 2 to account for input and output layers
+							UE::Learning::Agents::Policy::Private::GetBuilderActivationFunction(PolicyActivationFunction)),
+						Builder.MakeDenormalize(
+							ActionEncodedVectorSize,
+							Builder.MakeValuesZero(ActionEncodedVectorSize),
+							Builder.MakeValuesConstant(ActionEncodedVectorSize, PolicyInitialEncodedActionScale))
+						})
+				));
+			break;
+		case ELearningAgentsMemoryCell::StandardGRU:
+			Builder.WriteFileDataAndReset(FileData, PolicyInputSize, PolicyOutputSize,
+				Builder.MakeMemoryBackbone(
+					MemoryStateSize,
+					Builder.MakeMLP(
+						ObservationEncodedVectorSize,
+						PolicyHiddenLayerSize,
+						PolicyHiddenLayerSize,
+						PolicyHiddenLayerNum / 2 + 2, // Add 2 to account for input and output layers
+						UE::Learning::Agents::Policy::Private::GetBuilderActivationFunction(PolicyActivationFunction),
+						true),
+					Builder.MakeGRUCellLayer(
+						PolicyHiddenLayerSize,
+						MemoryStateSize,
+						MemoryCellWeightSettings),
+					Builder.MakeSequence({
+						Builder.MakeMLP(
+							MemoryStateSize,
+							ActionEncodedVectorSize,
+							PolicyHiddenLayerSize,
+							PolicyHiddenLayerNum / 2 + 2, // Add 2 to account for input and output layers
+							UE::Learning::Agents::Policy::Private::GetBuilderActivationFunction(PolicyActivationFunction)),
+						Builder.MakeDenormalize(
+							ActionEncodedVectorSize,
+							Builder.MakeValuesZero(ActionEncodedVectorSize),
+							Builder.MakeValuesConstant(ActionEncodedVectorSize, PolicyInitialEncodedActionScale))
+						})
+				));
+			break;
+		default:
+			checkNoEntry();
+		}
 
 		check(PolicyInputSize == ObservationEncodedVectorSize + MemoryStateSize);
 		check(PolicyOutputSize == ActionEncodedVectorSize + MemoryStateSize);
@@ -414,6 +469,19 @@ void ULearningAgentsPolicy::SetupPolicy(
 	Manager->AddListener(this);
 }
 
+TSharedRef<FJsonObject> ULearningAgentsPolicy::AsJsonConfig() const
+{
+	TSharedRef<FJsonObject> ConfigObject = MakeShared<FJsonObject>();
+
+	ConfigObject->SetNumberField(TEXT("HiddenLayerNum"), PolicySettings.HiddenLayerNum);
+	ConfigObject->SetNumberField(TEXT("HiddenLayerSize"), PolicySettings.HiddenLayerSize);
+	ConfigObject->SetStringField(TEXT("MemoryCell"), StaticEnum<ELearningAgentsMemoryCell>()->GetNameStringByValue(static_cast<int64>(PolicySettings.MemoryCell)));
+	ConfigObject->SetNumberField(TEXT("MemoryStateSize"), GetMemoryStateSize());
+	ConfigObject->SetNumberField(TEXT("InitialEncodedActionScale"), PolicySettings.InitialEncodedActionScale);
+	ConfigObject->SetStringField(TEXT("ActivationFunction"), StaticEnum<ELearningAgentsActivationFunction>()->GetNameStringByValue(static_cast<int64>(PolicySettings.ActivationFunction)));
+	return ConfigObject;
+}
+
 void ULearningAgentsPolicy::OnAgentsAdded_Implementation(const TArray<int32>& AgentIds)
 {
 	if (!IsSetup())
diff --git a/Engine/Plugins/Experimental/LearningAgents/Source/LearningAgents/Public/LearningAgentsNeuralNetwork.h b/Engine/Plugins/Experimental/LearningAgents/Source/LearningAgents/Public/LearningAgentsNeuralNetwork.h
--- a/Engine/Plugins/Experimental/LearningAgents/Source/LearningAgents/Public/LearningAgentsNeuralNetwork.h
+++ b/Engine/Plugins/Experimental/LearningAgents/Source/LearningAgents/Public/LearningAgentsNeuralNetwork.h
@@ -20,6 +20,9 @@ enum class ELearningAgentsActivationFunction : uint8
 
 	/** TanH Activation - Smooth activation function that is slower to train and evaluate but sometimes more stable for certain tasks. */
 	TanH	UMETA(DisplayName = "TanH"),
+
+	/** GELU Activation - Gaussian Error Linear Unit. */
+	GELU	UMETA(DisplayName = "GELU"),
 };
 
 class ULearningNeuralNetworkData;
diff --git a/Engine/Plugins/Experimental/LearningAgents/Source/LearningAgents/Public/LearningAgentsPolicy.h b/Engine/Plugins/Experimental/LearningAgents/Source/LearningAgents/Public/LearningAgentsPolicy.h
--- a/Engine/Plugins/Experimental/LearningAgents/Source/LearningAgents/Public/LearningAgentsPolicy.h
+++ b/Engine/Plugins/Experimental/LearningAgents/Source/LearningAgents/Public/LearningAgentsPolicy.h
@@ -19,6 +19,15 @@ namespace UE::Learning
 }
 
 class ULearningAgentsNeuralNetwork;
+class FJsonObject;
+
+UENUM(BlueprintType)
+enum class ELearningAgentsMemoryCell : uint8
+{
+	NoMemoryCell UMETA(DisplayName = "No Memory Cell"),
+	LearningAgentsGRU UMETA(DisplayName = "Learning Agents GRU"),
+	StandardGRU UMETA(DisplayName = "Standard GRU")
+};
 
 /** The configurable settings for a ULearningAgentsPolicy. */
 USTRUCT(BlueprintType, Category = "LearningAgents")
@@ -36,9 +45,9 @@ public:
 	UPROPERTY(EditAnywhere, Category = "LearningAgents", meta = (ClampMin = "1", UIMin = "1"))
 	int32 HiddenLayerSize = 128;
 
-	/** If this network should use memory or not. Setting this to false is equivalent to setting MemoryStateSize to 0. */
+	/** The type of memory cell to use. Setting this to NoMemoryCell is equivalent to setting MemoryStateSize to 0. */
 	UPROPERTY(EditAnywhere, Category = "LearningAgents")
-	bool bUseMemory = true;
+	ELearningAgentsMemoryCell MemoryCell = ELearningAgentsMemoryCell::LearningAgentsGRU;
 
 	/** Number of neurons in the memory state of the policy network */
 	UPROPERTY(EditAnywhere, Category = "LearningAgents", meta = (ClampMin = "0", UIMin = "0"))
@@ -137,9 +146,11 @@ public:
 		const bool bReinitializeEncoderNetwork = true,
 		const bool bReinitializePolicyNetwork = true,
 		const bool bReinitializeDecoderNetwork = true,
-		const FLearningAgentsPolicySettings& PolicySettings = FLearningAgentsPolicySettings(),
+		const FLearningAgentsPolicySettings& InPolicySettings = FLearningAgentsPolicySettings(),
 		const int32 Seed = 1234);
 
+	UE_API TSharedRef<FJsonObject> AsJsonConfig() const;
+
 public:
 
 	//~ Begin ULearningAgentsManagerListener Interface
@@ -300,6 +311,9 @@ private:
 	/** Temp buffers used to record the set of agents that are valid for evaluation */
 	TArray<int32> ValidAgentIds;
 	UE::Learning::FIndexSet ValidAgentSet;
+
+	FLearningAgentsPolicySettings PolicySettings;
+
 };
 
 #undef UE_API
diff --git a/Engine/Plugins/Experimental/LearningAgents/Source/LearningAgentsTraining/Private/LearningAgentsImitationTrainer.cpp b/Engine/Plugins/Experimental/LearningAgents/Source/LearningAgentsTraining/Private/LearningAgentsImitationTrainer.cpp
--- a/Engine/Plugins/Experimental/LearningAgents/Source/LearningAgentsTraining/Private/LearningAgentsImitationTrainer.cpp
+++ b/Engine/Plugins/Experimental/LearningAgents/Source/LearningAgentsTraining/Private/LearningAgentsImitationTrainer.cpp
@@ -450,6 +450,7 @@ TSharedRef<FJsonObject> ULearningAgentsImitationTrainer::CreateTrainerConfig(con
 
 	// Add Imitation Specific Config Entries
 	ConfigObject->SetObjectField(TEXT("BehaviorCloningSettings"), TrainingSettings.AsJsonConfig());
+	ConfigObject->SetObjectField(TEXT("PolicySettings"), Policy->AsJsonConfig());
 	ConfigObject->SetNumberField(TEXT("MemoryStateNum"), Policy->GetMemoryStateSize());
 
 	// Add Training Workflow Config Entries
