LACombat: Split schemas into a separate map & make recordings filterable with tags

#changelist validated
#virtualized

diff --git a/Engine/Plugins/Experimental/LearningAgents/Source/LearningAgentsTraining/LearningAgentsTraining.Build.cs b/Engine/Plugins/Experimental/LearningAgents/Source/LearningAgentsTraining/LearningAgentsTraining.Build.cs
--- a/Engine/Plugins/Experimental/LearningAgents/Source/LearningAgentsTraining/LearningAgentsTraining.Build.cs
+++ b/Engine/Plugins/Experimental/LearningAgents/Source/LearningAgentsTraining/LearningAgentsTraining.Build.cs
@@ -8,23 +8,13 @@ public class LearningAgentsTraining : ModuleRules
 	{
 		PCHUsage = ModuleRules.PCHUsageMode.UseExplicitOrSharedPCHs;
 
-		PublicIncludePaths.AddRange(
-			new string[] {
-				// ... add public include paths required here ...
-			});
-
-		PrivateIncludePaths.AddRange(
-			new string[] {
-				// ... add other private include paths required here ...
-			});
-
 		PublicDependencyModuleNames.AddRange(
 			new string[]
 			{
+				"AIModule",
 				"Core",
+				"GameplayTags",
 				"NavigationSystem",
-				"AIModule",
-				// ... add other public dependencies that you statically link with here ...
 			});
 
 		PrivateDependencyModuleNames.AddRange(
@@ -37,18 +27,11 @@ public class LearningAgentsTraining : ModuleRules
 				"Learning",
 				"LearningAgents",
 				"LearningTraining",
-				// ... add private dependencies that you statically link with here ...	
 			});
 
 		if (Target.bBuildEditor)
 		{
 			PrivateDependencyModuleNames.Add("UnrealEd");
 		}
-
-		DynamicallyLoadedModuleNames.AddRange(
-			new string[]
-			{
-				// ... add any modules that your module loads dynamically here ...
-			});
 	}
 }
diff --git a/Engine/Plugins/Experimental/LearningAgents/Source/LearningAgentsTraining/Private/LearningAgentsImitationTrainer.cpp b/Engine/Plugins/Experimental/LearningAgents/Source/LearningAgentsTraining/Private/LearningAgentsImitationTrainer.cpp
--- a/Engine/Plugins/Experimental/LearningAgents/Source/LearningAgentsTraining/Private/LearningAgentsImitationTrainer.cpp
+++ b/Engine/Plugins/Experimental/LearningAgents/Source/LearningAgentsTraining/Private/LearningAgentsImitationTrainer.cpp
@@ -179,13 +179,25 @@ void ULearningAgentsImitationTrainer::BeginTraining(
 	int32 TotalStepNum = 0;
 	for (const FLearningAgentsRecord& Record : Recording->Records)
 	{
-		if (!UE::Learning::Trainer::IsObservationSchemaSubsetCompatible(Record.ObservationSchemaJson, Interactor->GetObservationSchema()->ObservationSchema, Interactor->GetObservationSchemaElement().SchemaElement))
+		if (ImitationTrainerSettings.RecordingFilterTag.IsValid() && !Record.Tags.HasTag(ImitationTrainerSettings.RecordingFilterTag))
+		{
+			UE_LOG(LogLearning, Log, TEXT("Skipping record %s because its missing tag %s."), *Record.Name, *ImitationTrainerSettings.RecordingFilterTag.ToString());
+			continue;
+		}
+
+		if (!Recording->Schemas.Contains(Record.SchemaName))
+		{
+			UE_LOG(LogLearning, Error, TEXT("Schema name %s not found in schema list. Is the schema added and are the names correct?"), *Record.SchemaName.ToString());
+			continue;
+		}
+
+		if (!UE::Learning::Trainer::IsObservationSchemaSubsetCompatible(Recording->Schemas[Record.SchemaName].ObservationSchemaJson, Interactor->GetObservationSchema()->ObservationSchema, Interactor->GetObservationSchemaElement().SchemaElement))
 		{
 			UE_LOG(LogLearning, Warning, TEXT("%s: Record has incompatible observation schema."), *GetName());
 			continue;
 		}
 
-		if (!UE::Learning::Trainer::IsActionSchemaSubsetCompatible(Record.ActionSchemaJson, Interactor->GetActionSchema()->ActionSchema, Interactor->GetActionSchemaElement().SchemaElement))
+		if (!UE::Learning::Trainer::IsActionSchemaSubsetCompatible(Recording->Schemas[Record.SchemaName].ActionSchemaJson, Interactor->GetActionSchema()->ActionSchema, Interactor->GetActionSchemaElement().SchemaElement))
 		{
 			UE_LOG(LogLearning, Warning, TEXT("%s: Record has incompatible action schema."), *GetName());
 			continue;
@@ -222,6 +234,18 @@ void ULearningAgentsImitationTrainer::BeginTraining(
 	int32 StepIdx = 0;
 	for (const FLearningAgentsRecord& Record : Recording->Records)
 	{
+		if (ImitationTrainerSettings.RecordingFilterTag.IsValid() && !Record.Tags.HasTag(ImitationTrainerSettings.RecordingFilterTag))
+		{
+			UE_LOG(LogLearning, Log, TEXT("Skipping record %s because its missing tag %s."), *Record.Name, *ImitationTrainerSettings.RecordingFilterTag.ToString());
+			continue;
+		}
+
+		if (!Recording->Schemas.Contains(Record.SchemaName))
+		{
+			UE_LOG(LogLearning, Error, TEXT("Schema name %s not found in schema list. Is the schema added and are the names correct?"), *Record.SchemaName.ToString());
+			continue;
+		}
+
 		if (Record.ObservationDimNum == ObservationNum && Record.ActionDimNum == ActionNum)
 		{
 			TLearningArrayView<2, const float> ObservationsView = TLearningArrayView<2, const float>(Record.ObservationData.GetData(), { Record.StepNum, Record.ObservationDimNum });
@@ -232,8 +256,8 @@ void ULearningAgentsImitationTrainer::BeginTraining(
 			UE::Learning::Array::Copy(RecordedObservations.Slice(StepIdx, Record.StepNum), ObservationsView);
 			UE::Learning::Array::Copy(RecordedActions.Slice(StepIdx, Record.StepNum), ActionsView);
 		}
-		else if (UE::Learning::Trainer::IsObservationSchemaSubsetCompatible(Record.ObservationSchemaJson, Interactor->GetObservationSchema()->ObservationSchema, Interactor->GetObservationSchemaElement().SchemaElement) &&
-			UE::Learning::Trainer::IsActionSchemaSubsetCompatible(Record.ActionSchemaJson, Interactor->GetActionSchema()->ActionSchema, Interactor->GetActionSchemaElement().SchemaElement))
+		else if (UE::Learning::Trainer::IsObservationSchemaSubsetCompatible(Recording->Schemas[Record.SchemaName].ObservationSchemaJson, Interactor->GetObservationSchema()->ObservationSchema, Interactor->GetObservationSchemaElement().SchemaElement) &&
+			UE::Learning::Trainer::IsActionSchemaSubsetCompatible(Recording->Schemas[Record.SchemaName].ActionSchemaJson, Interactor->GetActionSchema()->ActionSchema, Interactor->GetActionSchemaElement().SchemaElement))
 		{
 			TLearningArrayView<2, const float> ObservationsView = TLearningArrayView<2, const float>(Record.ObservationData.GetData(), { Record.StepNum, Record.ObservationDimNum });
 			TLearningArrayView<2, const float> ActionsView = TLearningArrayView<2, const float>(Record.ActionData.GetData(), { Record.StepNum, Record.ActionDimNum });
@@ -242,7 +266,7 @@ void ULearningAgentsImitationTrainer::BeginTraining(
 			RecordedEpisodeLengths[EpisodeIdx] = Record.StepNum;
 
 			TArray<int32> ObservationIndices = UE::Learning::Trainer::ComputeObservationSchemaSubsetIndices(
-				Record.ObservationSchemaJson, Interactor->GetObservationSchema()->ObservationSchema, Interactor->GetObservationSchemaElement().SchemaElement);
+				Recording->Schemas[Record.SchemaName].ObservationSchemaJson, Interactor->GetObservationSchema()->ObservationSchema, Interactor->GetObservationSchemaElement().SchemaElement);
 			check(ObservationIndices.Num() == ObservationNum);
 
 			for (int Step = 0; Step < Record.StepNum; Step++)
@@ -254,7 +278,7 @@ void ULearningAgentsImitationTrainer::BeginTraining(
 			}
 
 			TArray<int32> ActionIndices = UE::Learning::Trainer::ComputeActionSchemaSubsetIndices(
-				Record.ActionSchemaJson, Interactor->GetActionSchema()->ActionSchema, Interactor->GetActionSchemaElement().SchemaElement);
+				Recording->Schemas[Record.SchemaName].ActionSchemaJson, Interactor->GetActionSchema()->ActionSchema, Interactor->GetActionSchemaElement().SchemaElement);
 			check(ActionIndices.Num() == ActionNum);
 
 			for (int Step = 0; Step < Record.StepNum; Step++)
diff --git a/Engine/Plugins/Experimental/LearningAgents/Source/LearningAgentsTraining/Private/LearningAgentsRecorder.cpp b/Engine/Plugins/Experimental/LearningAgents/Source/LearningAgentsTraining/Private/LearningAgentsRecorder.cpp
--- a/Engine/Plugins/Experimental/LearningAgents/Source/LearningAgentsTraining/Private/LearningAgentsRecorder.cpp
+++ b/Engine/Plugins/Experimental/LearningAgents/Source/LearningAgentsTraining/Private/LearningAgentsRecorder.cpp
@@ -99,12 +99,6 @@ void ULearningAgentsRecorder::FAgentRecordBuffer::CopyToRecord(FLearningAgentsRe
 	Record.ObservationData.SetNumUninitialized(StepNum * Record.ObservationDimNum);
 	Record.ActionData.SetNumUninitialized(StepNum * Record.ActionDimNum);
 
-	TSharedRef<TJsonWriter<>> ObsJsonWriter = TJsonWriterFactory<>::Create(&Record.ObservationSchemaJson, 0);
-	FJsonSerializer::Serialize(UE::Learning::Trainer::ConvertObservationSchemaToJSON(ObservationSchema, ObservationSchemaElement).ToSharedRef(), ObsJsonWriter, true);
-	
-	TSharedRef<TJsonWriter<>> ActJsonWriter = TJsonWriterFactory<>::Create(&Record.ActionSchemaJson, 0);
-	FJsonSerializer::Serialize(UE::Learning::Trainer::ConvertActionSchemaToJSON(ActionSchema, ActionSchemaElement).ToSharedRef(), ActJsonWriter, true);
-
 	TLearningArrayView<2, float> ObservationsView = TLearningArrayView<2, float>(Record.ObservationData.GetData(), { Record.StepNum, Record.ObservationDimNum });
 	TLearningArrayView<2, float> ActionsView = TLearningArrayView<2, float>(Record.ActionData.GetData(), { Record.StepNum, Record.ActionDimNum });
 
diff --git a/Engine/Plugins/Experimental/LearningAgents/Source/LearningAgentsTraining/Private/LearningAgentsRecording.cpp b/Engine/Plugins/Experimental/LearningAgents/Source/LearningAgentsTraining/Private/LearningAgentsRecording.cpp
--- a/Engine/Plugins/Experimental/LearningAgents/Source/LearningAgentsTraining/Private/LearningAgentsRecording.cpp
+++ b/Engine/Plugins/Experimental/LearningAgents/Source/LearningAgentsTraining/Private/LearningAgentsRecording.cpp
@@ -6,6 +6,7 @@
 
 #include "UObject/Package.h"
 #include "Misc/FileHelper.h"
+#include "Misc/Paths.h"
 #include "Serialization/JsonReader.h"
 #include "Serialization/JsonSerializer.h"
 #include "Serialization/JsonWriter.h"
@@ -25,21 +26,64 @@ namespace UE::Learning::Agents::Recording::Private
 
 void ULearningAgentsRecording::ResetRecording()
 {
+	Schemas.Empty();
 	Records.Empty();
 	ForceMarkDirty();
 }
 
 void ULearningAgentsRecording::AppendRecording()
 {
-	AppendRecordingFromFile(RecordingFile, SchemaFile);
+	AppendRecordingFromFile(RecordingFile, NewSchemaName, NewTag);
 }
 
-void ULearningAgentsRecording::LoadRecordingFromFile(const FFilePath& File, const FFilePath& Schema)
+void ULearningAgentsRecording::AppendAllRecordingsFromFolder()
 {
-	// Load JSON
-	FString ObservationSchemaString;
-	FString ActionSchemaString;
+	if (RecordingFolder.Path.IsEmpty())
+	{
+		UE_LOG(LogLearning, Error, TEXT("%s: RecordingFolder path is empty."), *GetName());
+		return;
+	}
+
+	IFileManager& FileManager = IFileManager::Get();
+	FString BasePathAbs = FPaths::ConvertRelativePathToFull(RecordingFolder.Path);
+	FPaths::NormalizeDirectoryName(BasePathAbs);
+
+	TArray<FString> FileNames;
+	const FString SearchPath = FPaths::Combine(BasePathAbs, TEXT("*.record"));
+	FileManager.FindFiles(FileNames, *SearchPath, /*Files=*/true, /*Directories=*/false);
+
+	if (FileNames.Num() == 0)
+	{
+		UE_LOG(LogLearning, Warning, TEXT("%s: No .record files found in folder: %s"), *GetName(), *BasePathAbs);
+		return;
+	}
+
+	for (const FString& FileName : FileNames)
+	{
+		const FString FullPath = FPaths::Combine(BasePathAbs, FileName);
+
+		if (!FileManager.FileExists(*FullPath))
+		{
+			UE_LOG(LogLearning, Warning, TEXT("%s: Skipping missing file discovered earlier: %s"), *GetName(), *FullPath);
+			continue;
+		}
+
+		FFilePath FilePath;
+		FilePath.FilePath = FullPath;
 
+		AppendRecordingFromFile(FilePath, NewSchemaName, NewTag);
+	}
+
+	ForceMarkDirty();
+}
+
+void ULearningAgentsRecording::AppendSchema()
+{
+	AppendSchemaFromFile(NewSchemaName, SchemaFile);
+}
+
+void ULearningAgentsRecording::AppendSchemaFromFile(const FName& SchemaName, const FFilePath& Schema)
+{
 	FString FileString;
 	if (FFileHelper::LoadFileToString(FileString, *Schema.FilePath))
 	{
@@ -47,17 +91,23 @@ void ULearningAgentsRecording::LoadRecordingFromFile(const FFilePath& File, cons
 		TSharedPtr<FJsonObject> JsonObject;
 		if (FJsonSerializer::Deserialize(Reader, JsonObject) && JsonObject.IsValid())
 		{
+			FLearningAgentsSchema NewSchema;
+
 			const TSharedPtr<FJsonObject>& ObservationJson = JsonObject->GetObjectField(TEXT("ObservationSchema"));
-			TSharedRef<TJsonWriter<>> ObservationWriter = TJsonWriterFactory<>::Create(&ObservationSchemaString);
+			TSharedRef<TJsonWriter<>> ObservationWriter = TJsonWriterFactory<>::Create(&NewSchema.ObservationSchemaJson);
 			FJsonSerializer::Serialize(ObservationJson.ToSharedRef(), ObservationWriter);
 
 			const TSharedPtr<FJsonObject>& ActionJson = JsonObject->GetObjectField(TEXT("ActionSchema"));
-			TSharedRef<TJsonWriter<>> ActionWriter = TJsonWriterFactory<>::Create(&ActionSchemaString);
+			TSharedRef<TJsonWriter<>> ActionWriter = TJsonWriterFactory<>::Create(&NewSchema.ActionSchemaJson);
 			FJsonSerializer::Serialize(ActionJson.ToSharedRef(), ActionWriter);
+
+			Schemas.Add(SchemaName, NewSchema);
+			ForceMarkDirty();
 		}
 		else
 		{
 			UE_LOG(LogLearning, Error, TEXT("%s: Failed to load schema. Failed to parse JSON from: \"%s\""), *GetName(), *Schema.FilePath);
+			return;
 		}
 	}
 	else
@@ -65,9 +115,10 @@ void ULearningAgentsRecording::LoadRecordingFromFile(const FFilePath& File, cons
 		UE_LOG(LogLearning, Error, TEXT("%s: Failed to load schema. File not found: \"%s\""), *GetName(), *Schema.FilePath);
 		return;
 	}
+}
 
-	// Write Record
-
+void ULearningAgentsRecording::LoadRecordingFromFile(const FFilePath& File, const FName& SchemaName, const FGameplayTag& Tag)
+{
 	TArray<uint8> RecordingData;
 
 	if (FFileHelper::LoadFileToArray(RecordingData, *File.FilePath))
@@ -100,6 +151,7 @@ void ULearningAgentsRecording::LoadRecordingFromFile(const FFilePath& File, cons
 
 		int32 RecordNum;
 		UE::Learning::DeserializeFromBytes(Offset, RecordingData, RecordNum);
+		check(RecordNum > 0);
 
 		Records.SetNum(RecordNum);
 
@@ -113,8 +165,9 @@ void ULearningAgentsRecording::LoadRecordingFromFile(const FFilePath& File, cons
 			UE::Learning::Array::DeserializeFromBytes32<2, float>(Offset, RecordingData, Records[RecordIdx].ObservationData);
 			UE::Learning::Array::DeserializeFromBytes32<2, float>(Offset, RecordingData, Records[RecordIdx].ActionData);
 			
-			Records[RecordIdx].ObservationSchemaJson = ObservationSchemaString;
-			Records[RecordIdx].ActionSchemaJson = ActionSchemaString;
+			Records[RecordIdx].Name = FPaths::GetBaseFilename(File.FilePath);
+			Records[RecordIdx].SchemaName = SchemaName;
+			Records[RecordIdx].Tags.AddTag(Tag);
 		}
 
 		check(Offset == RecordingData.Num());
@@ -174,10 +227,10 @@ void ULearningAgentsRecording::SaveRecordingToFile(const FFilePath& File) const
 	}
 }
 
-void ULearningAgentsRecording::AppendRecordingFromFile(const FFilePath& File, const FFilePath& Schema)
+void ULearningAgentsRecording::AppendRecordingFromFile(const FFilePath& File, const FName& SchemaName, const FGameplayTag& Tag)
 {
 	ULearningAgentsRecording* TempRecording = NewObject<ULearningAgentsRecording>(this);
-	TempRecording->LoadRecordingFromFile(File, Schema);
+	TempRecording->LoadRecordingFromFile(File, SchemaName, Tag);
 
 	Records.Append(TempRecording->Records);
 	ForceMarkDirty();
diff --git a/Engine/Plugins/Experimental/LearningAgents/Source/LearningAgentsTraining/Public/LearningAgentsImitationTrainer.h b/Engine/Plugins/Experimental/LearningAgents/Source/LearningAgentsTraining/Public/LearningAgentsImitationTrainer.h
--- a/Engine/Plugins/Experimental/LearningAgents/Source/LearningAgentsTraining/Public/LearningAgentsImitationTrainer.h
+++ b/Engine/Plugins/Experimental/LearningAgents/Source/LearningAgentsTraining/Public/LearningAgentsImitationTrainer.h
@@ -5,6 +5,7 @@
 #include "LearningAgentsManagerListener.h"
 #include "LearningAgentsTrainer.h"
 
+#include "GameplayTagContainer.h"
 #include "Templates/SharedPointer.h"
 #include "UObject/ObjectPtr.h"
 
@@ -34,6 +35,10 @@ public:
 	UPROPERTY(EditAnywhere, Category = "LearningAgents", meta = (ClampMin = "0.0", UIMin = "0.0"))
 	float TrainerCommunicationTimeout = 10.0f;
 
+	/** Tag used to filter records from the recording. */
+	UPROPERTY(EditDefaultsOnly, BlueprintReadWrite, Category = "LearningAgents")
+	FGameplayTag RecordingFilterTag;
+
 };
 
 /** The configurable settings for the training process. */
diff --git a/Engine/Plugins/Experimental/LearningAgents/Source/LearningAgentsTraining/Public/LearningAgentsRecording.h b/Engine/Plugins/Experimental/LearningAgents/Source/LearningAgentsTraining/Public/LearningAgentsRecording.h
--- a/Engine/Plugins/Experimental/LearningAgents/Source/LearningAgentsTraining/Public/LearningAgentsRecording.h
+++ b/Engine/Plugins/Experimental/LearningAgents/Source/LearningAgentsTraining/Public/LearningAgentsRecording.h
@@ -5,13 +5,31 @@
 #include "LearningArray.h"
 
 #include "Containers/Array.h"
-#include "UObject/NameTypes.h"
 #include "Engine/DataAsset.h"
+#include "GameplayTagContainer.h"
+#include "UObject/NameTypes.h"
 
 #include "LearningAgentsRecording.generated.h"
 
 #define UE_API LEARNINGAGENTSTRAINING_API
 
+/** A schema entry. */
+USTRUCT(BlueprintType)
+struct FLearningAgentsSchema
+{
+	GENERATED_BODY()
+
+public:
+
+	/** Observation schema in JSON format. */
+	UPROPERTY(BlueprintReadOnly, VisibleAnywhere, Category = "LearningAgents")
+	FString ObservationSchemaJson;
+
+	/** Action schema in JSON format. */
+	UPROPERTY(BlueprintReadOnly, VisibleAnywhere, Category = "LearningAgents")
+	FString ActionSchemaJson;
+};
+
 /** A single recording of a series of observations and actions. */
 USTRUCT(BlueprintType)
 struct FLearningAgentsRecord
@@ -20,6 +38,14 @@ struct FLearningAgentsRecord
 
 public:
 
+	/** Record's name. */
+	UPROPERTY(BlueprintReadWrite, EditAnywhere, Category = "LearningAgents")
+	FString Name;
+
+	/** Schema's name. Must correspond with a schema in the recording. */
+	UPROPERTY(BlueprintReadWrite, EditAnywhere, Category = "LearningAgents")
+	FName SchemaName;
+
 	/** The number of observations and actions recorded. */
 	UPROPERTY(BlueprintReadOnly, VisibleAnywhere, Category = "LearningAgents")
 	int32 StepNum = 0;
@@ -40,11 +66,9 @@ public:
 	UPROPERTY(BlueprintReadOnly, VisibleAnywhere, Category = "LearningAgents")
 	int32 ActionCompatibilityHash = 0;
 
-	UPROPERTY(BlueprintReadOnly, VisibleAnywhere, Category = "LearningAgents")
-	FString ObservationSchemaJson;
-	
-	UPROPERTY(BlueprintReadOnly, VisibleAnywhere, Category = "LearningAgents")
-	FString ActionSchemaJson;
+	/** The tags to be referenced in the imitation trainer. */
+	UPROPERTY(BlueprintReadWrite, EditAnywhere, Category = "LearningAgents")
+	FGameplayTagContainer Tags;
 
 	/** Observation data */
 	UPROPERTY(BlueprintReadOnly, Category = "LearningAgents")
@@ -53,7 +77,6 @@ public:
 	/** Action data */
 	UPROPERTY(BlueprintReadOnly, Category = "LearningAgents")
 	TArray<float> ActionData;
-
 };
 
 /** Data asset representing an array of records. */
@@ -78,9 +101,21 @@ public:
 	UFUNCTION(CallInEditor, Category = "LearningAgents")
 	UE_API void AppendRecording();
 
+	/** Appends all .record files from the specified folder to this recording asset */
+	UFUNCTION(CallInEditor, Category = "LearningAgents", meta = (DisplayName = "Append All Recordings from Folder"))
+	UE_API void AppendAllRecordingsFromFolder();
+
+	/** Appends the schema from "SchemaFile" to this recording.  */
+	UFUNCTION(CallInEditor, Category = "LearningAgents")
+	UE_API void AppendSchema();
+
+	/** Append schema to this recording from a file. */
+	UFUNCTION(BlueprintCallable, Category = "LearningAgents", meta = (RelativePath))
+	UE_API void AppendSchemaFromFile(const FName& SchemaName, const FFilePath& Schema);
+
 	/** Load this recording from a file. */
 	UFUNCTION(BlueprintCallable, Category = "LearningAgents", meta = (RelativePath))
-	UE_API void LoadRecordingFromFile(const FFilePath& File, const FFilePath& Schema);
+	UE_API void LoadRecordingFromFile(const FFilePath& File, const FName& SchemaName, const FGameplayTag& Tag);
 
 	/** Save this recording to a file. */
 	UFUNCTION(BlueprintCallable, BlueprintPure = false, Category = "LearningAgents", meta = (RelativePath))
@@ -88,7 +123,7 @@ public:
 
 	/** Append to this recording from a file. */
 	UFUNCTION(BlueprintCallable, Category = "LearningAgents", meta = (RelativePath))
-	UE_API void AppendRecordingFromFile(const FFilePath& File, const FFilePath& Schema);
+	UE_API void AppendRecordingFromFile(const FFilePath& File, const FName& SchemaName, const FGameplayTag& Tag);
 
 	/** Loads this recording from the given recording asset */
 	UFUNCTION(BlueprintCallable, Category = "LearningAgents")
@@ -143,10 +178,26 @@ public:
 	UPROPERTY(EditInstanceOnly, Category = "LearningAgents")
 	FFilePath RecordingFile;
 
-	/** A schema file. Used in combination with the "AppendRecording" button in the editor. */
+	/** A folder containing .record files to load. Used in combination with the "Append All Recordings From Folder" button in the editor.*/
+	UPROPERTY(EditInstanceOnly, Category = "LearningAgents")
+	FDirectoryPath RecordingFolder;
+
+	/** The tag to apply to new records. Used in combination with the "AppendRecording" and "Append All Recordings from Folder" buttons in the editor. */
+	UPROPERTY(EditInstanceOnly, Category = "LearningAgents")
+	FGameplayTag NewTag;
+
+	/** The schema name. Used in combination with the "AppendSchema" and "AppendRecording" buttons in the editor. */
+	UPROPERTY(EditInstanceOnly, Category = "LearningAgents")
+	FName NewSchemaName;
+
+	/** A schema file. Used in combination with the "AppendSchema" button in the editor. */
 	UPROPERTY(EditInstanceOnly, Category = "LearningAgents")
 	FFilePath SchemaFile;
 
+	/** Map of schemas. */
+	UPROPERTY(EditInstanceOnly, Category = "LearningAgents")
+	TMap<FName, FLearningAgentsSchema> Schemas;
+
 	/** Set of records. */
 	UPROPERTY(EditInstanceOnly, Category = "LearningAgents")
 	TArray<FLearningAgentsRecord> Records;
