LACombat: New grid sensor, LOS observations, v17 schema

#rnx
#DisabledValidations Assets
#virtualized

diff --git a/Engine/Plugins/Experimental/LearningAgents/Source/LearningAgents/Private/LearningAgentsObservations.cpp b/Engine/Plugins/Experimental/LearningAgents/Source/LearningAgents/Private/LearningAgentsObservations.cpp
--- a/Engine/Plugins/Experimental/LearningAgents/Source/LearningAgents/Private/LearningAgentsObservations.cpp
+++ b/Engine/Plugins/Experimental/LearningAgents/Source/LearningAgents/Private/LearningAgentsObservations.cpp
@@ -2953,7 +2953,7 @@ FLearningAgentsObservationObjectElement ULearningAgentsObservations::MakeLidarOb
 	FHitResult TraceHit;
 	const bool bHit = Object->GetWorld()->LineTraceSingleByObjectType(TraceHit, RayStartWorld, RayEndWorld, ObjectQueryParams, TraceParams);
 
-	const float Encoded = bHit ? TraceHit.Distance / RayRange : 1.0f; 
+	const float Encoded = bHit ? FMath::Clamp(1.0f - TraceHit.Distance / RayRange, 0.0f, 1.0f) : 0.0f;
 
 #if UE_LEARNING_AGENTS_ENABLE_VISUAL_LOG
 	if (bVisualLoggerEnabled && VisualLoggerListener)
