[Learning Agents] Added a config for evaluation in the imitation learning trainer

#rnx

diff --git a/Engine/Plugins/Experimental/LearningAgents/Source/LearningAgentsTraining/Private/LearningAgentsImitationTrainer.cpp b/Engine/Plugins/Experimental/LearningAgents/Source/LearningAgentsTraining/Private/LearningAgentsImitationTrainer.cpp
--- a/Engine/Plugins/Experimental/LearningAgents/Source/LearningAgentsTraining/Private/LearningAgentsImitationTrainer.cpp
+++ b/Engine/Plugins/Experimental/LearningAgents/Source/LearningAgentsTraining/Private/LearningAgentsImitationTrainer.cpp
@@ -37,6 +37,11 @@ TSharedRef<FJsonObject> FLearningAgentsImitationTrainerTrainingSettings::AsJsonC
 	ConfigObject->SetNumberField(TEXT("Seed"), RandomSeed);
 	ConfigObject->SetBoolField(TEXT("SaveSnapshots"), bSaveSnapshots);
 	ConfigObject->SetNumberField(TEXT("IterationsPerSnapshot"), IterationsPerSnapshot);
+	// Eval
+	ConfigObject->SetBoolField(TEXT("RunEvaluation"), bRunEvaluation);
+	ConfigObject->SetNumberField(TEXT("TrainEvalDatasetSplit"), TrainEvalDatasetSplit);
+	ConfigObject->SetNumberField(TEXT("EvaluationFrequency"), EvaluationFrequency);
+	ConfigObject->SetNumberField(TEXT("BatchCountPerEvaluation"), BatchCountPerEvaluation);
 
 	return ConfigObject;
 }
diff --git a/Engine/Plugins/Experimental/LearningAgents/Source/LearningAgentsTraining/Public/LearningAgentsImitationTrainer.h b/Engine/Plugins/Experimental/LearningAgents/Source/LearningAgentsTraining/Public/LearningAgentsImitationTrainer.h
--- a/Engine/Plugins/Experimental/LearningAgents/Source/LearningAgentsTraining/Public/LearningAgentsImitationTrainer.h
+++ b/Engine/Plugins/Experimental/LearningAgents/Source/LearningAgentsTraining/Public/LearningAgentsImitationTrainer.h
@@ -139,6 +139,23 @@ public:
 	UPROPERTY(EditAnywhere, Category = "LearningAgents")
 	FString MLflowTrackingUri = "";
 
+	/** Should evaluation run during the training process. */
+	UPROPERTY(EditAnywhere, Category = "LearningAgents")
+	bool bRunEvaluation = true;
+
+	/** How much data should be used for evaluation */
+	UPROPERTY(EditAnywhere, Category = "LearningAgents", meta = (ClampMin = "0.0", UIMin = "0.0", ClampMax = "1.0", UIMax = "1.0", EditCondition = "bRunEvaluation"))
+	float TrainEvalDatasetSplit = 0.1f;
+
+	/** How many training iteration loops between an evaluation run */
+	UPROPERTY(EditAnywhere, Category = "LearningAgents", meta = (ClampMin = "0", UIMin = "0", EditCondition = "bRunEvaluation"))
+	int32 EvaluationFrequency = 50;
+
+	/** How many batches to perform evaluation on? Randomly chosen each evaluation */
+	UPROPERTY(EditAnywhere, Category = "LearningAgents", meta = (ClampMin = "0", UIMin = "0", EditCondition = "bRunEvaluation"))
+	int32 BatchCountPerEvaluation = 10;
+
+
 	UE_API TSharedRef<FJsonObject> AsJsonConfig() const;
 };
 
