[LearningAgents] Disabled evaluation options in the imitation learning trainer and added a tooltip that this option is currently not supported in Python

#rnx

diff --git a/Engine/Plugins/Experimental/LearningAgents/Source/LearningAgentsTraining/Public/LearningAgentsImitationTrainer.h b/Engine/Plugins/Experimental/LearningAgents/Source/LearningAgentsTraining/Public/LearningAgentsImitationTrainer.h
--- a/Engine/Plugins/Experimental/LearningAgents/Source/LearningAgentsTraining/Public/LearningAgentsImitationTrainer.h
+++ b/Engine/Plugins/Experimental/LearningAgents/Source/LearningAgentsTraining/Public/LearningAgentsImitationTrainer.h
@@ -139,20 +139,20 @@ public:
 	UPROPERTY(EditAnywhere, Category = "LearningAgents")
 	FString MLflowTrackingUri = "";
 
-	/** Should evaluation run during the training process. */
-	UPROPERTY(EditAnywhere, Category = "LearningAgents")
+	/** Should evaluation run during the training process. Currently not used in Python */
+	UPROPERTY(VisibleDefaultsOnly, Category = "LearningAgents")
 	bool bRunEvaluation = true;
 
-	/** How much data should be used for evaluation */
-	UPROPERTY(EditAnywhere, Category = "LearningAgents", meta = (ClampMin = "0.0", UIMin = "0.0", ClampMax = "1.0", UIMax = "1.0", EditCondition = "bRunEvaluation"))
+	/** How much data should be used for evaluation. Currently not used in Python */
+	UPROPERTY(VisibleDefaultsOnly, Category = "LearningAgents", meta = (ClampMin = "0.0", UIMin = "0.0", ClampMax = "1.0", UIMax = "1.0", EditCondition = "bRunEvaluation"))
 	float TrainEvalDatasetSplit = 0.1f;
 
-	/** How many training iteration loops between an evaluation run */
-	UPROPERTY(EditAnywhere, Category = "LearningAgents", meta = (ClampMin = "0", UIMin = "0", EditCondition = "bRunEvaluation"))
+	/** How many training iteration loops between an evaluation run. Currently not used in Python */
+	UPROPERTY(VisibleDefaultsOnly, Category = "LearningAgents", meta = (ClampMin = "0", UIMin = "0", EditCondition = "bRunEvaluation"))
 	int32 EvaluationFrequency = 50;
 
-	/** How many batches to perform evaluation on? Randomly chosen each evaluation */
-	UPROPERTY(EditAnywhere, Category = "LearningAgents", meta = (ClampMin = "0", UIMin = "0", EditCondition = "bRunEvaluation"))
+	/** How many batches to perform evaluation on? Randomly chosen each evaluation. Currently not used in Python */
+	UPROPERTY(VisibleDefaultsOnly, Category = "LearningAgents", meta = (ClampMin = "0", UIMin = "0", EditCondition = "bRunEvaluation"))
 	int32 BatchCountPerEvaluation = 10;
 
 
