[LearningAgents] - schema subset bug fix

#changelist validated
#virtualized

diff --git a/Engine/Plugins/Experimental/LearningAgents/Source/LearningTraining/Private/LearningTrainer.cpp b/Engine/Plugins/Experimental/LearningAgents/Source/LearningTraining/Private/LearningTrainer.cpp
--- a/Engine/Plugins/Experimental/LearningAgents/Source/LearningTraining/Private/LearningTrainer.cpp
+++ b/Engine/Plugins/Experimental/LearningAgents/Source/LearningTraining/Private/LearningTrainer.cpp
@@ -979,6 +979,8 @@ namespace UE::Learning::Trainer
 			const Observation::FSchemaEncodingParameters Parameters = ObservationSchema.GetEncoding(ObservationSchemaElement);
 			
 			ComputeObservationSchemaSubsetIndices(NextIndex, Indices, Object->GetObjectField(TEXT("Element")).ToSharedRef(), ObservationSchema, Parameters.Element);
+
+			return;
 		}
 
 		// We didn't return earlier so we received an invalid Type
