LearningAgents: Ensure instead of check during GetAgent

diff --git a/Engine/Plugins/Experimental/LearningAgents/Source/LearningAgents/Private/LearningAgentsManager.cpp b/Engine/Plugins/Experimental/LearningAgents/Source/LearningAgents/Private/LearningAgentsManager.cpp
--- a/Engine/Plugins/Experimental/LearningAgents/Source/LearningAgents/Private/LearningAgentsManager.cpp
+++ b/Engine/Plugins/Experimental/LearningAgents/Source/LearningAgents/Private/LearningAgentsManager.cpp
@@ -449,14 +449,22 @@ bool ULearningAgentsManager::HasAgent(const int32 AgentId) const
 
 const UObject* ULearningAgentsManager::GetAgent(const int32 AgentId) const
 {
-	check(HasAgent(AgentId));
-	return Agents[AgentId];
+	if (ensure(HasAgent(AgentId)))
+	{
+		return Agents[AgentId];
+	}
+
+	return nullptr;
 }
 
 UObject* ULearningAgentsManager::GetAgent(const int32 AgentId)
 {
-	check(HasAgent(AgentId));
-	return Agents[AgentId];
+	if (ensure(HasAgent(AgentId)))
+	{
+		return Agents[AgentId];
+	}
+
+	return nullptr;
 }
 
 const TArray<int32>& ULearningAgentsManager::GetAllAgentIds() const
