[LearningAgents] Missing types in ComputeObservationSchemaSubsetIndices and bugfix

diff --git a/Engine/Plugins/Experimental/LearningAgents/Source/LearningTraining/Private/LearningTrainer.cpp b/Engine/Plugins/Experimental/LearningAgents/Source/LearningTraining/Private/LearningTrainer.cpp
--- a/Engine/Plugins/Experimental/LearningAgents/Source/LearningTraining/Private/LearningTrainer.cpp
+++ b/Engine/Plugins/Experimental/LearningAgents/Source/LearningTraining/Private/LearningTrainer.cpp
@@ -982,7 +982,22 @@ namespace UE::Learning::Trainer
 
 			return;
 		}
+		else if (Object->GetStringField(TEXT("Type")) == TEXT("Conv1d"))
+		{
+			const Observation::FSchemaEncodingParameters Parameters = ObservationSchema.GetEncoding(ObservationSchemaElement);
+
+			ComputeObservationSchemaSubsetIndices(NextIndex, Indices, Object->GetObjectField(TEXT("Element")).ToSharedRef(), ObservationSchema, Parameters.Element);
 
+			return;
+		}
+		else if (Object->GetStringField(TEXT("Type")) == TEXT("Conv2d"))
+		{
+			const Observation::FSchemaEncodingParameters Parameters = ObservationSchema.GetEncoding(ObservationSchemaElement);
+
+			ComputeObservationSchemaSubsetIndices(NextIndex, Indices, Object->GetObjectField(TEXT("Element")).ToSharedRef(), ObservationSchema, Parameters.Element);
+
+			return;
+		}
 		// We didn't return earlier so we received an invalid Type
 		checkNoEntry();
 	}
@@ -1453,6 +1468,8 @@ namespace UE::Learning::Trainer
 			const Action::FSchemaEncodingParameters Parameters = ActionSchema.GetEncoding(ActionSchemaElement);
 
 			ComputeActionSchemaSubsetIndices(NextIndex, Indices, Object->GetObjectField(TEXT("Element")).ToSharedRef(), ActionSchema, Parameters.Element);
+			
+			return;
 		}
 
 		// We didn't return earlier so we received an invalid Type
