[LearningAgents] bug fix to conv1d conv2d serialization

diff --git a/Engine/Plugins/Experimental/LearningAgents/Source/LearningTraining/Private/LearningTrainer.cpp b/Engine/Plugins/Experimental/LearningAgents/Source/LearningTraining/Private/LearningTrainer.cpp
--- a/Engine/Plugins/Experimental/LearningAgents/Source/LearningTraining/Private/LearningTrainer.cpp
+++ b/Engine/Plugins/Experimental/LearningAgents/Source/LearningTraining/Private/LearningTrainer.cpp
@@ -276,7 +276,7 @@ namespace UE::Learning::Trainer
 			Object->SetNumberField(TEXT("OutChannels"), Parameters.OutChannels);
 			Object->SetNumberField(TEXT("KernelSize"), Parameters.KernelSize);
 			Object->SetNumberField(TEXT("Padding"), Parameters.Padding);
-			Object->SetNumberField(TEXT("PaddingMode"), (uint32)Parameters.Padding);
+			Object->SetNumberField(TEXT("PaddingMode"), (uint32)Parameters.PaddingMode);
 			Object->SetObjectField(TEXT("Element"), ConvertObservationSchemaToJSON(ObservationSchema, Parameters.Element));
 			break;
 		}
@@ -292,7 +292,7 @@ namespace UE::Learning::Trainer
 			Object->SetNumberField(TEXT("KernelSize"), Parameters.KernelSize);
 			Object->SetNumberField(TEXT("Stride"), Parameters.Stride);
 			Object->SetNumberField(TEXT("Padding"), Parameters.Padding);
-			Object->SetNumberField(TEXT("PaddingMode"), (uint32)Parameters.Padding);
+			Object->SetNumberField(TEXT("PaddingMode"), (uint32)Parameters.PaddingMode);
 			Object->SetObjectField(TEXT("Element"), ConvertObservationSchemaToJSON(ObservationSchema, Parameters.Element));
 			break;
 		}
