[LearningAgents] LearningAgentsRecording

- made records loading faster and added recursive folder load

diff --git a/Engine/Plugins/Experimental/LearningAgents/Source/LearningAgentsTraining/Private/LearningAgentsRecording.cpp b/Engine/Plugins/Experimental/LearningAgents/Source/LearningAgentsTraining/Private/LearningAgentsRecording.cpp
--- a/Engine/Plugins/Experimental/LearningAgents/Source/LearningAgentsTraining/Private/LearningAgentsRecording.cpp
+++ b/Engine/Plugins/Experimental/LearningAgents/Source/LearningAgentsTraining/Private/LearningAgentsRecording.cpp
@@ -49,8 +49,7 @@ void ULearningAgentsRecording::AppendAllRecordingsFromFolder()
 	FPaths::NormalizeDirectoryName(BasePathAbs);
 
 	TArray<FString> FileNames;
-	const FString SearchPath = FPaths::Combine(BasePathAbs, TEXT("*.record"));
-	FileManager.FindFiles(FileNames, *SearchPath, /*Files=*/true, /*Directories=*/false);
+	FileManager.FindFilesRecursive(FileNames, *BasePathAbs, TEXT("*.record"), true, false);
 
 	if (FileNames.Num() == 0)
 	{
@@ -58,20 +57,33 @@ void ULearningAgentsRecording::AppendAllRecordingsFromFolder()
 		return;
 	}
 
+	TArray<FLearningAgentsRecord> RecordsFromDir;
+
 	for (const FString& FileName : FileNames)
 	{
-		const FString FullPath = FPaths::Combine(BasePathAbs, FileName);
-
-		if (!FileManager.FileExists(*FullPath))
+		if (!FileManager.FileExists(*FileName))
 		{
-			UE_LOG(LogLearning, Warning, TEXT("%s: Skipping missing file discovered earlier: %s"), *GetName(), *FullPath);
+			UE_LOG(LogLearning, Warning, TEXT("%s: Skipping missing file discovered earlier: %s"), *GetName(), *FileName);
 			continue;
 		}
 
 		FFilePath FilePath;
-		FilePath.FilePath = FullPath;
+		FilePath.FilePath = FileName;
+
+		TArray<FLearningAgentsRecord> RecordsFromFile;
+		LoadRecordingFromFile(FilePath, NewSchemaName, NewTag, RecordsFromFile);
 
-		AppendRecordingFromFile(FilePath, NewSchemaName, NewTag);
+		if (!RecordsFromFile.IsEmpty())
+		{
+			RecordsFromDir.Reserve(RecordsFromDir.Num() + RecordsFromFile.Num());
+			RecordsFromDir.Append(MoveTemp(RecordsFromFile));
+		}
+	}
+
+	if (!RecordsFromDir.IsEmpty())
+	{
+		Records.Reserve(Records.Num() + RecordsFromDir.Num());
+		Records.Append(MoveTemp(RecordsFromDir));
 	}
 
 	ForceMarkDirty();
@@ -117,7 +129,7 @@ void ULearningAgentsRecording::AppendSchemaFromFile(const FName& SchemaName, con
 	}
 }
 
-void ULearningAgentsRecording::LoadRecordingFromFile(const FFilePath& File, const FName& SchemaName, const FGameplayTag& Tag)
+void ULearningAgentsRecording::LoadRecordingFromFile(const FFilePath& File, const FName& SchemaName, const FGameplayTag& Tag, TArray<FLearningAgentsRecord>& OutRecords)
 {
 	TArray<uint8> RecordingData;
 
@@ -153,26 +165,28 @@ void ULearningAgentsRecording::LoadRecordingFromFile(const FFilePath& File, cons
 		UE::Learning::DeserializeFromBytes(Offset, RecordingData, RecordNum);
 		check(RecordNum > 0);
 
-		Records.SetNum(RecordNum);
-
+		OutRecords.Reset(RecordNum);
+		
 		for (int32 RecordIdx = 0; RecordIdx < RecordNum; RecordIdx++)
 		{
-			UE::Learning::DeserializeFromBytes(Offset, RecordingData, Records[RecordIdx].StepNum);
-			UE::Learning::DeserializeFromBytes(Offset, RecordingData, Records[RecordIdx].ObservationDimNum);
-			UE::Learning::DeserializeFromBytes(Offset, RecordingData, Records[RecordIdx].ActionDimNum);
-			UE::Learning::DeserializeFromBytes(Offset, RecordingData, Records[RecordIdx].ObservationCompatibilityHash);
-			UE::Learning::DeserializeFromBytes(Offset, RecordingData, Records[RecordIdx].ActionCompatibilityHash);
-			UE::Learning::Array::DeserializeFromBytes32<2, float>(Offset, RecordingData, Records[RecordIdx].ObservationData);
-			UE::Learning::Array::DeserializeFromBytes32<2, float>(Offset, RecordingData, Records[RecordIdx].ActionData);
+			FLearningAgentsRecord Record;
+
+			UE::Learning::DeserializeFromBytes(Offset, RecordingData, Record.StepNum);
+			UE::Learning::DeserializeFromBytes(Offset, RecordingData, Record.ObservationDimNum);
+			UE::Learning::DeserializeFromBytes(Offset, RecordingData, Record.ActionDimNum);
+			UE::Learning::DeserializeFromBytes(Offset, RecordingData, Record.ObservationCompatibilityHash);
+			UE::Learning::DeserializeFromBytes(Offset, RecordingData, Record.ActionCompatibilityHash);
+			UE::Learning::Array::DeserializeFromBytes32<2, float>(Offset, RecordingData, Record.ObservationData);
+			UE::Learning::Array::DeserializeFromBytes32<2, float>(Offset, RecordingData, Record.ActionData);
 			
-			Records[RecordIdx].Name = FPaths::GetBaseFilename(File.FilePath);
-			Records[RecordIdx].SchemaName = SchemaName;
-			Records[RecordIdx].Tags.AddTag(Tag);
+			Record.Name = FPaths::GetBaseFilename(File.FilePath);
+			Record.SchemaName = SchemaName;
+			Record.Tags.AddTag(Tag);
+
+			OutRecords.Add(MoveTemp(Record));
 		}
 
 		check(Offset == RecordingData.Num());
-
-		ForceMarkDirty();
 	}
 	else
 	{
@@ -229,11 +243,14 @@ void ULearningAgentsRecording::SaveRecordingToFile(const FFilePath& File) const
 
 void ULearningAgentsRecording::AppendRecordingFromFile(const FFilePath& File, const FName& SchemaName, const FGameplayTag& Tag)
 {
-	ULearningAgentsRecording* TempRecording = NewObject<ULearningAgentsRecording>(this);
-	TempRecording->LoadRecordingFromFile(File, SchemaName, Tag);
-
-	Records.Append(TempRecording->Records);
-	ForceMarkDirty();
+	TArray<FLearningAgentsRecord> RecordsFromFile;
+	LoadRecordingFromFile(File, SchemaName, Tag, RecordsFromFile);
+	if (!RecordsFromFile.IsEmpty())
+	{
+		Records.Reserve(Records.Num() + RecordsFromFile.Num());
+		Records.Append(MoveTemp(RecordsFromFile));
+		ForceMarkDirty();
+	}
 }
 
 void ULearningAgentsRecording::LoadRecordingFromAsset(ULearningAgentsRecording* RecordingAsset)
diff --git a/Engine/Plugins/Experimental/LearningAgents/Source/LearningAgentsTraining/Public/LearningAgentsRecording.h b/Engine/Plugins/Experimental/LearningAgents/Source/LearningAgentsTraining/Public/LearningAgentsRecording.h
--- a/Engine/Plugins/Experimental/LearningAgents/Source/LearningAgentsTraining/Public/LearningAgentsRecording.h
+++ b/Engine/Plugins/Experimental/LearningAgents/Source/LearningAgentsTraining/Public/LearningAgentsRecording.h
@@ -115,7 +115,7 @@ public:
 
 	/** Load this recording from a file. */
 	UFUNCTION(BlueprintCallable, Category = "LearningAgents", meta = (RelativePath))
-	UE_API void LoadRecordingFromFile(const FFilePath& File, const FName& SchemaName, const FGameplayTag& Tag);
+	UE_API void LoadRecordingFromFile(const FFilePath& File, const FName& SchemaName, const FGameplayTag& Tag, TArray<FLearningAgentsRecord>& OutRecords);
 
 	/** Save this recording to a file. */
 	UFUNCTION(BlueprintCallable, BlueprintPure = false, Category = "LearningAgents", meta = (RelativePath))
