LearningAgents: Add a missing ForceMarkDirty

diff --git a/Engine/Plugins/Experimental/LearningAgents/Source/LearningAgentsTraining/Private/LearningAgentsRecording.cpp b/Engine/Plugins/Experimental/LearningAgents/Source/LearningAgentsTraining/Private/LearningAgentsRecording.cpp
--- a/Engine/Plugins/Experimental/LearningAgents/Source/LearningAgentsTraining/Private/LearningAgentsRecording.cpp
+++ b/Engine/Plugins/Experimental/LearningAgents/Source/LearningAgentsTraining/Private/LearningAgentsRecording.cpp
@@ -102,6 +102,7 @@ void ULearningAgentsRecording::AppendSchemaFromFile(const FName& SchemaName, con
 			FJsonSerializer::Serialize(ActionJson.ToSharedRef(), ActionWriter);
 
 			Schemas.Add(SchemaName, NewSchema);
+			ForceMarkDirty();
 		}
 		else
 		{
@@ -150,6 +151,7 @@ void ULearningAgentsRecording::LoadRecordingFromFile(const FFilePath& File, cons
 
 		int32 RecordNum;
 		UE::Learning::DeserializeFromBytes(Offset, RecordingData, RecordNum);
+		check(RecordNum > 0);
 
 		Records.SetNum(RecordNum);
 
