LearningAgents: Added some more Frame Attribute functions.

diff --git a/Engine/Plugins/Experimental/LearningAgents/Source/Learning/Private/LearningFrameAttribute.cpp b/Engine/Plugins/Experimental/LearningAgents/Source/Learning/Private/LearningFrameAttribute.cpp
--- a/Engine/Plugins/Experimental/LearningAgents/Source/Learning/Private/LearningFrameAttribute.cpp
+++ b/Engine/Plugins/Experimental/LearningAgents/Source/Learning/Private/LearningFrameAttribute.cpp
@@ -2,6 +2,8 @@
 
 #include "LearningFrameAttribute.h"
 
+#include "Math/Transform.h"
+
 #if UE_LEARNING_ISPC
 #include "Learning.ispc.generated.h"
 #endif
@@ -10,6 +12,14 @@
 {
 	namespace FrameAttribute::Private
 	{
+		static inline float ApproxAcos(const float X)
+		{
+			const float Y = FMath::Abs(X);
+			const float P = -0.1565827 * Y + 1.570796;
+			const float Q = P * FMath::Sqrt(FMath::Max(1.0 - Y, 0.0));
+			return X >= 0.0 ? Q : UE_PI - Q;
+		}
+
 		static inline bool FindMinimum(
 			int32& OutChannelIdx,
 			int32& OutFrameIdx,
@@ -183,13 +193,48 @@
 			}
 		}
 
-		static inline void LengthSquared(TLearningArrayView<1, float> InOut, const TLearningArrayView<1, const float> In)
+		static inline void NegInplace(TLearningArrayView<1, float> InOut)
 		{
-			check(InOut.Num() == In.Num());
 			const int32 ValueNum = InOut.Num();
 			for (int32 ValueIdx = 0; ValueIdx < ValueNum; ValueIdx++)
 			{
-				InOut[ValueIdx] += In[ValueIdx] * In[ValueIdx];
+				InOut[ValueIdx] = -InOut[ValueIdx];
 			}
 		}
 
+		static inline void InvInplace(TLearningArrayView<1, float> InOut)
+		{
+			const int32 ValueNum = InOut.Num();
+			for (int32 ValueIdx = 0; ValueIdx < ValueNum; ValueIdx++)
+			{
+				InOut[ValueIdx] = 1.0f / InOut[ValueIdx];
+			}
+		}
+
+		static inline void AbsInplace(TLearningArrayView<1, float> InOut)
+		{
+			const int32 ValueNum = InOut.Num();
+			for (int32 ValueIdx = 0; ValueIdx < ValueNum; ValueIdx++)
+			{
+				InOut[ValueIdx] = FMath::Abs(InOut[ValueIdx]);
+			}
+		}
+
+		static inline void LogInplace(TLearningArrayView<1, float> InOut)
+		{
+			const int32 ValueNum = InOut.Num();
+			for (int32 ValueIdx = 0; ValueIdx < ValueNum; ValueIdx++)
+			{
+				InOut[ValueIdx] = FMath::Loge(InOut[ValueIdx]);
+			}
+		}
+
+		static inline void ExpInplace(TLearningArrayView<1, float> InOut)
+		{
+			const int32 ValueNum = InOut.Num();
+			for (int32 ValueIdx = 0; ValueIdx < ValueNum; ValueIdx++)
+			{
+				InOut[ValueIdx] = FMath::Exp(InOut[ValueIdx]);
+			}
+		}
+
@@ -202,6 +247,25 @@
 			}
 		}
 
+		static inline void ApproxAcosInplace(TLearningArrayView<1, float> InOut)
+		{
+			const int32 ValueNum = InOut.Num();
+			for (int32 ValueIdx = 0; ValueIdx < ValueNum; ValueIdx++)
+			{
+				InOut[ValueIdx] = ApproxAcos(InOut[ValueIdx]);
+			}
+		}
+
+		static inline void LengthSquared(TLearningArrayView<1, float> InOut, const TLearningArrayView<1, const float> In)
+		{
+			check(InOut.Num() == In.Num());
+			const int32 ValueNum = InOut.Num();
+			for (int32 ValueIdx = 0; ValueIdx < ValueNum; ValueIdx++)
+			{
+				InOut[ValueIdx] += In[ValueIdx] * In[ValueIdx];
+			}
+		}
+
 		static inline void Normalize(TLearningArrayView<2, float> Out, const TLearningArrayView<2, const float> In)
 		{
 			check(Out.Num<0>() == In.Num<0>());
@@ -306,6 +370,62 @@
 			}
 		}
 
+		static inline void ConstantDot(TLearningArrayView<1, float> InOut, const float Lhs, const TLearningArrayView<1, const float> Rhs)
+		{
+			check(InOut.Num() == Rhs.Num());
+			const int32 ValueNum = InOut.Num();
+			for (int32 ValueIdx = 0; ValueIdx < ValueNum; ValueIdx++)
+			{
+				InOut[ValueIdx] += Lhs * Rhs[ValueIdx];
+			}
+		}
+
+		static inline void ConstantDistSquared(TLearningArrayView<1, float> InOut, const float Lhs, const TLearningArrayView<1, const float> Rhs)
+		{
+			check(InOut.Num() == Rhs.Num());
+			const int32 ValueNum = InOut.Num();
+			for (int32 ValueIdx = 0; ValueIdx < ValueNum; ValueIdx++)
+			{
+				InOut[ValueIdx] += FMath::Square(Lhs - Rhs[ValueIdx]);
+			}
+		}
+
+		static inline void AddConstantInplace(TLearningArrayView<1, float> InOut, const float Value)
+		{
+			const int32 ValueNum = InOut.Num();
+			for (int32 ValueIdx = 0; ValueIdx < ValueNum; ValueIdx++)
+			{
+				InOut[ValueIdx] += Value;
+			}
+		}
+
+		static inline void SubConstantInplace(TLearningArrayView<1, float> InOut, const float Value)
+		{
+			const int32 ValueNum = InOut.Num();
+			for (int32 ValueIdx = 0; ValueIdx < ValueNum; ValueIdx++)
+			{
+				InOut[ValueIdx] -= Value;
+			}
+		}
+
+		static inline void MulConstantInplace(TLearningArrayView<1, float> InOut, const float Value)
+		{
+			const int32 ValueNum = InOut.Num();
+			for (int32 ValueIdx = 0; ValueIdx < ValueNum; ValueIdx++)
+			{
+				InOut[ValueIdx] *= Value;
+			}
+		}
+
+		static inline void DivConstantInplace(TLearningArrayView<1, float> InOut, const float Value)
+		{
+			const int32 ValueNum = InOut.Num();
+			for (int32 ValueIdx = 0; ValueIdx < ValueNum; ValueIdx++)
+			{
+				InOut[ValueIdx] /= Value;
+			}
+		}
+
 		static inline void LogicalAnd(TLearningArrayView<1, float> Out, const TLearningArrayView<1, const float> Lhs, const TLearningArrayView<1, const float> Rhs)
 		{
 			check(Out.Num() == Lhs.Num());
@@ -952,6 +1072,31 @@
 			}
 		}
 
+		static inline void QuatApproxShortestAngleBetweenConstant(TLearningArrayView<2, float> Out, const TLearningArrayView<2, const float> Lhs, const FQuat4f Rhs)
+		{
+			check(Out.Num<0>() == 1);
+			check(Lhs.Num<0>() == 4);
+			check(Out.Num<1>() == Lhs.Num<1>());
+
+			const int32 ValueNum = Out.Num<1>();
+			for (int32 ValueIdx = 0; ValueIdx < ValueNum; ValueIdx++)
+			{
+				const float DotProduct0 =
+					Lhs[0][ValueIdx] * +Rhs.X +
+					Lhs[1][ValueIdx] * +Rhs.Y +
+					Lhs[2][ValueIdx] * +Rhs.Z +
+					Lhs[3][ValueIdx] * +Rhs.W;
+
+				const float DotProduct1 =
+					Lhs[0][ValueIdx] * -Rhs.X +
+					Lhs[1][ValueIdx] * -Rhs.Y +
+					Lhs[2][ValueIdx] * -Rhs.Z +
+					Lhs[3][ValueIdx] * -Rhs.W;
+
+				Out[0][ValueIdx] = ApproxAcos(FMath::Max(DotProduct0, DotProduct1));
+			}
+		}
+
 		static inline void QuatConstantMul(TLearningArrayView<2, float> Out, const FQuat4f Lhs, const TLearningArrayView<2, const float> Rhs)
 		{
 			check(Out.Num<0>() == 4);
@@ -1153,6 +1298,359 @@
 			OutStd.Z = FMath::Sqrt(OutStd.Z);
 		}
 
+
+		static inline void TransformInv(
+			TLearningArrayView<2, float> Out,
+			const TLearningArrayView<2, const float> In)
+		{
+			check(Out.Num<0>() == 10);
+			check(In.Num<0>() == 10);
+			check(Out.Num<1>() == In.Num<1>());
+
+			const int32 ValueNum = In.Num<1>();
+
+			for (int32 ValueIdx = 0; ValueIdx < ValueNum; ValueIdx++)
+			{
+				const FTransform3f T = FTransform3f(
+					FQuat4f(In[3][ValueIdx], In[4][ValueIdx], In[5][ValueIdx], In[6][ValueIdx]),
+					FVector3f(In[0][ValueIdx], In[1][ValueIdx], In[2][ValueIdx]),
+					FVector3f(In[7][ValueIdx], In[8][ValueIdx], In[9][ValueIdx])).Inverse();
+
+				Out[0][ValueIdx] = T.GetLocation().X;
+				Out[1][ValueIdx] = T.GetLocation().Y;
+				Out[2][ValueIdx] = T.GetLocation().Z;
+				Out[3][ValueIdx] = T.GetRotation().X;
+				Out[4][ValueIdx] = T.GetRotation().Y;
+				Out[5][ValueIdx] = T.GetRotation().Z;
+				Out[6][ValueIdx] = T.GetRotation().W;
+				Out[7][ValueIdx] = T.GetScale3D().X;
+				Out[8][ValueIdx] = T.GetScale3D().Y;
+				Out[9][ValueIdx] = T.GetScale3D().Z;
+			}
+		}
+
+		static inline void TransformMul(
+			TLearningArrayView<1, float> OutPosX,
+			TLearningArrayView<1, float> OutPosY,
+			TLearningArrayView<1, float> OutPosZ,
+			TLearningArrayView<1, float> OutRotX,
+			TLearningArrayView<1, float> OutRotY,
+			TLearningArrayView<1, float> OutRotZ,
+			TLearningArrayView<1, float> OutRotW,
+			TLearningArrayView<1, float> OutSclX,
+			TLearningArrayView<1, float> OutSclY,
+			TLearningArrayView<1, float> OutSclZ,
+			const TLearningArrayView<1, const float> LhsPosX,
+			const TLearningArrayView<1, const float> LhsPosY,
+			const TLearningArrayView<1, const float> LhsPosZ,
+			const TLearningArrayView<1, const float> LhsRotX,
+			const TLearningArrayView<1, const float> LhsRotY,
+			const TLearningArrayView<1, const float> LhsRotZ,
+			const TLearningArrayView<1, const float> LhsRotW,
+			const TLearningArrayView<1, const float> LhsSclX,
+			const TLearningArrayView<1, const float> LhsSclY,
+			const TLearningArrayView<1, const float> LhsSclZ,
+			const TLearningArrayView<1, const float> RhsPosX,
+			const TLearningArrayView<1, const float> RhsPosY,
+			const TLearningArrayView<1, const float> RhsPosZ,
+			const TLearningArrayView<1, const float> RhsRotX,
+			const TLearningArrayView<1, const float> RhsRotY,
+			const TLearningArrayView<1, const float> RhsRotZ,
+			const TLearningArrayView<1, const float> RhsRotW,
+			const TLearningArrayView<1, const float> RhsSclX,
+			const TLearningArrayView<1, const float> RhsSclY,
+			const TLearningArrayView<1, const float> RhsSclZ)
+		{
+			const int32 ValueNum = OutPosX.Num();
+
+			for (int32 ValueIdx = 0; ValueIdx < ValueNum; ValueIdx++)
+			{
+				const FTransform3f L = FTransform3f(
+					FQuat4f(LhsRotX[ValueIdx], LhsRotY[ValueIdx], LhsRotZ[ValueIdx], LhsRotW[ValueIdx]),
+					FVector3f(LhsPosX[ValueIdx], LhsPosY[ValueIdx], LhsPosZ[ValueIdx]),
+					FVector3f(LhsSclX[ValueIdx], LhsSclY[ValueIdx], LhsSclZ[ValueIdx]));
+
+				const FTransform3f R = FTransform3f(
+					FQuat4f(RhsRotX[ValueIdx], RhsRotY[ValueIdx], RhsRotZ[ValueIdx], RhsRotW[ValueIdx]),
+					FVector3f(RhsPosX[ValueIdx], RhsPosY[ValueIdx], RhsPosZ[ValueIdx]),
+					FVector3f(RhsSclX[ValueIdx], RhsSclY[ValueIdx], RhsSclZ[ValueIdx]));
+
+				const FTransform3f O = L * R;
+
+				OutPosX[ValueIdx] = O.GetLocation().X;
+				OutPosY[ValueIdx] = O.GetLocation().Y;
+				OutPosZ[ValueIdx] = O.GetLocation().Z;
+				OutRotX[ValueIdx] = O.GetRotation().X;
+				OutRotY[ValueIdx] = O.GetRotation().Y;
+				OutRotZ[ValueIdx] = O.GetRotation().Z;
+				OutRotW[ValueIdx] = O.GetRotation().W;
+				OutSclX[ValueIdx] = O.GetScale3D().X;
+				OutSclY[ValueIdx] = O.GetScale3D().Y;
+				OutSclZ[ValueIdx] = O.GetScale3D().Z;
+			}
+		}
+		static inline void TransformDiv(
+			TLearningArrayView<1, float> OutPosX,
+			TLearningArrayView<1, float> OutPosY,
+			TLearningArrayView<1, float> OutPosZ,
+			TLearningArrayView<1, float> OutRotX,
+			TLearningArrayView<1, float> OutRotY,
+			TLearningArrayView<1, float> OutRotZ,
+			TLearningArrayView<1, float> OutRotW,
+			TLearningArrayView<1, float> OutSclX,
+			TLearningArrayView<1, float> OutSclY,
+			TLearningArrayView<1, float> OutSclZ,
+			const TLearningArrayView<1, const float> LhsPosX,
+			const TLearningArrayView<1, const float> LhsPosY,
+			const TLearningArrayView<1, const float> LhsPosZ,
+			const TLearningArrayView<1, const float> LhsRotX,
+			const TLearningArrayView<1, const float> LhsRotY,
+			const TLearningArrayView<1, const float> LhsRotZ,
+			const TLearningArrayView<1, const float> LhsRotW,
+			const TLearningArrayView<1, const float> LhsSclX,
+			const TLearningArrayView<1, const float> LhsSclY,
+			const TLearningArrayView<1, const float> LhsSclZ,
+			const TLearningArrayView<1, const float> RhsPosX,
+			const TLearningArrayView<1, const float> RhsPosY,
+			const TLearningArrayView<1, const float> RhsPosZ,
+			const TLearningArrayView<1, const float> RhsRotX,
+			const TLearningArrayView<1, const float> RhsRotY,
+			const TLearningArrayView<1, const float> RhsRotZ,
+			const TLearningArrayView<1, const float> RhsRotW,
+			const TLearningArrayView<1, const float> RhsSclX,
+			const TLearningArrayView<1, const float> RhsSclY,
+			const TLearningArrayView<1, const float> RhsSclZ)
+		{
+			const int32 ValueNum = OutPosX.Num();
+
+			for (int32 ValueIdx = 0; ValueIdx < ValueNum; ValueIdx++)
+			{
+				const FTransform3f L = FTransform3f(
+					FQuat4f(LhsRotX[ValueIdx], LhsRotY[ValueIdx], LhsRotZ[ValueIdx], LhsRotW[ValueIdx]),
+					FVector3f(LhsPosX[ValueIdx], LhsPosY[ValueIdx], LhsPosZ[ValueIdx]),
+					FVector3f(LhsSclX[ValueIdx], LhsSclY[ValueIdx], LhsSclZ[ValueIdx]));
+
+				const FTransform3f R = FTransform3f(
+					FQuat4f(RhsRotX[ValueIdx], RhsRotY[ValueIdx], RhsRotZ[ValueIdx], RhsRotW[ValueIdx]),
+					FVector3f(RhsPosX[ValueIdx], RhsPosY[ValueIdx], RhsPosZ[ValueIdx]),
+					FVector3f(RhsSclX[ValueIdx], RhsSclY[ValueIdx], RhsSclZ[ValueIdx]));
+
+				const FTransform3f O = L * R.Inverse();
+
+				OutPosX[ValueIdx] = O.GetLocation().X;
+				OutPosY[ValueIdx] = O.GetLocation().Y;
+				OutPosZ[ValueIdx] = O.GetLocation().Z;
+				OutRotX[ValueIdx] = O.GetRotation().X;
+				OutRotY[ValueIdx] = O.GetRotation().Y;
+				OutRotZ[ValueIdx] = O.GetRotation().Z;
+				OutRotW[ValueIdx] = O.GetRotation().W;
+				OutSclX[ValueIdx] = O.GetScale3D().X;
+				OutSclY[ValueIdx] = O.GetScale3D().Y;
+				OutSclZ[ValueIdx] = O.GetScale3D().Z;
+			}
+		}
+
+		static inline void TransformMulConstant(
+			TLearningArrayView<2, float> Out,
+			const TLearningArrayView<2, const float> Lhs,
+			const FTransform3f Rhs)
+		{
+			check(Out.Num<0>() == 10);
+			check(Lhs.Num<0>() == 10);
+			check(Out.Num<1>() == Lhs.Num<1>());
+
+			const int32 ValueNum = Out.Num<1>();
+
+			for (int32 ValueIdx = 0; ValueIdx < ValueNum; ValueIdx++)
+			{
+				const FTransform3f L = FTransform3f(
+					FQuat4f(Lhs[3][ValueIdx], Lhs[4][ValueIdx], Lhs[5][ValueIdx], Lhs[6][ValueIdx]),
+					FVector3f(Lhs[0][ValueIdx], Lhs[1][ValueIdx], Lhs[2][ValueIdx]),
+					FVector3f(Lhs[7][ValueIdx], Lhs[8][ValueIdx], Lhs[9][ValueIdx]));
+
+				const FTransform3f O = L * Rhs;
+
+				Out[0][ValueIdx] = O.GetLocation().X;
+				Out[1][ValueIdx] = O.GetLocation().Y;
+				Out[2][ValueIdx] = O.GetLocation().Z;
+				Out[3][ValueIdx] = O.GetRotation().X;
+				Out[4][ValueIdx] = O.GetRotation().Y;
+				Out[5][ValueIdx] = O.GetRotation().Z;
+				Out[6][ValueIdx] = O.GetRotation().W;
+				Out[7][ValueIdx] = O.GetScale3D().X;
+				Out[8][ValueIdx] = O.GetScale3D().Y;
+				Out[9][ValueIdx] = O.GetScale3D().Z;
+			}
+		}
+
+		static inline void TransformDivConstant(
+			TLearningArrayView<2, float> Out,
+			const TLearningArrayView<2, const float> Lhs,
+			const FTransform3f Rhs)
+		{
+			check(Out.Num<0>() == 10);
+			check(Lhs.Num<0>() == 10);
+			check(Out.Num<1>() == Lhs.Num<1>());
+
+			const int32 ValueNum = Out.Num<1>();
+
+			for (int32 ValueIdx = 0; ValueIdx < ValueNum; ValueIdx++)
+			{
+				const FTransform3f L = FTransform3f(
+					FQuat4f(Lhs[3][ValueIdx], Lhs[4][ValueIdx], Lhs[5][ValueIdx], Lhs[6][ValueIdx]),
+					FVector3f(Lhs[0][ValueIdx], Lhs[1][ValueIdx], Lhs[2][ValueIdx]),
+					FVector3f(Lhs[7][ValueIdx], Lhs[8][ValueIdx], Lhs[9][ValueIdx]));
+
+				const FTransform3f O = L * Rhs.Inverse();
+
+				Out[0][ValueIdx] = O.GetLocation().X;
+				Out[1][ValueIdx] = O.GetLocation().Y;
+				Out[2][ValueIdx] = O.GetLocation().Z;
+				Out[3][ValueIdx] = O.GetRotation().X;
+				Out[4][ValueIdx] = O.GetRotation().Y;
+				Out[5][ValueIdx] = O.GetRotation().Z;
+				Out[6][ValueIdx] = O.GetRotation().W;
+				Out[7][ValueIdx] = O.GetScale3D().X;
+				Out[8][ValueIdx] = O.GetScale3D().Y;
+				Out[9][ValueIdx] = O.GetScale3D().Z;
+			}
+		}
+
+		static inline void TransformConstantMul(
+			TLearningArrayView<2, float> Out,
+			const FTransform3f Lhs,
+			const TLearningArrayView<2, const float> Rhs)
+		{
+			check(Out.Num<0>() == 10);
+			check(Rhs.Num<0>() == 10);
+			check(Out.Num<1>() == Rhs.Num<1>());
+
+			const int32 ValueNum = Out.Num<1>();
+
+			for (int32 ValueIdx = 0; ValueIdx < ValueNum; ValueIdx++)
+			{
+				const FTransform3f R = FTransform3f(
+					FQuat4f(Rhs[3][ValueIdx], Rhs[4][ValueIdx], Rhs[5][ValueIdx], Rhs[6][ValueIdx]),
+					FVector3f(Rhs[0][ValueIdx], Rhs[1][ValueIdx], Rhs[2][ValueIdx]),
+					FVector3f(Rhs[7][ValueIdx], Rhs[8][ValueIdx], Rhs[9][ValueIdx]));
+
+				const FTransform3f O = Lhs * R;
+
+				Out[0][ValueIdx] = O.GetLocation().X;
+				Out[1][ValueIdx] = O.GetLocation().Y;
+				Out[2][ValueIdx] = O.GetLocation().Z;
+				Out[3][ValueIdx] = O.GetRotation().X;
+				Out[4][ValueIdx] = O.GetRotation().Y;
+				Out[5][ValueIdx] = O.GetRotation().Z;
+				Out[6][ValueIdx] = O.GetRotation().W;
+				Out[7][ValueIdx] = O.GetScale3D().X;
+				Out[8][ValueIdx] = O.GetScale3D().Y;
+				Out[9][ValueIdx] = O.GetScale3D().Z;
+			}
+		}
+
+		static inline void TransformConstantDiv(
+			TLearningArrayView<2, float> Out,
+			const FTransform3f Lhs,
+			const TLearningArrayView<2, const float> Rhs)
+		{
+			check(Out.Num<0>() == 10);
+			check(Rhs.Num<0>() == 10);
+			check(Out.Num<1>() == Rhs.Num<1>());
+
+			const int32 ValueNum = Out.Num<1>();
+
+			for (int32 ValueIdx = 0; ValueIdx < ValueNum; ValueIdx++)
+			{
+				const FTransform3f R = FTransform3f(
+					FQuat4f(Rhs[3][ValueIdx], Rhs[4][ValueIdx], Rhs[5][ValueIdx], Rhs[6][ValueIdx]),
+					FVector3f(Rhs[0][ValueIdx], Rhs[1][ValueIdx], Rhs[2][ValueIdx]),
+					FVector3f(Rhs[7][ValueIdx], Rhs[8][ValueIdx], Rhs[9][ValueIdx]));
+
+				const FTransform3f O = Lhs * R.Inverse();
+
+				Out[0][ValueIdx] = O.GetLocation().X;
+				Out[1][ValueIdx] = O.GetLocation().Y;
+				Out[2][ValueIdx] = O.GetLocation().Z;
+				Out[3][ValueIdx] = O.GetRotation().X;
+				Out[4][ValueIdx] = O.GetRotation().Y;
+				Out[5][ValueIdx] = O.GetRotation().Z;
+				Out[6][ValueIdx] = O.GetRotation().W;
+				Out[7][ValueIdx] = O.GetScale3D().X;
+				Out[8][ValueIdx] = O.GetScale3D().Y;
+				Out[9][ValueIdx] = O.GetScale3D().Z;
+			}
+		}
+
+		static inline void TransformMake(
+			TLearningArrayView<1, float> OutPosX,
+			TLearningArrayView<1, float> OutPosY,
+			TLearningArrayView<1, float> OutPosZ,
+			TLearningArrayView<1, float> OutRotX,
+			TLearningArrayView<1, float> OutRotY,
+			TLearningArrayView<1, float> OutRotZ,
+			TLearningArrayView<1, float> OutRotW,
+			TLearningArrayView<1, float> OutSclX,
+			TLearningArrayView<1, float> OutSclY,
+			TLearningArrayView<1, float> OutSclZ,
+			const TLearningArrayView<1, const float> InPosX,
+			const TLearningArrayView<1, const float> InPosY,
+			const TLearningArrayView<1, const float> InPosZ,
+			const TLearningArrayView<1, const float> InRotX,
+			const TLearningArrayView<1, const float> InRotY,
+			const TLearningArrayView<1, const float> InRotZ,
+			const TLearningArrayView<1, const float> InRotW)
+		{
+			Array::Copy(OutPosX, InPosX);
+			Array::Copy(OutPosY, InPosY);
+			Array::Copy(OutPosZ, InPosZ);
+
+			Array::Copy(OutRotX, InRotX);
+			Array::Copy(OutRotY, InRotY);
+			Array::Copy(OutRotZ, InRotZ);
+			Array::Copy(OutRotW, InRotW);
+
+			Array::Set(OutSclX, 1.0f);
+			Array::Set(OutSclY, 1.0f);
+			Array::Set(OutSclZ, 1.0f);
+		}
+
+		static inline void TransformMake(
+			TLearningArrayView<1, float> OutPosX,
+			TLearningArrayView<1, float> OutPosY,
+			TLearningArrayView<1, float> OutPosZ,
+			TLearningArrayView<1, float> OutRotX,
+			TLearningArrayView<1, float> OutRotY,
+			TLearningArrayView<1, float> OutRotZ,
+			TLearningArrayView<1, float> OutRotW,
+			TLearningArrayView<1, float> OutSclX,
+			TLearningArrayView<1, float> OutSclY,
+			TLearningArrayView<1, float> OutSclZ,
+			const TLearningArrayView<1, const float> InPosX,
+			const TLearningArrayView<1, const float> InPosY,
+			const TLearningArrayView<1, const float> InPosZ,
+			const TLearningArrayView<1, const float> InRotX,
+			const TLearningArrayView<1, const float> InRotY,
+			const TLearningArrayView<1, const float> InRotZ,
+			const TLearningArrayView<1, const float> InRotW,
+			const TLearningArrayView<1, const float> InSclX,
+			const TLearningArrayView<1, const float> InSclY,
+			const TLearningArrayView<1, const float> InSclZ)
+		{
+			Array::Copy(OutPosX, InPosX);
+			Array::Copy(OutPosY, InPosY);
+			Array::Copy(OutPosZ, InPosZ);
+
+			Array::Copy(OutRotX, InRotX);
+			Array::Copy(OutRotY, InRotY);
+			Array::Copy(OutRotZ, InRotZ);
+			Array::Copy(OutRotW, InRotW);
+
+			Array::Copy(OutSclX, InSclX);
+			Array::Copy(OutSclY, InSclY);
+			Array::Copy(OutSclZ, InSclZ);
+		}
+
 	}
 
 	void FFrameAttribute::Check() const
@@ -1338,6 +1836,13 @@
 			Op(Out, In, Out.FrameRangeSet.GetAllRangeOffsets(), Out.FrameRangeSet.GetAllRangeLengths());
 		}
 
+		void InplaceOp(
+			FFrameAttribute& InOut,
+			const InplaceOpFunction Op)
+		{
+			Op(InOut, InOut.FrameRangeSet.GetAllRangeOffsets(), InOut.FrameRangeSet.GetAllRangeLengths());
+		}
+
 		void BinaryOp(
 			FFrameAttribute& Out,
 			const int32 OutChannelNum,
@@ -1720,6 +2225,144 @@
 				});
 		}
 
+		void Sqrt(FFrameAttribute& Out, const FFrameAttribute& In)
+		{
+			UnaryOp(Out, In.GetChannelNum(), In, [](
+				FFrameAttribute& Out,
+				const FFrameAttribute& In,
+				const TLearningArrayView<1, const int32> RangeOffsets,
+				const TLearningArrayView<1, const int32> RangeLengths) {
+
+					const int32 ChannelNum = Out.GetChannelNum();
+
+					for (int32 ChannelIdx = 0; ChannelIdx < ChannelNum; ChannelIdx++)
+					{
+						Private::Sqrt(
+							Out.GetChannelAttributeData(ChannelIdx),
+							In.GetChannelAttributeData(ChannelIdx));
+					}
+				});
+		}
+
+		void NegInplace(FFrameAttribute& InOut)
+		{
+			InplaceOp(InOut, [](
+				FFrameAttribute& InOut,
+				const TLearningArrayView<1, const int32> RangeOffsets,
+				const TLearningArrayView<1, const int32> RangeLengths) {
+
+					const int32 ChannelNum = InOut.GetChannelNum();
+
+					for (int32 ChannelIdx = 0; ChannelIdx < ChannelNum; ChannelIdx++)
+					{
+						Private::NegInplace(
+							InOut.GetChannelAttributeData(ChannelIdx));
+					}
+				});
+		}
+
+		void InvInplace(FFrameAttribute& InOut)
+		{
+			InplaceOp(InOut, [](
+				FFrameAttribute& InOut,
+				const TLearningArrayView<1, const int32> RangeOffsets,
+				const TLearningArrayView<1, const int32> RangeLengths) {
+
+					const int32 ChannelNum = InOut.GetChannelNum();
+
+					for (int32 ChannelIdx = 0; ChannelIdx < ChannelNum; ChannelIdx++)
+					{
+						Private::InvInplace(
+							InOut.GetChannelAttributeData(ChannelIdx));
+					}
+				});
+		}
+
+		void AbsInplace(FFrameAttribute& InOut)
+		{
+			InplaceOp(InOut, [](
+				FFrameAttribute& InOut,
+				const TLearningArrayView<1, const int32> RangeOffsets,
+				const TLearningArrayView<1, const int32> RangeLengths) {
+
+					const int32 ChannelNum = InOut.GetChannelNum();
+
+					for (int32 ChannelIdx = 0; ChannelIdx < ChannelNum; ChannelIdx++)
+					{
+						Private::AbsInplace(
+							InOut.GetChannelAttributeData(ChannelIdx));
+					}
+				});
+		}
+
+		void LogInplace(FFrameAttribute& InOut)
+		{
+			InplaceOp(InOut, [](
+				FFrameAttribute& InOut,
+				const TLearningArrayView<1, const int32> RangeOffsets,
+				const TLearningArrayView<1, const int32> RangeLengths) {
+
+					const int32 ChannelNum = InOut.GetChannelNum();
+
+					for (int32 ChannelIdx = 0; ChannelIdx < ChannelNum; ChannelIdx++)
+					{
+						Private::LogInplace(
+							InOut.GetChannelAttributeData(ChannelIdx));
+					}
+				});
+		}
+
+		void ExpInplace(FFrameAttribute& InOut)
+		{
+			InplaceOp(InOut, [](
+				FFrameAttribute& InOut,
+				const TLearningArrayView<1, const int32> RangeOffsets,
+				const TLearningArrayView<1, const int32> RangeLengths) {
+
+					const int32 ChannelNum = InOut.GetChannelNum();
+
+					for (int32 ChannelIdx = 0; ChannelIdx < ChannelNum; ChannelIdx++)
+					{
+						Private::ExpInplace(
+							InOut.GetChannelAttributeData(ChannelIdx));
+					}
+				});
+		}
+
+		void SqrtInplace(FFrameAttribute& InOut)
+		{
+			InplaceOp(InOut, [](
+				FFrameAttribute& InOut,
+				const TLearningArrayView<1, const int32> RangeOffsets,
+				const TLearningArrayView<1, const int32> RangeLengths) {
+
+					const int32 ChannelNum = InOut.GetChannelNum();
+
+					for (int32 ChannelIdx = 0; ChannelIdx < ChannelNum; ChannelIdx++)
+					{
+						Private::SqrtInplace(
+							InOut.GetChannelAttributeData(ChannelIdx));
+					}
+				});
+		}
+
+		void ApproxAcosInplace(FFrameAttribute& InOut)
+		{
+			InplaceOp(InOut, [](
+				FFrameAttribute& InOut,
+				const TLearningArrayView<1, const int32> RangeOffsets,
+				const TLearningArrayView<1, const int32> RangeLengths) {
+
+					const int32 ChannelNum = InOut.GetChannelNum();
+
+					for (int32 ChannelIdx = 0; ChannelIdx < ChannelNum; ChannelIdx++)
+					{
+						Private::ApproxAcosInplace(
+							InOut.GetChannelAttributeData(ChannelIdx));
+					}
+				});
+		}
+
 		void Length(FFrameAttribute& Out, const FFrameAttribute& In)
 		{
 			UnaryOp(Out, 1, In, [](
@@ -1772,6 +2415,26 @@
 				});
 		}
 
+		void Slice(FFrameAttribute& Out, const FFrameAttribute& In, const int32 StartChannelIdx, const int32 ChannelNum)
+		{
+			check(StartChannelIdx >= 0);
+			check(StartChannelIdx + ChannelNum <= In.GetChannelNum());
+
+			UnaryOp(Out, ChannelNum, In, [StartChannelIdx, ChannelNum](
+				FFrameAttribute& Out,
+				const FFrameAttribute& In,
+				const TLearningArrayView<1, const int32> RangeOffsets,
+				const TLearningArrayView<1, const int32> RangeLengths) {
+
+					check(Out.GetChannelNum() == ChannelNum);
+					for (int32 ChannelIdx = 0; ChannelIdx < ChannelNum; ChannelIdx++)
+					{
+						Array::Copy(Out.GetChannelAttributeData(ChannelIdx), In.GetChannelAttributeData(ChannelIdx + StartChannelIdx));
+					}
+
+				});
+		}
+
 		void AddConstant(FFrameAttribute& Out, const FFrameAttribute& Lhs, const TLearningArrayView<1, const float> Rhs)
 		{
 			UnaryOp(Out, Lhs.GetChannelNum(), Lhs, [Rhs](
@@ -1932,6 +2595,146 @@
 				});
 		}
 
+		void ConstantDot(FFrameAttribute& Out, const TLearningArrayView<1, const float> Lhs, const FFrameAttribute& Rhs)
+		{
+			UnaryOp(Out, 1, Rhs, [Lhs](
+				FFrameAttribute& Out,
+				const FFrameAttribute& In,
+				const TLearningArrayView<1, const int32> RangeOffsets,
+				const TLearningArrayView<1, const int32> RangeLengths) {
+
+					const int32 ChannelNum = In.GetChannelNum();
+
+					Array::Zero(Out.GetAttributeData());
+
+					for (int32 ChannelIdx = 0; ChannelIdx < ChannelNum; ChannelIdx++)
+					{
+						Private::ConstantDot(
+							Out.GetChannelAttributeData(0),
+							Lhs[ChannelIdx],
+							In.GetChannelAttributeData(ChannelIdx));
+					}
+				});
+		}
+
+		void ConstantSquaredDist(FFrameAttribute& Out, const TLearningArrayView<1, const float> Lhs, const FFrameAttribute& Rhs)
+		{
+			UnaryOp(Out, 1, Rhs, [Lhs](
+				FFrameAttribute& Out,
+				const FFrameAttribute& In,
+				const TLearningArrayView<1, const int32> RangeOffsets,
+				const TLearningArrayView<1, const int32> RangeLengths) {
+
+					const int32 ChannelNum = In.GetChannelNum();
+
+					Array::Zero(Out.GetAttributeData());
+
+					for (int32 ChannelIdx = 0; ChannelIdx < ChannelNum; ChannelIdx++)
+					{
+						Private::ConstantDistSquared(
+							Out.GetChannelAttributeData(0),
+							Lhs[ChannelIdx],
+							In.GetChannelAttributeData(ChannelIdx));
+					}
+				});
+		}
+
+		void ConstantDist(FFrameAttribute& Out, const TLearningArrayView<1, const float> Lhs, const FFrameAttribute& Rhs)
+		{
+			UnaryOp(Out, 1, Rhs, [Lhs](
+				FFrameAttribute& Out,
+				const FFrameAttribute& In,
+				const TLearningArrayView<1, const int32> RangeOffsets,
+				const TLearningArrayView<1, const int32> RangeLengths) {
+
+					const int32 ChannelNum = In.GetChannelNum();
+
+					Array::Zero(Out.GetAttributeData());
+
+					for (int32 ChannelIdx = 0; ChannelIdx < ChannelNum; ChannelIdx++)
+					{
+						Private::ConstantDistSquared(
+							Out.GetChannelAttributeData(0),
+							Lhs[ChannelIdx],
+							In.GetChannelAttributeData(ChannelIdx));
+					}
+
+					Private::SqrtInplace(Out.GetChannelAttributeData(0));
+				});
+		}
+
+		void AddConstantInplace(FFrameAttribute& InOut, const TLearningArrayView<1, const float> Rhs)
+		{
+			InplaceOp(InOut, [Rhs](
+				FFrameAttribute& InOut,
+				const TLearningArrayView<1, const int32> RangeOffsets,
+				const TLearningArrayView<1, const int32> RangeLengths) {
+
+					const int32 ChannelNum = InOut.GetChannelNum();
+
+					for (int32 ChannelIdx = 0; ChannelIdx < ChannelNum; ChannelIdx++)
+					{
+						Private::AddConstantInplace(
+							InOut.GetChannelAttributeData(ChannelIdx),
+							Rhs[ChannelIdx]);
+					}
+				});
+		}
+
+		void SubConstantInplace(FFrameAttribute& InOut, const TLearningArrayView<1, const float> Rhs)
+		{
+			InplaceOp(InOut, [Rhs](
+				FFrameAttribute& InOut,
+				const TLearningArrayView<1, const int32> RangeOffsets,
+				const TLearningArrayView<1, const int32> RangeLengths) {
+
+					const int32 ChannelNum = InOut.GetChannelNum();
+
+					for (int32 ChannelIdx = 0; ChannelIdx < ChannelNum; ChannelIdx++)
+					{
+						Private::SubConstantInplace(
+							InOut.GetChannelAttributeData(ChannelIdx),
+							Rhs[ChannelIdx]);
+					}
+				});
+		}
+
+		void MulConstantInplace(FFrameAttribute& InOut, const TLearningArrayView<1, const float> Rhs)
+		{
+			InplaceOp(InOut, [Rhs](
+				FFrameAttribute& InOut,
+				const TLearningArrayView<1, const int32> RangeOffsets,
+				const TLearningArrayView<1, const int32> RangeLengths) {
+
+					const int32 ChannelNum = InOut.GetChannelNum();
+
+					for (int32 ChannelIdx = 0; ChannelIdx < ChannelNum; ChannelIdx++)
+					{
+						Private::MulConstantInplace(
+							InOut.GetChannelAttributeData(ChannelIdx),
+							Rhs[ChannelIdx]);
+					}
+				});
+		}
+
+		void DivConstantInplace(FFrameAttribute& InOut, const TLearningArrayView<1, const float> Rhs)
+		{
+			InplaceOp(InOut, [Rhs](
+				FFrameAttribute& InOut,
+				const TLearningArrayView<1, const int32> RangeOffsets,
+				const TLearningArrayView<1, const int32> RangeLengths) {
+
+					const int32 ChannelNum = InOut.GetChannelNum();
+
+					for (int32 ChannelIdx = 0; ChannelIdx < ChannelNum; ChannelIdx++)
+					{
+						Private::DivConstantInplace(
+							InOut.GetChannelAttributeData(ChannelIdx),
+							Rhs[ChannelIdx]);
+					}
+				});
+		}
+
 		void Sum(FFrameAttribute& Out, const TArrayView<const ConstFrameAttributePtr> Inputs)
 		{
 			const int32 InputNum = Inputs.Num();
@@ -2913,6 +3716,23 @@
 				});
 		}
 
+		void QuatApproxShortestAngleBetweenConstant(FFrameAttribute& Out, const FFrameAttribute& Lhs, const FQuat4f Rhs)
+		{
+			check(Lhs.GetChannelNum() == 4);
+
+			UnaryOp(Out, 1, Lhs, [Rhs](
+				FFrameAttribute& Out,
+				const FFrameAttribute& Lhs,
+				const TLearningArrayView<1, const int32> RangeOffsets,
+				const TLearningArrayView<1, const int32> RangeLengths) {
+
+					Private::QuatApproxShortestAngleBetweenConstant(
+						Out.GetAttributeData(),
+						Lhs.GetAttributeData(),
+						Rhs);
+				});
+		}
+
 		void QuatConstantMul(FFrameAttribute& Out, const FQuat4f Lhs, const FFrameAttribute& Rhs)
 		{
 			check(Rhs.GetChannelNum() == 4);
@@ -3033,5 +3853,308 @@
 				});
 		}
 
+		void TransformMake(FFrameAttribute& Out, const FFrameAttribute& Location, const FFrameAttribute& Rotation, const FFrameAttribute& Scale)
+		{
+			check(Location.GetChannelNum() == 3);
+			check(Rotation.GetChannelNum() == 4);
+			check(Scale.GetChannelNum() == 3);
+
+			NaryOp(Out, 10, {
+				(ConstFrameAttributePtr)&Location,
+				(ConstFrameAttributePtr)&Rotation,
+				(ConstFrameAttributePtr)&Scale }, [](
+					FFrameAttribute& Out,
+					const TArrayView<const ConstFrameAttributePtr> In,
+					const TLearningArrayView<1, const int32> OutRangeOffsets,
+					const TArrayView<const TLearningArrayView<1, const int32>> InRangeOffsets,
+					const TLearningArrayView<1, const int32> RangeLengths) {
+
+						const int32 RangeNum = RangeLengths.Num();
+
+						for (int32 RangeIdx = 0; RangeIdx < RangeNum; RangeIdx++)
+						{
+							Private::TransformMake(
+								Out.GetChannelRangeAttributeData(0, OutRangeOffsets[RangeIdx], RangeLengths[RangeIdx]),
+								Out.GetChannelRangeAttributeData(1, OutRangeOffsets[RangeIdx], RangeLengths[RangeIdx]),
+								Out.GetChannelRangeAttributeData(2, OutRangeOffsets[RangeIdx], RangeLengths[RangeIdx]),
+								Out.GetChannelRangeAttributeData(3, OutRangeOffsets[RangeIdx], RangeLengths[RangeIdx]),
+								Out.GetChannelRangeAttributeData(4, OutRangeOffsets[RangeIdx], RangeLengths[RangeIdx]),
+								Out.GetChannelRangeAttributeData(5, OutRangeOffsets[RangeIdx], RangeLengths[RangeIdx]),
+								Out.GetChannelRangeAttributeData(6, OutRangeOffsets[RangeIdx], RangeLengths[RangeIdx]),
+								Out.GetChannelRangeAttributeData(7, OutRangeOffsets[RangeIdx], RangeLengths[RangeIdx]),
+								Out.GetChannelRangeAttributeData(8, OutRangeOffsets[RangeIdx], RangeLengths[RangeIdx]),
+								Out.GetChannelRangeAttributeData(9, OutRangeOffsets[RangeIdx], RangeLengths[RangeIdx]),
+
+								In[0]->GetChannelRangeAttributeData(0, InRangeOffsets[0][RangeIdx], RangeLengths[RangeIdx]),
+								In[0]->GetChannelRangeAttributeData(1, InRangeOffsets[0][RangeIdx], RangeLengths[RangeIdx]),
+								In[0]->GetChannelRangeAttributeData(2, InRangeOffsets[0][RangeIdx], RangeLengths[RangeIdx]),
+
+								In[1]->GetChannelRangeAttributeData(0, InRangeOffsets[1][RangeIdx], RangeLengths[RangeIdx]),
+								In[1]->GetChannelRangeAttributeData(1, InRangeOffsets[1][RangeIdx], RangeLengths[RangeIdx]),
+								In[1]->GetChannelRangeAttributeData(2, InRangeOffsets[1][RangeIdx], RangeLengths[RangeIdx]),
+								In[1]->GetChannelRangeAttributeData(3, InRangeOffsets[1][RangeIdx], RangeLengths[RangeIdx]),
+								
+								In[2]->GetChannelRangeAttributeData(0, InRangeOffsets[2][RangeIdx], RangeLengths[RangeIdx]),
+								In[2]->GetChannelRangeAttributeData(1, InRangeOffsets[2][RangeIdx], RangeLengths[RangeIdx]),
+								In[2]->GetChannelRangeAttributeData(2, InRangeOffsets[2][RangeIdx], RangeLengths[RangeIdx]));
+						}
+				});
+		}
+
+		void TransformMake(FFrameAttribute& Out, const FFrameAttribute& Location, const FFrameAttribute& Rotation)
+		{
+			check(Location.GetChannelNum() == 3);
+			check(Rotation.GetChannelNum() == 4);
+
+			BinaryOp(Out, 10, Location, Rotation, [](
+				FFrameAttribute& Out,
+				const FFrameAttribute& Lhs,
+				const FFrameAttribute& Rhs,
+				const TLearningArrayView<1, const int32> OutRangeOffsets,
+				const TLearningArrayView<1, const int32> LhsRangeOffsets,
+				const TLearningArrayView<1, const int32> RhsRangeOffsets,
+				const TLearningArrayView<1, const int32> RangeLengths) {
+
+					const int32 RangeNum = RangeLengths.Num();
+
+					for (int32 RangeIdx = 0; RangeIdx < RangeNum; RangeIdx++)
+					{
+						Private::TransformMake(
+							Out.GetChannelRangeAttributeData(0, OutRangeOffsets[RangeIdx], RangeLengths[RangeIdx]),
+							Out.GetChannelRangeAttributeData(1, OutRangeOffsets[RangeIdx], RangeLengths[RangeIdx]),
+							Out.GetChannelRangeAttributeData(2, OutRangeOffsets[RangeIdx], RangeLengths[RangeIdx]),
+							Out.GetChannelRangeAttributeData(3, OutRangeOffsets[RangeIdx], RangeLengths[RangeIdx]),
+							Out.GetChannelRangeAttributeData(4, OutRangeOffsets[RangeIdx], RangeLengths[RangeIdx]),
+							Out.GetChannelRangeAttributeData(5, OutRangeOffsets[RangeIdx], RangeLengths[RangeIdx]),
+							Out.GetChannelRangeAttributeData(6, OutRangeOffsets[RangeIdx], RangeLengths[RangeIdx]),
+							Out.GetChannelRangeAttributeData(7, OutRangeOffsets[RangeIdx], RangeLengths[RangeIdx]),
+							Out.GetChannelRangeAttributeData(8, OutRangeOffsets[RangeIdx], RangeLengths[RangeIdx]),
+							Out.GetChannelRangeAttributeData(9, OutRangeOffsets[RangeIdx], RangeLengths[RangeIdx]),
+
+							Lhs.GetChannelRangeAttributeData(0, LhsRangeOffsets[RangeIdx], RangeLengths[RangeIdx]),
+							Lhs.GetChannelRangeAttributeData(1, LhsRangeOffsets[RangeIdx], RangeLengths[RangeIdx]),
+							Lhs.GetChannelRangeAttributeData(2, LhsRangeOffsets[RangeIdx], RangeLengths[RangeIdx]),
+
+							Rhs.GetChannelRangeAttributeData(0, RhsRangeOffsets[RangeIdx], RangeLengths[RangeIdx]),
+							Rhs.GetChannelRangeAttributeData(1, RhsRangeOffsets[RangeIdx], RangeLengths[RangeIdx]),
+							Rhs.GetChannelRangeAttributeData(2, RhsRangeOffsets[RangeIdx], RangeLengths[RangeIdx]),
+							Rhs.GetChannelRangeAttributeData(3, RhsRangeOffsets[RangeIdx], RangeLengths[RangeIdx]));
+					}
+				});
+		}
+
+		void TransformMul(FFrameAttribute& Out, const FFrameAttribute& Lhs, const FFrameAttribute& Rhs)
+		{
+			check(Lhs.GetChannelNum() == 10);
+			check(Rhs.GetChannelNum() == 10);
+
+			BinaryOp(Out, Lhs.GetChannelNum(), Lhs, Rhs, [](
+				FFrameAttribute& Out,
+				const FFrameAttribute& Lhs,
+				const FFrameAttribute& Rhs,
+				const TLearningArrayView<1, const int32> OutRangeOffsets,
+				const TLearningArrayView<1, const int32> LhsRangeOffsets,
+				const TLearningArrayView<1, const int32> RhsRangeOffsets,
+				const TLearningArrayView<1, const int32> RangeLengths)
+				{
+					const int32 RangeNum = RangeLengths.Num();
+
+					for (int32 RangeIdx = 0; RangeIdx < RangeNum; RangeIdx++)
+					{
+						Private::TransformMul(
+							Out.GetChannelRangeAttributeData(0, OutRangeOffsets[RangeIdx], RangeLengths[RangeIdx]),
+							Out.GetChannelRangeAttributeData(1, OutRangeOffsets[RangeIdx], RangeLengths[RangeIdx]),
+							Out.GetChannelRangeAttributeData(2, OutRangeOffsets[RangeIdx], RangeLengths[RangeIdx]),
+							Out.GetChannelRangeAttributeData(3, OutRangeOffsets[RangeIdx], RangeLengths[RangeIdx]),
+							Out.GetChannelRangeAttributeData(4, OutRangeOffsets[RangeIdx], RangeLengths[RangeIdx]),
+							Out.GetChannelRangeAttributeData(5, OutRangeOffsets[RangeIdx], RangeLengths[RangeIdx]),
+							Out.GetChannelRangeAttributeData(6, OutRangeOffsets[RangeIdx], RangeLengths[RangeIdx]),
+							Out.GetChannelRangeAttributeData(7, OutRangeOffsets[RangeIdx], RangeLengths[RangeIdx]),
+							Out.GetChannelRangeAttributeData(8, OutRangeOffsets[RangeIdx], RangeLengths[RangeIdx]),
+							Out.GetChannelRangeAttributeData(9, OutRangeOffsets[RangeIdx], RangeLengths[RangeIdx]),
+
+							Lhs.GetChannelRangeAttributeData(0, LhsRangeOffsets[RangeIdx], RangeLengths[RangeIdx]),
+							Lhs.GetChannelRangeAttributeData(1, LhsRangeOffsets[RangeIdx], RangeLengths[RangeIdx]),
+							Lhs.GetChannelRangeAttributeData(2, LhsRangeOffsets[RangeIdx], RangeLengths[RangeIdx]),
+							Lhs.GetChannelRangeAttributeData(3, LhsRangeOffsets[RangeIdx], RangeLengths[RangeIdx]),
+							Lhs.GetChannelRangeAttributeData(4, LhsRangeOffsets[RangeIdx], RangeLengths[RangeIdx]),
+							Lhs.GetChannelRangeAttributeData(5, LhsRangeOffsets[RangeIdx], RangeLengths[RangeIdx]),
+							Lhs.GetChannelRangeAttributeData(6, LhsRangeOffsets[RangeIdx], RangeLengths[RangeIdx]),
+							Lhs.GetChannelRangeAttributeData(7, LhsRangeOffsets[RangeIdx], RangeLengths[RangeIdx]),
+							Lhs.GetChannelRangeAttributeData(8, LhsRangeOffsets[RangeIdx], RangeLengths[RangeIdx]),
+							Lhs.GetChannelRangeAttributeData(9, LhsRangeOffsets[RangeIdx], RangeLengths[RangeIdx]),
+
+							Rhs.GetChannelRangeAttributeData(0, RhsRangeOffsets[RangeIdx], RangeLengths[RangeIdx]),
+							Rhs.GetChannelRangeAttributeData(1, RhsRangeOffsets[RangeIdx], RangeLengths[RangeIdx]),
+							Rhs.GetChannelRangeAttributeData(2, RhsRangeOffsets[RangeIdx], RangeLengths[RangeIdx]),
+							Rhs.GetChannelRangeAttributeData(3, RhsRangeOffsets[RangeIdx], RangeLengths[RangeIdx]),
+							Rhs.GetChannelRangeAttributeData(4, RhsRangeOffsets[RangeIdx], RangeLengths[RangeIdx]),
+							Rhs.GetChannelRangeAttributeData(5, RhsRangeOffsets[RangeIdx], RangeLengths[RangeIdx]),
+							Rhs.GetChannelRangeAttributeData(6, RhsRangeOffsets[RangeIdx], RangeLengths[RangeIdx]),
+							Rhs.GetChannelRangeAttributeData(7, RhsRangeOffsets[RangeIdx], RangeLengths[RangeIdx]),
+							Rhs.GetChannelRangeAttributeData(8, RhsRangeOffsets[RangeIdx], RangeLengths[RangeIdx]),
+							Rhs.GetChannelRangeAttributeData(9, RhsRangeOffsets[RangeIdx], RangeLengths[RangeIdx]));
+					}
+				});
+		}
+
+		void TransformDiv(FFrameAttribute& Out, const FFrameAttribute& Lhs, const FFrameAttribute& Rhs)
+		{
+			check(Lhs.GetChannelNum() == 10);
+			check(Rhs.GetChannelNum() == 10);
+
+
+			BinaryOp(Out, Lhs.GetChannelNum(), Lhs, Rhs, [](
+				FFrameAttribute& Out,
+				const FFrameAttribute& Lhs,
+				const FFrameAttribute& Rhs,
+				const TLearningArrayView<1, const int32> OutRangeOffsets,
+				const TLearningArrayView<1, const int32> LhsRangeOffsets,
+				const TLearningArrayView<1, const int32> RhsRangeOffsets,
+				const TLearningArrayView<1, const int32> RangeLengths)
+				{
+					const int32 RangeNum = RangeLengths.Num();
+
+					for (int32 RangeIdx = 0; RangeIdx < RangeNum; RangeIdx++)
+					{
+						Private::TransformDiv(
+							Out.GetChannelRangeAttributeData(0, OutRangeOffsets[RangeIdx], RangeLengths[RangeIdx]),
+							Out.GetChannelRangeAttributeData(1, OutRangeOffsets[RangeIdx], RangeLengths[RangeIdx]),
+							Out.GetChannelRangeAttributeData(2, OutRangeOffsets[RangeIdx], RangeLengths[RangeIdx]),
+							Out.GetChannelRangeAttributeData(3, OutRangeOffsets[RangeIdx], RangeLengths[RangeIdx]),
+							Out.GetChannelRangeAttributeData(4, OutRangeOffsets[RangeIdx], RangeLengths[RangeIdx]),
+							Out.GetChannelRangeAttributeData(5, OutRangeOffsets[RangeIdx], RangeLengths[RangeIdx]),
+							Out.GetChannelRangeAttributeData(6, OutRangeOffsets[RangeIdx], RangeLengths[RangeIdx]),
+							Out.GetChannelRangeAttributeData(7, OutRangeOffsets[RangeIdx], RangeLengths[RangeIdx]),
+							Out.GetChannelRangeAttributeData(8, OutRangeOffsets[RangeIdx], RangeLengths[RangeIdx]),
+							Out.GetChannelRangeAttributeData(9, OutRangeOffsets[RangeIdx], RangeLengths[RangeIdx]),
+
+							Lhs.GetChannelRangeAttributeData(0, LhsRangeOffsets[RangeIdx], RangeLengths[RangeIdx]),
+							Lhs.GetChannelRangeAttributeData(1, LhsRangeOffsets[RangeIdx], RangeLengths[RangeIdx]),
+							Lhs.GetChannelRangeAttributeData(2, LhsRangeOffsets[RangeIdx], RangeLengths[RangeIdx]),
+							Lhs.GetChannelRangeAttributeData(3, LhsRangeOffsets[RangeIdx], RangeLengths[RangeIdx]),
+							Lhs.GetChannelRangeAttributeData(4, LhsRangeOffsets[RangeIdx], RangeLengths[RangeIdx]),
+							Lhs.GetChannelRangeAttributeData(5, LhsRangeOffsets[RangeIdx], RangeLengths[RangeIdx]),
+							Lhs.GetChannelRangeAttributeData(6, LhsRangeOffsets[RangeIdx], RangeLengths[RangeIdx]),
+							Lhs.GetChannelRangeAttributeData(7, LhsRangeOffsets[RangeIdx], RangeLengths[RangeIdx]),
+							Lhs.GetChannelRangeAttributeData(8, LhsRangeOffsets[RangeIdx], RangeLengths[RangeIdx]),
+							Lhs.GetChannelRangeAttributeData(9, LhsRangeOffsets[RangeIdx], RangeLengths[RangeIdx]),
+
+							Rhs.GetChannelRangeAttributeData(0, RhsRangeOffsets[RangeIdx], RangeLengths[RangeIdx]),
+							Rhs.GetChannelRangeAttributeData(1, RhsRangeOffsets[RangeIdx], RangeLengths[RangeIdx]),
+							Rhs.GetChannelRangeAttributeData(2, RhsRangeOffsets[RangeIdx], RangeLengths[RangeIdx]),
+							Rhs.GetChannelRangeAttributeData(3, RhsRangeOffsets[RangeIdx], RangeLengths[RangeIdx]),
+							Rhs.GetChannelRangeAttributeData(4, RhsRangeOffsets[RangeIdx], RangeLengths[RangeIdx]),
+							Rhs.GetChannelRangeAttributeData(5, RhsRangeOffsets[RangeIdx], RangeLengths[RangeIdx]),
+							Rhs.GetChannelRangeAttributeData(6, RhsRangeOffsets[RangeIdx], RangeLengths[RangeIdx]),
+							Rhs.GetChannelRangeAttributeData(7, RhsRangeOffsets[RangeIdx], RangeLengths[RangeIdx]),
+							Rhs.GetChannelRangeAttributeData(8, RhsRangeOffsets[RangeIdx], RangeLengths[RangeIdx]),
+							Rhs.GetChannelRangeAttributeData(9, RhsRangeOffsets[RangeIdx], RangeLengths[RangeIdx]));
+					}
+				});
+		}
+
+		void TransformMulConstant(FFrameAttribute& Out, const FFrameAttribute& Lhs, const FTransform3f& Rhs)
+		{
+			check(Lhs.GetChannelNum() == 10);
+
+			UnaryOp(Out, Lhs.GetChannelNum(), Lhs, [&Rhs](
+				FFrameAttribute& Out,
+				const FFrameAttribute& In,
+				const TLearningArrayView<1, const int32> RangeOffsets,
+				const TLearningArrayView<1, const int32> RangeLengths)
+				{
+					Private::TransformMulConstant(
+						Out.GetAttributeData(),
+						In.GetAttributeData(),
+						Rhs);
+				});
+		}
+
+		void TransformDivConstant(FFrameAttribute& Out, const FFrameAttribute& Lhs, const FTransform3f& Rhs)
+		{
+			check(Lhs.GetChannelNum() == 10);
+
+			UnaryOp(Out, Lhs.GetChannelNum(), Lhs, [&Rhs](
+				FFrameAttribute& Out,
+				const FFrameAttribute& In,
+				const TLearningArrayView<1, const int32> RangeOffsets,
+				const TLearningArrayView<1, const int32> RangeLengths)
+				{
+					Private::TransformDivConstant(
+						Out.GetAttributeData(),
+						In.GetAttributeData(),
+						Rhs);
+				});
+		}
+
+		void TransformConstantMul(FFrameAttribute& Out, const FTransform3f& Lhs, const FFrameAttribute& Rhs)
+		{
+			check(Rhs.GetChannelNum() == 10);
+
+			UnaryOp(Out, Rhs.GetChannelNum(), Rhs, [&Lhs](
+				FFrameAttribute& Out,
+				const FFrameAttribute& In,
+				const TLearningArrayView<1, const int32> RangeOffsets,
+				const TLearningArrayView<1, const int32> RangeLengths)
+				{
+					Private::TransformConstantMul(
+						Out.GetAttributeData(),
+						Lhs,
+						In.GetAttributeData());
+				});
+		}
+
+		void TransformConstantDiv(FFrameAttribute& Out, const FTransform3f& Lhs, const FFrameAttribute& Rhs)
+		{
+			check(Rhs.GetChannelNum() == 10);
+
+			UnaryOp(Out, Rhs.GetChannelNum(), Rhs, [&Lhs](
+				FFrameAttribute& Out,
+				const FFrameAttribute& In,
+				const TLearningArrayView<1, const int32> RangeOffsets,
+				const TLearningArrayView<1, const int32> RangeLengths)
+				{
+					Private::TransformConstantDiv(
+						Out.GetAttributeData(),
+						Lhs,
+						In.GetAttributeData());
+				});
+		}
+
+		void TransformInv(FFrameAttribute& Out, const FFrameAttribute& In)
+		{
+			check(In.GetChannelNum() == 10);
+
+			UnaryOp(Out, In.GetChannelNum(), In, [](
+				UE::Learning::FFrameAttribute& Out,
+				const UE::Learning::FFrameAttribute& In,
+				const TLearningArrayView<1, const int32> RangeOffsets,
+				const TLearningArrayView<1, const int32> RangeLengths) {
+
+					Private::TransformInv(
+						Out.GetAttributeData(),
+						In.GetAttributeData());
+				});
+		}
+
+		void TransformLocation(FFrameAttribute& Out, const FFrameAttribute& In)
+		{
+			check(In.GetChannelNum() == 10);
+			Slice(Out, In, 0, 3);
+		}
+
+		void TransformRotation(FFrameAttribute& Out, const FFrameAttribute& In)
+		{
+			check(In.GetChannelNum() == 10);
+			Slice(Out, In, 3, 4);
+		}
+
+		void TransformScale(FFrameAttribute& Out, const FFrameAttribute& In)
+		{
+			check(In.GetChannelNum() == 10);
+			Slice(Out, In, 7, 3);
+		}
+
 	}
 }
diff --git a/Engine/Plugins/Experimental/LearningAgents/Source/Learning/Public/LearningFrameAttribute.h b/Engine/Plugins/Experimental/LearningAgents/Source/Learning/Public/LearningFrameAttribute.h
--- a/Engine/Plugins/Experimental/LearningAgents/Source/Learning/Public/LearningFrameAttribute.h
+++ b/Engine/Plugins/Experimental/LearningAgents/Source/Learning/Public/LearningFrameAttribute.h
@@ -110,6 +110,12 @@ namespace UE::Learning
 			const TLearningArrayView<1, const int32> RangeOffsets,
 			const TLearningArrayView<1, const int32> RangeLengths)>;
 
+		/** In-place op function. Takes a single frame attribute and modifies it in-place. */
+		using InplaceOpFunction = TFunctionRef<void(
+			FFrameAttribute& InOut,
+			const TLearningArrayView<1, const int32> RangeOffsets,
+			const TLearningArrayView<1, const int32> RangeLengths)>;
+
 		/** Binary op function. Takes a two frame attributes as input and produces as output a single frame attribute. */
 		using BinaryOpFunction = TFunctionRef<void(
 			FFrameAttribute& Out,
@@ -131,12 +137,14 @@ namespace UE::Learning
 			const TArrayView<const TLearningArrayView<1, const int32>> InRangeOffsets,
 			const TLearningArrayView<1, const int32> RangeLengths)>;
 
+
 		/** Computes the intersection of a frame attribute and frame range set */
 		UE_API void Intersection(FFrameAttribute& OutFrameAttribute, const FFrameAttribute& FrameAttribute, const FFrameRangeSet& FrameRangeSet);
 
 		/** Computes the frame range set where the given channel is non-zero in the frame attribute. */
 		UE_API void NonZeroFrameRangeSet(FFrameRangeSet& OutFrameRangeSet, const FFrameAttribute& FrameAttribute, const int32 ChannelIdx);
 
+
 		/** Computes a reduction on a frame attribute */
 		UE_API void ReduceOp(
 			const FFrameAttribute& In,
@@ -156,6 +164,11 @@ namespace UE::Learning
 			const FFrameAttribute& In,
 			const UnaryOpFunction Op);
 
+		/** Modifies a frame attribute in-place using the given op */
+		UE_API void InplaceOp(
+			FFrameAttribute& InOut,
+			const InplaceOpFunction Op);
+
 		/** Creates a frame attribute from two others and the given op. Performs an intersection of the Lhs and Rhs if they do not match. */
 		UE_API void BinaryOp(
 			FFrameAttribute& Out,
@@ -171,6 +184,7 @@ namespace UE::Learning
 			const TArrayView<const ConstFrameAttributePtr> Inputs,
 			const NaryOpFunction Op);
 
+
 		/** Find the channel and frame with the smallest value */
 		UE_API bool FindMinimum(
 			int32& OutChannelIdx,
@@ -185,6 +199,7 @@ namespace UE::Learning
 			float& OutMaximumValue, 
 			const FFrameAttribute& In);
 
+
 		/** Create a frame attribute made up of zeros */
 		UE_API void Zeros(FFrameAttribute& Out, const FFrameRangeSet& FrameRangeSet, const int32 ChannelNum);
 
@@ -194,6 +209,7 @@ namespace UE::Learning
 		/** Create a frame attribute filled with the given value at each frame */
 		UE_API void Fill(FFrameAttribute& Out, const FFrameRangeSet& FrameRangeSet, const TLearningArrayView<1, const float> Values);
 
+
 		/** Add two frame attributes. Channel numbers must match. */
 		UE_API void Add(FFrameAttribute& Out, const FFrameAttribute& Lhs, const FFrameAttribute& Rhs);
 
@@ -206,9 +222,11 @@ namespace UE::Learning
 		/** Divide two frame attributes. Channel numbers must match. */
 		UE_API void Div(FFrameAttribute& Out, const FFrameAttribute& Lhs, const FFrameAttribute& Rhs);
 
+
 		/** Compute the dot product of two frame attributes. Channel numbers must match. */
 		UE_API void Dot(FFrameAttribute& Out, const FFrameAttribute& Lhs, const FFrameAttribute& Rhs);
 
+
 		/** Negate a frame attribute. */
 		UE_API void Neg(FFrameAttribute& Out, const FFrameAttribute& In);
 
@@ -224,6 +242,32 @@ namespace UE::Learning
 		/** Compute the exp of a frame attribute. */
 		UE_API void Exp(FFrameAttribute& Out, const FFrameAttribute& In);
 
+		/** Compute the sqrt of a frame attribute. */
+		UE_API void Sqrt(FFrameAttribute& Out, const FFrameAttribute& In);
+
+
+		/** Compute the negation of a frame attribute in place. */
+		UE_API void NegInplace(FFrameAttribute& InOut);
+
+		/** Compute 1/x of a frame attribute in place. */
+		UE_API void InvInplace(FFrameAttribute& InOut);
+
+		/** Compute the absolute value of a frame attribute in place. */
+		UE_API void AbsInplace(FFrameAttribute& InOut);
+
+		/** Compute the log of a frame attribute in place. */
+		UE_API void LogInplace(FFrameAttribute& InOut);
+
+		/** Compute the exp of a frame attribute in place. */
+		UE_API void ExpInplace(FFrameAttribute& InOut);
+
+		/** Compute the square root of a frame attribute in place. */
+		UE_API void SqrtInplace(FFrameAttribute& InOut);
+
+		/** Compute the approximate acos of a frame attribute in place. */
+		UE_API void ApproxAcosInplace(FFrameAttribute& InOut);
+
+
 		/** Compute the length of a frame attribute over the channels dimension. */
 		UE_API void Length(FFrameAttribute& Out, const FFrameAttribute& In);
 
@@ -233,6 +277,9 @@ namespace UE::Learning
 		/** Extract a single channel from frame attribute. */
 		UE_API void Index(FFrameAttribute& Out, const FFrameAttribute& In, const int32 ChannelIdx);
 
+		/** Extract a range of channels from frame attribute. */
+		UE_API void Slice(FFrameAttribute& Out, const FFrameAttribute& In, const int32 StartChannelIdx, const int32 ChannelNum);
+
 		/** Add a constant value to a frame attribute. Channel number must match the size of Rhs. */
 		UE_API void AddConstant(FFrameAttribute& Out, const FFrameAttribute& Lhs, const TLearningArrayView<1, const float> Rhs);
 
@@ -245,6 +292,7 @@ namespace UE::Learning
 		/** Divide a frame attribute by a constant value. Channel number must match the size of Rhs. */
 		UE_API void DivConstant(FFrameAttribute& Out, const FFrameAttribute& Lhs, const TLearningArrayView<1, const float> Rhs);
 
+
 		/** Add a constant value to a frame attribute. Channel number must match the size of Lhs. */
 		UE_API void ConstantAdd(FFrameAttribute& Out, const TLearningArrayView<1, const float> Lhs, const FFrameAttribute& Rhs);
 
@@ -257,12 +305,37 @@ namespace UE::Learning
 		/** Divide a frame attribute by a constant value. Channel number must match the size of Lhs. */
 		UE_API void ConstantDiv(FFrameAttribute& Out, const TLearningArrayView<1, const float> Lhs, const FFrameAttribute& Rhs);
 
+
+		/** Gets the dot product with a constant value. Channel number must match the size of Lhs. */
+		UE_API void ConstantDot(FFrameAttribute& Out, const TLearningArrayView<1, const float> Lhs, const FFrameAttribute& Rhs);
+
+		/** Gets the squared euclidean distance to a constant value. Channel number must match the size of Lhs. */
+		UE_API void ConstantSquaredDist(FFrameAttribute& Out, const TLearningArrayView<1, const float> Lhs, const FFrameAttribute& Rhs);
+
+		/** Gets the euclidean distance to a constant value. Channel number must match the size of Lhs. */
+		UE_API void ConstantDist(FFrameAttribute& Out, const TLearningArrayView<1, const float> Lhs, const FFrameAttribute& Rhs);
+
+
+		/** Add a frame attribute by a constant value in-place. Channel number must match the size of Rhs. */
+		UE_API void AddConstantInplace(FFrameAttribute& InOut, const TLearningArrayView<1, const float> Rhs);
+
+		/** Subtract a frame attribute by a constant value in-place. Channel number must match the size of Rhs. */
+		UE_API void SubConstantInplace(FFrameAttribute& InOut, const TLearningArrayView<1, const float> Rhs);
+
+		/** Multiply a frame attribute by a constant value in-place. Channel number must match the size of Rhs. */
+		UE_API void MulConstantInplace(FFrameAttribute& InOut, const TLearningArrayView<1, const float> Rhs);
+
+		/** Divide a frame attribute by a constant value in-place. Channel number must match the size of Rhs. */
+		UE_API void DivConstantInplace(FFrameAttribute& InOut, const TLearningArrayView<1, const float> Rhs);
+
+
 		/** Compute the sum of an array of frame attributes. Channel numbers must match. */
 		UE_API void Sum(FFrameAttribute& Out, const TArrayView<const ConstFrameAttributePtr> Inputs);
 
 		/** Compute the product of an array of frame attributes. Channel numbers must match. */
 		UE_API void Prod(FFrameAttribute& Out, const TArrayView<const ConstFrameAttributePtr> Inputs);
 
+
 		/** Computes the logical and of two frame attributes. Channel numbers must match. */
 		UE_API void LogicalAnd(FFrameAttribute& Out, const FFrameAttribute& Lhs, const FFrameAttribute& Rhs);
 
@@ -272,6 +345,7 @@ namespace UE::Learning
 		/** Computes the logical not of a frame attribute. */
 		UE_API void LogicalNot(FFrameAttribute& Out, const FFrameAttribute& In);
 
+
 		/** Computes if the values in one frame attribute are greater than another. Channel numbers must match. */
 		UE_API void Gt(FFrameAttribute& Out, const FFrameAttribute& Lhs, const FFrameAttribute& Rhs);
 
@@ -290,6 +364,7 @@ namespace UE::Learning
 		/** Computes if the values in one frame attribute are not equal to another. Channel numbers must match. */
 		UE_API void Neq(FFrameAttribute& Out, const FFrameAttribute& Lhs, const FFrameAttribute& Rhs);
 
+
 		/** Computes if the values in a frame attribute are greater than a constant. Channel number must match the size of Rhs. */
 		UE_API void GtConstant(FFrameAttribute& Out, const FFrameAttribute& Lhs, const TLearningArrayView<1, const float> Rhs);
 
@@ -308,6 +383,7 @@ namespace UE::Learning
 		/** Computes if the values in a frame attribute are not equal to a constant. Channel number must match the size of Rhs. */
 		UE_API void NeqConstant(FFrameAttribute& Out, const FFrameAttribute& Lhs, const TLearningArrayView<1, const float> Rhs);
 
+
 		/** Computes if the values in a frame attribute are greater than a constant. Channel number must match the size of Lhs. */
 		UE_API void ConstantGt(FFrameAttribute& Out, const TLearningArrayView<1, const float> Lhs, const FFrameAttribute& Rhs);
 		
@@ -326,6 +402,7 @@ namespace UE::Learning
 		/** Computes if the values in a frame attribute are not equal to a constant. Channel number must match the size of Lhs. */
 		UE_API void ConstantNeq(FFrameAttribute& Out, const TLearningArrayView<1, const float> Lhs, const FFrameAttribute& Rhs);
 
+
 		/** Applies a Gaussian smoothing filter to the ranges of values in the frame attribute. */
 		UE_API void FilterGaussian(FFrameAttribute& Out, const FFrameAttribute& In, const float StdInFrames);
 
@@ -338,6 +415,7 @@ namespace UE::Learning
 		/** Computes the mean and log std across all frames in the log space. OutMean and OutLogStd should be size of the number of channels */
 		UE_API void LogMeanStd(TLearningArrayView<1, float> OutMean, TLearningArrayView<1, float> OutLogStd, const FFrameAttribute& In);
 
+
 		/** Computes the quaternion multiplication of two frame attributes. Both arguments are assumed to have 4 channels. */
 		UE_API void QuatMul(FFrameAttribute& Out, const FFrameAttribute& Lhs, const FFrameAttribute& Rhs);
 
@@ -392,6 +470,9 @@ namespace UE::Learning
 		/** Computes the quaternion between the given frame attribute and vector. Lhs is assumed to have 3 channels. */
 		UE_API void QuatBetweenConstant(FFrameAttribute& Out, const FFrameAttribute& Lhs, const FVector3f Rhs);
 
+		/** Computes a fast approximation of the shortest angle between the given frame attribute and the quaternion. Lhs is assumed to have 4 channels. */
+		UE_API void QuatApproxShortestAngleBetweenConstant(FFrameAttribute& Out, const FFrameAttribute& Lhs, const FQuat4f Rhs);
+
 		/** Computes the quaternion multiplication of a frame attribute by the given quaternion. Rhs is assumed to have 4 channels. */
 		UE_API void QuatConstantMul(FFrameAttribute& Out, const FQuat4f Lhs, const FFrameAttribute& Rhs);
 
@@ -415,6 +496,43 @@ namespace UE::Learning
 
 		/** Computes the quaternion mean and std across all frames. */
 		UE_API void QuatMeanStd(FQuat4f& OutMean, FVector3f& OutStd, const FFrameAttribute& In);
+
+
+		/** Make a transform attribute from a location, rotation, and scale */
+		UE_API void TransformMake(FFrameAttribute& Out, const FFrameAttribute& Location, const FFrameAttribute& Rotation, const FFrameAttribute& Scale);
+
+		/** Make a transform attribute from a location and rotation */
+		UE_API void TransformMake(FFrameAttribute& Out, const FFrameAttribute& Location, const FFrameAttribute& Rotation);
+
+		/** Computes the transform multiplication of two frame attributes. Both arguments are assumed to have 10 channels. */
+		UE_API void TransformMul(FFrameAttribute& Out, const FFrameAttribute& Lhs, const FFrameAttribute& Rhs);
+
+		/** Computes the transform division of two frame attributes. Both arguments are assumed to have 10 channels. */
+		UE_API void TransformDiv(FFrameAttribute& Out, const FFrameAttribute& Lhs, const FFrameAttribute& Rhs);
+
+		/** Computes the transform multiplication of two frame attributes. Argument is assumed to have 10 channels. */
+		UE_API void TransformMulConstant(FFrameAttribute& Out, const FFrameAttribute& Lhs, const FTransform3f& Rhs);
+
+		/** Computes the transform division of a frame attribute and a constrant value. Argument is assumed to have 10 channels. */
+		UE_API void TransformDivConstant(FFrameAttribute& Out, const FFrameAttribute& Lhs, const FTransform3f& Rhs);
+
+		/** Computes the transform multiplication of a frame attribute and a constrant value Argument is assumed to have 10 channels. */
+		UE_API void TransformConstantMul(FFrameAttribute& Out, const FTransform3f& Lhs, const FFrameAttribute& Rhs);
+
+		/** Computes the transform division of a frame attribute and a constrant value. Argument is assumed to have 10 channels. */
+		UE_API void TransformConstantDiv(FFrameAttribute& Out, const FTransform3f& Lhs, const FFrameAttribute& Rhs);
+
+		/** Compute the transform inverse of a frame attribute. Argument is assumed to have 10 channels. */
+		UE_API void TransformInv(FFrameAttribute& Out, const FFrameAttribute& In);
+
+		/** Gets the location part of a transform. Argument is assumed to have 10 channels. */
+		UE_API void TransformLocation(FFrameAttribute& Out, const FFrameAttribute& In);
+
+		/** Gets the rotation part of a transform. Argument is assumed to have 10 channels. */
+		UE_API void TransformRotation(FFrameAttribute& Out, const FFrameAttribute& In);
+
+		/** Gets the scale part of a transform. Argument is assumed to have 10 channels. */
+		UE_API void TransformScale(FFrameAttribute& Out, const FFrameAttribute& In);
 	}
 }
 
