LearningAgents: Fixed docstring

diff --git a/Engine/Plugins/Experimental/LearningAgents/Source/Learning/Public/LearningFrameAttribute.h b/Engine/Plugins/Experimental/LearningAgents/Source/Learning/Public/LearningFrameAttribute.h
--- a/Engine/Plugins/Experimental/LearningAgents/Source/Learning/Public/LearningFrameAttribute.h
+++ b/Engine/Plugins/Experimental/LearningAgents/Source/Learning/Public/LearningFrameAttribute.h
@@ -510,16 +510,16 @@ namespace UE::Learning
 		/** Computes the transform division of two frame attributes. Both arguments are assumed to have 10 channels. */
 		UE_API void TransformDiv(FFrameAttribute& Out, const FFrameAttribute& Lhs, const FFrameAttribute& Rhs);
 
-		/** Computes the transform multiplication of two frame attributes. Argument is assumed to have 10 channels. */
+		/** Computes the transform multiplication of a frame attribute and a constant value. Argument is assumed to have 10 channels. */
 		UE_API void TransformMulConstant(FFrameAttribute& Out, const FFrameAttribute& Lhs, const FTransform3f& Rhs);
 
-		/** Computes the transform division of a frame attribute and a constrant value. Argument is assumed to have 10 channels. */
+		/** Computes the transform division of a frame attribute and a constant value. Argument is assumed to have 10 channels. */
 		UE_API void TransformDivConstant(FFrameAttribute& Out, const FFrameAttribute& Lhs, const FTransform3f& Rhs);
 
-		/** Computes the transform multiplication of a frame attribute and a constrant value Argument is assumed to have 10 channels. */
+		/** Computes the transform multiplication of a frame attribute and a constant value. Argument is assumed to have 10 channels. */
 		UE_API void TransformConstantMul(FFrameAttribute& Out, const FTransform3f& Lhs, const FFrameAttribute& Rhs);
 
-		/** Computes the transform division of a frame attribute and a constrant value. Argument is assumed to have 10 channels. */
+		/** Computes the transform division of a frame attribute and a constant value. Argument is assumed to have 10 channels. */
 		UE_API void TransformConstantDiv(FFrameAttribute& Out, const FTransform3f& Lhs, const FFrameAttribute& Rhs);
 
 		/** Compute the transform inverse of a frame attribute. Argument is assumed to have 10 channels. */
