LearningAgents: Added checks for prior probabilities sizes

diff --git a/Engine/Plugins/Experimental/LearningAgents/Source/LearningAgents/Private/LearningAgentsActions.cpp b/Engine/Plugins/Experimental/LearningAgents/Source/LearningAgents/Private/LearningAgentsActions.cpp
--- a/Engine/Plugins/Experimental/LearningAgents/Source/LearningAgents/Private/LearningAgentsActions.cpp
+++ b/Engine/Plugins/Experimental/LearningAgents/Source/LearningAgents/Private/LearningAgentsActions.cpp
@@ -1111,8 +1111,14 @@ FLearningAgentsActionSchemaElement ULearningAgentsActions::SpecifyExclusiveDiscr
 		UE_LOG(LogLearning, Warning, TEXT("SpecifyExclusiveDiscreteActionFromArrayView: Specifying zero-sized Exclusive Discrete Action."));
 	}
 
+	if (PriorProbabilities.Num() > Size)
+	{
+		UE_LOG(LogLearning, Error, TEXT("SpecifyExclusiveDiscreteActionFromArrayView: Action Size is '%i' but got '%i' PriorProbabilities."), Size, PriorProbabilities.Num());
+		return FLearningAgentsActionSchemaElement();
+	}
+
 	TArray<float, TInlineAllocator<16>> NormalizedPriorProbabilities;
-	NormalizedPriorProbabilities.Init(1.0f / Size, Size);
+	NormalizedPriorProbabilities.Init(Size == 0 ? 0.0f : 1.0f / Size, Size);
 	for (int32 Idx = 0; Idx < PriorProbabilities.Num(); Idx++)
 	{
 		NormalizedPriorProbabilities[Idx] = PriorProbabilities[Idx];
@@ -1124,13 +1130,13 @@ FLearningAgentsActionSchemaElement ULearningAgentsActions::SpecifyExclusiveDiscr
 
 FLearningAgentsActionSchemaElement ULearningAgentsActions::SpecifyNamedExclusiveDiscreteAction(ULearningAgentsActionSchema* Schema, const TArray<FName>& Names, const TMap<FName, float>& PriorProbabilities, const FName Tag)
 {
-	if (Names.Num() == 0)
+	const int32 NameNum = Names.Num();
+
+	if (NameNum == 0)
 	{
 		UE_LOG(LogLearning, Warning, TEXT("SpecifyNamedExclusiveDiscreteAction: Specifying zero-sized Named Exclusive Discrete Action."));
 	}
 
-	const int32 NameNum = Names.Num();
-
 	TArray<int32, TInlineAllocator<16>> NameIndices;
 	TArray<float, TInlineAllocator<16>> NamePriorProbabilities;
 	NameIndices.Empty(NameNum);
@@ -1141,7 +1147,7 @@ FLearningAgentsActionSchemaElement ULearningAgentsActions::SpecifyNamedExclusive
 	{
 		NameIndices.Add(Index);
 		const float* PriorProb = PriorProbabilities.Find(Name);
-		NamePriorProbabilities.Add(PriorProb ? *PriorProb : 1.0f / NameNum);
+		NamePriorProbabilities.Add(PriorProb ? *PriorProb : NameNum == 0 ? 0.0f : 1.0f / NameNum);
 		Index++;
 	}
 
@@ -1173,14 +1179,22 @@ FLearningAgentsActionSchemaElement ULearningAgentsActions::SpecifyNamedExclusive
 		return FLearningAgentsActionSchemaElement();
 	}
 
-	if (Names.Num() == 0)
+	const int32 NameNum = Names.Num();
+
+	if (NameNum == 0)
 	{
 		UE_LOG(LogLearning, Error, TEXT("SpecifyNamedExclusiveDiscreteActionFromArrayViews: Specifying zero-sized Named Exclusive Discrete Action."));
 		return FLearningAgentsActionSchemaElement();
 	}
 
+	if (PriorProbabilities.Num() > NameNum)
+	{
+		UE_LOG(LogLearning, Error, TEXT("SpecifyNamedExclusiveDiscreteActionFromArrayViews: Action Size is '%i' but got '%i' PriorProbabilities."), NameNum, PriorProbabilities.Num());
+		return FLearningAgentsActionSchemaElement();
+	}
+
 	TArray<float, TInlineAllocator<16>> NormalizedPriorProbabilities;
-	NormalizedPriorProbabilities.Init(1.0f / Names.Num(), Names.Num());
+	NormalizedPriorProbabilities.Init(NameNum == 0 ? 0.0f : 1.0f / NameNum, NameNum);
 	for (int32 Idx = 0; Idx < PriorProbabilities.Num(); Idx++)
 	{
 		NormalizedPriorProbabilities[Idx] = PriorProbabilities[Idx];
@@ -1209,6 +1223,12 @@ FLearningAgentsActionSchemaElement ULearningAgentsActions::SpecifyInclusiveDiscr
 		return FLearningAgentsActionSchemaElement();
 	}
 
+	if (PriorProbabilities.Num() > Size)
+	{
+		UE_LOG(LogLearning, Error, TEXT("SpecifyInclusiveDiscreteActionFromArrayView: Action Size is '%i' but got '%i' PriorProbabilities."), Size, PriorProbabilities.Num());
+		return FLearningAgentsActionSchemaElement();
+	}
+
 	if (Size == 0)
 	{
 		UE_LOG(LogLearning, Warning, TEXT("SpecifyInclusiveDiscreteActionFromArrayView: Specifying zero-sized Inclusive Discrete Action."));
@@ -1227,13 +1247,13 @@ FLearningAgentsActionSchemaElement ULearningAgentsActions::SpecifyInclusiveDiscr
 
 FLearningAgentsActionSchemaElement ULearningAgentsActions::SpecifyNamedInclusiveDiscreteAction(ULearningAgentsActionSchema* Schema, const TArray<FName> Names, const TMap<FName, float>& PriorProbabilities, const FName Tag)
 {
-	if (Names.Num() == 0)
+	const int32 NameNum = Names.Num();
+
+	if (NameNum == 0)
 	{
 		UE_LOG(LogLearning, Warning, TEXT("SpecifyNamedInclusiveDiscreteAction: Specifying zero-sized Named Inclusive Discrete Action."));
 	}
 
-	const int32 NameNum = Names.Num();
-
 	TArray<int32, TInlineAllocator<16>> NameIndices;
 	TArray<float, TInlineAllocator<16>> NamePriorProbabilities;
 	NameIndices.Empty(NameNum);
@@ -1244,7 +1264,7 @@ FLearningAgentsActionSchemaElement ULearningAgentsActions::SpecifyNamedInclusive
 	{
 		NameIndices.Add(Index);
 		const float* PriorProb = PriorProbabilities.Find(Name);
-		NamePriorProbabilities.Add(PriorProb ? *PriorProb : 1.0f / NameNum);
+		NamePriorProbabilities.Add(PriorProb ? *PriorProb : NameNum == 0 ? 0.0f : 1.0f / NameNum);
 		Index++;
 	}
 
@@ -1276,13 +1296,21 @@ FLearningAgentsActionSchemaElement ULearningAgentsActions::SpecifyNamedInclusive
 		return FLearningAgentsActionSchemaElement();
 	}
 
-	if (Names.Num() == 0)
+	const int32 NameNum = Names.Num();
+
+	if (PriorProbabilities.Num() > NameNum)
+	{
+		UE_LOG(LogLearning, Error, TEXT("SpecifyNamedInclusiveDiscreteActionFromArrayViews: Action Size is '%i' but got '%i' PriorProbabilities."), NameNum, PriorProbabilities.Num());
+		return FLearningAgentsActionSchemaElement();
+	}
+
+	if (NameNum == 0)
 	{
 		UE_LOG(LogLearning, Warning, TEXT("SpecifyNamedInclusiveDiscreteActionFromArrayViews: Specifying zero-sized Named Inclusive Discrete Action."));
 	}
 
 	TArray<float, TInlineAllocator<16>> NormalizedPriorProbabilities;
-	NormalizedPriorProbabilities.Init(1.0f / Names.Num(), Names.Num());
+	NormalizedPriorProbabilities.Init(NameNum == 0 ? 0.0f : 1.0f / NameNum, NameNum);
 	for (int32 Idx = 0; Idx < PriorProbabilities.Num(); Idx++)
 	{
 		NormalizedPriorProbabilities[Idx] = PriorProbabilities[Idx];
@@ -1409,7 +1437,7 @@ FLearningAgentsActionSchemaElement ULearningAgentsActions::SpecifyExclusiveUnion
 		SubElementNames.Add(Element.Key);
 		SubElements.Add(Element.Value);
 		const float* PriorProb = PriorProbabilities.Find(Element.Key);
-		SubElementPriorProbabilities.Add(PriorProb ? *PriorProb : 1.0f / SubElementNum);
+		SubElementPriorProbabilities.Add(PriorProb ? *PriorProb : SubElementNum == 0 ? 0.0f : 1.0f / SubElementNum);
 		Index++;
 	}
 
@@ -1460,6 +1488,12 @@ FLearningAgentsActionSchemaElement ULearningAgentsActions::SpecifyExclusiveUnion
 		return FLearningAgentsActionSchemaElement();
 	}
 
+	if (PriorProbabilities.Num() > Elements.Num())
+	{
+		UE_LOG(LogLearning, Error, TEXT("SpecifyExclusiveUnionActionFromArrayViews: Action Size is '%i' but got '%i' PriorProbabilities."), Elements.Num(), PriorProbabilities.Num());
+		return FLearningAgentsActionSchemaElement();
+	}
+
 	if (UE::Learning::Agents::Action::Private::ContainsDuplicates(ElementNames))
 	{
 		UE_LOG(LogLearning, Error, TEXT("SpecifyExclusiveUnionActionFromArrayViews: Element Names contain duplicates."));
@@ -1481,7 +1515,7 @@ FLearningAgentsActionSchemaElement ULearningAgentsActions::SpecifyExclusiveUnion
 	}
 
 	TArray<float, TInlineAllocator<16>> NormalizedPriorProbabilities;
-	NormalizedPriorProbabilities.Init(1.0f / Elements.Num(), Elements.Num());
+	NormalizedPriorProbabilities.Init(Elements.Num() == 0 ? 0.0f : 1.0f / Elements.Num(), Elements.Num());
 	for (int32 Idx = 0; Idx < PriorProbabilities.Num(); Idx++)
 	{
 		NormalizedPriorProbabilities[Idx] = PriorProbabilities[Idx];
@@ -1516,7 +1550,7 @@ FLearningAgentsActionSchemaElement ULearningAgentsActions::SpecifyInclusiveUnion
 		SubElementNames.Add(Element.Key);
 		SubElements.Add(Element.Value);
 		const float* PriorProb = PriorProbabilities.Find(Element.Key);
-		SubElementPriorProbabilities.Add(PriorProb ? *PriorProb : 1.0f / SubElementNum);
+		SubElementPriorProbabilities.Add(PriorProb ? *PriorProb : SubElementNum == 0 ? 0.0f : 1.0f / SubElementNum);
 		Index++;
 	}
 
@@ -1567,6 +1601,12 @@ FLearningAgentsActionSchemaElement ULearningAgentsActions::SpecifyInclusiveUnion
 		return FLearningAgentsActionSchemaElement();
 	}
 
+	if (PriorProbabilities.Num() > Elements.Num())
+	{
+		UE_LOG(LogLearning, Error, TEXT("SpecifyInclusiveUnionActionFromArrayViews: Action Size is '%i' but got '%i' PriorProbabilities."), Elements.Num(), PriorProbabilities.Num());
+		return FLearningAgentsActionSchemaElement();
+	}
+
 	if (UE::Learning::Agents::Action::Private::ContainsDuplicates(ElementNames))
 	{
 		UE_LOG(LogLearning, Error, TEXT("SpecifyInclusiveUnionActionFromArrayViews: Element Names contain duplicates."));
@@ -1639,8 +1679,14 @@ FLearningAgentsActionSchemaElement ULearningAgentsActions::SpecifyEnumAction(ULe
 		return FLearningAgentsActionSchemaElement();
 	}
 
+	if (Enum->NumEnums() <= 0)
+	{
+		UE_LOG(LogLearning, Error, TEXT("SpecifyEnumAction: Invalid Enum."));
+		return FLearningAgentsActionSchemaElement();
+	}
+
 	TArray<float, TInlineAllocator<16>> PriorProbabilitiesArray;
-	PriorProbabilitiesArray.Init(1.0f / (Enum->NumEnums() - 1), Enum->NumEnums() - 1);
+	PriorProbabilitiesArray.Init((Enum->NumEnums() - 1) == 0 ? 0.0f : 1.0f / (Enum->NumEnums() - 1), Enum->NumEnums() - 1);
 	for (const TPair<uint8, float> Prior : PriorProbabilities)
 	{
 		const int32 EnumIndex = Enum->GetIndexByValue(Prior.Key);
@@ -1677,6 +1723,12 @@ FLearningAgentsActionSchemaElement ULearningAgentsActions::SpecifyBitmaskAction(
 		return FLearningAgentsActionSchemaElement();
 	}
 
+	if (Enum->NumEnums() <= 0)
+	{
+		UE_LOG(LogLearning, Error, TEXT("SpecifyBitmaskAction: Invalid Bitmask."));
+		return FLearningAgentsActionSchemaElement();
+	}
+
 	if (Enum->NumEnums() - 1 > 32)
 	{
 		UE_LOG(LogLearning, Error, TEXT("SpecifyBitmaskAction: Too many values in Enum to use as Bitmask (%i)."), Enum->NumEnums() - 1);
diff --git a/Engine/Plugins/Experimental/LearningAgents/Source/LearningAgents/Private/LearningAgentsManager.cpp b/Engine/Plugins/Experimental/LearningAgents/Source/LearningAgents/Private/LearningAgentsManager.cpp
--- a/Engine/Plugins/Experimental/LearningAgents/Source/LearningAgents/Private/LearningAgentsManager.cpp
+++ b/Engine/Plugins/Experimental/LearningAgents/Source/LearningAgents/Private/LearningAgentsManager.cpp
@@ -341,7 +341,7 @@ UObject* ULearningAgentsManager::GetAgent(const int32 AgentId, const TSubclassOf
 {
 	if (!HasAgent(AgentId))
 	{
-		UE_LOG(LogLearning, Error, TEXT("%s: AgentId %d not found. Be sure to only use AgentIds returned by AddAgent and check that the agent has not be removed."), *GetName(), AgentId);
+		UE_LOG(LogLearning, Error, TEXT("%s: AgentId %d not found. Be sure to only use AgentIds returned by AddAgent and check that the agent has not been removed."), *GetName(), AgentId);
 		return nullptr;
 	}
 	
@@ -368,7 +368,7 @@ void ULearningAgentsManager::GetAgents(TArray<UObject*>& OutAgents, const TArray
 	{
 		if (!HasAgent(AgentIds[AgentIdIdx]))
 		{
-			UE_LOG(LogLearning, Error, TEXT("%s: AgentId %d not found. Be sure to only use AgentIds returned by AddAgent and check that the agent has not be removed."), *GetName(), AgentIds[AgentIdIdx]);
+			UE_LOG(LogLearning, Error, TEXT("%s: AgentId %d not found. Be sure to only use AgentIds returned by AddAgent and check that the agent has not been removed."), *GetName(), AgentIds[AgentIdIdx]);
 			OutAgents[AgentIdIdx] = nullptr;
 			continue;
 		}
