[LearningAgents] In support of the ability to have custom recording assets

diff --git a/Engine/Plugins/Experimental/LearningAgents/Source/LearningAgentsTraining/Private/LearningAgentsRecording.cpp b/Engine/Plugins/Experimental/LearningAgents/Source/LearningAgentsTraining/Private/LearningAgentsRecording.cpp
--- a/Engine/Plugins/Experimental/LearningAgents/Source/LearningAgentsTraining/Private/LearningAgentsRecording.cpp
+++ b/Engine/Plugins/Experimental/LearningAgents/Source/LearningAgentsTraining/Private/LearningAgentsRecording.cpp
@@ -24,6 +24,16 @@ namespace UE::Learning::Agents::Recording::Private
 	static constexpr int32 VersionNumber = 1;
 }
 
+int32 UE::Learning::Agents::Recording::GetMagicNumber() 
+{ 
+	return UE::Learning::Agents::Recording::Private::MagicNumber; 
+}
+
+int32 UE::Learning::Agents::Recording::GetVersionNumber() 
+{ 
+	return UE::Learning::Agents::Recording::Private::VersionNumber; 
+}
+
 void ULearningAgentsRecording::ResetRecording()
 {
 	Schemas.Empty();
diff --git a/Engine/Plugins/Experimental/LearningAgents/Source/LearningAgentsTraining/Public/LearningAgentsRecording.h b/Engine/Plugins/Experimental/LearningAgents/Source/LearningAgentsTraining/Public/LearningAgentsRecording.h
--- a/Engine/Plugins/Experimental/LearningAgents/Source/LearningAgentsTraining/Public/LearningAgentsRecording.h
+++ b/Engine/Plugins/Experimental/LearningAgents/Source/LearningAgentsTraining/Public/LearningAgentsRecording.h
@@ -13,6 +13,12 @@
 
 #define UE_API LEARNINGAGENTSTRAINING_API
 
+namespace UE::Learning::Agents::Recording
+{
+	UE_API int32 GetMagicNumber();
+	UE_API int32 GetVersionNumber();
+}
+
 /** A schema entry. */
 USTRUCT(BlueprintType)
 struct FLearningAgentsSchema
@@ -95,7 +101,7 @@ public:
 
 	/** Resets this recording asset to be empty. */
 	UFUNCTION(CallInEditor, Category = "LearningAgents")
-	UE_API void ResetRecording();
+	UE_API virtual void ResetRecording();
 
 	/** Appends the recording from "RecordingFile" to this recording.  */
 	UFUNCTION(CallInEditor, Category = "LearningAgents")
@@ -115,7 +121,7 @@ public:
 
 	/** Load this recording from a file. */
 	UFUNCTION(BlueprintCallable, Category = "LearningAgents", meta = (RelativePath))
-	UE_API void LoadRecordingFromFile(const FFilePath& File, const FName& SchemaName, const FGameplayTag& Tag, TArray<FLearningAgentsRecord>& OutRecords);
+	UE_API virtual void LoadRecordingFromFile(const FFilePath& File, const FName& SchemaName, const FGameplayTag& Tag, TArray<FLearningAgentsRecord>& OutRecords);
 
 	/** Save this recording to a file. */
 	UFUNCTION(BlueprintCallable, BlueprintPure = false, Category = "LearningAgents", meta = (RelativePath))
