[LearningAgents] Added event delegates to recorder

#rnx
#changelist validated
#virtualized

diff --git a/Engine/Plugins/Experimental/LearningAgents/Source/LearningAgentsTraining/Private/LearningAgentsRecorder.cpp b/Engine/Plugins/Experimental/LearningAgents/Source/LearningAgentsTraining/Private/LearningAgentsRecorder.cpp
--- a/Engine/Plugins/Experimental/LearningAgents/Source/LearningAgentsTraining/Private/LearningAgentsRecorder.cpp
+++ b/Engine/Plugins/Experimental/LearningAgents/Source/LearningAgentsTraining/Private/LearningAgentsRecorder.cpp
@@ -341,6 +341,8 @@ void ULearningAgentsRecorder::EndRecording()
 	}
 
 	bIsRecording = false;
+
+	OnRecordingEnd.Broadcast();
 }
 
 void ULearningAgentsRecorder::EndRecordingAndDiscard()
@@ -365,6 +367,8 @@ void ULearningAgentsRecorder::EndRecordingAndDiscard()
 	}
 
 	bIsRecording = false;
+
+	OnRecordingEnd.Broadcast();
 }
 
 void ULearningAgentsRecorder::BeginRecording()
@@ -387,6 +391,8 @@ void ULearningAgentsRecorder::BeginRecording()
 	}
 
 	bIsRecording = true;
+
+	OnRecordingBegin.Broadcast();
 }
 
 const ULearningAgentsRecording* ULearningAgentsRecorder::GetRecordingAsset() const
diff --git a/Engine/Plugins/Experimental/LearningAgents/Source/LearningAgentsTraining/Public/LearningAgentsRecorder.h b/Engine/Plugins/Experimental/LearningAgents/Source/LearningAgentsTraining/Public/LearningAgentsRecorder.h
--- a/Engine/Plugins/Experimental/LearningAgents/Source/LearningAgentsTraining/Public/LearningAgentsRecorder.h
+++ b/Engine/Plugins/Experimental/LearningAgents/Source/LearningAgentsTraining/Public/LearningAgentsRecorder.h
@@ -15,6 +15,9 @@
 
 #define UE_API LEARNINGAGENTSTRAINING_API
 
+DECLARE_DYNAMIC_MULTICAST_DELEGATE(FOnRecordingBeginSignature);
+DECLARE_DYNAMIC_MULTICAST_DELEGATE(FOnRecordingEndSignature);
+
 struct FLearningAgentsRecord;
 class ULearningAgentsRecording;
 
@@ -120,6 +123,14 @@ public:
 	UFUNCTION(BlueprintPure, Category = "LearningAgents")
 	UE_API bool IsRecording() const;
 
+	/** Called after the recording has started */
+	UPROPERTY(BlueprintAssignable, Category = "LearningAgents")
+	FOnRecordingBeginSignature OnRecordingBegin;
+
+	/** Called after the recording has stopped */
+	UPROPERTY(BlueprintAssignable, Category = "LearningAgents")
+	FOnRecordingEndSignature OnRecordingEnd;
+	
 	/**
 	 * While recording, adds the current buffered observations and actions of the added agents to the internal buffer. Call this after 
 	 * GatherObservations and either EvaluateAgentController (if recording a human/AI demonstration) or DecodeAndSampleActions (if recording from 
