LearningAgents: Support larger recordings by adding some indirection

#rnx

diff --git a/Engine/Plugins/Experimental/LearningAgents/Source/LearningAgentsTraining/Private/LearningAgentsImitationTrainer.cpp b/Engine/Plugins/Experimental/LearningAgents/Source/LearningAgentsTraining/Private/LearningAgentsImitationTrainer.cpp
--- a/Engine/Plugins/Experimental/LearningAgents/Source/LearningAgentsTraining/Private/LearningAgentsImitationTrainer.cpp
+++ b/Engine/Plugins/Experimental/LearningAgents/Source/LearningAgentsTraining/Private/LearningAgentsImitationTrainer.cpp
@@ -253,8 +253,8 @@ void ULearningAgentsImitationTrainer::BeginTraining(
 
 		if (Record.ObservationDimNum == ObservationNum && Record.ActionDimNum == ActionNum)
 		{
-			TLearningArrayView<2, const float> ObservationsView = TLearningArrayView<2, const float>(Record.ObservationData.GetData(), { Record.StepNum, Record.ObservationDimNum });
-			TLearningArrayView<2, const float> ActionsView = TLearningArrayView<2, const float>(Record.ActionData.GetData(), { Record.StepNum, Record.ActionDimNum });
+			TLearningArrayView<2, const float> ObservationsView = TLearningArrayView<2, const float>(Record.Blob->ObservationData.GetData(), { Record.StepNum, Record.ObservationDimNum });
+			TLearningArrayView<2, const float> ActionsView = TLearningArrayView<2, const float>(Record.Blob->ActionData.GetData(), { Record.StepNum, Record.ActionDimNum });
 
 			RecordedEpisodeStarts[EpisodeIdx] = StepIdx;
 			RecordedEpisodeLengths[EpisodeIdx] = Record.StepNum;
@@ -264,8 +264,8 @@ void ULearningAgentsImitationTrainer::BeginTraining(
 		else if (UE::Learning::Trainer::IsObservationSchemaSubsetCompatible(Recording->Schemas[Record.SchemaName].ObservationSchemaJson, Interactor->GetObservationSchema()->ObservationSchema, Interactor->GetObservationSchemaElement().SchemaElement) &&
 			UE::Learning::Trainer::IsActionSchemaSubsetCompatible(Recording->Schemas[Record.SchemaName].ActionSchemaJson, Interactor->GetActionSchema()->ActionSchema, Interactor->GetActionSchemaElement().SchemaElement))
 		{
-			TLearningArrayView<2, const float> ObservationsView = TLearningArrayView<2, const float>(Record.ObservationData.GetData(), { Record.StepNum, Record.ObservationDimNum });
-			TLearningArrayView<2, const float> ActionsView = TLearningArrayView<2, const float>(Record.ActionData.GetData(), { Record.StepNum, Record.ActionDimNum });
+			TLearningArrayView<2, const float> ObservationsView = TLearningArrayView<2, const float>(Record.Blob->ObservationData.GetData(), { Record.StepNum, Record.ObservationDimNum });
+			TLearningArrayView<2, const float> ActionsView = TLearningArrayView<2, const float>(Record.Blob->ActionData.GetData(), { Record.StepNum, Record.ActionDimNum });
 
 			RecordedEpisodeStarts[EpisodeIdx] = StepIdx;
 			RecordedEpisodeLengths[EpisodeIdx] = Record.StepNum;
diff --git a/Engine/Plugins/Experimental/LearningAgents/Source/LearningAgentsTraining/Private/LearningAgentsRecorder.cpp b/Engine/Plugins/Experimental/LearningAgents/Source/LearningAgentsTraining/Private/LearningAgentsRecorder.cpp
--- a/Engine/Plugins/Experimental/LearningAgents/Source/LearningAgentsTraining/Private/LearningAgentsRecorder.cpp
+++ b/Engine/Plugins/Experimental/LearningAgents/Source/LearningAgentsTraining/Private/LearningAgentsRecorder.cpp
@@ -96,11 +96,11 @@ void ULearningAgentsRecorder::FAgentRecordBuffer::CopyToRecord(FLearningAgentsRe
 	Record.ActionDimNum = GetAction(0).Num();
 	Record.ObservationCompatibilityHash = ObservationCompatibilityHash;
 	Record.ActionCompatibilityHash = ActionCompatibilityHash;
-	Record.ObservationData.SetNumUninitialized(StepNum * Record.ObservationDimNum);
-	Record.ActionData.SetNumUninitialized(StepNum * Record.ActionDimNum);
+	Record.Blob->ObservationData.SetNumUninitialized(StepNum * Record.ObservationDimNum);
+	Record.Blob->ActionData.SetNumUninitialized(StepNum * Record.ActionDimNum);
 
-	TLearningArrayView<2, float> ObservationsView = TLearningArrayView<2, float>(Record.ObservationData.GetData(), { Record.StepNum, Record.ObservationDimNum });
-	TLearningArrayView<2, float> ActionsView = TLearningArrayView<2, float>(Record.ActionData.GetData(), { Record.StepNum, Record.ActionDimNum });
+	TLearningArrayView<2, float> ObservationsView = TLearningArrayView<2, float>(Record.Blob->ObservationData.GetData(), { Record.StepNum, Record.ObservationDimNum });
+	TLearningArrayView<2, float> ActionsView = TLearningArrayView<2, float>(Record.Blob->ActionData.GetData(), { Record.StepNum, Record.ActionDimNum });
 
 	for (int32 StepIdx = 0; StepIdx < StepNum; StepIdx++)
 	{
diff --git a/Engine/Plugins/Experimental/LearningAgents/Source/LearningAgentsTraining/Private/LearningAgentsRecording.cpp b/Engine/Plugins/Experimental/LearningAgents/Source/LearningAgentsTraining/Private/LearningAgentsRecording.cpp
--- a/Engine/Plugins/Experimental/LearningAgents/Source/LearningAgentsTraining/Private/LearningAgentsRecording.cpp
+++ b/Engine/Plugins/Experimental/LearningAgents/Source/LearningAgentsTraining/Private/LearningAgentsRecording.cpp
@@ -180,14 +180,15 @@ void ULearningAgentsRecording::LoadRecordingFromFile(const FFilePath& File, cons
 		for (int32 RecordIdx = 0; RecordIdx < RecordNum; RecordIdx++)
 		{
 			FLearningAgentsRecord Record;
+			Record.Blob = NewObject<ULearningAgentsRecordBlob>(this);
 
 			UE::Learning::DeserializeFromBytes(Offset, RecordingData, Record.StepNum);
 			UE::Learning::DeserializeFromBytes(Offset, RecordingData, Record.ObservationDimNum);
 			UE::Learning::DeserializeFromBytes(Offset, RecordingData, Record.ActionDimNum);
 			UE::Learning::DeserializeFromBytes(Offset, RecordingData, Record.ObservationCompatibilityHash);
 			UE::Learning::DeserializeFromBytes(Offset, RecordingData, Record.ActionCompatibilityHash);
-			UE::Learning::Array::DeserializeFromBytes32<2, float>(Offset, RecordingData, Record.ObservationData);
-			UE::Learning::Array::DeserializeFromBytes32<2, float>(Offset, RecordingData, Record.ActionData);
+			UE::Learning::Array::DeserializeFromBytes32<2, float>(Offset, RecordingData, Record.Blob->ObservationData);
+			UE::Learning::Array::DeserializeFromBytes32<2, float>(Offset, RecordingData, Record.Blob->ActionData);
 			
 			Record.Name = FPaths::GetBaseFilename(File.FilePath);
 			Record.SchemaName = SchemaName;
@@ -239,8 +240,8 @@ void ULearningAgentsRecording::SaveRecordingToFile(const FFilePath& File) const
 		UE::Learning::SerializeToBytes(Offset, RecordingData, Records[RecordIdx].ActionDimNum);
 		UE::Learning::SerializeToBytes(Offset, RecordingData, Records[RecordIdx].ObservationCompatibilityHash);
 		UE::Learning::SerializeToBytes(Offset, RecordingData, Records[RecordIdx].ActionCompatibilityHash);
-		UE::Learning::Array::SerializeToBytes32<2, float>(Offset, RecordingData, { Records[RecordIdx].StepNum, Records[RecordIdx].ObservationDimNum }, Records[RecordIdx].ObservationData);
-		UE::Learning::Array::SerializeToBytes32<2, float>(Offset, RecordingData, { Records[RecordIdx].StepNum, Records[RecordIdx].ActionDimNum }, Records[RecordIdx].ActionData);
+		UE::Learning::Array::SerializeToBytes32<2, float>(Offset, RecordingData, { Records[RecordIdx].StepNum, Records[RecordIdx].ObservationDimNum }, Records[RecordIdx].Blob->ObservationData);
+		UE::Learning::Array::SerializeToBytes32<2, float>(Offset, RecordingData, { Records[RecordIdx].StepNum, Records[RecordIdx].ActionDimNum }, Records[RecordIdx].Blob->ActionData);
 	}
 
 	check(Offset == RecordingData.Num());
@@ -352,7 +353,7 @@ void ULearningAgentsRecording::GetObservationVector(TArray<float>& OutObservatio
 		return;
 	}
 
-	OutObservationVector = MakeArrayView(Records[Record].ObservationData).Slice(Records[Record].ObservationDimNum * Step, Records[Record].ObservationDimNum);
+	OutObservationVector = MakeArrayView(Records[Record].Blob->ObservationData).Slice(Records[Record].ObservationDimNum * Step, Records[Record].ObservationDimNum);
 	OutObservationCompatibilityHash = Records[Record].ObservationCompatibilityHash;
 }
 
@@ -374,7 +375,7 @@ void ULearningAgentsRecording::GetActionVector(TArray<float>& OutActionVector, i
 		return;
 	}
 
-	OutActionVector = MakeArrayView(Records[Record].ActionData).Slice(Records[Record].ActionDimNum * Step, Records[Record].ActionDimNum);
+	OutActionVector = MakeArrayView(Records[Record].Blob->ActionData).Slice(Records[Record].ActionDimNum * Step, Records[Record].ActionDimNum);
 	OutActionCompatibilityHash = Records[Record].ActionCompatibilityHash;
 }
 
diff --git a/Engine/Plugins/Experimental/LearningAgents/Source/LearningAgentsTraining/Public/LearningAgentsRecording.h b/Engine/Plugins/Experimental/LearningAgents/Source/LearningAgentsTraining/Public/LearningAgentsRecording.h
--- a/Engine/Plugins/Experimental/LearningAgents/Source/LearningAgentsTraining/Public/LearningAgentsRecording.h
+++ b/Engine/Plugins/Experimental/LearningAgents/Source/LearningAgentsTraining/Public/LearningAgentsRecording.h
@@ -36,6 +36,21 @@ public:
 	FString ActionSchemaJson;
 };
 
+/** Container for the record's binary data. Kept separate to keep Slate from loading it. */
+UCLASS()
+class ULearningAgentsRecordBlob : public UObject
+{
+	GENERATED_BODY()
+
+public:
+
+	UPROPERTY()
+	TArray<float> ObservationData;
+
+	UPROPERTY()
+	TArray<float> ActionData;
+};
+
 /** A single recording of a series of observations and actions. */
 USTRUCT(BlueprintType)
 struct FLearningAgentsRecord
@@ -76,13 +91,9 @@ public:
 	UPROPERTY(BlueprintReadWrite, EditAnywhere, Category = "LearningAgents")
 	FGameplayTagContainer Tags;
 
-	/** Observation data */
-	UPROPERTY(BlueprintReadOnly, Category = "LearningAgents")
-	TArray<float> ObservationData;
-
-	/** Action data */
-	UPROPERTY(BlueprintReadOnly, Category = "LearningAgents")
-	TArray<float> ActionData;
+	/** The binary data. */
+	UPROPERTY()
+	TObjectPtr<ULearningAgentsRecordBlob> Blob;
 };
 
 /** Data asset representing an array of records. */
