[LearningAgent] Serialize encoding function for action encoding layers

diff --git a/Engine/Plugins/Experimental/LearningAgents/Source/Learning/Public/LearningAction.h b/Engine/Plugins/Experimental/LearningAgents/Source/Learning/Public/LearningAction.h
--- a/Engine/Plugins/Experimental/LearningAgents/Source/Learning/Public/LearningAction.h
+++ b/Engine/Plugins/Experimental/LearningAgents/Source/Learning/Public/LearningAction.h
@@ -161,6 +161,18 @@ namespace UE::Learning::Action
 		GELU = 3,
 	};
 
+	static inline FString GetEncodingActivationAsString(const EEncodingActivationFunction ActivationFunction)
+	{
+		switch (ActivationFunction)
+		{
+		case UE::Learning::Action::EEncodingActivationFunction::ELU: return TEXT("ELU");
+		case UE::Learning::Action::EEncodingActivationFunction::ReLU: return TEXT("ReLU");
+		case UE::Learning::Action::EEncodingActivationFunction::TanH: return TEXT("TanH");
+		case UE::Learning::Action::EEncodingActivationFunction::GELU: return TEXT("GELU");
+		default: checkNoEntry(); return TEXT("");
+		}
+	}
+
 	struct FSchemaEncodingParameters
 	{
 		// The sub-element.
diff --git a/Engine/Plugins/Experimental/LearningAgents/Source/LearningTraining/Private/LearningTrainer.cpp b/Engine/Plugins/Experimental/LearningAgents/Source/LearningTraining/Private/LearningTrainer.cpp
--- a/Engine/Plugins/Experimental/LearningAgents/Source/LearningTraining/Private/LearningTrainer.cpp
+++ b/Engine/Plugins/Experimental/LearningAgents/Source/LearningTraining/Private/LearningTrainer.cpp
@@ -451,6 +451,8 @@ namespace UE::Learning::Trainer
 
 			Object->SetStringField(TEXT("Type"), TEXT("Encoding"));
 			Object->SetNumberField(TEXT("EncodingSize"), Parameters.EncodingSize);
+			Object->SetNumberField(TEXT("EncodingDepth"), Parameters.LayerNum);
+			Object->SetStringField(TEXT("ActivationFunction"), UE::Learning::Action::GetEncodingActivationAsString(Parameters.ActivationFunction));
 			Object->SetObjectField(TEXT("Element"), ConvertActionSchemaToJSON(ActionSchema, Parameters.Element));
 			break;
 		}
@@ -984,7 +986,7 @@ namespace UE::Learning::Trainer
 		}
 		else if (Object->GetStringField(TEXT("Type")) == TEXT("Conv1d"))
 		{
-			const Observation::FSchemaEncodingParameters Parameters = ObservationSchema.GetEncoding(ObservationSchemaElement);
+			const Observation::FSchemaConv1dParameters Parameters = ObservationSchema.GetConv1d(ObservationSchemaElement);
 
 			ComputeObservationSchemaSubsetIndices(NextIndex, Indices, Object->GetObjectField(TEXT("Element")).ToSharedRef(), ObservationSchema, Parameters.Element);
 
@@ -992,7 +994,7 @@ namespace UE::Learning::Trainer
 		}
 		else if (Object->GetStringField(TEXT("Type")) == TEXT("Conv2d"))
 		{
-			const Observation::FSchemaEncodingParameters Parameters = ObservationSchema.GetEncoding(ObservationSchemaElement);
+			const Observation::FSchemaConv2dParameters Parameters = ObservationSchema.GetConv2d(ObservationSchemaElement);
 
 			ComputeObservationSchemaSubsetIndices(NextIndex, Indices, Object->GetObjectField(TEXT("Element")).ToSharedRef(), ObservationSchema, Parameters.Element);
 
@@ -1003,7 +1005,7 @@ namespace UE::Learning::Trainer
 	}
 
 	bool IsActionSchemaSubsetCompatible(const FString& SourceJsonString, const Action::FSchema& ActionSchema, const Action::FSchemaElement& ActionSchemaElement)
-	{
+	{		
 		TSharedPtr<FJsonObject> JsonObject;
 		TSharedRef<TJsonReader<>> Reader = TJsonReaderFactory<>::Create(SourceJsonString);
 
@@ -1227,7 +1229,9 @@ namespace UE::Learning::Trainer
 			const Action::FSchemaEncodingParameters Parameters = ActionSchema.GetEncoding(ActionSchemaElement);
 
 			if (Object->GetStringField(TEXT("Type")) != TEXT("Encoding") ||
-				Object->GetNumberField(TEXT("EncodingSize")) != Parameters.EncodingSize)
+				Object->GetNumberField(TEXT("EncodingSize")) != Parameters.EncodingSize ||
+				Object->GetNumberField(TEXT("EncodingDepth")) != Parameters.LayerNum ||
+				Object->GetStringField(TEXT("ActivationFunction")) != UE::Learning::Action::GetEncodingActivationAsString(Parameters.ActivationFunction))
 			{
 				return false;
 			}
