Learning Agents: Recorder crash fix for nullptr record blob

#rnx

diff --git a/Engine/Plugins/Experimental/LearningAgents/Source/LearningAgentsTraining/Private/LearningAgentsRecorder.cpp b/Engine/Plugins/Experimental/LearningAgents/Source/LearningAgentsTraining/Private/LearningAgentsRecorder.cpp
--- a/Engine/Plugins/Experimental/LearningAgents/Source/LearningAgentsTraining/Private/LearningAgentsRecorder.cpp
+++ b/Engine/Plugins/Experimental/LearningAgents/Source/LearningAgentsTraining/Private/LearningAgentsRecorder.cpp
@@ -249,7 +249,9 @@ void ULearningAgentsRecorder::OnAgentsRemoved_Implementation(const TArray<int32>
 	{
 		if (!RecordBuffers[AgentId].IsEmpty())
 		{
-			RecordBuffers[AgentId].CopyToRecord(Recording->Records.AddDefaulted_GetRef());
+			FLearningAgentsRecord& Record = Recording->Records.AddDefaulted_GetRef();
+			Record.Blob = NewObject<ULearningAgentsRecordBlob>();
+			RecordBuffers[AgentId].CopyToRecord(Record);
 			RecordBuffers[AgentId].Empty();
 		}
 	}
@@ -267,7 +269,9 @@ void ULearningAgentsRecorder::OnAgentsReset_Implementation(const TArray<int32>&
 	{
 		if (!RecordBuffers[AgentId].IsEmpty())
 		{
-			RecordBuffers[AgentId].CopyToRecord(Recording->Records.AddDefaulted_GetRef());
+			FLearningAgentsRecord& Record = Recording->Records.AddDefaulted_GetRef();
+			Record.Blob = NewObject<ULearningAgentsRecordBlob>();
+			RecordBuffers[AgentId].CopyToRecord(Record);
 			RecordBuffers[AgentId].Empty();
 		}
 	}
@@ -335,7 +339,9 @@ void ULearningAgentsRecorder::EndRecording()
 	{
 		if (!RecordBuffers[AgentId].IsEmpty())
 		{
-			RecordBuffers[AgentId].CopyToRecord(Recording->Records.AddDefaulted_GetRef());
+			FLearningAgentsRecord& Record = Recording->Records.AddDefaulted_GetRef();
+			Record.Blob = NewObject<ULearningAgentsRecordBlob>();
+			RecordBuffers[AgentId].CopyToRecord(Record);
 			RecordBuffers[AgentId].Empty();
 		}
 	}
diff --git a/Engine/Plugins/Experimental/LearningAgents/Source/LearningAgentsTraining/Private/LearningAgentsRecording.cpp b/Engine/Plugins/Experimental/LearningAgents/Source/LearningAgentsTraining/Private/LearningAgentsRecording.cpp
--- a/Engine/Plugins/Experimental/LearningAgents/Source/LearningAgentsTraining/Private/LearningAgentsRecording.cpp
+++ b/Engine/Plugins/Experimental/LearningAgents/Source/LearningAgentsTraining/Private/LearningAgentsRecording.cpp
@@ -180,7 +180,7 @@ void ULearningAgentsRecording::LoadRecordingFromFile(const FFilePath& File, cons
 		for (int32 RecordIdx = 0; RecordIdx < RecordNum; RecordIdx++)
 		{
 			FLearningAgentsRecord Record;
-			Record.Blob = NewObject<ULearningAgentsRecordBlob>(this);
+			Record.Blob = NewObject<ULearningAgentsRecordBlob>();
 
 			UE::Learning::DeserializeFromBytes(Offset, RecordingData, Record.StepNum);
 			UE::Learning::DeserializeFromBytes(Offset, RecordingData, Record.ObservationDimNum);
