[LearningAgents] Trivial adding gelu to encoding activation option

diff --git a/Engine/Plugins/Experimental/LearningAgents/Source/LearningAgents/Private/LearningAgentsActions.cpp b/Engine/Plugins/Experimental/LearningAgents/Source/LearningAgents/Private/LearningAgentsActions.cpp
--- a/Engine/Plugins/Experimental/LearningAgents/Source/LearningAgents/Private/LearningAgentsActions.cpp
+++ b/Engine/Plugins/Experimental/LearningAgents/Source/LearningAgents/Private/LearningAgentsActions.cpp
@@ -993,6 +993,7 @@ namespace UE::Learning::Agents::Action::Private
 		case ELearningAgentsActivationFunction::ReLU: return Learning::Action::EEncodingActivationFunction::ReLU;
 		case ELearningAgentsActivationFunction::ELU: return Learning::Action::EEncodingActivationFunction::ELU;
 		case ELearningAgentsActivationFunction::TanH: return Learning::Action::EEncodingActivationFunction::TanH;
+		case ELearningAgentsActivationFunction::GELU: return Learning::Action::EEncodingActivationFunction::GELU;
 		default: checkNoEntry(); return Learning::Action::EEncodingActivationFunction::ReLU;
 		}
 	}
