LearningAgents: Added ability to manually set generation in Observation and Action structures.

diff --git a/Engine/Plugins/Experimental/LearningAgents/Source/Learning/Private/LearningAction.cpp b/Engine/Plugins/Experimental/LearningAgents/Source/Learning/Private/LearningAction.cpp
--- a/Engine/Plugins/Experimental/LearningAgents/Source/Learning/Private/LearningAction.cpp
+++ b/Engine/Plugins/Experimental/LearningAgents/Source/Learning/Private/LearningAction.cpp
@@ -556,6 +556,11 @@ namespace UE::Learning::Action
 		return Parameters;
 	}
 
+	void FSchema::SetGeneration(const uint32 InGeneration)
+	{
+		Generation = InGeneration;
+	}
+
 	uint32 FSchema::GetGeneration() const
 	{
 		return Generation;
@@ -961,6 +966,11 @@ namespace UE::Learning::Action
 		return Parameters;
 	}
 
+	void FObject::SetGeneration(const uint32 InGeneration)
+	{
+		Generation = InGeneration;
+	}
+
 	uint32 FObject::GetGeneration() const
 	{
 		return Generation;
@@ -1396,6 +1406,11 @@ namespace UE::Learning::Action
 		return Parameters;
 	}
 
+	void FModifier::SetGeneration(const uint32 InGeneration)
+	{
+		Generation = InGeneration;
+	}
+
 	uint32 FModifier::GetGeneration() const
 	{
 		return Generation;
diff --git a/Engine/Plugins/Experimental/LearningAgents/Source/Learning/Private/LearningObservation.cpp b/Engine/Plugins/Experimental/LearningAgents/Source/Learning/Private/LearningObservation.cpp
--- a/Engine/Plugins/Experimental/LearningAgents/Source/Learning/Private/LearningObservation.cpp
+++ b/Engine/Plugins/Experimental/LearningAgents/Source/Learning/Private/LearningObservation.cpp
@@ -576,6 +576,11 @@ namespace UE::Learning::Observation
 		return Parameters;
 	}
 
+	void FSchema::SetGeneration(const uint32 InGeneration)
+	{
+		Generation = InGeneration;
+	}
+
 	uint32 FSchema::GetGeneration() const
 	{
 		return Generation;
@@ -1079,6 +1084,11 @@ namespace UE::Learning::Observation
 		return Parameters;
 	}
 
+	void FObject::SetGeneration(const uint32 InGeneration)
+	{
+		Generation = InGeneration;
+	}
+
 	uint32 FObject::GetGeneration() const
 	{
 		return Generation;
diff --git a/Engine/Plugins/Experimental/LearningAgents/Source/Learning/Public/LearningAction.h b/Engine/Plugins/Experimental/LearningAgents/Source/Learning/Public/LearningAction.h
--- a/Engine/Plugins/Experimental/LearningAgents/Source/Learning/Public/LearningAction.h
+++ b/Engine/Plugins/Experimental/LearningAgents/Source/Learning/Public/LearningAction.h
@@ -241,6 +241,9 @@ namespace UE::Learning::Action
 		// Get the action modifier vector size of the given element
 		UE_API int32 GetActionModifierVectorSize(const FSchemaElement Element) const;
 
+		// Set the current generation
+		UE_API void SetGeneration(const uint32 InGeneration);
+
 		// Get the current generation
 		UE_API uint32 GetGeneration() const;
 		
@@ -470,6 +473,9 @@ namespace UE::Learning::Action
 		// Gets the tag of the given element
 		UE_API FName GetTag(const FObjectElement Element) const;
 
+		// Set the current generation
+		UE_API void SetGeneration(const uint32 InGeneration);
+
 		// Get the current generation
 		UE_API uint32 GetGeneration() const;
 
@@ -630,6 +636,9 @@ namespace UE::Learning::Action
 		// Gets the tag of the given element
 		UE_API FName GetTag(const FModifierElement Element) const;
 
+		// Set the current generation
+		UE_API void SetGeneration(const uint32 InGeneration);
+
 		// Get the current generation
 		UE_API uint32 GetGeneration() const;
 
diff --git a/Engine/Plugins/Experimental/LearningAgents/Source/Learning/Public/LearningObservation.h b/Engine/Plugins/Experimental/LearningAgents/Source/Learning/Public/LearningObservation.h
--- a/Engine/Plugins/Experimental/LearningAgents/Source/Learning/Public/LearningObservation.h
+++ b/Engine/Plugins/Experimental/LearningAgents/Source/Learning/Public/LearningObservation.h
@@ -320,6 +320,9 @@ namespace UE::Learning::Observation
 		// Get the observation distribution vector size of the given element
 		UE_API int32 GetObservationDistributionVectorSize(const FSchemaElement Element) const;
 
+		// Set the current generation
+		UE_API void SetGeneration(const uint32 InGeneration);
+
 		// Get the current generation
 		UE_API uint32 GetGeneration() const;
 
@@ -605,6 +608,9 @@ namespace UE::Learning::Observation
 		// Gets the tag of the given element
 		UE_API FName GetTag(const FObjectElement Element) const;
 
+		// Set the current generation
+		UE_API void SetGeneration(const uint32 InGeneration);
+
 		// Get the current generation
 		UE_API uint32 GetGeneration() const;
 
