LearningAgents: Added name duplication error-out.

diff --git a/Engine/Plugins/Experimental/LearningAgents/Source/LearningAgents/Private/LearningAgentsActions.cpp b/Engine/Plugins/Experimental/LearningAgents/Source/LearningAgents/Private/LearningAgentsActions.cpp
--- a/Engine/Plugins/Experimental/LearningAgents/Source/LearningAgents/Private/LearningAgentsActions.cpp
+++ b/Engine/Plugins/Experimental/LearningAgents/Source/LearningAgents/Private/LearningAgentsActions.cpp
@@ -1194,6 +1194,12 @@ FLearningAgentsActionSchemaElement ULearningAgentsActions::SpecifyNamedExclusive
 		return FLearningAgentsActionSchemaElement();
 	}
 
+	if (UE::Learning::Agents::Action::Private::ContainsDuplicates(Names))
+	{
+		UE_LOG(LogLearning, Error, TEXT("SpecifyNamedExclusiveDiscreteActionFromArrayViews: Names contain duplicates."));
+		return FLearningAgentsActionSchemaElement();
+	}
+
 	TArray<float, TInlineAllocator<16>> NormalizedPriorProbabilities;
 	NormalizedPriorProbabilities.Init(NameNum == 0 ? 0.0f : 1.0f / NameNum, NameNum);
 	for (int32 Idx = 0; Idx < PriorProbabilities.Num(); Idx++)
@@ -1310,6 +1316,12 @@ FLearningAgentsActionSchemaElement ULearningAgentsActions::SpecifyNamedInclusive
 		UE_LOG(LogLearning, Warning, TEXT("SpecifyNamedInclusiveDiscreteActionFromArrayViews: Specifying zero-sized Named Inclusive Discrete Action."));
 	}
 
+	if (UE::Learning::Agents::Action::Private::ContainsDuplicates(Names))
+	{
+		UE_LOG(LogLearning, Error, TEXT("SpecifyNamedInclusiveDiscreteActionFromArrayViews: Names contain duplicates."));
+		return FLearningAgentsActionSchemaElement();
+	}
+
 	TArray<float, TInlineAllocator<16>> NormalizedPriorProbabilities;
 	NormalizedPriorProbabilities.Init(NameNum == 0 ? 0.0f : 1.0f / NameNum, NameNum);
 	for (int32 Idx = 0; Idx < PriorProbabilities.Num(); Idx++)
diff --git a/Engine/Plugins/Experimental/LearningAgents/Source/LearningAgents/Private/LearningAgentsObservations.cpp b/Engine/Plugins/Experimental/LearningAgents/Source/LearningAgents/Private/LearningAgentsObservations.cpp
--- a/Engine/Plugins/Experimental/LearningAgents/Source/LearningAgents/Private/LearningAgentsObservations.cpp
+++ b/Engine/Plugins/Experimental/LearningAgents/Source/LearningAgents/Private/LearningAgentsObservations.cpp
@@ -761,6 +761,12 @@ FLearningAgentsObservationSchemaElement ULearningAgentsObservations::SpecifyName
 		return FLearningAgentsObservationSchemaElement();
 	}
 
+	if (UE::Learning::Agents::Observation::Private::ContainsDuplicates(ElementNames))
+	{
+		UE_LOG(LogLearning, Error, TEXT("SpecifyNamedExclusiveDiscreteObservationFromArrayView: Element Names contain duplicates."));
+		return FLearningAgentsObservationSchemaElement();
+	}
+
 	return { Schema->ObservationSchema.CreateNamedDiscreteExclusive({ ElementNames }, Tag) };
 }
 
@@ -782,6 +788,12 @@ FLearningAgentsObservationSchemaElement ULearningAgentsObservations::SpecifyName
 		UE_LOG(LogLearning, Warning, TEXT("SpecifyNamedInclusiveDiscreteObservationFromArrayView: Specifying zero-sized Named Inclusive Discrete Observation."));
 	}
 
+	if (UE::Learning::Agents::Observation::Private::ContainsDuplicates(ElementNames))
+	{
+		UE_LOG(LogLearning, Error, TEXT("SpecifyNamedInclusiveDiscreteObservationFromArrayView: Element Names contain duplicates."));
+		return FLearningAgentsObservationSchemaElement();
+	}
+
 	return { Schema->ObservationSchema.CreateNamedDiscreteInclusive({ ElementNames }, Tag) };
 }
 
