LearningAgents: Added some additional frame attribute functions.

diff --git a/Engine/Plugins/Experimental/LearningAgents/Source/Learning/Private/LearningFrameAttribute.cpp b/Engine/Plugins/Experimental/LearningAgents/Source/Learning/Private/LearningFrameAttribute.cpp
--- a/Engine/Plugins/Experimental/LearningAgents/Source/Learning/Private/LearningFrameAttribute.cpp
+++ b/Engine/Plugins/Experimental/LearningAgents/Source/Learning/Private/LearningFrameAttribute.cpp
@@ -202,6 +202,24 @@ namespace UE::Learning
 			}
 		}
 
+		static inline void RepeatFirstRangeEntryInplace(
+			TLearningArrayView<1, float> InOut,
+			const TLearningArrayView<1, const int32> RangeOffsets,
+			const TLearningArrayView<1, const int32> RangeLengths)
+		{
+			const int32 RangeNum = RangeOffsets.Num();
+
+			for (int32 RangeIdx = 0; RangeIdx < RangeNum; RangeIdx++)
+			{
+				const int32 RangeOffset = RangeOffsets[RangeIdx];
+				const int32 RangeLength = RangeLengths[RangeIdx];
+				for (int32 ValueIdx = RangeOffset; ValueIdx < RangeOffset + RangeLength; ValueIdx++)
+				{
+					InOut[ValueIdx] = InOut[RangeOffset];
+				}
+			}
+		}
+
 		static inline void InvInplace(TLearningArrayView<1, float> InOut)
 		{
 			const int32 ValueNum = InOut.Num();
@@ -2130,6 +2148,18 @@ namespace UE::Learning
 				});
 		}
 
+		void Copy(FFrameAttribute& Out, const FFrameAttribute& In)
+		{
+			UnaryOp(Out, In.GetChannelNum(), In, [](
+				FFrameAttribute& Out,
+				const FFrameAttribute& In,
+				const TLearningArrayView<1, const int32> RangeOffsets,
+				const TLearningArrayView<1, const int32> RangeLengths) {
+
+					Array::Copy(Out.AttributeData, In.AttributeData);
+				});
+		}
+
 		void Neg(FFrameAttribute& Out, const FFrameAttribute& In)
 		{
 			UnaryOp(Out, In.GetChannelNum(), In, [](
@@ -2261,6 +2291,25 @@ namespace UE::Learning
 				});
 		}
 
+		void RepeatFirstRangeEntryInplace(FFrameAttribute& InOut)
+		{
+			InplaceOp(InOut, [](
+				FFrameAttribute& InOut,
+				const TLearningArrayView<1, const int32> RangeOffsets,
+				const TLearningArrayView<1, const int32> RangeLengths) {
+
+					const int32 ChannelNum = InOut.GetChannelNum();
+
+					for (int32 ChannelIdx = 0; ChannelIdx < ChannelNum; ChannelIdx++)
+					{
+						Private::RepeatFirstRangeEntryInplace(
+							InOut.GetChannelAttributeData(ChannelIdx),
+							RangeOffsets,
+							RangeLengths);
+					}
+				});
+		}
+
 		void InvInplace(FFrameAttribute& InOut)
 		{
 			InplaceOp(InOut, [](
diff --git a/Engine/Plugins/Experimental/LearningAgents/Source/Learning/Public/LearningFrameAttribute.h b/Engine/Plugins/Experimental/LearningAgents/Source/Learning/Public/LearningFrameAttribute.h
--- a/Engine/Plugins/Experimental/LearningAgents/Source/Learning/Public/LearningFrameAttribute.h
+++ b/Engine/Plugins/Experimental/LearningAgents/Source/Learning/Public/LearningFrameAttribute.h
@@ -226,6 +226,8 @@ namespace UE::Learning
 		/** Compute the dot product of two frame attributes. Channel numbers must match. */
 		UE_API void Dot(FFrameAttribute& Out, const FFrameAttribute& Lhs, const FFrameAttribute& Rhs);
 
+		/** Copy a frame attribute. */
+		UE_API void Copy(FFrameAttribute& Out, const FFrameAttribute& In);
 
 		/** Negate a frame attribute. */
 		UE_API void Neg(FFrameAttribute& Out, const FFrameAttribute& In);
@@ -267,6 +269,9 @@ namespace UE::Learning
 		/** Compute the approximate acos of a frame attribute in place. */
 		UE_API void ApproxAcosInplace(FFrameAttribute& InOut);
 
+		/** Repeats the first entry in the range across the rest of the range */
+		UE_API void RepeatFirstRangeEntryInplace(FFrameAttribute& InOut);
+
 
 		/** Compute the length of a frame attribute over the channels dimension. */
 		UE_API void Length(FFrameAttribute& Out, const FFrameAttribute& In);
