LearningAgents: Remove firing OnAgentsAdded and OnAgentsRemoved if they are empty

diff --git a/Engine/Plugins/Experimental/LearningAgents/Source/LearningAgents/Private/LearningAgentsManager.cpp b/Engine/Plugins/Experimental/LearningAgents/Source/LearningAgents/Private/LearningAgentsManager.cpp
--- a/Engine/Plugins/Experimental/LearningAgents/Source/LearningAgents/Private/LearningAgentsManager.cpp
+++ b/Engine/Plugins/Experimental/LearningAgents/Source/LearningAgents/Private/LearningAgentsManager.cpp
@@ -486,13 +486,22 @@ void ULearningAgentsManager::AddListener(ULearningAgentsManagerListener* Listene
 {
 	Listeners.Add(Listener);
 	OnEventAgentIds = OccupiedAgentIds;
-	Listener->OnAgentsAdded(OnEventAgentIds);
+	
+	if (!OnEventAgentIds.IsEmpty())
+	{
+		Listener->OnAgentsAdded(OnEventAgentIds);
+	}
 }
 
 void ULearningAgentsManager::RemoveListener(ULearningAgentsManagerListener* Listener)
 {
 	OnEventAgentIds = OccupiedAgentIds;
-	Listener->OnAgentsRemoved(OnEventAgentIds);
+	
+	if (!OnEventAgentIds.IsEmpty())
+	{
+		Listener->OnAgentsRemoved(OnEventAgentIds);
+	}
+
 	Listeners.Remove(Listener);
 }
 
