LearningAgents: Fixed bug in array serialization

#Synced-CL 48363837

diff --git a/Engine/Plugins/Experimental/LearningAgents/Source/Learning/Public/LearningArray.h b/Engine/Plugins/Experimental/LearningAgents/Source/Learning/Public/LearningArray.h
--- a/Engine/Plugins/Experimental/LearningAgents/Source/Learning/Public/LearningArray.h
+++ b/Engine/Plugins/Experimental/LearningAgents/Source/Learning/Public/LearningArray.h
@@ -711,6 +711,22 @@ namespace UE::Learning
 			}
 		}
 
+		/**
+		* Check if the shapes of two arrays are equal
+		*/
+		template<uint8 InDimNum>
+		inline bool ShapeEqual(
+			const TLearningArrayShape<InDimNum>& Lhs,
+			const TLearningArrayShape<InDimNum>& Rhs)
+		{
+			for (uint32 Idx = 0; Idx < InDimNum; Idx++)
+			{
+				if (Lhs[Idx] != Rhs[Idx]) { return false; }
+			}
+
+			return true;
+		}
+
 		/**
 		* Check if the contents of one array equals the contents of another.
 		*/
@@ -884,7 +900,7 @@ namespace UE::Learning
 			{
 				for (uint8 ShapeIdx = 0; ShapeIdx < InDimNum; ShapeIdx++)
 				{
-					FDefaultAllocator::SizeType Num = Array.Num(ShapeIdx);
+					int64 Num = Array.Num(ShapeIdx);
 					Ar << Num;
 				}
 				Ar.Serialize(Array.GetData(), Array.Num() * sizeof(InElementType));
