LearningAgents: Added Frame Attribute Equality function

#Synced-CL 48363837

diff --git a/Engine/Plugins/Experimental/LearningAgents/Source/Learning/Private/LearningFrameAttribute.cpp b/Engine/Plugins/Experimental/LearningAgents/Source/Learning/Private/LearningFrameAttribute.cpp
--- a/Engine/Plugins/Experimental/LearningAgents/Source/Learning/Private/LearningFrameAttribute.cpp
+++ b/Engine/Plugins/Experimental/LearningAgents/Source/Learning/Private/LearningFrameAttribute.cpp
@@ -1719,6 +1719,11 @@ namespace UE::Learning
 
 	namespace FrameAttribute
 	{
+		bool Equal(const FFrameAttribute& Lhs, const FFrameAttribute& Rhs)
+		{
+			return FrameRangeSet::Equal(Lhs.FrameRangeSet, Rhs.FrameRangeSet) && Array::Equal(Lhs.AttributeData, Rhs.AttributeData);
+		}
+
 		void Intersection(FFrameAttribute& OutFrameAttribute, const FFrameAttribute& FrameAttribute, const FFrameRangeSet& FrameRangeSet)
 		{
 			if (FrameRangeSet::Equal(FrameAttribute.FrameRangeSet, FrameRangeSet))
diff --git a/Engine/Plugins/Experimental/LearningAgents/Source/Learning/Public/LearningFrameAttribute.h b/Engine/Plugins/Experimental/LearningAgents/Source/Learning/Public/LearningFrameAttribute.h
--- a/Engine/Plugins/Experimental/LearningAgents/Source/Learning/Public/LearningFrameAttribute.h
+++ b/Engine/Plugins/Experimental/LearningAgents/Source/Learning/Public/LearningFrameAttribute.h
@@ -137,6 +137,8 @@ namespace UE::Learning
 			const TArrayView<const TLearningArrayView<1, const int32>> InRangeOffsets,
 			const TLearningArrayView<1, const int32> RangeLengths)>;
 
+		/** Checks if two frame attributes are equal */
+		UE_API bool Equal(const FFrameAttribute& Lhs, const FFrameAttribute& Rhs);
 
 		/** Computes the intersection of a frame attribute and frame range set */
 		UE_API void Intersection(FFrameAttribute& OutFrameAttribute, const FFrameAttribute& FrameAttribute, const FFrameRangeSet& FrameRangeSet);
