LearningAgents: Added serialization to Observation and Action structures.

#Synced-CL 48452931

diff --git a/Engine/Plugins/Experimental/LearningAgents/Source/Learning/Private/LearningAction.cpp b/Engine/Plugins/Experimental/LearningAgents/Source/Learning/Private/LearningAction.cpp
--- a/Engine/Plugins/Experimental/LearningAgents/Source/Learning/Private/LearningAction.cpp
+++ b/Engine/Plugins/Experimental/LearningAgents/Source/Learning/Private/LearningAction.cpp
@@ -10,6 +10,24 @@ namespace UE::Learning::Action
 {
 	namespace Private
 	{
+		template<typename T>
+		static inline void ArraySerializeRaw(FArchive& Ar, TArray<T>& Array)
+		{
+			if (Ar.IsSaving())
+			{
+				int32 Num = Array.Num();
+				Ar << Num;
+				Ar.Serialize(Array.GetData(), Array.Num() * sizeof(T));
+			}
+			else if (Ar.IsLoading())
+			{
+				int32 Num = 0;
+				Ar << Num;
+				Array.SetNumUninitialized(Num);
+				Ar.Serialize(Array.GetData(), Array.Num() * sizeof(T));
+			}
+		}
+
 		static inline bool ContainsDuplicates(const TArrayView<const FName> ElementNames)
 		{
 			TSet<FName, DefaultKeyFuncs<FName>, TInlineSetAllocator<32>> ElementNameSet;
@@ -573,6 +591,7 @@ namespace UE::Learning::Action
 		EncodedVectorSizes.Empty();
 		ActionVectorSizes.Empty();
 		ActionDistributionVectorSizes.Empty();
+		ActionModifierVectorSizes.Empty();
 		TypeDataIndices.Empty();
 
 		ContinuousData.Empty();
@@ -605,6 +624,7 @@ namespace UE::Learning::Action
 		EncodedVectorSizes.Reset();
 		ActionVectorSizes.Reset();
 		ActionDistributionVectorSizes.Reset();
+		ActionModifierVectorSizes.Reset();
 		TypeDataIndices.Reset();
 
 		ContinuousData.Reset();
@@ -625,6 +645,58 @@ namespace UE::Learning::Action
 		Generation++;
 	}
 
+	bool FSchema::Serialize(FArchive& Ar)
+	{
+		static constexpr int32 Magic = 0x25344a5c;
+		static constexpr int32 Version = 0;
+
+		if (Ar.IsLoading())
+		{
+			int64 Offset = Ar.Tell();
+			int32 LoadMagic = 0;
+			int32 LoadVersion = 0;
+			Ar << LoadMagic;
+			if (!ensure(LoadMagic == Magic)) { Ar.Seek(Offset); Empty(); return false; }
+			Ar << LoadVersion;
+			if (!ensure(LoadVersion == Version)) { Ar.Seek(Offset); Empty(); return false; }
+		}
+		else
+		{
+			int32 SaveMagic = Magic;
+			int32 SaveVersion = Version;
+			Ar << SaveMagic;
+			Ar << SaveVersion;
+		}
+
+		Ar << Generation;
+
+		Ar << Types;
+		Ar << Tags;
+		Ar << EncodedVectorSizes;
+		Ar << ActionVectorSizes;
+		Ar << ActionDistributionVectorSizes;
+		Ar << ActionModifierVectorSizes;
+		Ar << TypeDataIndices;
+
+		Private::ArraySerializeRaw(Ar, ContinuousData);
+		Private::ArraySerializeRaw(Ar, DiscreteExclusiveData);
+		Private::ArraySerializeRaw(Ar, DiscreteInclusiveData);
+		Private::ArraySerializeRaw(Ar, NamedDiscreteExclusiveData);
+		Private::ArraySerializeRaw(Ar, NamedDiscreteInclusiveData);
+		Private::ArraySerializeRaw(Ar, AndData);
+		Private::ArraySerializeRaw(Ar, OrExclusiveData);
+		Private::ArraySerializeRaw(Ar, OrInclusiveData);
+		Private::ArraySerializeRaw(Ar, ArrayData);
+		Private::ArraySerializeRaw(Ar, EncodingData);
+
+		Ar << SubElementNames;
+		Ar << SubElementObjects;
+
+		Ar << PriorProbabilities;
+
+		return true;
+	}
+
 	FObjectElement FObject::CreateNull(const FName Tag)
 	{
 		const int32 Index = Types.Add(EType::Null);
@@ -1019,6 +1091,50 @@ namespace UE::Learning::Action
 		Generation++;
 	}
 
+	bool FObject::Serialize(FArchive& Ar)
+	{
+		static constexpr int32 Magic = 0x535a63fb;
+		static constexpr int32 Version = 0;
+
+		if (Ar.IsLoading())
+		{
+			int64 Offset = Ar.Tell();
+			int32 LoadMagic = 0;
+			int32 LoadVersion = 0;
+			Ar << LoadMagic;
+			if (!ensure(LoadMagic == Magic)) { Ar.Seek(Offset); Empty(); return false; }
+			Ar << LoadVersion;
+			if (!ensure(LoadVersion == Version)) { Ar.Seek(Offset); Empty(); return false; }
+		}
+		else
+		{
+			int32 SaveMagic = Magic;
+			int32 SaveVersion = Version;
+			Ar << SaveMagic;
+			Ar << SaveVersion;
+		}
+
+		Ar << Generation;
+
+		Ar << Types;
+		Ar << Tags;
+
+		Ar << ContinuousDataOffsets;
+		Ar << ContinuousDataNums;
+
+		Ar << DiscreteDataOffsets;
+		Ar << DiscreteDataNums;
+
+		Ar << ElementDataOffsets;
+		Ar << ElementDataNums;
+
+		Ar << ContinuousValues;
+		Ar << DiscreteValues;
+		Ar << SubElementObjects;
+		Ar << SubElementNames;
+
+		return true;
+	}
 
 	FModifierElement FModifier::CreateNull(const FName Tag)
 	{
@@ -1467,6 +1583,51 @@ namespace UE::Learning::Action
 		Generation++;
 	}
 
+	bool FModifier::Serialize(FArchive& Ar)
+	{
+		static constexpr int32 Magic = 0x5342479b;
+		static constexpr int32 Version = 0;
+
+		if (Ar.IsLoading())
+		{
+			int64 Offset = Ar.Tell();
+			int32 LoadMagic = 0;
+			int32 LoadVersion = 0;
+			Ar << LoadMagic;
+			if (!ensure(LoadMagic == Magic)) { Ar.Seek(Offset); Empty(); return false; }
+			Ar << LoadVersion;
+			if (!ensure(LoadVersion == Version)) { Ar.Seek(Offset); Empty(); return false; }
+		}
+		else
+		{
+			int32 SaveMagic = Magic;
+			int32 SaveVersion = Version;
+			Ar << SaveMagic;
+			Ar << SaveVersion;
+		}
+
+		Ar << Generation;
+
+		Ar << Types;
+		Ar << Tags;
+		Ar << ContinuousDataOffsets;
+		Ar << ContinuousDataNums;
+		Ar << DiscreteDataOffsets;
+		Ar << DiscreteDataNums;
+		Ar << ElementDataOffsets;
+		Ar << ElementDataNums;
+		Ar << MaskedDataOffsets;
+		Ar << MaskedDataNums;
+
+		Ar << ContinuousMaskeds;
+		Ar << ContinuousMaskedValues;
+		Ar << DiscreteValues;
+		Ar << SubElementModifiers;
+		Ar << SubElementNames;
+		Ar << MaskedElementNames;
+
+		return true;
+	}
 
 	namespace Private
 	{
@@ -3620,4 +3781,55 @@ namespace UE::Learning::Action
 		}
 	}
 
+}
+
+FArchive& operator<<(FArchive& Ar, UE::Learning::Action::FSchemaElement& SchemaElement)
+{
+	Ar << SchemaElement.Generation;
+	Ar << SchemaElement.Index;
+	return Ar;
+}
+
+FArchive& operator<<(FArchive& Ar, UE::Learning::Action::FObjectElement& ObjectElement)
+{
+	Ar << ObjectElement.Generation;
+	Ar << ObjectElement.Index;
+	return Ar;
+}
+
+FArchive& operator<<(FArchive& Ar, UE::Learning::Action::FModifierElement& ModifierElement)
+{
+	Ar << ModifierElement.Generation;
+	Ar << ModifierElement.Index;
+	return Ar;
+}
+
+bool operator==(const UE::Learning::Action::FSchemaElement& Lhs, const UE::Learning::Action::FSchemaElement& Rhs)
+{
+	return Lhs.Generation == Rhs.Generation && Lhs.Index == Rhs.Index;
+}
+
+bool operator==(const UE::Learning::Action::FObjectElement& Lhs, const UE::Learning::Action::FObjectElement& Rhs)
+{
+	return Lhs.Generation == Rhs.Generation && Lhs.Index == Rhs.Index;
+}
+
+bool operator==(const UE::Learning::Action::FModifierElement& Lhs, const UE::Learning::Action::FModifierElement& Rhs)
+{
+	return Lhs.Generation == Rhs.Generation && Lhs.Index == Rhs.Index;
+}
+
+uint32 GetTypeHash(const UE::Learning::Action::FSchemaElement& SchemaElement)
+{
+	return GetTypeHash(TPair<int32, int32>{ SchemaElement.Generation, SchemaElement.Index });
+}
+
+uint32 GetTypeHash(const UE::Learning::Action::FObjectElement& ObjectElement)
+{
+	return GetTypeHash(TPair<int32, int32>{ ObjectElement.Generation, ObjectElement.Index });
+}
+
+uint32 GetTypeHash(const UE::Learning::Action::FModifierElement& ModifierElement)
+{
+	return GetTypeHash(TPair<int32, int32>{ ModifierElement.Generation, ModifierElement.Index });
 }
\ No newline at end of file
diff --git a/Engine/Plugins/Experimental/LearningAgents/Source/Learning/Private/LearningObservation.cpp b/Engine/Plugins/Experimental/LearningAgents/Source/Learning/Private/LearningObservation.cpp
--- a/Engine/Plugins/Experimental/LearningAgents/Source/Learning/Private/LearningObservation.cpp
+++ b/Engine/Plugins/Experimental/LearningAgents/Source/Learning/Private/LearningObservation.cpp
@@ -9,6 +9,24 @@ namespace UE::Learning::Observation
 {
 	namespace Private
 	{
+		template<typename T>
+		static inline void ArraySerializeRaw(FArchive& Ar, TArray<T>& Array)
+		{
+			if (Ar.IsSaving())
+			{
+				int32 Num = Array.Num();
+				Ar << Num;
+				Ar.Serialize(Array.GetData(), Array.Num() * sizeof(T));
+			}
+			else if (Ar.IsLoading())
+			{
+				int32 Num = 0;
+				Ar << Num;
+				Array.SetNumUninitialized(Num);
+				Ar.Serialize(Array.GetData(), Array.Num() * sizeof(T));
+			}
+		}
+
 		static inline bool ContainsDuplicates(const TArrayView<const FName> ElementNames)
 		{
 			TSet<FName, DefaultKeyFuncs<FName>, TInlineSetAllocator<32>> ElementNameSet;
@@ -606,6 +624,8 @@ namespace UE::Learning::Observation
 		ArrayData.Empty();
 		SetData.Empty();
 		EncodingData.Empty();
+		Conv1dData.Empty();
+		Conv2dData.Empty();
 
 		SubElementNames.Empty();
 		SubElementObjects.Empty();
@@ -638,6 +658,8 @@ namespace UE::Learning::Observation
 		ArrayData.Reset();
 		SetData.Reset();
 		EncodingData.Reset();
+		Conv1dData.Reset();
+		Conv2dData.Reset();
 
 		SubElementNames.Reset();
 		SubElementObjects.Reset();
@@ -645,6 +667,58 @@ namespace UE::Learning::Observation
 		Generation++;
 	}
 
+	bool FSchema::Serialize(FArchive& Ar)
+	{
+		static constexpr int32 Magic = 0x7db41964;
+		static constexpr int32 Version = 0;
+
+		if (Ar.IsLoading())
+		{
+			int64 Offset = Ar.Tell();
+			int32 LoadMagic = 0;
+			int32 LoadVersion = 0;
+			Ar << LoadMagic;
+			if (!ensure(LoadMagic == Magic)) { Ar.Seek(Offset); Empty(); return false; }
+			Ar << LoadVersion;
+			if (!ensure(LoadVersion == Version)) { Ar.Seek(Offset); Empty(); return false; }
+		}
+		else
+		{
+			int32 SaveMagic = Magic;
+			int32 SaveVersion = Version;
+			Ar << SaveMagic;
+			Ar << SaveVersion;
+		}
+
+		Ar << Generation;
+
+		Ar << Types;
+		Ar << Tags;
+		Ar << ObservationVectorSizes;
+		Ar << EncodedVectorSizes;
+		Ar << ObservationDistributionVectorSizes;
+		Ar << TypeDataIndices;
+
+		Private::ArraySerializeRaw(Ar, ContinuousData);
+		Private::ArraySerializeRaw(Ar, DiscreteExclusiveData);
+		Private::ArraySerializeRaw(Ar, DiscreteInclusiveData);
+		Private::ArraySerializeRaw(Ar, NamedDiscreteExclusiveData);
+		Private::ArraySerializeRaw(Ar, NamedDiscreteInclusiveData);
+		Private::ArraySerializeRaw(Ar, AndData);
+		Private::ArraySerializeRaw(Ar, OrExclusiveData);
+		Private::ArraySerializeRaw(Ar, OrInclusiveData);
+		Private::ArraySerializeRaw(Ar, ArrayData);
+		Private::ArraySerializeRaw(Ar, SetData);
+		Private::ArraySerializeRaw(Ar, EncodingData);
+		Private::ArraySerializeRaw(Ar, Conv1dData);
+		Private::ArraySerializeRaw(Ar, Conv2dData);
+
+		Ar << SubElementNames;
+		Ar << SubElementObjects;
+
+		return true;
+	}
+
 	FObjectElement FObject::CreateNull(const FName Tag)
 	{
 		const int32 Index = Types.Add(EType::Null);
@@ -1102,8 +1176,8 @@ namespace UE::Learning::Observation
 		ContinuousDataOffsets.Empty();
 		ContinuousDataNums.Empty();
 
-		DiscreteDataOffsets.Reset();
-		DiscreteDataNums.Reset();
+		DiscreteDataOffsets.Empty();
+		DiscreteDataNums.Empty();
 
 		SubElementDataOffsets.Empty();
 		SubElementDataNums.Empty();
@@ -1143,6 +1217,51 @@ namespace UE::Learning::Observation
 		Generation++;
 	}
 
+	bool FObject::Serialize(FArchive& Ar)
+	{
+		static constexpr int32 Magic = 0xf71282d9;
+		static constexpr int32 Version = 0;
+
+		if (Ar.IsLoading())
+		{
+			int64 Offset = Ar.Tell();
+			int32 LoadMagic = 0;
+			int32 LoadVersion = 0;
+			Ar << LoadMagic;
+			if (!ensure(LoadMagic == Magic)) { Ar.Seek(Offset); Empty(); return false; }
+			Ar << LoadVersion;
+			if (!ensure(LoadVersion == Version)) { Ar.Seek(Offset); Empty(); return false; }
+		}
+		else
+		{
+			int32 SaveMagic = Magic;
+			int32 SaveVersion = Version;
+			Ar << SaveMagic;
+			Ar << SaveVersion;
+		}
+
+		Ar << Generation;
+
+		Ar << Types;
+		Ar << Tags;
+
+		Ar << ContinuousDataOffsets;
+		Ar << ContinuousDataNums;
+
+		Ar << DiscreteDataOffsets;
+		Ar << DiscreteDataNums;
+
+		Ar << SubElementDataOffsets;
+		Ar << SubElementDataNums;
+
+		Ar << ContinuousValues;
+		Ar << DiscreteValues;
+		Ar << SubElementNames;
+		Ar << SubElementObjects;
+
+		return true;
+	}
+
 	namespace Private
 	{
 		static inline NNE::RuntimeBasic::FModelBuilder::EActivationFunction GetNNEActivationFunction(const EEncodingActivationFunction ActivationFunction)
@@ -3236,4 +3355,37 @@ namespace UE::Learning::Observation
 		}
 	}
 
+}
+
+FArchive& operator<<(FArchive& Ar, UE::Learning::Observation::FSchemaElement& SchemaElement)
+{
+	Ar << SchemaElement.Generation;
+	Ar << SchemaElement.Index;
+	return Ar;
+}
+
+FArchive& operator<<(FArchive& Ar, UE::Learning::Observation::FObjectElement& ObjectElement)
+{
+	Ar << ObjectElement.Generation;
+	Ar << ObjectElement.Index;
+	return Ar;
+}
+bool operator==(const UE::Learning::Observation::FSchemaElement& Lhs, const UE::Learning::Observation::FSchemaElement& Rhs)
+{
+	return Lhs.Generation == Rhs.Generation && Lhs.Index == Rhs.Index;
+}
+
+bool operator==(const UE::Learning::Observation::FObjectElement& Lhs, const UE::Learning::Observation::FObjectElement& Rhs)
+{
+	return Lhs.Generation == Rhs.Generation && Lhs.Index == Rhs.Index;
+}
+
+uint32 GetTypeHash(const UE::Learning::Observation::FSchemaElement& SchemaElement)
+{
+	return GetTypeHash(TPair<int32, int32>{ SchemaElement.Generation, SchemaElement.Index });
+}
+
+uint32 GetTypeHash(const UE::Learning::Observation::FObjectElement& ObjectElement)
+{
+	return GetTypeHash(TPair<int32, int32>{ ObjectElement.Generation, ObjectElement.Index });
 }
\ No newline at end of file
diff --git a/Engine/Plugins/Experimental/LearningAgents/Source/Learning/Public/LearningAction.h b/Engine/Plugins/Experimental/LearningAgents/Source/Learning/Public/LearningAction.h
--- a/Engine/Plugins/Experimental/LearningAgents/Source/Learning/Public/LearningAction.h
+++ b/Engine/Plugins/Experimental/LearningAgents/Source/Learning/Public/LearningAction.h
@@ -256,6 +256,9 @@ namespace UE::Learning::Action
 		// Reset all internal buffers (without freeing memory). This invalidates all existing elements.
 		UE_API void Reset();
 
+		// Serializes an action schema. Returns false if serialization fails.
+		UE_API bool Serialize(FArchive& Ar);
+
 	private:
 
 		struct FContinuousData
@@ -488,6 +491,9 @@ namespace UE::Learning::Action
 		// Reset all internal buffers (without freeing memory). This invalidates all existing elements.
 		UE_API void Reset();
 
+		// Serializes an action object. Returns false if serialization fails.
+		UE_API bool Serialize(FArchive& Ar);
+
 	private:
 
 		uint32 Generation = 0;
@@ -651,6 +657,9 @@ namespace UE::Learning::Action
 		// Reset all internal buffers (without freeing memory). This invalidates all existing elements.
 		UE_API void Reset();
 
+		// Serializes an action modifier object. Returns false if serialization fails.
+		UE_API bool Serialize(FArchive& Ar);
+
 	private:
 
 		uint32 Generation = 0;
@@ -845,4 +854,14 @@ namespace UE::Learning::Action
 		const TLearningArrayView<1, const float> ActionModifierVector);
 }
 
+UE_API FArchive& operator<<(FArchive& Ar, UE::Learning::Action::FSchemaElement& SchemaElement);
+UE_API FArchive& operator<<(FArchive& Ar, UE::Learning::Action::FObjectElement& ObjectElement);
+UE_API FArchive& operator<<(FArchive& Ar, UE::Learning::Action::FModifierElement& ModifierElement);
+UE_API bool operator==(const UE::Learning::Action::FSchemaElement& Lhs, const UE::Learning::Action::FSchemaElement& Rhs);
+UE_API bool operator==(const UE::Learning::Action::FObjectElement& Lhs, const UE::Learning::Action::FObjectElement& Rhs);
+UE_API bool operator==(const UE::Learning::Action::FModifierElement& Lhs, const UE::Learning::Action::FModifierElement& Rhs);
+UE_API uint32 GetTypeHash(const UE::Learning::Action::FSchemaElement& SchemaElement);
+UE_API uint32 GetTypeHash(const UE::Learning::Action::FObjectElement& ObjectElement);
+UE_API uint32 GetTypeHash(const UE::Learning::Action::FModifierElement& ModifierElement);
+
 #undef UE_API
diff --git a/Engine/Plugins/Experimental/LearningAgents/Source/Learning/Public/LearningObservation.h b/Engine/Plugins/Experimental/LearningAgents/Source/Learning/Public/LearningObservation.h
--- a/Engine/Plugins/Experimental/LearningAgents/Source/Learning/Public/LearningObservation.h
+++ b/Engine/Plugins/Experimental/LearningAgents/Source/Learning/Public/LearningObservation.h
@@ -335,6 +335,9 @@ namespace UE::Learning::Observation
 		// Reset all internal buffers (without freeing memory). This invalidates all existing elements.
 		UE_API void Reset();
 
+		// Serializes an observation schema. Returns false if serialization fails.
+		UE_API bool Serialize(FArchive& Ar);
+
 	private:
 
 		struct FContinuousData
@@ -623,6 +626,9 @@ namespace UE::Learning::Observation
 		// Reset all internal buffers (without freeing memory). This invalidates all existing elements.
 		UE_API void Reset();
 
+		// Serializes an observation object. Returns false if serialization fails.
+		UE_API bool Serialize(FArchive& Ar);
+
 	private:
 
 		uint32 Generation = 0;
@@ -821,4 +827,11 @@ namespace UE::Learning::Observation
 
 }
 
+UE_API FArchive& operator<<(FArchive& Ar, UE::Learning::Observation::FSchemaElement& SchemaElement);
+UE_API FArchive& operator<<(FArchive& Ar, UE::Learning::Observation::FObjectElement& ObjectElement);
+UE_API bool operator==(const UE::Learning::Observation::FSchemaElement& Lhs, const UE::Learning::Observation::FSchemaElement& Rhs);
+UE_API bool operator==(const UE::Learning::Observation::FObjectElement& Lhs, const UE::Learning::Observation::FObjectElement& Rhs);
+UE_API uint32 GetTypeHash(const UE::Learning::Observation::FSchemaElement& SchemaElement);
+UE_API uint32 GetTypeHash(const UE::Learning::Observation::FObjectElement& ObjectElement);
+
 #undef UE_API
