LearningAgents: Added mean frame attribute function.

#Synced-CL 48723911

diff --git a/Engine/Plugins/Experimental/LearningAgents/Source/Learning/Private/LearningFrameAttribute.cpp b/Engine/Plugins/Experimental/LearningAgents/Source/Learning/Private/LearningFrameAttribute.cpp
--- a/Engine/Plugins/Experimental/LearningAgents/Source/Learning/Private/LearningFrameAttribute.cpp
+++ b/Engine/Plugins/Experimental/LearningAgents/Source/Learning/Private/LearningFrameAttribute.cpp
@@ -729,6 +729,17 @@ namespace UE::Learning
 #endif
 		}
 
+		static inline void Mean(float& OutMean, const TLearningArrayView<1, const float> In)
+		{
+			const int32 Num = In.Num();
+
+			OutMean = 0.0f;
+			for (int32 Idx = 0; Idx < Num; Idx++)
+			{
+				OutMean += (In[Idx] - OutMean) / (Idx + 1);
+			}
+		}
+
 		static inline void MeanStd(float& OutMean, float& OutStd, const TLearningArrayView<1, const float> In)
 		{
 			const int32 Num = In.Num();
@@ -3366,6 +3377,22 @@ namespace UE::Learning
 				});
 		}
 
+		void Mean(TLearningArrayView<1, float> OutMean, const FFrameAttribute& In)
+		{
+			ReduceOp(In, [&OutMean](
+				const FFrameAttribute& In,
+				const TLearningArrayView<1, const int32> RangeOffsets,
+				const TLearningArrayView<1, const int32> RangeLengths) {
+
+					const int32 ChannelNum = In.GetChannelNum();
+
+					for (int32 ChannelIdx = 0; ChannelIdx < ChannelNum; ChannelIdx++)
+					{
+						Private::Mean(OutMean[ChannelIdx], In.GetChannelAttributeData(ChannelIdx));
+					}
+				});
+		}
+
 		void MeanStd(TLearningArrayView<1, float> OutMean, TLearningArrayView<1, float> OutStd, const FFrameAttribute& In)
 		{
 			ReduceOp(In, [&OutMean, &OutStd](
diff --git a/Engine/Plugins/Experimental/LearningAgents/Source/Learning/Public/LearningFrameAttribute.h b/Engine/Plugins/Experimental/LearningAgents/Source/Learning/Public/LearningFrameAttribute.h
--- a/Engine/Plugins/Experimental/LearningAgents/Source/Learning/Public/LearningFrameAttribute.h
+++ b/Engine/Plugins/Experimental/LearningAgents/Source/Learning/Public/LearningFrameAttribute.h
@@ -416,6 +416,9 @@ namespace UE::Learning
 		/** Applies a Majority Vote filter to the ranges of values in the frame attribute. */
 		UE_API void FilterMajorityVote(FFrameAttribute& Out, const FFrameAttribute& In, const int32 FilterWidthFrames);
 
+		/** Computes the mean across all frames. OutMean should be size of the number of channels */
+		UE_API void Mean(TLearningArrayView<1, float> OutMean, const FFrameAttribute& In);
+
 		/** Computes the mean and std across all frames. OutMean and OutStd should be size of the number of channels */
 		UE_API void MeanStd(TLearningArrayView<1, float> OutMean, TLearningArrayView<1, float> OutStd, const FFrameAttribute& In);
 
