LearningAgents: Fix record blobs becoming invalid

#rnx
#virtualized
#Synced-CL 49127937

diff --git a/Engine/Plugins/Experimental/LearningAgents/Source/LearningAgentsTraining/Private/LearningAgentsRecorder.cpp b/Engine/Plugins/Experimental/LearningAgents/Source/LearningAgentsTraining/Private/LearningAgentsRecorder.cpp
--- a/Engine/Plugins/Experimental/LearningAgents/Source/LearningAgentsTraining/Private/LearningAgentsRecorder.cpp
+++ b/Engine/Plugins/Experimental/LearningAgents/Source/LearningAgentsTraining/Private/LearningAgentsRecorder.cpp
@@ -250,7 +250,7 @@ void ULearningAgentsRecorder::OnAgentsRemoved_Implementation(const TArray<int32>
 		if (!RecordBuffers[AgentId].IsEmpty())
 		{
 			FLearningAgentsRecord& Record = Recording->Records.AddDefaulted_GetRef();
-			Record.Blob = NewObject<ULearningAgentsRecordBlob>();
+			Record.Blob = NewObject<ULearningAgentsRecordBlob>(Recording);
 			RecordBuffers[AgentId].CopyToRecord(Record);
 			RecordBuffers[AgentId].Empty();
 		}
@@ -270,7 +270,7 @@ void ULearningAgentsRecorder::OnAgentsReset_Implementation(const TArray<int32>&
 		if (!RecordBuffers[AgentId].IsEmpty())
 		{
 			FLearningAgentsRecord& Record = Recording->Records.AddDefaulted_GetRef();
-			Record.Blob = NewObject<ULearningAgentsRecordBlob>();
+			Record.Blob = NewObject<ULearningAgentsRecordBlob>(Recording);
 			RecordBuffers[AgentId].CopyToRecord(Record);
 			RecordBuffers[AgentId].Empty();
 		}
@@ -340,7 +340,7 @@ void ULearningAgentsRecorder::EndRecording()
 		if (!RecordBuffers[AgentId].IsEmpty())
 		{
 			FLearningAgentsRecord& Record = Recording->Records.AddDefaulted_GetRef();
-			Record.Blob = NewObject<ULearningAgentsRecordBlob>();
+			Record.Blob = NewObject<ULearningAgentsRecordBlob>(Recording);
 			RecordBuffers[AgentId].CopyToRecord(Record);
 			RecordBuffers[AgentId].Empty();
 		}
diff --git a/Engine/Plugins/Experimental/LearningAgents/Source/LearningAgentsTraining/Private/LearningAgentsRecording.cpp b/Engine/Plugins/Experimental/LearningAgents/Source/LearningAgentsTraining/Private/LearningAgentsRecording.cpp
--- a/Engine/Plugins/Experimental/LearningAgents/Source/LearningAgentsTraining/Private/LearningAgentsRecording.cpp
+++ b/Engine/Plugins/Experimental/LearningAgents/Source/LearningAgentsTraining/Private/LearningAgentsRecording.cpp
@@ -180,7 +180,7 @@ void ULearningAgentsRecording::LoadRecordingFromFile(const FFilePath& File, cons
 		for (int32 RecordIdx = 0; RecordIdx < RecordNum; RecordIdx++)
 		{
 			FLearningAgentsRecord Record;
-			Record.Blob = NewObject<ULearningAgentsRecordBlob>();
+			Record.Blob = NewObject<ULearningAgentsRecordBlob>(this);
 
 			UE::Learning::DeserializeFromBytes(Offset, RecordingData, Record.StepNum);
 			UE::Learning::DeserializeFromBytes(Offset, RecordingData, Record.ObservationDimNum);
