LearningAgents: Prevent duplicate agents adds in AddAgent & AddAgents

#rnx
#Synced-CL 49742619

diff --git a/Engine/Plugins/Experimental/LearningAgents/Source/LearningAgents/Private/LearningAgentsManager.cpp b/Engine/Plugins/Experimental/LearningAgents/Source/LearningAgents/Private/LearningAgentsManager.cpp
--- a/Engine/Plugins/Experimental/LearningAgents/Source/LearningAgents/Private/LearningAgentsManager.cpp
+++ b/Engine/Plugins/Experimental/LearningAgents/Source/LearningAgents/Private/LearningAgentsManager.cpp
@@ -98,6 +98,12 @@ int32 ULearningAgentsManager::AddAgent(UObject* Agent)
 		return INDEX_NONE;
 	}
 
+	if (Agents.Contains(Agent))
+	{
+		UE_LOG(LogLearning, Error, TEXT("%s: Agent %s already added to this manager."), *GetName(), *Agent->GetName());
+		return INDEX_NONE;
+	}
+
 	// Add Agent
 
 	const int32 AgentId = VacantAgentIds.Pop();
@@ -143,6 +149,12 @@ void ULearningAgentsManager::AddAgents(TArray<int32>& OutAgentIds, const TArray<
 			continue;
 		}
 
+		if (Agents.Contains(InAgents[AgentIdx]))
+		{
+			UE_LOG(LogLearning, Error, TEXT("%s: Agent %s already added to this manager."), *GetName(), *InAgents[AgentIdx]->GetName());
+			continue;
+		}
+
 		// Add Agent
 
 		OutAgentIds[AgentIdx] = VacantAgentIds.Pop();
