LearningAgents: Added function for getting all Frames from a FrameSet

#Synced-CL 49932463

diff --git a/Engine/Plugins/Experimental/LearningAgents/Source/Learning/Private/LearningFrameSet.cpp b/Engine/Plugins/Experimental/LearningAgents/Source/Learning/Private/LearningFrameSet.cpp
--- a/Engine/Plugins/Experimental/LearningAgents/Source/Learning/Private/LearningFrameSet.cpp
+++ b/Engine/Plugins/Experimental/LearningAgents/Source/Learning/Private/LearningFrameSet.cpp
@@ -246,6 +246,8 @@ namespace UE::Learning
 	int32 FFrameSet::GetEntryFrame(const int32 EntryIdx, const int32 FrameIdx) const { return Frames[EntryFrameOffsets[EntryIdx] + FrameIdx]; }
 	float FFrameSet::GetEntryFrameTime(const int32 EntryIdx, const int32 FrameIdx, const float FrameDeltaTime) const { return GetEntryFrame(EntryIdx, FrameIdx) * FrameDeltaTime; }
 	int32 FFrameSet::GetEntryOffset(const int32 EntryIdx) const { return EntryFrameOffsets[EntryIdx]; }
+	TLearningArrayView<1, const int32> FFrameSet::GetAllFrames() const { return Frames; }
+
 	int32 FFrameSet::GetTotalFrameNum() const { return Frames.Num(); }
 	
 	bool FFrameSet::ContainsSequence(const int32 Sequence) const
diff --git a/Engine/Plugins/Experimental/LearningAgents/Source/Learning/Public/LearningFrameSet.h b/Engine/Plugins/Experimental/LearningAgents/Source/Learning/Public/LearningFrameSet.h
--- a/Engine/Plugins/Experimental/LearningAgents/Source/Learning/Public/LearningFrameSet.h
+++ b/Engine/Plugins/Experimental/LearningAgents/Source/Learning/Public/LearningFrameSet.h
@@ -74,6 +74,9 @@ namespace UE::Learning
 		/** Gets the flat offset associated with a particular entry if you were to flatten this data structure */
 		UE_API int32 GetEntryOffset(const int32 EntryIdx) const;
 
+		/** Gets an array of the frames */
+		UE_API TLearningArrayView<1, const int32> GetAllFrames() const;
+
 		/** Gets the total number of frames if you were to flatten this data structure */
 		UE_API int32 GetTotalFrameNum() const;
 
