LearningAgents: Optimizations and additions to FrameAttribute functions

#Synced-CL 49932463

diff --git a/Engine/Plugins/Experimental/LearningAgents/Source/Learning/Private/Learning.ispc b/Engine/Plugins/Experimental/LearningAgents/Source/Learning/Private/Learning.ispc
--- a/Engine/Plugins/Experimental/LearningAgents/Source/Learning/Private/Learning.ispc
+++ b/Engine/Plugins/Experimental/LearningAgents/Source/Learning/Private/Learning.ispc
@@ -1,7 +1,7 @@
 // Copyright Epic Games, Inc. All Rights Reserved.
 
-#ifndef MAX_flt
-#define MAX_flt			(3.402823466e+38F)
+#ifndef UE_MAX_FLT
+#define UE_MAX_FLT			(3.402823466e+38F)
 #endif
 
 #ifndef INDEX_NONE
@@ -12,6 +12,10 @@
 #define UE_SMALL_NUMBER (1.e-8f)
 #endif
 
+#ifndef UE_PI
+#define UE_PI 					(3.1415926535897932f)	/* Extra digits if needed: 3.1415926535897932384626433832795f */
+#endif
+
 inline float LearningSigmoid(const float X)
 {
 	return 1.0f / (1.0f + exp(-X));
@@ -30,6 +34,14 @@ inline float LearningInvExpApprox(float X)
 	return 1 / (1 + A * X + B * X * X + C * X * X * X);
 }
 
+inline float LearningApproxAcos(float X)
+{
+	const float Y = abs(X);
+	const float P = -0.1565827 * Y + 1.570796;
+	const float Q = P * sqrt(max(1.0 - Y, 0.0));
+	return X >= 0.0 ? Q : UE_PI - Q;
+}
+
 /////////////////////////////////////////////////////////
 // 
 // Functions matching those in LearningRandom.cpp
@@ -539,7 +551,7 @@ export void LearningKMeansUpdateAssignmentsFromCenters(
 {
 	for (uniform int SampleIdx = 0; SampleIdx < SampleNum; SampleIdx++)
 	{
-		uniform float BestDistance = MAX_flt;
+		uniform float BestDistance = UE_MAX_FLT;
 		uniform int32 BestIndex = INDEX_NONE;
 
 		for (uniform int32 ClusterIdx = 0; ClusterIdx < ClusterNum; ClusterIdx++)
@@ -574,7 +586,7 @@ export void LearningKMeansUpdateAssignmentsFromBounds(
 {
 	for (uniform int SampleIdx = 0; SampleIdx < SampleNum; SampleIdx++)
 	{
-		uniform float BestDistance = MAX_flt;
+		uniform float BestDistance = UE_MAX_FLT;
 		uniform int32 BestIndex = INDEX_NONE;
 
 		for (uniform int32 ClusterIdx = 0; ClusterIdx < ClusterNum; ClusterIdx++)
@@ -635,8 +647,8 @@ export void LearningKMeansUpdateBounds(
 {
 	foreach(ClusterIdx = 0 ... ClusterNum, DimIdx = 0 ... DimNum)
 	{
-		OutMins[ClusterIdx * DimNum + DimIdx] = +MAX_flt;
-		OutMaxs[ClusterIdx * DimNum + DimIdx] = -MAX_flt;
+		OutMins[ClusterIdx * DimNum + DimIdx] = +UE_MAX_FLT;
+		OutMaxs[ClusterIdx * DimNum + DimIdx] = -UE_MAX_FLT;
 	}
 
 	for (uniform int SampleIdx = 0; SampleIdx < SampleNum; SampleIdx++)
@@ -745,4 +757,213 @@ export void LearningFilterMajorityVote(
 
 		Out[ValueIdx] = Total > 0;
 	}
+}
+
+export void LearningFrameAttributeAdd(
+	uniform float Out[], 
+	uniform const float Lhs[],
+	uniform const float Rhs[],
+	uniform int ValueNum)
+{
+	foreach(ValueIdx = 0 ... ValueNum)
+	{
+		Out[ValueIdx] = Lhs[ValueIdx] + Rhs[ValueIdx];
+	}
+}
+
+export void LearningFrameAttributeSub(
+	uniform float Out[],
+	uniform const float Lhs[],
+	uniform const float Rhs[],
+	uniform int ValueNum)
+{
+	foreach(ValueIdx = 0 ... ValueNum)
+	{
+		Out[ValueIdx] = Lhs[ValueIdx] - Rhs[ValueIdx];
+	}
+}
+
+export void LearningFrameAttributeMul(
+	uniform float Out[],
+	uniform const float Lhs[],
+	uniform const float Rhs[],
+	uniform int ValueNum)
+{
+	foreach(ValueIdx = 0 ... ValueNum)
+	{
+		Out[ValueIdx] = Lhs[ValueIdx] * Rhs[ValueIdx];
+	}
+}
+
+export void LearningFrameAttributeDiv(
+	uniform float Out[],
+	uniform const float Lhs[],
+	uniform const float Rhs[],
+	uniform int ValueNum)
+{
+	foreach(ValueIdx = 0 ... ValueNum)
+	{
+		Out[ValueIdx] = Lhs[ValueIdx] / Rhs[ValueIdx];
+	}
+}
+
+export void LearningFrameAttributeAddConstantInplace(
+	uniform float InOut[],
+	uniform const float Value,
+	uniform const int ValueNum)
+{
+	foreach(ValueIdx = 0 ... ValueNum)
+	{
+		InOut[ValueIdx] += Value;
+	}
+}
+
+export void LearningFrameAttributeSubConstantInplace(
+	uniform float InOut[],
+	uniform const float Value,
+	uniform const int ValueNum)
+{
+	foreach(ValueIdx = 0 ... ValueNum)
+	{
+		InOut[ValueIdx] -= Value;
+	}
+}
+
+export void LearningFrameAttributeMulConstantInplace(
+	uniform float InOut[], 
+	uniform const float Value,
+	uniform const int ValueNum)
+{
+	foreach(ValueIdx = 0 ... ValueNum)
+	{
+		InOut[ValueIdx] *= Value;
+	}
+}
+
+export void LearningFrameAttributeDivConstantInplace(
+	uniform float InOut[],
+	uniform const float Value,
+	uniform const int ValueNum)
+{
+	foreach(ValueIdx = 0 ... ValueNum)
+	{
+		InOut[ValueIdx] /= Value;
+	}
+}
+
+export void LearningFrameAttributeConstantDot(
+	uniform float InOut[], 
+	uniform const float Lhs, 
+	uniform const float Rhs[],
+	uniform const int ValueNum)
+{
+	foreach(ValueIdx = 0 ... ValueNum)
+	{
+		InOut[ValueIdx] += Lhs * Rhs[ValueIdx];
+	}
+}
+
+export void LearningFrameAttributeConstantDistSquared(
+	uniform float InOut[], 
+	uniform const float Lhs, 
+	uniform const float Rhs[],
+	uniform const int ValueNum)
+{
+	foreach(ValueIdx = 0 ... ValueNum)
+	{
+		InOut[ValueIdx] += LearningSquare(Lhs - Rhs[ValueIdx]);
+	}
+}
+
+export void LearningFrameAttributeApproxAcosInplace(
+	uniform float InOut[],
+	uniform const int ValueNum)
+{
+	foreach(ValueIdx = 0 ... ValueNum)
+	{
+		InOut[ValueIdx] = LearningApproxAcos(InOut[ValueIdx]);
+	}
+}
+
+export void LearningFrameAttributeSqrtInplace(
+	uniform float InOut[],
+	uniform const int ValueNum)
+{
+	foreach(ValueIdx = 0 ... ValueNum)
+	{
+		InOut[ValueIdx] = sqrt(InOut[ValueIdx]);
+	}
+}
+
+export uniform bool LearningFrameAttributeFindMinimum(
+	uniform int32& OutChannelIdx,
+	uniform int32& OutFrameIdx,
+	uniform float& OutMinimumValue,
+	uniform const float Values[],
+	uniform const int ChannelNum,
+	uniform const int FrameNum)
+{
+	OutChannelIdx = INDEX_NONE;
+	OutFrameIdx = INDEX_NONE;
+	OutMinimumValue = UE_MAX_FLT;
+
+	for (uniform int ChannelIdx = 0; ChannelIdx < ChannelNum; ChannelIdx++)
+	{
+		// TODO: Vectorize
+		for (uniform int FrameIdx = 0; FrameIdx < FrameNum; FrameIdx++)
+		{
+			if (Values[ChannelIdx * FrameNum + FrameIdx] < OutMinimumValue)
+			{
+				OutChannelIdx = ChannelIdx;
+				OutFrameIdx = FrameIdx;
+				OutMinimumValue = Values[ChannelIdx * FrameNum + FrameIdx];
+			}
+		}
+	}
+
+	return OutChannelIdx != INDEX_NONE;
+}
+
+export uniform bool LearningFrameAttributeFindMaximum(
+	uniform int32& OutChannelIdx,
+	uniform int32& OutFrameIdx,
+	uniform float& OutMaximumValue,
+	uniform const float Values[],
+	uniform const int ChannelNum,
+	uniform const int FrameNum)
+{
+	OutChannelIdx = INDEX_NONE;
+	OutFrameIdx = INDEX_NONE;
+	OutMaximumValue = -UE_MAX_FLT;
+
+	for (uniform int ChannelIdx = 0; ChannelIdx < ChannelNum; ChannelIdx++)
+	{
+		// TODO: Vectorize
+		for (uniform int FrameIdx = 0; FrameIdx < FrameNum; FrameIdx++)
+		{
+			if (Values[ChannelIdx * FrameNum + FrameIdx] > OutMaximumValue)
+			{
+				OutChannelIdx = ChannelIdx;
+				OutFrameIdx = FrameIdx;
+				OutMaximumValue = Values[ChannelIdx * FrameNum + FrameIdx];
+			}
+		}
+	}
+
+	return OutChannelIdx != INDEX_NONE;
+}
+
+
+export uniform bool LearningInt32ArraysEqual(
+	uniform const int32 Lhs[],
+	uniform const int32 Rhs[],
+	uniform const int ValueNum)
+{
+	// TODO: Vectorize
+	for (uniform int ValueIdx = 0; ValueIdx < ValueNum; ValueIdx++)
+	{
+		if (Lhs[ValueIdx] != Rhs[ValueIdx]) { return false; }
+	}
+
+	return true;
 }
\ No newline at end of file
diff --git a/Engine/Plugins/Experimental/LearningAgents/Source/Learning/Private/LearningFrameAttribute.cpp b/Engine/Plugins/Experimental/LearningAgents/Source/Learning/Private/LearningFrameAttribute.cpp
--- a/Engine/Plugins/Experimental/LearningAgents/Source/Learning/Private/LearningFrameAttribute.cpp
+++ b/Engine/Plugins/Experimental/LearningAgents/Source/Learning/Private/LearningFrameAttribute.cpp
@@ -26,6 +26,15 @@ namespace UE::Learning
 			float& OutMinimumValue,
 			const TLearningArrayView<2, const float> Values)
 		{
+#if UE_LEARNING_ISPC
+			return ispc::LearningFrameAttributeFindMinimum(
+				OutChannelIdx,
+				OutFrameIdx,
+				OutMinimumValue,
+				Values.GetData(),
+				Values.Num<0>(),
+				Values.Num<1>());
+#else
 			OutChannelIdx = INDEX_NONE;
 			OutFrameIdx = INDEX_NONE;
 			OutMinimumValue = UE_MAX_FLT;
@@ -47,6 +56,7 @@ namespace UE::Learning
 			}
 
 			return OutChannelIdx != INDEX_NONE;
+#endif
 		}
 		
 		static inline bool FindMaximum(
@@ -55,6 +65,15 @@ namespace UE::Learning
 			float& OutMaximumValue,
 			const TLearningArrayView<2, const float> Values)
 		{
+#if UE_LEARNING_ISPC
+			return ispc::LearningFrameAttributeFindMaximum(
+				OutChannelIdx,
+				OutFrameIdx,
+				OutMaximumValue,
+				Values.GetData(),
+				Values.Num<0>(),
+				Values.Num<1>());
+#else
 			OutChannelIdx = INDEX_NONE;
 			OutFrameIdx = INDEX_NONE;
 			OutMaximumValue = -UE_MAX_FLT;
@@ -76,50 +95,87 @@ namespace UE::Learning
 			}
 
 			return OutChannelIdx != INDEX_NONE;
+#endif
 		}
 
 		static inline void Add(TLearningArrayView<1, float> Out, const TLearningArrayView<1, const float> Lhs, const TLearningArrayView<1, const float> Rhs)
 		{
 			check(Out.Num() == Lhs.Num());
-			check(Out.Num() == Rhs.Num());
+			check(Out.Num() == Rhs.Num()); 
+
+#if UE_LEARNING_ISPC
+			ispc::LearningFrameAttributeAdd(
+				Out.GetData(),
+				Lhs.GetData(),
+				Rhs.GetData(),
+				Out.Num());
+#else
 			const int32 ValueNum = Out.Num();
 			for (int32 ValueIdx = 0; ValueIdx < ValueNum; ValueIdx++)
 			{
 				Out[ValueIdx] = Lhs[ValueIdx] + Rhs[ValueIdx];
 			}
+#endif
 		}
 
 		static inline void Sub(TLearningArrayView<1, float> Out, const TLearningArrayView<1, const float> Lhs, const TLearningArrayView<1, const float> Rhs)
 		{
 			check(Out.Num() == Lhs.Num());
 			check(Out.Num() == Rhs.Num());
+
+#if UE_LEARNING_ISPC
+			ispc::LearningFrameAttributeSub(
+				Out.GetData(),
+				Lhs.GetData(),
+				Rhs.GetData(),
+				Out.Num());
+#else
 			const int32 ValueNum = Out.Num();
 			for (int32 ValueIdx = 0; ValueIdx < ValueNum; ValueIdx++)
 			{
 				Out[ValueIdx] = Lhs[ValueIdx] - Rhs[ValueIdx];
 			}
+#endif
 		}
 
 		static inline void Mul(TLearningArrayView<1, float> Out, const TLearningArrayView<1, const float> Lhs, const TLearningArrayView<1, const float> Rhs)
 		{
 			check(Out.Num() == Lhs.Num());
 			check(Out.Num() == Rhs.Num());
+
+#if UE_LEARNING_ISPC
+			ispc::LearningFrameAttributeMul(
+				Out.GetData(),
+				Lhs.GetData(),
+				Rhs.GetData(),
+				Out.Num());
+#else
 			const int32 ValueNum = Out.Num();
 			for (int32 ValueIdx = 0; ValueIdx < ValueNum; ValueIdx++)
 			{
 				Out[ValueIdx] = Lhs[ValueIdx] * Rhs[ValueIdx];
 			}
+#endif
 		}
 
 		static inline void Div(TLearningArrayView<1, float> Out, const TLearningArrayView<1, const float> Lhs, const TLearningArrayView<1, const float> Rhs)
 		{
 			check(Out.Num() == Lhs.Num());
 			check(Out.Num() == Rhs.Num());
+
+#if UE_LEARNING_ISPC
+			ispc::LearningFrameAttributeDiv(
+				Out.GetData(),
+				Lhs.GetData(),
+				Rhs.GetData(),
+				Out.Num());
+#else
 			const int32 ValueNum = Out.Num();
 			for (int32 ValueIdx = 0; ValueIdx < ValueNum; ValueIdx++)
 			{
 				Out[ValueIdx] = Lhs[ValueIdx] / Rhs[ValueIdx];
 			}
+#endif
 		}
 
 		static inline void Dot(TLearningArrayView<1, float> InOut, const TLearningArrayView<1, const float> Lhs, const TLearningArrayView<1, const float> Rhs)
@@ -258,20 +314,53 @@ namespace UE::Learning
 
 		static inline void SqrtInplace(TLearningArrayView<1, float> InOut)
 		{
+#if UE_LEARNING_ISPC
+			ispc::LearningFrameAttributeSqrtInplace(
+				InOut.GetData(),
+				InOut.Num());
+#else
 			const int32 ValueNum = InOut.Num();
 			for (int32 ValueIdx = 0; ValueIdx < ValueNum; ValueIdx++)
 			{
 				InOut[ValueIdx] = FMath::Sqrt(InOut[ValueIdx]);
 			}
+#endif
 		}
 
 		static inline void ApproxAcosInplace(TLearningArrayView<1, float> InOut)
 		{
+#if UE_LEARNING_ISPC
+			ispc::LearningFrameAttributeApproxAcosInplace(
+				InOut.GetData(),
+				InOut.Num());
+#else
 			const int32 ValueNum = InOut.Num();
 			for (int32 ValueIdx = 0; ValueIdx < ValueNum; ValueIdx++)
 			{
 				InOut[ValueIdx] = ApproxAcos(InOut[ValueIdx]);
 			}
+#endif
+		}
+
+		static inline void NormalizeInplace(TLearningArrayView<2, float> InOut)
+		{
+			const int32 ChannelNum = InOut.Num<0>();
+			const int32 FrameNum = InOut.Num<1>();
+			for (int32 FrameIdx = 0; FrameIdx < FrameNum; FrameIdx++)
+			{
+				float LengthSquared = 0.0f;
+				for (int32 ChannelIdx = 0; ChannelIdx < ChannelNum; ChannelIdx++)
+				{
+					LengthSquared += InOut[ChannelIdx][FrameIdx] * InOut[ChannelIdx][FrameIdx];
+				}
+
+				const float Length = FMath::Sqrt(LengthSquared);
+
+				for (int32 ChannelIdx = 0; ChannelIdx < ChannelNum; ChannelIdx++)
+				{
+					InOut[ChannelIdx][FrameIdx] = InOut[ChannelIdx][FrameIdx] / Length;
+				}
+			}
 		}
 
 		static inline void LengthSquared(TLearningArrayView<1, float> InOut, const TLearningArrayView<1, const float> In)
@@ -391,57 +480,103 @@ namespace UE::Learning
 		static inline void ConstantDot(TLearningArrayView<1, float> InOut, const float Lhs, const TLearningArrayView<1, const float> Rhs)
 		{
 			check(InOut.Num() == Rhs.Num());
+
+#if UE_LEARNING_ISPC
+			ispc::LearningFrameAttributeConstantDot(
+				InOut.GetData(),
+				Lhs,
+				Rhs.GetData(),
+				InOut.Num());
+#else
 			const int32 ValueNum = InOut.Num();
 			for (int32 ValueIdx = 0; ValueIdx < ValueNum; ValueIdx++)
 			{
 				InOut[ValueIdx] += Lhs * Rhs[ValueIdx];
 			}
+#endif
 		}
 
 		static inline void ConstantDistSquared(TLearningArrayView<1, float> InOut, const float Lhs, const TLearningArrayView<1, const float> Rhs)
 		{
 			check(InOut.Num() == Rhs.Num());
+
+#if UE_LEARNING_ISPC
+			ispc::LearningFrameAttributeConstantDistSquared(
+				InOut.GetData(),
+				Lhs,
+				Rhs.GetData(),
+				InOut.Num());
+#else
 			const int32 ValueNum = InOut.Num();
 			for (int32 ValueIdx = 0; ValueIdx < ValueNum; ValueIdx++)
 			{
 				InOut[ValueIdx] += FMath::Square(Lhs - Rhs[ValueIdx]);
 			}
+#endif
 		}
 
 		static inline void AddConstantInplace(TLearningArrayView<1, float> InOut, const float Value)
 		{
+#if UE_LEARNING_ISPC
+			ispc::LearningFrameAttributeAddConstantInplace(
+				InOut.GetData(),
+				Value,
+				InOut.Num());
+#else
 			const int32 ValueNum = InOut.Num();
 			for (int32 ValueIdx = 0; ValueIdx < ValueNum; ValueIdx++)
 			{
 				InOut[ValueIdx] += Value;
 			}
+#endif
 		}
 
 		static inline void SubConstantInplace(TLearningArrayView<1, float> InOut, const float Value)
 		{
+#if UE_LEARNING_ISPC
+			ispc::LearningFrameAttributeSubConstantInplace(
+				InOut.GetData(),
+				Value,
+				InOut.Num());
+#else
 			const int32 ValueNum = InOut.Num();
 			for (int32 ValueIdx = 0; ValueIdx < ValueNum; ValueIdx++)
 			{
 				InOut[ValueIdx] -= Value;
 			}
+#endif
 		}
 
 		static inline void MulConstantInplace(TLearningArrayView<1, float> InOut, const float Value)
 		{
+#if UE_LEARNING_ISPC
+			ispc::LearningFrameAttributeMulConstantInplace(
+				InOut.GetData(),
+				Value,
+				InOut.Num());
+#else
 			const int32 ValueNum = InOut.Num();
 			for (int32 ValueIdx = 0; ValueIdx < ValueNum; ValueIdx++)
 			{
 				InOut[ValueIdx] *= Value;
 			}
+#endif
 		}
 
 		static inline void DivConstantInplace(TLearningArrayView<1, float> InOut, const float Value)
 		{
+#if UE_LEARNING_ISPC
+			ispc::LearningFrameAttributeDivConstantInplace(
+				InOut.GetData(),
+				Value,
+				InOut.Num());
+#else
 			const int32 ValueNum = InOut.Num();
 			for (int32 ValueIdx = 0; ValueIdx < ValueNum; ValueIdx++)
 			{
 				InOut[ValueIdx] /= Value;
 			}
+#endif
 		}
 
 		static inline void LogicalAnd(TLearningArrayView<1, float> Out, const TLearningArrayView<1, const float> Lhs, const TLearningArrayView<1, const float> Rhs)
@@ -2025,6 +2160,19 @@ namespace UE::Learning
 		{
 			check(Lhs.GetChannelNum() == Rhs.GetChannelNum());
 
+			if (FrameRangeSet::Equal(Lhs.FrameRangeSet, Rhs.FrameRangeSet))
+			{
+				Out.FrameRangeSet = Lhs.FrameRangeSet;
+				Out.AttributeData.SetNumUninitialized({ Lhs.GetChannelNum(), Lhs.GetTotalFrameNum() });
+
+				Private::Add(
+					Out.GetAttributeData().Flatten(),
+					Lhs.GetAttributeData().Flatten(),
+					Rhs.GetAttributeData().Flatten());
+
+				return;
+			}
+
 			BinaryOp(Out, Lhs.GetChannelNum(), Lhs, Rhs, [](
 				FFrameAttribute& Out,
 				const FFrameAttribute& Lhs,
@@ -2053,6 +2201,19 @@ namespace UE::Learning
 		{
 			check(Lhs.GetChannelNum() == Rhs.GetChannelNum());
 
+			if (FrameRangeSet::Equal(Lhs.FrameRangeSet, Rhs.FrameRangeSet))
+			{
+				Out.FrameRangeSet = Lhs.FrameRangeSet;
+				Out.AttributeData.SetNumUninitialized({ Lhs.GetChannelNum(), Lhs.GetTotalFrameNum() });
+
+				Private::Sub(
+					Out.GetAttributeData().Flatten(),
+					Lhs.GetAttributeData().Flatten(),
+					Rhs.GetAttributeData().Flatten());
+
+				return;
+			}
+
 			BinaryOp(Out, Lhs.GetChannelNum(), Lhs, Rhs, [](
 				FFrameAttribute& Out,
 				const FFrameAttribute& Lhs,
@@ -2081,6 +2242,19 @@ namespace UE::Learning
 		{
 			check(Lhs.GetChannelNum() == Rhs.GetChannelNum());
 
+			if (FrameRangeSet::Equal(Lhs.FrameRangeSet, Rhs.FrameRangeSet))
+			{
+				Out.FrameRangeSet = Lhs.FrameRangeSet;
+				Out.AttributeData.SetNumUninitialized({ Lhs.GetChannelNum(), Lhs.GetTotalFrameNum() });
+
+				Private::Mul(
+					Out.GetAttributeData().Flatten(),
+					Lhs.GetAttributeData().Flatten(),
+					Rhs.GetAttributeData().Flatten());
+
+				return;
+			}
+
 			BinaryOp(Out, Lhs.GetChannelNum(), Lhs, Rhs, [](
 				FFrameAttribute& Out,
 				const FFrameAttribute& Lhs,
@@ -2109,6 +2283,19 @@ namespace UE::Learning
 		{
 			check(Lhs.GetChannelNum() == Rhs.GetChannelNum());
 
+			if (FrameRangeSet::Equal(Lhs.FrameRangeSet, Rhs.FrameRangeSet))
+			{
+				Out.FrameRangeSet = Lhs.FrameRangeSet;
+				Out.AttributeData.SetNumUninitialized({ Lhs.GetChannelNum(), Lhs.GetTotalFrameNum() });
+
+				Private::Div(
+					Out.GetAttributeData().Flatten(),
+					Lhs.GetAttributeData().Flatten(),
+					Rhs.GetAttributeData().Flatten());
+
+				return;
+			}
+
 			BinaryOp(Out, Lhs.GetChannelNum(), Lhs, Rhs, [](
 				FFrameAttribute& Out,
 				const FFrameAttribute& Lhs,
@@ -2428,6 +2615,18 @@ namespace UE::Learning
 				});
 		}
 
+		void NormalizeInplace(FFrameAttribute& InOut)
+		{
+			InplaceOp(InOut, [](
+				FFrameAttribute& InOut,
+				const TLearningArrayView<1, const int32> RangeOffsets,
+				const TLearningArrayView<1, const int32> RangeLengths) {
+
+					Private::NormalizeInplace(InOut.GetAttributeData());
+
+				});
+		}
+
 		void Length(FFrameAttribute& Out, const FFrameAttribute& In)
 		{
 			UnaryOp(Out, 1, In, [](
diff --git a/Engine/Plugins/Experimental/LearningAgents/Source/Learning/Private/LearningFrameRangeSet.cpp b/Engine/Plugins/Experimental/LearningAgents/Source/Learning/Private/LearningFrameRangeSet.cpp
--- a/Engine/Plugins/Experimental/LearningAgents/Source/Learning/Private/LearningFrameRangeSet.cpp
+++ b/Engine/Plugins/Experimental/LearningAgents/Source/Learning/Private/LearningFrameRangeSet.cpp
@@ -5,6 +5,10 @@
 
 #include "Async/ParallelFor.h"
 
+#if UE_LEARNING_ISPC
+#include "Learning.ispc.generated.h"
+#endif
+
 namespace UE::Learning
 {
 	namespace FrameRangeSet::Private
@@ -124,6 +128,25 @@ namespace UE::Learning
 			return false;
 		}
 
+		static inline bool RangesContainsRange(
+			const TLearningArrayView<1, const int32> Starts,
+			const TLearningArrayView<1, const int32> Lengths,
+			const int32 Start,
+			const int32 Length)
+		{
+			check(Starts.Num() == Lengths.Num());
+
+			const int32 RangeNum = Starts.Num();
+			for (int32 RangeIdx = 0; RangeIdx < RangeNum; RangeIdx++)
+			{
+				if (Starts[RangeIdx] <= Start && Start + Length <= Starts[RangeIdx] + Lengths[RangeIdx])
+				{
+					return true;
+				}
+			}
+			return false;
+		}
+
 		static inline bool RangesContainsTime(
 			const TLearningArrayView<1, const int32> Starts,
 			const TLearningArrayView<1, const int32> Lengths,
@@ -1418,6 +1441,14 @@ namespace UE::Learning
 			Frame);
 	}
 
+	bool FFrameRangeSet::EntryContains(const int32 EntryIdx, const int32 Frame) const
+	{
+		return FrameRangeSet::Private::RangesContains(
+			GetEntryRangeStarts(EntryIdx),
+			GetEntryRangeLengths(EntryIdx),
+			Frame);
+	}
+
 	bool FFrameRangeSet::IntersectsRange(const int32 Sequence, const int32 Start, const int32 Length) const
 	{
 		const int32 EntryIdx = FindSequenceEntry(Sequence);
@@ -1429,6 +1460,14 @@ namespace UE::Learning
 			Length);
 	}
 
+	bool FFrameRangeSet::EntryContainsRange(const int32 EntryIdx, const int32 Start, const int32 Length) const
+	{
+		return FrameRangeSet::Private::RangesContainsRange(
+			GetEntryRangeStarts(EntryIdx),
+			GetEntryRangeLengths(EntryIdx),
+			Start, Length);
+	}
+
 	bool FFrameRangeSet::ContainsTime(const int32 Sequence, const float Time, const float FrameDeltaTime) const
 	{
 		const int32 EntryIdx = FindSequenceEntry(Sequence);
@@ -1592,6 +1631,17 @@ namespace UE::Learning
 	{
 		bool Equal(const FFrameRangeSet& Lhs, const FFrameRangeSet& Rhs)
 		{
+#if UE_LEARNING_ISPC
+			return
+				Lhs.EntrySequences.Num() == Rhs.EntrySequences.Num() &&
+				Lhs.RangeStarts.Num() == Rhs.RangeStarts.Num() &&
+				ispc::LearningInt32ArraysEqual(Lhs.EntrySequences.GetData(), Rhs.EntrySequences.GetData(), Lhs.EntrySequences.Num()) &&
+				ispc::LearningInt32ArraysEqual(Lhs.EntryRangeOffsets.GetData(), Rhs.EntryRangeOffsets.GetData(), Lhs.EntryRangeOffsets.Num()) &&
+				ispc::LearningInt32ArraysEqual(Lhs.EntryRangeNums.GetData(), Rhs.EntryRangeNums.GetData(), Lhs.EntryRangeNums.Num()) &&
+				ispc::LearningInt32ArraysEqual(Lhs.RangeStarts.GetData(), Rhs.RangeStarts.GetData(), Lhs.RangeStarts.Num()) &&
+				ispc::LearningInt32ArraysEqual(Lhs.RangeLengths.GetData(), Rhs.RangeLengths.GetData(), Lhs.RangeLengths.Num()) &&
+				ispc::LearningInt32ArraysEqual(Lhs.RangeOffsets.GetData(), Rhs.RangeOffsets.GetData(), Lhs.RangeOffsets.Num());
+#else
 			return
 				Lhs.EntrySequences.Num() == Rhs.EntrySequences.Num() &&
 				Lhs.RangeStarts.Num() == Rhs.RangeStarts.Num() &&
@@ -1601,6 +1651,7 @@ namespace UE::Learning
 				Array::Equal(Lhs.RangeStarts, Rhs.RangeStarts) &&
 				Array::Equal(Lhs.RangeLengths, Rhs.RangeLengths) &&
 				Array::Equal(Lhs.RangeOffsets, Rhs.RangeOffsets);
+#endif
 		}
 
 		void Union(FFrameRangeSet& OutFrameRangeSet, const FFrameSet& FrameSet, const FFrameRangeSet& FrameRangeSet)
@@ -2760,6 +2811,56 @@ namespace UE::Learning
 			}
 		}
 
+		void MakeFrameSetFromRangeInterval(FFrameSet& OutFrameSet, const FFrameRangeSet& FrameRangeSet, const int32 IntervalFrames)
+		{
+			check(IntervalFrames > 0);
+
+			const int32 EntryNum = FrameRangeSet.GetEntryNum();
+
+			OutFrameSet.EntrySequences = FrameRangeSet.EntrySequences;
+			OutFrameSet.EntryFrameOffsets.SetNumUninitialized({ EntryNum });
+			OutFrameSet.EntryFrameNums.SetNumUninitialized({ EntryNum });
+			Array::Zero(OutFrameSet.EntryFrameOffsets);
+			Array::Zero(OutFrameSet.EntryFrameNums);
+
+			int32 TotalFrameNum = 0;
+			for (int32 EntryIdx = 0; EntryIdx < EntryNum; EntryIdx++)
+			{
+				const int32 RangeNum = FrameRangeSet.GetEntryRangeNum(EntryIdx);
+
+				for (int32 RangeIdx = 0; RangeIdx < RangeNum; RangeIdx++)
+				{
+					check(FrameRangeSet.GetEntryRangeLength(EntryIdx, RangeIdx) > 0);
+					TotalFrameNum += (FrameRangeSet.GetEntryRangeLength(EntryIdx, RangeIdx) - 1) / IntervalFrames + 1;
+				}
+			}
+
+			OutFrameSet.Frames.SetNumUninitialized({ TotalFrameNum });
+
+			int32 FrameOffset = 0;
+			for (int32 EntryIdx = 0; EntryIdx < EntryNum; EntryIdx++)
+			{
+				OutFrameSet.EntryFrameOffsets[EntryIdx] = FrameOffset;
+
+				const int32 RangeNum = FrameRangeSet.GetEntryRangeNum(EntryIdx);
+
+				for (int32 RangeIdx = 0; RangeIdx < RangeNum; RangeIdx++)
+				{
+					const int32 RangeLength = FrameRangeSet.GetEntryRangeLength(EntryIdx, RangeIdx);
+					const int32 RangeStart = FrameRangeSet.GetEntryRangeStart(EntryIdx, RangeIdx);
+
+					for (int32 FrameIdx = RangeStart; FrameIdx < RangeStart + RangeLength; FrameIdx += IntervalFrames)
+					{
+						OutFrameSet.EntryFrameNums[EntryIdx]++;
+						OutFrameSet.Frames[FrameOffset] = FrameIdx;
+						FrameOffset++;
+					}
+				}
+			}
+
+			check(FrameOffset == TotalFrameNum);
+		}
+
 		void RangesBeforeFrameSet(FFrameRangeSet& OutFrameRangeSet, const FFrameRangeSet& FrameRangeSet, const FFrameSet& FrameSet)
 		{
 			OutFrameRangeSet.Empty();
diff --git a/Engine/Plugins/Experimental/LearningAgents/Source/Learning/Public/LearningFrameAttribute.h b/Engine/Plugins/Experimental/LearningAgents/Source/Learning/Public/LearningFrameAttribute.h
--- a/Engine/Plugins/Experimental/LearningAgents/Source/Learning/Public/LearningFrameAttribute.h
+++ b/Engine/Plugins/Experimental/LearningAgents/Source/Learning/Public/LearningFrameAttribute.h
@@ -271,6 +271,9 @@ namespace UE::Learning
 		/** Compute the approximate acos of a frame attribute in place. */
 		UE_API void ApproxAcosInplace(FFrameAttribute& InOut);
 
+		/** Normalize the frame attribute in place. */
+		UE_API void NormalizeInplace(FFrameAttribute& InOut);
+
 		/** Repeats the first entry in the range across the rest of the range */
 		UE_API void RepeatFirstRangeEntryInplace(FFrameAttribute& InOut);
 
diff --git a/Engine/Plugins/Experimental/LearningAgents/Source/Learning/Public/LearningFrameRangeSet.h b/Engine/Plugins/Experimental/LearningAgents/Source/Learning/Public/LearningFrameRangeSet.h
--- a/Engine/Plugins/Experimental/LearningAgents/Source/Learning/Public/LearningFrameRangeSet.h
+++ b/Engine/Plugins/Experimental/LearningAgents/Source/Learning/Public/LearningFrameRangeSet.h
@@ -114,9 +114,15 @@ namespace UE::Learning
 		/** Checks if this FrameRangeSet contains a given sequence and frame in that sequence */
 		UE_API bool Contains(const int32 Sequence, const int32 Frame) const;
 
+		/** Checks if the given Entry of this FrameRangeSet contains a frame */
+		UE_API bool EntryContains(const int32 EntryIdx, const int32 Frame) const;
+
 		/** Checks if this FrameRangeSet intersects a sequence and any part of a range in that sequence */
 		UE_API bool IntersectsRange(const int32 Sequence, const int32 Start, const int32 Length) const;
 
+		/** Checks if the given Entry of this FrameRangeSet fully contains a range */
+		UE_API bool EntryContainsRange(const int32 EntryIdx, const int32 Start, const int32 Length) const;
+
 		/** Checks if this FrameRangeSet contains a given sequence and time in that sequence */
 		UE_API bool ContainsTime(const int32 Sequence, const float Time, const float FrameDeltaTime) const;
 
@@ -224,6 +230,9 @@ namespace UE::Learning
 		/** Create a frame set from all of the range ends of a given frame range set */
 		UE_API void MakeFrameSetFromRangeEnds(FFrameSet& OutFrameSet, const FFrameRangeSet& FrameRangeSet);
 
+		/** Create a frame set from a given frame range set, sampling frames at a regular interval */
+		UE_API void MakeFrameSetFromRangeInterval(FFrameSet& OutFrameSet, const FFrameRangeSet& FrameRangeSet, const int32 IntervalFrames = 10);
+
 		/** Trim a frame range set to just the period of time before the first occurrence of a frame from a given frame set */
 		UE_API void RangesBeforeFrameSet(FFrameRangeSet& OutFrameRangeSet, const FFrameRangeSet& FrameRangeSet, const FFrameSet& FrameSet);
 
