LearningAgents: Plugin split into LearningAgents and LearningCore

#Synced-CL 50215877

diff --git a/Engine/Plugins/Experimental/LearningAgents/LearningAgents.uplugin b/Engine/Plugins/Experimental/LearningAgents/LearningAgents.uplugin
--- a/Engine/Plugins/Experimental/LearningAgents/LearningAgents.uplugin
+++ b/Engine/Plugins/Experimental/LearningAgents/LearningAgents.uplugin
@@ -16,16 +16,6 @@
 	"IsExperimentalVersion": true,
 	"Installed": false,
 	"Modules": [
-		{
-			"Name": "Learning",
-			"Type": "Runtime",
-			"LoadingPhase": "Default"
-		},
-		{
-			"Name": "LearningTraining",
-			"Type": "Runtime",
-			"LoadingPhase": "Default"
-		},
 		{
 			"Name": "LearningAgents",
 			"Type": "Runtime",
@@ -48,6 +38,10 @@
 		}
 	],
 	"Plugins": [
+		{
+			"Name": "LearningCore",
+			"Enabled": true
+		},
 		{
 			"Name": "PythonMLPackages",
 			"Enabled": true
