LearningAgents: Fix memory leak in LearningAgentsRecorder

#Synced-CL 50093269

diff --git a/Engine/Plugins/Experimental/LearningAgents/Source/LearningAgentsTraining/Private/LearningAgentsRecorder.cpp b/Engine/Plugins/Experimental/LearningAgents/Source/LearningAgentsTraining/Private/LearningAgentsRecorder.cpp
--- a/Engine/Plugins/Experimental/LearningAgents/Source/LearningAgentsTraining/Private/LearningAgentsRecorder.cpp
+++ b/Engine/Plugins/Experimental/LearningAgents/Source/LearningAgentsTraining/Private/LearningAgentsRecorder.cpp
@@ -64,7 +64,7 @@ void ULearningAgentsRecorder::FAgentRecordBuffer::Push(
 	const TLearningArrayView<1, const float> Observation,
 	const TLearningArrayView<1, const float> Action)
 {
-	if (StepNum / ChunkSize <= Observations.Num())
+	if (StepNum % ChunkSize == 0)
 	{
 		Observations.AddDefaulted_GetRef().SetNumUninitialized({ ChunkSize, Observation.Num() });
 		Actions.AddDefaulted_GetRef().SetNumUninitialized({ ChunkSize, Action.Num() });
