LearningAgents: New records will automatically be given a schema name and tag from the recording

#rnx
#Synced-CL 50354366

diff --git a/Engine/Plugins/Experimental/LearningAgents/Source/LearningAgentsTraining/Private/LearningAgentsRecorder.cpp b/Engine/Plugins/Experimental/LearningAgents/Source/LearningAgentsTraining/Private/LearningAgentsRecorder.cpp
--- a/Engine/Plugins/Experimental/LearningAgents/Source/LearningAgentsTraining/Private/LearningAgentsRecorder.cpp
+++ b/Engine/Plugins/Experimental/LearningAgents/Source/LearningAgentsTraining/Private/LearningAgentsRecorder.cpp
@@ -251,6 +251,9 @@ void ULearningAgentsRecorder::OnAgentsRemoved_Implementation(const TArray<int32>
 		{
 			FLearningAgentsRecord& Record = Recording->Records.AddDefaulted_GetRef();
 			Record.Blob = NewObject<ULearningAgentsRecordBlob>(Recording);
+			Record.SchemaName = Recording->NewSchemaName;
+			Record.Tags.AddTag(Recording->NewTag);
+
 			RecordBuffers[AgentId].CopyToRecord(Record);
 			RecordBuffers[AgentId].Empty();
 		}
@@ -271,6 +274,9 @@ void ULearningAgentsRecorder::OnAgentsReset_Implementation(const TArray<int32>&
 		{
 			FLearningAgentsRecord& Record = Recording->Records.AddDefaulted_GetRef();
 			Record.Blob = NewObject<ULearningAgentsRecordBlob>(Recording);
+			Record.SchemaName = Recording->NewSchemaName;
+			Record.Tags.AddTag(Recording->NewTag);
+
 			RecordBuffers[AgentId].CopyToRecord(Record);
 			RecordBuffers[AgentId].Empty();
 		}
@@ -341,6 +347,9 @@ void ULearningAgentsRecorder::EndRecording()
 		{
 			FLearningAgentsRecord& Record = Recording->Records.AddDefaulted_GetRef();
 			Record.Blob = NewObject<ULearningAgentsRecordBlob>(Recording);
+			Record.SchemaName = Recording->NewSchemaName;
+			Record.Tags.AddTag(Recording->NewTag);
+
 			RecordBuffers[AgentId].CopyToRecord(Record);
 			RecordBuffers[AgentId].Empty();
 		}
