READY TO SUBMIT

[LearningAgents]
- added bActive flag to all observations and actions that route to NullObservation / NullAction if inactive

[LACombat]
- updated combat interactor toward data oriented design: all actions/data configs can be loaded from one source, no more specialized structs for scales/biases etc
- header reduced from 1300 to <200 lines
- split out all rule based guard logic like lethality reducer/action stickiness/guards into a separate component so the interactor is pure ML specs
- kept old version of interactor for legacy and safety

#Synced-CL 50475176

diff --git a/Engine/Plugins/Experimental/LearningAgents/Source/LearningAgents/Private/LearningAgentsActions.cpp b/Engine/Plugins/Experimental/LearningAgents/Source/LearningAgents/Private/LearningAgentsActions.cpp
--- a/Engine/Plugins/Experimental/LearningAgents/Source/LearningAgents/Private/LearningAgentsActions.cpp
+++ b/Engine/Plugins/Experimental/LearningAgents/Source/LearningAgents/Private/LearningAgentsActions.cpp
@@ -1060,14 +1060,19 @@ FLearningAgentsActionSchemaElement ULearningAgentsActions::SpecifyNullAction(ULe
 	return { Schema->ActionSchema.CreateNull(Tag) };
 }
 
-FLearningAgentsActionSchemaElement ULearningAgentsActions::SpecifyContinuousAction(ULearningAgentsActionSchema* Schema, const int32 Size, const float Scale, const FName Tag)
+FLearningAgentsActionSchemaElement ULearningAgentsActions::SpecifyContinuousAction(ULearningAgentsActionSchema* Schema, const int32 Size, const float Scale, const bool bActive, const FName Tag)
 {
+	if (!bActive)
+	{
+		return SpecifyNullAction(Schema, Tag);
+	}
+
 	if (!Schema)
 	{
 		UE_LOG(LogLearning, Error, TEXT("SpecifyContinuousAction: Schema is nullptr."));
 		return FLearningAgentsActionSchemaElement();
 	}
-
+	
 	if (Size < 0)
 	{
 		UE_LOG(LogLearning, Error, TEXT("SpecifyContinuousAction: Invalid Continuous Action Size '%i'."), Size);
@@ -1088,13 +1093,23 @@ FLearningAgentsActionSchemaElement ULearningAgentsActions::SpecifyContinuousActi
 	return { Schema->ActionSchema.CreateContinuous({ Size, Scale }, Tag) };
 }
 
-FLearningAgentsActionSchemaElement ULearningAgentsActions::SpecifyExclusiveDiscreteAction(ULearningAgentsActionSchema* Schema, const int32 Size, const TArray<float>& PriorProbabilities, const FName Tag)
+FLearningAgentsActionSchemaElement ULearningAgentsActions::SpecifyExclusiveDiscreteAction(ULearningAgentsActionSchema* Schema, const int32 Size, const TArray<float>& PriorProbabilities, const bool bActive, const FName Tag)
 {
-	return SpecifyExclusiveDiscreteActionFromArrayView(Schema, Size, PriorProbabilities, Tag);
+	if (!bActive)
+	{
+		return SpecifyNullAction(Schema, Tag);
+	}
+
+	return SpecifyExclusiveDiscreteActionFromArrayView(Schema, Size, PriorProbabilities, bActive, Tag);
 }
 
-FLearningAgentsActionSchemaElement ULearningAgentsActions::SpecifyExclusiveDiscreteActionFromArrayView(ULearningAgentsActionSchema* Schema, const int32 Size, const TArrayView<const float> PriorProbabilities, const FName Tag)
+FLearningAgentsActionSchemaElement ULearningAgentsActions::SpecifyExclusiveDiscreteActionFromArrayView(ULearningAgentsActionSchema* Schema, const int32 Size, const TArrayView<const float> PriorProbabilities, const bool bActive, const FName Tag)
 {
+	if (!bActive)
+	{
+		return SpecifyNullAction(Schema, Tag);
+	}
+
 	if (!Schema)
 	{
 		UE_LOG(LogLearning, Error, TEXT("SpecifyExclusiveDiscreteActionFromArrayView: Schema is nullptr."));
@@ -1129,8 +1144,13 @@ FLearningAgentsActionSchemaElement ULearningAgentsActions::SpecifyExclusiveDiscr
 	return { Schema->ActionSchema.CreateDiscreteExclusive({ Size, MakeArrayView(NormalizedPriorProbabilities) }, Tag)};
 }
 
-FLearningAgentsActionSchemaElement ULearningAgentsActions::SpecifyNamedExclusiveDiscreteAction(ULearningAgentsActionSchema* Schema, const TArray<FName>& Names, const TMap<FName, float>& PriorProbabilities, const FName Tag)
+FLearningAgentsActionSchemaElement ULearningAgentsActions::SpecifyNamedExclusiveDiscreteAction(ULearningAgentsActionSchema* Schema, const TArray<FName>& Names, const TMap<FName, float>& PriorProbabilities, const bool bActive, const FName Tag)
 {
+	if (!bActive)
+	{
+		return SpecifyNullAction(Schema, Tag);
+	}
+
 	const int32 NameNum = Names.Num();
 
 	if (NameNum == 0)
@@ -1169,11 +1189,16 @@ FLearningAgentsActionSchemaElement ULearningAgentsActions::SpecifyNamedExclusive
 		SortedPriorProbabilities[Idx] = NamePriorProbabilities[NameIndices[Idx]];
 	}
 
-	return SpecifyNamedExclusiveDiscreteActionFromArrayViews(Schema, SortedNames, SortedPriorProbabilities, Tag);
+	return SpecifyNamedExclusiveDiscreteActionFromArrayViews(Schema, SortedNames, SortedPriorProbabilities, bActive, Tag);
 }
 
-FLearningAgentsActionSchemaElement ULearningAgentsActions::SpecifyNamedExclusiveDiscreteActionFromArrayViews(ULearningAgentsActionSchema* Schema, const TArrayView<const FName> Names, const TArrayView<const float> PriorProbabilities, const FName Tag)
+FLearningAgentsActionSchemaElement ULearningAgentsActions::SpecifyNamedExclusiveDiscreteActionFromArrayViews(ULearningAgentsActionSchema* Schema, const TArrayView<const FName> Names, const TArrayView<const float> PriorProbabilities, const bool bActive, const FName Tag)
 {
+	if (!bActive)
+	{
+		return SpecifyNullAction(Schema, Tag);
+	}
+
 	if (!Schema)
 	{
 		UE_LOG(LogLearning, Error, TEXT("SpecifyNamedExclusiveDiscreteActionFromArrayViews: Schema is nullptr."));
@@ -1211,13 +1236,23 @@ FLearningAgentsActionSchemaElement ULearningAgentsActions::SpecifyNamedExclusive
 	return { Schema->ActionSchema.CreateNamedDiscreteExclusive({ Names, MakeArrayView(NormalizedPriorProbabilities) }, Tag) };
 }
 
-FLearningAgentsActionSchemaElement ULearningAgentsActions::SpecifyInclusiveDiscreteAction(ULearningAgentsActionSchema* Schema, const int32 Size, const TArray<float>& PriorProbabilities, const FName Tag)
+FLearningAgentsActionSchemaElement ULearningAgentsActions::SpecifyInclusiveDiscreteAction(ULearningAgentsActionSchema* Schema, const int32 Size, const TArray<float>& PriorProbabilities, const bool bActive, const FName Tag)
 {
-	return SpecifyInclusiveDiscreteActionFromArrayView(Schema, Size, PriorProbabilities, Tag);
+	if (!bActive)
+	{
+		return SpecifyNullAction(Schema, Tag);
+	}
+
+	return SpecifyInclusiveDiscreteActionFromArrayView(Schema, Size, PriorProbabilities, bActive, Tag);
 }
 
-FLearningAgentsActionSchemaElement ULearningAgentsActions::SpecifyInclusiveDiscreteActionFromArrayView(ULearningAgentsActionSchema* Schema, const int32 Size, const TArrayView<const float> PriorProbabilities, const FName Tag)
+FLearningAgentsActionSchemaElement ULearningAgentsActions::SpecifyInclusiveDiscreteActionFromArrayView(ULearningAgentsActionSchema* Schema, const int32 Size, const TArrayView<const float> PriorProbabilities, const bool bActive, const FName Tag)
 {
+	if (!bActive)
+	{
+		return SpecifyNullAction(Schema, Tag);
+	}
+
 	if (!Schema)
 	{
 		UE_LOG(LogLearning, Error, TEXT("SpecifyInclusiveDiscreteActionFromArrayView: Schema is nullptr."));
@@ -1252,8 +1287,13 @@ FLearningAgentsActionSchemaElement ULearningAgentsActions::SpecifyInclusiveDiscr
 	return { Schema->ActionSchema.CreateDiscreteInclusive({ Size, MakeArrayView(NormalizedPriorProbabilities) }, Tag) };
 }
 
-FLearningAgentsActionSchemaElement ULearningAgentsActions::SpecifyNamedInclusiveDiscreteAction(ULearningAgentsActionSchema* Schema, const TArray<FName> Names, const TMap<FName, float>& PriorProbabilities, const FName Tag)
+FLearningAgentsActionSchemaElement ULearningAgentsActions::SpecifyNamedInclusiveDiscreteAction(ULearningAgentsActionSchema* Schema, const TArray<FName> Names, const TMap<FName, float>& PriorProbabilities, const bool bActive, const FName Tag)
 {
+	if (!bActive)
+	{
+		return SpecifyNullAction(Schema, Tag);
+	}
+
 	const int32 NameNum = Names.Num();
 
 	if (NameNum == 0)
@@ -1292,11 +1332,16 @@ FLearningAgentsActionSchemaElement ULearningAgentsActions::SpecifyNamedInclusive
 		SortedPriorProbabilities[Idx] = NamePriorProbabilities[NameIndices[Idx]];
 	}
 
-	return SpecifyNamedInclusiveDiscreteActionFromArrayViews(Schema, SortedNames, SortedPriorProbabilities, Tag);
+	return SpecifyNamedInclusiveDiscreteActionFromArrayViews(Schema, SortedNames, SortedPriorProbabilities, bActive, Tag);
 }
 
-FLearningAgentsActionSchemaElement ULearningAgentsActions::SpecifyNamedInclusiveDiscreteActionFromArrayViews(ULearningAgentsActionSchema* Schema, const TArrayView<const FName> Names, const TArrayView<const float> PriorProbabilities, const FName Tag)
+FLearningAgentsActionSchemaElement ULearningAgentsActions::SpecifyNamedInclusiveDiscreteActionFromArrayViews(ULearningAgentsActionSchema* Schema, const TArrayView<const FName> Names, const TArrayView<const float> PriorProbabilities, const bool bActive, const FName Tag)
 {
+	if (!bActive)
+	{
+		return SpecifyNullAction(Schema, Tag);
+	}
+
 	if (!Schema)
 	{
 		UE_LOG(LogLearning, Error, TEXT("SpecifyNamedInclusiveDiscreteActionFromArrayViews: Schema is nullptr."));
@@ -1333,8 +1378,13 @@ FLearningAgentsActionSchemaElement ULearningAgentsActions::SpecifyNamedInclusive
 	return { Schema->ActionSchema.CreateNamedDiscreteInclusive({ Names, MakeArrayView(NormalizedPriorProbabilities) }, Tag) };
 }
 
-FLearningAgentsActionSchemaElement ULearningAgentsActions::SpecifyStructAction(ULearningAgentsActionSchema* Schema, const TMap<FName, FLearningAgentsActionSchemaElement>& Elements, const FName Tag)
+FLearningAgentsActionSchemaElement ULearningAgentsActions::SpecifyStructAction(ULearningAgentsActionSchema* Schema, const TMap<FName, FLearningAgentsActionSchemaElement>& Elements, const bool bActive, const FName Tag)
 {
+	if (!bActive)
+	{
+		return SpecifyNullAction(Schema, Tag);
+	}
+
 	if (Elements.Num() == 0)
 	{
 		UE_LOG(LogLearning, Warning, TEXT("SpecifyStructAction: Specifying zero-sized Struct Action."));
@@ -1375,16 +1425,26 @@ FLearningAgentsActionSchemaElement ULearningAgentsActions::SpecifyStructAction(U
 		SortedSubElements[Idx] = SubElements[SubElementIndices[Idx]];
 	}
 
-	return SpecifyStructActionFromArrayViews(Schema, SortedSubElementNames, SortedSubElements, Tag);
+	return SpecifyStructActionFromArrayViews(Schema, SortedSubElementNames, SortedSubElements, bActive, Tag);
 }
 
-FLearningAgentsActionSchemaElement ULearningAgentsActions::SpecifyStructActionFromArrays(ULearningAgentsActionSchema* Schema, const TArray<FName>& ElementNames, const TArray<FLearningAgentsActionSchemaElement>& Elements, const FName Tag)
+FLearningAgentsActionSchemaElement ULearningAgentsActions::SpecifyStructActionFromArrays(ULearningAgentsActionSchema* Schema, const TArray<FName>& ElementNames, const TArray<FLearningAgentsActionSchemaElement>& Elements, const bool bActive, const FName Tag)
 {
-	return SpecifyStructActionFromArrayViews(Schema, ElementNames, Elements, Tag);
+	if (!bActive)
+	{
+		return SpecifyNullAction(Schema, Tag);
+	}
+
+	return SpecifyStructActionFromArrayViews(Schema, ElementNames, Elements, bActive, Tag);
 }
 
-FLearningAgentsActionSchemaElement ULearningAgentsActions::SpecifyStructActionFromArrayViews(ULearningAgentsActionSchema* Schema, const TArrayView<const FName> ElementNames, const TArrayView<const FLearningAgentsActionSchemaElement> Elements, const FName Tag)
+FLearningAgentsActionSchemaElement ULearningAgentsActions::SpecifyStructActionFromArrayViews(ULearningAgentsActionSchema* Schema, const TArrayView<const FName> ElementNames, const TArrayView<const FLearningAgentsActionSchemaElement> Elements, const bool bActive, const FName Tag)
 {
+	if (!bActive)
+	{
+		return SpecifyNullAction(Schema, Tag);
+	}
+
 	if (!Schema)
 	{
 		UE_LOG(LogLearning, Error, TEXT("SpecifyStructActionFromArrayViews: Schema is nullptr."));
@@ -1425,8 +1485,13 @@ FLearningAgentsActionSchemaElement ULearningAgentsActions::SpecifyStructActionFr
 	return { Schema->ActionSchema.CreateAnd({ ElementNames, SubElements }, Tag) };
 }
 
-FLearningAgentsActionSchemaElement ULearningAgentsActions::SpecifyExclusiveUnionAction(ULearningAgentsActionSchema* Schema, const TMap<FName, FLearningAgentsActionSchemaElement>& Elements, const TMap<FName, float>& PriorProbabilities, const FName Tag)
+FLearningAgentsActionSchemaElement ULearningAgentsActions::SpecifyExclusiveUnionAction(ULearningAgentsActionSchema* Schema, const TMap<FName, FLearningAgentsActionSchemaElement>& Elements, const TMap<FName, float>& PriorProbabilities, const bool bActive, const FName Tag)
 {
+	if (!bActive)
+	{
+		return SpecifyNullAction(Schema, Tag);
+	}
+
 	if (Elements.Num() == 0)
 	{
 		UE_LOG(LogLearning, Warning, TEXT("SpecifyExclusiveUnionAction: Specifying zero-sized Exclusive Union Action."));
@@ -1474,16 +1539,26 @@ FLearningAgentsActionSchemaElement ULearningAgentsActions::SpecifyExclusiveUnion
 		SortedSubElementPriorProbabilities[Idx] = SubElementPriorProbabilities[SubElementIndices[Idx]];
 	}
 
-	return SpecifyExclusiveUnionActionFromArrayViews(Schema, SortedSubElementNames, SortedSubElements, SortedSubElementPriorProbabilities, Tag);
+	return SpecifyExclusiveUnionActionFromArrayViews(Schema, SortedSubElementNames, SortedSubElements, SortedSubElementPriorProbabilities, bActive, Tag);
 }
 
-FLearningAgentsActionSchemaElement ULearningAgentsActions::SpecifyExclusiveUnionActionFromArrays(ULearningAgentsActionSchema* Schema, const TArray<FName>& ElementNames, const TArray<FLearningAgentsActionSchemaElement>& Elements, const TArray<float>& PriorProbabilities, const FName Tag)
+FLearningAgentsActionSchemaElement ULearningAgentsActions::SpecifyExclusiveUnionActionFromArrays(ULearningAgentsActionSchema* Schema, const TArray<FName>& ElementNames, const TArray<FLearningAgentsActionSchemaElement>& Elements, const TArray<float>& PriorProbabilities, const bool bActive, const FName Tag)
 {
-	return SpecifyExclusiveUnionActionFromArrayViews(Schema, ElementNames, Elements, PriorProbabilities, Tag);
+	if (!bActive)
+	{
+		return SpecifyNullAction(Schema, Tag);
+	}
+
+	return SpecifyExclusiveUnionActionFromArrayViews(Schema, ElementNames, Elements, PriorProbabilities, bActive, Tag);
 }
 
-FLearningAgentsActionSchemaElement ULearningAgentsActions::SpecifyExclusiveUnionActionFromArrayViews(ULearningAgentsActionSchema* Schema, const TArrayView<const FName> ElementNames, const TArrayView<const FLearningAgentsActionSchemaElement> Elements, const TArrayView<const float> PriorProbabilities, const FName Tag)
+FLearningAgentsActionSchemaElement ULearningAgentsActions::SpecifyExclusiveUnionActionFromArrayViews(ULearningAgentsActionSchema* Schema, const TArrayView<const FName> ElementNames, const TArrayView<const FLearningAgentsActionSchemaElement> Elements, const TArrayView<const float> PriorProbabilities, const bool bActive, const FName Tag)
 {
+	if (!bActive)
+	{
+		return SpecifyNullAction(Schema, Tag);
+	}
+
 	if (!Schema)
 	{
 		UE_LOG(LogLearning, Error, TEXT("SpecifyExclusiveUnionActionFromArrayViews: Schema is nullptr."));
@@ -1538,8 +1613,13 @@ FLearningAgentsActionSchemaElement ULearningAgentsActions::SpecifyExclusiveUnion
 	return { Schema->ActionSchema.CreateOrExclusive({ ElementNames, SubElements, MakeArrayView(NormalizedPriorProbabilities) }, Tag)};
 }
 
-FLearningAgentsActionSchemaElement ULearningAgentsActions::SpecifyInclusiveUnionAction(ULearningAgentsActionSchema* Schema, const TMap<FName, FLearningAgentsActionSchemaElement>& Elements, const TMap<FName, float>& PriorProbabilities, const FName Tag)
+FLearningAgentsActionSchemaElement ULearningAgentsActions::SpecifyInclusiveUnionAction(ULearningAgentsActionSchema* Schema, const TMap<FName, FLearningAgentsActionSchemaElement>& Elements, const TMap<FName, float>& PriorProbabilities, const bool bActive, const FName Tag)
 {
+	if (!bActive)
+	{
+		return SpecifyNullAction(Schema, Tag);
+	}
+
 	if (Elements.Num() == 0)
 	{
 		UE_LOG(LogLearning, Warning, TEXT("SpecifyInclusiveUnionAction: Specifying zero-sized Inclusive Union Action."));
@@ -1587,16 +1667,26 @@ FLearningAgentsActionSchemaElement ULearningAgentsActions::SpecifyInclusiveUnion
 		SortedSubElementPriorProbabilities[Idx] = SubElementPriorProbabilities[SubElementIndices[Idx]];
 	}
 
-	return SpecifyInclusiveUnionActionFromArrayViews(Schema, SortedSubElementNames, SortedSubElements, SortedSubElementPriorProbabilities, Tag);
+	return SpecifyInclusiveUnionActionFromArrayViews(Schema, SortedSubElementNames, SortedSubElements, SortedSubElementPriorProbabilities, bActive, Tag);
 }
 
-FLearningAgentsActionSchemaElement ULearningAgentsActions::SpecifyInclusiveUnionActionFromArrays(ULearningAgentsActionSchema* Schema, const TArray<FName> ElementNames, const TArray<FLearningAgentsActionSchemaElement>& Elements, const TArray<float>& PriorProbabilities, const FName Tag)
+FLearningAgentsActionSchemaElement ULearningAgentsActions::SpecifyInclusiveUnionActionFromArrays(ULearningAgentsActionSchema* Schema, const TArray<FName> ElementNames, const TArray<FLearningAgentsActionSchemaElement>& Elements, const TArray<float>& PriorProbabilities, const bool bActive, const FName Tag)
 {
-	return SpecifyInclusiveUnionActionFromArrayViews(Schema, ElementNames, Elements, PriorProbabilities, Tag);
+	if (!bActive)
+	{
+		return SpecifyNullAction(Schema, Tag);
+	}
+
+	return SpecifyInclusiveUnionActionFromArrayViews(Schema, ElementNames, Elements, PriorProbabilities, bActive, Tag);
 }
 
-FLearningAgentsActionSchemaElement ULearningAgentsActions::SpecifyInclusiveUnionActionFromArrayViews(ULearningAgentsActionSchema* Schema, const TArrayView<const FName> ElementNames, const TArrayView<const FLearningAgentsActionSchemaElement> Elements, const TArrayView<const float> PriorProbabilities, const FName Tag)
+FLearningAgentsActionSchemaElement ULearningAgentsActions::SpecifyInclusiveUnionActionFromArrayViews(ULearningAgentsActionSchema* Schema, const TArrayView<const FName> ElementNames, const TArrayView<const FLearningAgentsActionSchemaElement> Elements, const TArrayView<const float> PriorProbabilities, const bool bActive, const FName Tag)
 {
+	if (!bActive)
+	{
+		return SpecifyNullAction(Schema, Tag);
+	}
+
 	if (!Schema)
 	{
 		UE_LOG(LogLearning, Error, TEXT("SpecifyInclusiveUnionActionFromArrayViews: Schema is nullptr."));
@@ -1651,8 +1741,13 @@ FLearningAgentsActionSchemaElement ULearningAgentsActions::SpecifyInclusiveUnion
 	return { Schema->ActionSchema.CreateOrInclusive({ ElementNames, SubElements, MakeArrayView(NormalizedPriorProbabilities) }, Tag)};
 }
 
-FLearningAgentsActionSchemaElement ULearningAgentsActions::SpecifyStaticArrayAction(ULearningAgentsActionSchema* Schema, const FLearningAgentsActionSchemaElement Element, const int32 Num, const FName Tag)
+FLearningAgentsActionSchemaElement ULearningAgentsActions::SpecifyStaticArrayAction(ULearningAgentsActionSchema* Schema, const FLearningAgentsActionSchemaElement Element, const int32 Num, const bool bActive, const FName Tag)
 {
+	if (!bActive)
+	{
+		return SpecifyNullAction(Schema, Tag);
+	}
+
 	if (!Schema)
 	{
 		UE_LOG(LogLearning, Error, TEXT("SpecifyStaticArrayAction: Schema is nullptr."));
@@ -1679,13 +1774,23 @@ FLearningAgentsActionSchemaElement ULearningAgentsActions::SpecifyStaticArrayAct
 	return { Schema->ActionSchema.CreateArray({ Element.SchemaElement, Num }, Tag) };
 }
 
-FLearningAgentsActionSchemaElement ULearningAgentsActions::SpecifyPairAction(ULearningAgentsActionSchema* Schema, const FLearningAgentsActionSchemaElement Key, const FLearningAgentsActionSchemaElement Value, const FName Tag)
+FLearningAgentsActionSchemaElement ULearningAgentsActions::SpecifyPairAction(ULearningAgentsActionSchema* Schema, const FLearningAgentsActionSchemaElement Key, const FLearningAgentsActionSchemaElement Value, const bool bActive, const FName Tag)
 {
-	return SpecifyStructActionFromArrayViews(Schema, { TEXT("Key"), TEXT("Value") }, { Key, Value }, Tag);
+	if (!bActive)
+	{
+		return SpecifyNullAction(Schema, Tag);
+	}
+
+	return SpecifyStructActionFromArrayViews(Schema, { TEXT("Key"), TEXT("Value") }, { Key, Value }, bActive, Tag);
 }
 
-FLearningAgentsActionSchemaElement ULearningAgentsActions::SpecifyEnumAction(ULearningAgentsActionSchema* Schema, const UEnum* Enum, const TMap<uint8, float>& PriorProbabilities, const FName Tag)
+FLearningAgentsActionSchemaElement ULearningAgentsActions::SpecifyEnumAction(ULearningAgentsActionSchema* Schema, const UEnum* Enum, const TMap<uint8, float>& PriorProbabilities, const bool bActive, const FName Tag)
 {
+	if (!bActive)
+	{
+		return SpecifyNullAction(Schema, Tag);
+	}
+
 	if (!Enum)
 	{
 		UE_LOG(LogLearning, Error, TEXT("SpecifyEnumAction: Enum is nullptr."));
@@ -1709,27 +1814,42 @@ FLearningAgentsActionSchemaElement ULearningAgentsActions::SpecifyEnumAction(ULe
 		}
 	}
 
-	return SpecifyEnumActionFromArrayView(Schema, Enum, PriorProbabilitiesArray, Tag);
+	return SpecifyEnumActionFromArrayView(Schema, Enum, PriorProbabilitiesArray, bActive, Tag);
 }
 
-FLearningAgentsActionSchemaElement ULearningAgentsActions::SpecifyEnumActionFromArray(ULearningAgentsActionSchema* Schema, const UEnum* Enum, const TArray<float>& PriorProbabilities, const FName Tag)
+FLearningAgentsActionSchemaElement ULearningAgentsActions::SpecifyEnumActionFromArray(ULearningAgentsActionSchema* Schema, const UEnum* Enum, const TArray<float>& PriorProbabilities, const bool bActive, const FName Tag)
 {
-	return SpecifyEnumActionFromArrayView(Schema, Enum, PriorProbabilities, Tag);
+	if (!bActive)
+	{
+		return SpecifyNullAction(Schema, Tag);
+	}
+
+	return SpecifyEnumActionFromArrayView(Schema, Enum, PriorProbabilities, bActive, Tag);
 }
 
-FLearningAgentsActionSchemaElement ULearningAgentsActions::SpecifyEnumActionFromArrayView(ULearningAgentsActionSchema* Schema, const UEnum* Enum, const TArrayView<const float> PriorProbabilities, const FName Tag)
+FLearningAgentsActionSchemaElement ULearningAgentsActions::SpecifyEnumActionFromArrayView(ULearningAgentsActionSchema* Schema, const UEnum* Enum, const TArrayView<const float> PriorProbabilities, const bool bActive, const FName Tag)
 {
+	if (!bActive)
+	{
+		return SpecifyNullAction(Schema, Tag);
+	}
+
 	if (!Enum)
 	{
 		UE_LOG(LogLearning, Error, TEXT("SpecifyEnumActionFromArrayView: Enum is nullptr."));
 		return FLearningAgentsActionSchemaElement();
 	}
 
-	return SpecifyExclusiveDiscreteActionFromArrayView(Schema, Enum->NumEnums() - 1, PriorProbabilities, Tag);
+	return SpecifyExclusiveDiscreteActionFromArrayView(Schema, Enum->NumEnums() - 1, PriorProbabilities, bActive, Tag);
 }
 
-FLearningAgentsActionSchemaElement ULearningAgentsActions::SpecifyBitmaskAction(ULearningAgentsActionSchema* Schema, const UEnum* Enum, const TMap<uint8, float>& PriorProbabilities, const FName Tag)
+FLearningAgentsActionSchemaElement ULearningAgentsActions::SpecifyBitmaskAction(ULearningAgentsActionSchema* Schema, const UEnum* Enum, const TMap<uint8, float>& PriorProbabilities, const bool bActive, const FName Tag)
 {
+	if (!bActive)
+	{
+		return SpecifyNullAction(Schema, Tag);
+	}
+
 	if (!Enum)
 	{
 		UE_LOG(LogLearning, Error, TEXT("SpecifyBitmaskAction: Enum is nullptr."));
@@ -1759,16 +1879,26 @@ FLearningAgentsActionSchemaElement ULearningAgentsActions::SpecifyBitmaskAction(
 		}
 	}
 
-	return SpecifyBitmaskActionFromArrayView(Schema, Enum, PriorProbabilitiesArray, Tag);
+	return SpecifyBitmaskActionFromArrayView(Schema, Enum, PriorProbabilitiesArray, bActive, Tag);
 }
 
-FLearningAgentsActionSchemaElement ULearningAgentsActions::SpecifyBitmaskActionFromArray(ULearningAgentsActionSchema* Schema, const UEnum* Enum, const TArray<float>& PriorProbabilities, const FName Tag)
+FLearningAgentsActionSchemaElement ULearningAgentsActions::SpecifyBitmaskActionFromArray(ULearningAgentsActionSchema* Schema, const UEnum* Enum, const TArray<float>& PriorProbabilities, const bool bActive, const FName Tag)
 {
-	return SpecifyBitmaskActionFromArrayView(Schema, Enum, PriorProbabilities, Tag);
+	if (!bActive)
+	{
+		return SpecifyNullAction(Schema, Tag);
+	}
+
+	return SpecifyBitmaskActionFromArrayView(Schema, Enum, PriorProbabilities, bActive, Tag);
 }
 
-FLearningAgentsActionSchemaElement ULearningAgentsActions::SpecifyBitmaskActionFromArrayView(ULearningAgentsActionSchema* Schema, const UEnum* Enum, const TArrayView<const float> PriorProbabilities, const FName Tag)
+FLearningAgentsActionSchemaElement ULearningAgentsActions::SpecifyBitmaskActionFromArrayView(ULearningAgentsActionSchema* Schema, const UEnum* Enum, const TArrayView<const float> PriorProbabilities, const bool bActive, const FName Tag)
 {
+	if (!bActive)
+	{
+		return SpecifyNullAction(Schema, Tag);
+	}
+
 	if (!Enum)
 	{
 		UE_LOG(LogLearning, Error, TEXT("SpecifyBitmaskActionFromArrayView: Enum is nullptr."));
@@ -1781,21 +1911,36 @@ FLearningAgentsActionSchemaElement ULearningAgentsActions::SpecifyBitmaskActionF
 		return FLearningAgentsActionSchemaElement();
 	}
 
-	return SpecifyInclusiveDiscreteActionFromArrayView(Schema, Enum->NumEnums() - 1, PriorProbabilities, Tag);
+	return SpecifyInclusiveDiscreteActionFromArrayView(Schema, Enum->NumEnums() - 1, PriorProbabilities, bActive, Tag);
 }
 
-FLearningAgentsActionSchemaElement ULearningAgentsActions::SpecifyOptionalAction(ULearningAgentsActionSchema* Schema, const FLearningAgentsActionSchemaElement Element, const float PriorProbability, const FName Tag)
+FLearningAgentsActionSchemaElement ULearningAgentsActions::SpecifyOptionalAction(ULearningAgentsActionSchema* Schema, const FLearningAgentsActionSchemaElement Element, const float PriorProbability, const bool bActive, const FName Tag)
 {
-	return SpecifyExclusiveUnionActionFromArrayViews(Schema, { TEXT("Null"), TEXT("Valid") }, { SpecifyNullAction(Schema), Element }, { 1.0f - PriorProbability, PriorProbability }, Tag);
+	if (!bActive)
+	{
+		return SpecifyNullAction(Schema, Tag);
+	}
+
+	return SpecifyExclusiveUnionActionFromArrayViews(Schema, { TEXT("Null"), TEXT("Valid") }, { SpecifyNullAction(Schema), Element }, { 1.0f - PriorProbability, PriorProbability }, bActive, Tag);
 }
 
-FLearningAgentsActionSchemaElement ULearningAgentsActions::SpecifyEitherAction(ULearningAgentsActionSchema* Schema, const FLearningAgentsActionSchemaElement A, const FLearningAgentsActionSchemaElement B, const float PriorProbabilityOfA, const FName Tag)
+FLearningAgentsActionSchemaElement ULearningAgentsActions::SpecifyEitherAction(ULearningAgentsActionSchema* Schema, const FLearningAgentsActionSchemaElement A, const FLearningAgentsActionSchemaElement B, const float PriorProbabilityOfA, const bool bActive, const FName Tag)
 {
-	return SpecifyExclusiveUnionActionFromArrayViews(Schema, { TEXT("A"), TEXT("B") }, { A, B }, { 1.0f - PriorProbabilityOfA, PriorProbabilityOfA }, Tag);
+	if (!bActive)
+	{
+		return SpecifyNullAction(Schema, Tag);
+	}
+
+	return SpecifyExclusiveUnionActionFromArrayViews(Schema, { TEXT("A"), TEXT("B") }, { A, B }, { 1.0f - PriorProbabilityOfA, PriorProbabilityOfA }, bActive, Tag);
 }
 
-FLearningAgentsActionSchemaElement ULearningAgentsActions::SpecifyEncodingAction(ULearningAgentsActionSchema* Schema, const FLearningAgentsActionSchemaElement Element, const int32 EncodingSize, const int32 HiddenLayerNum, const ELearningAgentsActivationFunction ActivationFunction, const FName Tag)
+FLearningAgentsActionSchemaElement ULearningAgentsActions::SpecifyEncodingAction(ULearningAgentsActionSchema* Schema, const FLearningAgentsActionSchemaElement Element, const int32 EncodingSize, const int32 HiddenLayerNum, const ELearningAgentsActivationFunction ActivationFunction, const bool bActive, const FName Tag)
 {
+	if (!bActive)
+	{
+		return SpecifyNullAction(Schema, Tag);
+	}
+
 	if (!Schema)
 	{
 		UE_LOG(LogLearning, Error, TEXT("SpecifyEncodingAction: Schema is nullptr."));
@@ -1823,29 +1968,54 @@ FLearningAgentsActionSchemaElement ULearningAgentsActions::SpecifyEncodingAction
 	return { Schema->ActionSchema.CreateEncoding({ Element.SchemaElement, EncodingSize, HiddenLayerNum, UE::Learning::Agents::Action::Private::GetEncodingActivationFunction(ActivationFunction) }, Tag)};
 }
 
-FLearningAgentsActionSchemaElement ULearningAgentsActions::SpecifyBoolAction(ULearningAgentsActionSchema* Schema, const float PriorProbability, const FName Tag)
+FLearningAgentsActionSchemaElement ULearningAgentsActions::SpecifyBoolAction(ULearningAgentsActionSchema* Schema, const float PriorProbability, const bool bActive, const FName Tag)
 {
-	return SpecifyExclusiveDiscreteActionFromArrayView(Schema, 2, { 1.0f - PriorProbability, PriorProbability }, Tag);
+	if (!bActive)
+	{
+		return SpecifyNullAction(Schema, Tag);
+	}
+
+	return SpecifyExclusiveDiscreteActionFromArrayView(Schema, 2, { 1.0f - PriorProbability, PriorProbability }, bActive, Tag);
 }
 
-FLearningAgentsActionSchemaElement ULearningAgentsActions::SpecifyFloatAction(ULearningAgentsActionSchema* Schema, const float FloatScale, const FName Tag)
+FLearningAgentsActionSchemaElement ULearningAgentsActions::SpecifyFloatAction(ULearningAgentsActionSchema* Schema, const float FloatScale, const bool bActive, const FName Tag)
 {
-	return SpecifyContinuousAction(Schema, 1, FloatScale, Tag);
+	if (!bActive)
+	{
+		return SpecifyNullAction(Schema, Tag);
+	}
+
+	return SpecifyContinuousAction(Schema, 1, FloatScale, bActive, Tag);
 }
 
-FLearningAgentsActionSchemaElement ULearningAgentsActions::SpecifyLocationAction(ULearningAgentsActionSchema* Schema, const float LocationScale, const FName Tag)
+FLearningAgentsActionSchemaElement ULearningAgentsActions::SpecifyLocationAction(ULearningAgentsActionSchema* Schema, const float LocationScale, const bool bActive, const FName Tag)
 {
-	return SpecifyContinuousAction(Schema, 3, LocationScale, Tag);
+	if (!bActive)
+	{
+		return SpecifyNullAction(Schema, Tag);
+	}
+
+	return SpecifyContinuousAction(Schema, 3, LocationScale, bActive, Tag);
 }
 
-FLearningAgentsActionSchemaElement ULearningAgentsActions::SpecifyRotationAction(ULearningAgentsActionSchema* Schema, const float RotationScale, const FName Tag)
+FLearningAgentsActionSchemaElement ULearningAgentsActions::SpecifyRotationAction(ULearningAgentsActionSchema* Schema, const float RotationScale, const bool bActive, const FName Tag)
 {
-	return SpecifyContinuousAction(Schema, 3, RotationScale, Tag);
+	if (!bActive)
+	{
+		return SpecifyNullAction(Schema, Tag);
+	}
+
+	return SpecifyContinuousAction(Schema, 3, RotationScale, bActive, Tag);
 }
 
-FLearningAgentsActionSchemaElement ULearningAgentsActions::SpecifyScaleAction(ULearningAgentsActionSchema* Schema, const float ScaleScale, const FName Tag)
+FLearningAgentsActionSchemaElement ULearningAgentsActions::SpecifyScaleAction(ULearningAgentsActionSchema* Schema, const float ScaleScale, const bool bActive, const FName Tag)
 {
-	return SpecifyContinuousAction(Schema, 3, ScaleScale, Tag);
+	if (!bActive)
+	{
+		return SpecifyNullAction(Schema, Tag);
+	}
+
+	return SpecifyContinuousAction(Schema, 3, ScaleScale, bActive, Tag);
 }
 
 FLearningAgentsActionSchemaElement ULearningAgentsActions::SpecifyTransformAction(
@@ -1853,35 +2023,56 @@ FLearningAgentsActionSchemaElement ULearningAgentsActions::SpecifyTransformActio
 	const float LocationScale ,
 	const float RotationScale,
 	const float ScaleScale,
+	const bool bActive,
 	const FName Tag)
 {
+	if (!bActive)
+	{
+		return SpecifyNullAction(Schema, Tag);
+	}
+
 	return SpecifyStructActionFromArrayViews(Schema,
 		{
 			TEXT("Location"),
 			TEXT("Rotation"),
 			TEXT("Scale")
 		},
-		{
-			SpecifyLocationAction(Schema, LocationScale),
-			SpecifyRotationAction(Schema, RotationScale),
-			SpecifyScaleAction(Schema, ScaleScale)
-		}, 
-		Tag);
+			{
+				SpecifyLocationAction(Schema, LocationScale),
+				SpecifyRotationAction(Schema, RotationScale),
+				SpecifyScaleAction(Schema, ScaleScale)
+			},
+		bActive, Tag);
 }
 
-FLearningAgentsActionSchemaElement ULearningAgentsActions::SpecifyAngleAction(ULearningAgentsActionSchema* Schema, const float AngleScale, const FName Tag)
+FLearningAgentsActionSchemaElement ULearningAgentsActions::SpecifyAngleAction(ULearningAgentsActionSchema* Schema, const float AngleScale, const bool bActive, const FName Tag)
 {
-	return SpecifyContinuousAction(Schema, 1, AngleScale, Tag);
+	if (!bActive)
+	{
+		return SpecifyNullAction(Schema, Tag);
+	}
+
+	return SpecifyContinuousAction(Schema, 1, AngleScale, bActive, Tag);
 }
 
-FLearningAgentsActionSchemaElement ULearningAgentsActions::SpecifyVelocityAction(ULearningAgentsActionSchema* Schema, const float VelocityScale, const FName Tag)
+FLearningAgentsActionSchemaElement ULearningAgentsActions::SpecifyVelocityAction(ULearningAgentsActionSchema* Schema, const float VelocityScale, const bool bActive, const FName Tag)
 {
-	return SpecifyContinuousAction(Schema, 3, VelocityScale, Tag);
+	if (!bActive)
+	{
+		return SpecifyNullAction(Schema, Tag);
+	}
+
+	return SpecifyContinuousAction(Schema, 3, VelocityScale, bActive, Tag);
 }
 
-FLearningAgentsActionSchemaElement ULearningAgentsActions::SpecifyDirectionAction(ULearningAgentsActionSchema* Schema, const FName Tag)
+FLearningAgentsActionSchemaElement ULearningAgentsActions::SpecifyDirectionAction(ULearningAgentsActionSchema* Schema, const bool bActive, const FName Tag)
 {
-	return SpecifyContinuousAction(Schema, 3, 1.0f, Tag);
+	if (!bActive)
+	{
+		return SpecifyNullAction(Schema, Tag);
+	}
+
+	return SpecifyContinuousAction(Schema, 3, 1.0f, bActive, Tag);
 }
 
 void ULearningAgentsActions::LogAction(const ULearningAgentsActionObject* Object, const FLearningAgentsActionObjectElement Element)
@@ -1909,6 +2100,7 @@ FLearningAgentsActionObjectElement ULearningAgentsActions::MakeNullAction(ULearn
 FLearningAgentsActionObjectElement ULearningAgentsActions::MakeContinuousAction(
 	ULearningAgentsActionObject* Object,
 	const TArray<float>& Values,
+    const bool bActive,	
 	const FName Tag,
 	const bool bVisualLoggerEnabled,
 	ULearningAgentsManagerListener* VisualLoggerListener,
@@ -1916,12 +2108,18 @@ FLearningAgentsActionObjectElement ULearningAgentsActions::MakeContinuousAction(
 	const FVector VisualLoggerLocation,
 	const FLinearColor VisualLoggerColor)
 {
-	return MakeContinuousActionFromArrayView(Object, Values, Tag, bVisualLoggerEnabled, VisualLoggerListener, VisualLoggerAgentId, VisualLoggerLocation, VisualLoggerColor);
+	if (!bActive)
+	{
+		return MakeNullAction(Object, Tag);
+	}
+
+	return MakeContinuousActionFromArrayView(Object, Values, bActive, Tag, bVisualLoggerEnabled, VisualLoggerListener, VisualLoggerAgentId, VisualLoggerLocation, VisualLoggerColor);
 }
 
 FLearningAgentsActionObjectElement ULearningAgentsActions::MakeContinuousActionFromArrayView(
 	ULearningAgentsActionObject* Object,
 	const TArrayView<const float> Values,
+    const bool bActive,
 	const FName Tag,
 	const bool bVisualLoggerEnabled,
 	ULearningAgentsManagerListener* VisualLoggerListener,
@@ -1929,6 +2127,11 @@ FLearningAgentsActionObjectElement ULearningAgentsActions::MakeContinuousActionF
 	const FVector VisualLoggerLocation,
 	const FLinearColor VisualLoggerColor)
 {
+	if (!bActive)
+	{
+		return MakeNullAction(Object, Tag);
+	}
+
 	if (!Object)
 	{
 		UE_LOG(LogLearning, Error, TEXT("MakeContinuousActionFromArrayView: Object is nullptr."));
@@ -1962,6 +2165,7 @@ FLearningAgentsActionObjectElement ULearningAgentsActions::MakeContinuousActionF
 FLearningAgentsActionObjectElement ULearningAgentsActions::MakeExclusiveDiscreteAction(
 	ULearningAgentsActionObject* Object,
 	const int32 Index,
+    const bool bActive,
 	const FName Tag,
 	const bool bVisualLoggerEnabled,
 	ULearningAgentsManagerListener* VisualLoggerListener,
@@ -1969,6 +2173,11 @@ FLearningAgentsActionObjectElement ULearningAgentsActions::MakeExclusiveDiscrete
 	const FVector VisualLoggerLocation,
 	const FLinearColor VisualLoggerColor)
 {
+	if (!bActive)
+	{
+		return MakeNullAction(Object, Tag);
+	}
+
 	if (!Object)
 	{
 		UE_LOG(LogLearning, Error, TEXT("MakeExclusiveDiscreteAction: Object is nullptr."));
@@ -2002,6 +2211,7 @@ FLearningAgentsActionObjectElement ULearningAgentsActions::MakeExclusiveDiscrete
 FLearningAgentsActionObjectElement ULearningAgentsActions::MakeNamedExclusiveDiscreteAction(
 	ULearningAgentsActionObject* Object,
 	const FName Name,
+    const bool bActive,	
 	const FName Tag,
 	const bool bVisualLoggerEnabled,
 	ULearningAgentsManagerListener* VisualLoggerListener,
@@ -2009,6 +2219,11 @@ FLearningAgentsActionObjectElement ULearningAgentsActions::MakeNamedExclusiveDis
 	const FVector VisualLoggerLocation,
 	const FLinearColor VisualLoggerColor)
 {
+	if (!bActive)
+	{
+		return MakeNullAction(Object, Tag);
+	}
+
 	if (!Object)
 	{
 		UE_LOG(LogLearning, Error, TEXT("MakeNamedExclusiveDiscreteAction: Object is nullptr."));
@@ -2036,6 +2251,7 @@ FLearningAgentsActionObjectElement ULearningAgentsActions::MakeNamedExclusiveDis
 FLearningAgentsActionObjectElement ULearningAgentsActions::MakeInclusiveDiscreteAction(
 	ULearningAgentsActionObject* Object,
 	const TArray<int32>& Indices,
+    const bool bActive,
 	const FName Tag,
 	const bool bVisualLoggerEnabled,
 	ULearningAgentsManagerListener* VisualLoggerListener,
@@ -2043,12 +2259,18 @@ FLearningAgentsActionObjectElement ULearningAgentsActions::MakeInclusiveDiscrete
 	const FVector VisualLoggerLocation,
 	const FLinearColor VisualLoggerColor)
 {
-	return MakeInclusiveDiscreteActionFromArrayView(Object, Indices, Tag, bVisualLoggerEnabled, VisualLoggerListener, VisualLoggerAgentId, VisualLoggerLocation, VisualLoggerColor);
+	if (!bActive)
+	{
+		return MakeNullAction(Object, Tag);
+	}
+
+	return MakeInclusiveDiscreteActionFromArrayView(Object, Indices, bActive, Tag, bVisualLoggerEnabled, VisualLoggerListener, VisualLoggerAgentId, VisualLoggerLocation, VisualLoggerColor);
 }
 
 FLearningAgentsActionObjectElement ULearningAgentsActions::MakeInclusiveDiscreteActionFromArrayView(
 	ULearningAgentsActionObject* Object,
 	const TArrayView<const int32> Indices,
+    const bool bActive,	
 	const FName Tag,
 	const bool bVisualLoggerEnabled,
 	ULearningAgentsManagerListener* VisualLoggerListener,
@@ -2056,6 +2278,11 @@ FLearningAgentsActionObjectElement ULearningAgentsActions::MakeInclusiveDiscrete
 	const FVector VisualLoggerLocation,
 	const FLinearColor VisualLoggerColor)
 {
+	if (!bActive)
+	{
+		return MakeNullAction(Object, Tag);
+	}
+
 	if (!Object)
 	{
 		UE_LOG(LogLearning, Error, TEXT("MakeInclusiveDiscreteActionFromArrayView: Object is nullptr."));
@@ -2100,6 +2327,7 @@ FLearningAgentsActionObjectElement ULearningAgentsActions::MakeInclusiveDiscrete
 FLearningAgentsActionObjectElement ULearningAgentsActions::MakeNamedInclusiveDiscreteAction(
 	ULearningAgentsActionObject* Object,
 	const TArray<FName>& Names,
+    const bool bActive,	
 	const FName Tag,
 	const bool bVisualLoggerEnabled,
 	ULearningAgentsManagerListener* VisualLoggerListener,
@@ -2107,12 +2335,18 @@ FLearningAgentsActionObjectElement ULearningAgentsActions::MakeNamedInclusiveDis
 	const FVector VisualLoggerLocation,
 	const FLinearColor VisualLoggerColor)
 {
-	return MakeNamedInclusiveDiscreteActionFromArrayView(Object, Names, Tag, bVisualLoggerEnabled, VisualLoggerListener, VisualLoggerAgentId, VisualLoggerLocation, VisualLoggerColor);
+	if (!bActive)
+	{
+		return MakeNullAction(Object, Tag);
+	}
+
+	return MakeNamedInclusiveDiscreteActionFromArrayView(Object, Names, bActive, Tag, bVisualLoggerEnabled, VisualLoggerListener, VisualLoggerAgentId, VisualLoggerLocation, VisualLoggerColor);
 }
 
 FLearningAgentsActionObjectElement ULearningAgentsActions::MakeNamedInclusiveDiscreteActionFromArrayView(
 	ULearningAgentsActionObject* Object,
 	const TArrayView<const FName> Names,
+    const bool bActive,
 	const FName Tag,
 	const bool bVisualLoggerEnabled,
 	ULearningAgentsManagerListener* VisualLoggerListener,
@@ -2120,6 +2354,11 @@ FLearningAgentsActionObjectElement ULearningAgentsActions::MakeNamedInclusiveDis
 	const FVector VisualLoggerLocation,
 	const FLinearColor VisualLoggerColor)
 {
+	if (!bActive)
+	{
+		return MakeNullAction(Object, Tag);
+	}
+
 	if (!Object)
 	{
 		UE_LOG(LogLearning, Error, TEXT("MakeNamedInclusiveDiscreteActionFromArrayView: Object is nullptr."));
@@ -2150,8 +2389,13 @@ FLearningAgentsActionObjectElement ULearningAgentsActions::MakeNamedInclusiveDis
 	return { Object->ActionObject.CreateNamedDiscreteInclusive({ Names }, Tag) };
 }
 
-FLearningAgentsActionObjectElement ULearningAgentsActions::MakeStructAction(ULearningAgentsActionObject* Object, const TMap<FName, FLearningAgentsActionObjectElement>& Elements, const FName Tag)
+FLearningAgentsActionObjectElement ULearningAgentsActions::MakeStructAction(ULearningAgentsActionObject* Object, const TMap<FName, FLearningAgentsActionObjectElement>& Elements, const bool bActive, const FName Tag)
 {
+	if (!bActive)
+	{
+		return MakeNullAction(Object, Tag);
+	}
+
 	if (Elements.Num() == 0)
 	{
 		UE_LOG(LogLearning, Warning, TEXT("MakeStructAction: Creating zero-sized Struct Action."));
@@ -2170,16 +2414,26 @@ FLearningAgentsActionObjectElement ULearningAgentsActions::MakeStructAction(ULea
 		SubElementNames.Add(Element.Key);
 	}
 
-	return MakeStructActionFromArrayViews(Object, SubElementNames, SubElements, Tag);
+	return MakeStructActionFromArrayViews(Object, SubElementNames, SubElements, bActive, Tag);
 }
 
-FLearningAgentsActionObjectElement ULearningAgentsActions::MakeStructActionFromArrays(ULearningAgentsActionObject* Object, const TArray<FName>& ElementNames, const TArray<FLearningAgentsActionObjectElement>& Elements, const FName Tag)
+FLearningAgentsActionObjectElement ULearningAgentsActions::MakeStructActionFromArrays(ULearningAgentsActionObject* Object, const TArray<FName>& ElementNames, const TArray<FLearningAgentsActionObjectElement>& Elements, const bool bActive, const FName Tag)
 {
-	return MakeStructActionFromArrayViews(Object, ElementNames, Elements, Tag);
+	if (!bActive)
+	{
+		return MakeNullAction(Object, Tag);
+	}
+
+	return MakeStructActionFromArrayViews(Object, ElementNames, Elements, bActive, Tag);
 }
 
-FLearningAgentsActionObjectElement ULearningAgentsActions::MakeStructActionFromArrayViews(ULearningAgentsActionObject* Object, const TArrayView<const FName> ElementNames, const TArrayView<const FLearningAgentsActionObjectElement> Elements, const FName Tag)
+FLearningAgentsActionObjectElement ULearningAgentsActions::MakeStructActionFromArrayViews(ULearningAgentsActionObject* Object, const TArrayView<const FName> ElementNames, const TArrayView<const FLearningAgentsActionObjectElement> Elements, const bool bActive, const FName Tag)
 {
+	if (!bActive)
+	{
+		return MakeNullAction(Object, Tag);
+	}
+
 	if (!Object)
 	{
 		UE_LOG(LogLearning, Error, TEXT("MakeStructActionFromArrayViews: Object is nullptr."));
@@ -2220,8 +2474,13 @@ FLearningAgentsActionObjectElement ULearningAgentsActions::MakeStructActionFromA
 	return { Object->ActionObject.CreateAnd({ ElementNames, SubElements }, Tag) };
 }
 
-FLearningAgentsActionObjectElement ULearningAgentsActions::MakeExclusiveUnionAction(ULearningAgentsActionObject* Object, const FName ElementName, const FLearningAgentsActionObjectElement Element, const FName Tag)
+FLearningAgentsActionObjectElement ULearningAgentsActions::MakeExclusiveUnionAction(ULearningAgentsActionObject* Object, const FName ElementName, const FLearningAgentsActionObjectElement Element, const bool bActive, const FName Tag)
 {
+	if (!bActive)
+	{
+		return MakeNullAction(Object, Tag);
+	}
+
 	if (!Object)
 	{
 		UE_LOG(LogLearning, Error, TEXT("MakeExclusiveUnionAction: Object is nullptr."));
@@ -2237,8 +2496,13 @@ FLearningAgentsActionObjectElement ULearningAgentsActions::MakeExclusiveUnionAct
 	return { Object->ActionObject.CreateOrExclusive({ ElementName, Element.ObjectElement }, Tag) };
 }
 
-FLearningAgentsActionObjectElement ULearningAgentsActions::MakeInclusiveUnionAction(ULearningAgentsActionObject* Object, const TMap<FName, FLearningAgentsActionObjectElement>& Elements, const FName Tag)
+FLearningAgentsActionObjectElement ULearningAgentsActions::MakeInclusiveUnionAction(ULearningAgentsActionObject* Object, const TMap<FName, FLearningAgentsActionObjectElement>& Elements, const bool bActive, const FName Tag)
 {
+	if (!bActive)
+	{
+		return MakeNullAction(Object, Tag);
+	}
+
 	const int32 SubElementNum = Elements.Num();
 
 	TArray<FLearningAgentsActionObjectElement, TInlineAllocator<16>> SubElements;
@@ -2252,16 +2516,26 @@ FLearningAgentsActionObjectElement ULearningAgentsActions::MakeInclusiveUnionAct
 		SubElementNames.Add(Element.Key);
 	}
 
-	return MakeInclusiveUnionActionFromArrayViews(Object, SubElementNames, SubElements, Tag);
+	return MakeInclusiveUnionActionFromArrayViews(Object, SubElementNames, SubElements, bActive, Tag);
 }
 
-FLearningAgentsActionObjectElement ULearningAgentsActions::MakeInclusiveUnionActionFromArrays(ULearningAgentsActionObject* Object, const TArray<FName>& ElementNames, const TArray<FLearningAgentsActionObjectElement>& Elements, const FName Tag)
+FLearningAgentsActionObjectElement ULearningAgentsActions::MakeInclusiveUnionActionFromArrays(ULearningAgentsActionObject* Object, const TArray<FName>& ElementNames, const TArray<FLearningAgentsActionObjectElement>& Elements, const bool bActive, const FName Tag)
 {
-	return MakeInclusiveUnionActionFromArrayViews(Object, ElementNames, Elements, Tag);
+	if (!bActive)
+	{
+		return MakeNullAction(Object, Tag);
+	}
+
+	return MakeInclusiveUnionActionFromArrayViews(Object, ElementNames, Elements, bActive, Tag);
 }
 
-FLearningAgentsActionObjectElement ULearningAgentsActions::MakeInclusiveUnionActionFromArrayViews(ULearningAgentsActionObject* Object, const TArrayView<const FName> ElementNames, const TArrayView<const FLearningAgentsActionObjectElement> Elements, const FName Tag)
+FLearningAgentsActionObjectElement ULearningAgentsActions::MakeInclusiveUnionActionFromArrayViews(ULearningAgentsActionObject* Object, const TArrayView<const FName> ElementNames, const TArrayView<const FLearningAgentsActionObjectElement> Elements, const bool bActive, const FName Tag)
 {
+	if (!bActive)
+	{
+		return MakeNullAction(Object, Tag);
+	}
+
 	if (!Object)
 	{
 		UE_LOG(LogLearning, Error, TEXT("MakeInclusiveUnionActionFromArrayViews: Object is nullptr."));
@@ -2297,13 +2571,23 @@ FLearningAgentsActionObjectElement ULearningAgentsActions::MakeInclusiveUnionAct
 	return { Object->ActionObject.CreateOrInclusive({ ElementNames, SubElements }, Tag) };
 }
 
-FLearningAgentsActionObjectElement ULearningAgentsActions::MakeStaticArrayAction(ULearningAgentsActionObject* Object, const TArray<FLearningAgentsActionObjectElement>& Elements, const FName Tag)
+FLearningAgentsActionObjectElement ULearningAgentsActions::MakeStaticArrayAction(ULearningAgentsActionObject* Object, const TArray<FLearningAgentsActionObjectElement>& Elements, const bool bActive, const FName Tag)
 {
-	return MakeStaticArrayActionFromArrayView(Object, Elements, Tag);
+	if (!bActive)
+	{
+		return MakeNullAction(Object, Tag);
+	}
+
+	return MakeStaticArrayActionFromArrayView(Object, Elements, bActive, Tag);
 }
 
-FLearningAgentsActionObjectElement ULearningAgentsActions::MakeStaticArrayActionFromArrayView(ULearningAgentsActionObject* Object, const TArrayView<const FLearningAgentsActionObjectElement> Elements, const FName Tag)
+FLearningAgentsActionObjectElement ULearningAgentsActions::MakeStaticArrayActionFromArrayView(ULearningAgentsActionObject* Object, const TArrayView<const FLearningAgentsActionObjectElement> Elements, const bool bActive, const FName Tag)
 {
+	if (!bActive)
+	{
+		return MakeNullAction(Object, Tag);
+	}
+
 	if (!Object)
 	{
 		UE_LOG(LogLearning, Error, TEXT("MakeStaticArrayActionFromArrayView: Object is nullptr."));
@@ -2332,15 +2616,21 @@ FLearningAgentsActionObjectElement ULearningAgentsActions::MakeStaticArrayAction
 	return { Object->ActionObject.CreateArray({ SubElements }, Tag) };
 }
 
-FLearningAgentsActionObjectElement ULearningAgentsActions::MakePairAction(ULearningAgentsActionObject* Object, const FLearningAgentsActionObjectElement Key, const FLearningAgentsActionObjectElement Value, const FName Tag)
+FLearningAgentsActionObjectElement ULearningAgentsActions::MakePairAction(ULearningAgentsActionObject* Object, const FLearningAgentsActionObjectElement Key, const FLearningAgentsActionObjectElement Value, const bool bActive, const FName Tag)
 {
-	return MakeStructActionFromArrayViews(Object, { TEXT("Key"), TEXT("Value") }, { Key, Value }, Tag);
+	if (!bActive)
+	{
+		return MakeNullAction(Object, Tag);
+	}
+
+	return MakeStructActionFromArrayViews(Object, { TEXT("Key"), TEXT("Value") }, { Key, Value }, bActive, Tag);
 }
 
 FLearningAgentsActionObjectElement ULearningAgentsActions::MakeEnumAction(
 	ULearningAgentsActionObject* Object,
 	const UEnum* Enum,
 	const uint8 EnumValue,
+    const bool bActive,
 	const FName Tag,
 	const bool bVisualLoggerEnabled,
 	ULearningAgentsManagerListener* VisualLoggerListener,
@@ -2348,6 +2638,11 @@ FLearningAgentsActionObjectElement ULearningAgentsActions::MakeEnumAction(
 	const FVector VisualLoggerLocation,
 	const FLinearColor VisualLoggerColor)
 {
+	if (!bActive)
+	{
+		return MakeNullAction(Object, Tag);
+	}
+
 	if (!Enum)
 	{
 		UE_LOG(LogLearning, Error, TEXT("MakeEnumAction: Enum is nullptr."));
@@ -2380,13 +2675,14 @@ FLearningAgentsActionObjectElement ULearningAgentsActions::MakeEnumAction(
 	}
 #endif
 
-	return MakeExclusiveDiscreteAction(Object, EnumValueIndex, Tag);
+	return MakeExclusiveDiscreteAction(Object, EnumValueIndex, bActive, Tag);
 }
 
 FLearningAgentsActionObjectElement ULearningAgentsActions::MakeBitmaskAction(
 	ULearningAgentsActionObject* Object,
 	const UEnum* Enum,
 	const int32 BitmaskValue,
+    const bool bActive,	
 	const FName Tag,
 	const bool bVisualLoggerEnabled,
 	ULearningAgentsManagerListener* VisualLoggerListener,
@@ -2394,6 +2690,11 @@ FLearningAgentsActionObjectElement ULearningAgentsActions::MakeBitmaskAction(
 	const FVector VisualLoggerLocation,
 	const FLinearColor VisualLoggerColor)
 {
+	if (!bActive)
+	{
+		return MakeNullAction(Object, Tag);
+	}
+
 	if (!Enum)
 	{
 		UE_LOG(LogLearning, Error, TEXT("MakeBitmaskAction: Enum is nullptr."));
@@ -2450,44 +2751,79 @@ FLearningAgentsActionObjectElement ULearningAgentsActions::MakeBitmaskAction(
 	}
 #endif
 
-	return MakeInclusiveDiscreteActionFromArrayView(Object, BitmaskIndices, Tag);
+	return MakeInclusiveDiscreteActionFromArrayView(Object, BitmaskIndices, bActive, Tag);
 }
 
-FLearningAgentsActionObjectElement ULearningAgentsActions::MakeOptionalAction(ULearningAgentsActionObject* Object, const FLearningAgentsActionObjectElement Element, const ELearningAgentsOptionalAction Option, const FName Tag)
+FLearningAgentsActionObjectElement ULearningAgentsActions::MakeOptionalAction(ULearningAgentsActionObject* Object, const FLearningAgentsActionObjectElement Element, const ELearningAgentsOptionalAction Option, const bool bActive, const FName Tag)
 {
+	if (!bActive)
+	{
+		return MakeNullAction(Object, Tag);
+	}
+
 	return MakeExclusiveUnionAction(Object,
 		Option == ELearningAgentsOptionalAction::Null ? TEXT("Null") : TEXT("Valid"),
 		Option == ELearningAgentsOptionalAction::Null ? MakeNullAction(Object) : Element,
-		Tag);
+		bActive, Tag);
 }
 
-FLearningAgentsActionObjectElement ULearningAgentsActions::MakeOptionalNullAction(ULearningAgentsActionObject* Object, const FName Tag)
+FLearningAgentsActionObjectElement ULearningAgentsActions::MakeOptionalNullAction(ULearningAgentsActionObject* Object, const bool bActive, const FName Tag)
 {
-	return MakeExclusiveUnionAction(Object, TEXT("Null"), MakeNullAction(Object), Tag);
+	if (!bActive)
+	{
+		return MakeNullAction(Object, Tag);
+	}
+
+	return MakeExclusiveUnionAction(Object, TEXT("Null"), MakeNullAction(Object), bActive, Tag);
 }
 
-FLearningAgentsActionObjectElement ULearningAgentsActions::MakeOptionalValidAction(ULearningAgentsActionObject* Object, const FLearningAgentsActionObjectElement Element, const FName Tag)
+FLearningAgentsActionObjectElement ULearningAgentsActions::MakeOptionalValidAction(ULearningAgentsActionObject* Object, const FLearningAgentsActionObjectElement Element, const bool bActive, const FName Tag)
 {
-	return MakeExclusiveUnionAction(Object, TEXT("Valid"), Element, Tag);
+	if (!bActive)
+	{
+		return MakeNullAction(Object, Tag);
+	}
+
+	return MakeExclusiveUnionAction(Object, TEXT("Valid"), Element, bActive, Tag);
 }
 
-FLearningAgentsActionObjectElement ULearningAgentsActions::MakeEitherAction(ULearningAgentsActionObject* Object, const FLearningAgentsActionObjectElement Element, const ELearningAgentsEitherAction Either, const FName Tag)
+FLearningAgentsActionObjectElement ULearningAgentsActions::MakeEitherAction(ULearningAgentsActionObject* Object, const FLearningAgentsActionObjectElement Element, const ELearningAgentsEitherAction Either, const bool bActive, const FName Tag)
 {
-	return MakeExclusiveUnionAction(Object, Either == ELearningAgentsEitherAction::A ? TEXT("A") : TEXT("B"), Element, Tag);
+	if (!bActive)
+	{
+		return MakeNullAction(Object, Tag);
+	}
+
+	return MakeExclusiveUnionAction(Object, Either == ELearningAgentsEitherAction::A ? TEXT("A") : TEXT("B"), Element, bActive, Tag);
 }
 
-FLearningAgentsActionObjectElement ULearningAgentsActions::MakeEitherAAction(ULearningAgentsActionObject* Object, const FLearningAgentsActionObjectElement A, const FName Tag)
+FLearningAgentsActionObjectElement ULearningAgentsActions::MakeEitherAAction(ULearningAgentsActionObject* Object, const FLearningAgentsActionObjectElement A, const bool bActive, const FName Tag)
 {
-	return MakeExclusiveUnionAction(Object, TEXT("A"), A, Tag);
+	if (!bActive)
+	{
+		return MakeNullAction(Object, Tag);
+	}
+
+	return MakeExclusiveUnionAction(Object, TEXT("A"), A, bActive, Tag);
 }
 
-FLearningAgentsActionObjectElement ULearningAgentsActions::MakeEitherBAction(ULearningAgentsActionObject* Object, const FLearningAgentsActionObjectElement B, const FName Tag)
+FLearningAgentsActionObjectElement ULearningAgentsActions::MakeEitherBAction(ULearningAgentsActionObject* Object, const FLearningAgentsActionObjectElement B, const bool bActive, const FName Tag)
 {
-	return MakeExclusiveUnionAction(Object, TEXT("B"), B, Tag);
+	if (!bActive)
+	{
+		return MakeNullAction(Object, Tag);
+	}
+
+	return MakeExclusiveUnionAction(Object, TEXT("B"), B, bActive, Tag);
 }
 
-FLearningAgentsActionObjectElement ULearningAgentsActions::MakeEncodingAction(ULearningAgentsActionObject* Object, const FLearningAgentsActionObjectElement Element, const FName Tag)
+FLearningAgentsActionObjectElement ULearningAgentsActions::MakeEncodingAction(ULearningAgentsActionObject* Object, const FLearningAgentsActionObjectElement Element, const bool bActive, const FName Tag)
 {
+	if (!bActive)
+	{
+		return MakeNullAction(Object, Tag);
+	}
+
 	if (!Object)
 	{
 		UE_LOG(LogLearning, Error, TEXT("MakeEncodingAction: Object is nullptr."));
@@ -2506,6 +2842,7 @@ FLearningAgentsActionObjectElement ULearningAgentsActions::MakeEncodingAction(UL
 FLearningAgentsActionObjectElement ULearningAgentsActions::MakeBoolAction(
 	ULearningAgentsActionObject* Object,
 	const bool bValue,
+    const bool bActive,
 	const FName Tag,
 	const bool bVisualLoggerEnabled,
 	ULearningAgentsManagerListener* VisualLoggerListener,
@@ -2513,6 +2850,11 @@ FLearningAgentsActionObjectElement ULearningAgentsActions::MakeBoolAction(
 	const FVector VisualLoggerLocation,
 	const FLinearColor VisualLoggerColor)
 {
+	if (!bActive)
+	{
+		return MakeNullAction(Object, Tag);
+	}
+
 #if UE_LEARNING_AGENTS_ENABLE_VISUAL_LOG
 	if (bVisualLoggerEnabled && VisualLoggerListener)
 	{
@@ -2528,12 +2870,13 @@ FLearningAgentsActionObjectElement ULearningAgentsActions::MakeBoolAction(
 	}
 #endif
 
-	return MakeExclusiveDiscreteAction(Object, bValue ? 1 : 0, Tag);
+	return MakeExclusiveDiscreteAction(Object, bValue ? 1 : 0, bActive, Tag);
 }
 
 FLearningAgentsActionObjectElement ULearningAgentsActions::MakeFloatAction(
 	ULearningAgentsActionObject* Object,
 	const float Value,
+    const bool bActive,
 	const FName Tag,
 	const bool bVisualLoggerEnabled,
 	ULearningAgentsManagerListener* VisualLoggerListener,
@@ -2541,6 +2884,11 @@ FLearningAgentsActionObjectElement ULearningAgentsActions::MakeFloatAction(
 	const FVector VisualLoggerLocation,
 	const FLinearColor VisualLoggerColor)
 {
+	if (!bActive)
+	{
+		return MakeNullAction(Object, Tag);
+	}
+
 #if UE_LEARNING_AGENTS_ENABLE_VISUAL_LOG
 	if (bVisualLoggerEnabled && VisualLoggerListener)
 	{
@@ -2556,13 +2904,14 @@ FLearningAgentsActionObjectElement ULearningAgentsActions::MakeFloatAction(
 	}
 #endif
 
-	return MakeContinuousActionFromArrayView(Object, { Value }, Tag);
+	return MakeContinuousActionFromArrayView(Object, { Value }, bActive, Tag);
 }
 
 FLearningAgentsActionObjectElement ULearningAgentsActions::MakeLocationAction(
 	ULearningAgentsActionObject* Object,
 	const FVector Location,
 	const FTransform RelativeTransform,
+	const bool bActive,	
 	const FName Tag,
 	const bool bVisualLoggerEnabled,
 	ULearningAgentsManagerListener* VisualLoggerListener,
@@ -2570,6 +2919,11 @@ FLearningAgentsActionObjectElement ULearningAgentsActions::MakeLocationAction(
 	const FVector VisualLoggerLocation,
 	const FLinearColor VisualLoggerColor)
 {
+	if (!bActive)
+	{
+		return MakeNullAction(Object, Tag);
+	}
+
 	const FVector LocalLocation = RelativeTransform.InverseTransformPosition(Location);
 
 #if UE_LEARNING_AGENTS_ENABLE_VISUAL_LOG
@@ -2609,13 +2963,14 @@ FLearningAgentsActionObjectElement ULearningAgentsActions::MakeLocationAction(
 	return MakeContinuousActionFromArrayView(Object, {
 		(float)LocalLocation.X,
 		(float)LocalLocation.Y,
-		(float)LocalLocation.Z }, Tag);
+		(float)LocalLocation.Z }, bActive, Tag);
 }
 
 FLearningAgentsActionObjectElement ULearningAgentsActions::MakeRotationAction(
 	ULearningAgentsActionObject* Object,
 	const FRotator Rotation,
 	const FRotator RelativeRotation,
+    const bool bActive,	
 	const FName Tag,
 	const bool bVisualLoggerEnabled,
 	ULearningAgentsManagerListener* VisualLoggerListener,
@@ -2624,7 +2979,12 @@ FLearningAgentsActionObjectElement ULearningAgentsActions::MakeRotationAction(
 	const FVector VisualLoggerLocation,
 	const FLinearColor VisualLoggerColor)
 {
-	return MakeRotationActionFromQuat(Object, FQuat::MakeFromRotator(Rotation), FQuat::MakeFromRotator(RelativeRotation), Tag,
+	if (!bActive)
+	{
+		return MakeNullAction(Object, Tag);
+	}
+
+	return MakeRotationActionFromQuat(Object, FQuat::MakeFromRotator(Rotation), FQuat::MakeFromRotator(RelativeRotation), bActive, Tag,
 		bVisualLoggerEnabled, VisualLoggerListener, VisualLoggerAgentId, VisualLoggerRotationLocation, VisualLoggerLocation, VisualLoggerColor);
 }
 
@@ -2632,6 +2992,7 @@ FLearningAgentsActionObjectElement ULearningAgentsActions::MakeRotationActionFro
 	ULearningAgentsActionObject* Object,
 	const FQuat Rotation,
 	const FQuat RelativeRotation,
+	const bool bActive,
 	const FName Tag,
 	const bool bVisualLoggerEnabled,
 	ULearningAgentsManagerListener* VisualLoggerListener,
@@ -2640,6 +3001,11 @@ FLearningAgentsActionObjectElement ULearningAgentsActions::MakeRotationActionFro
 	const FVector VisualLoggerLocation,
 	const FLinearColor VisualLoggerColor)
 {
+	if (!bActive)
+	{
+		return MakeNullAction(Object, Tag);
+	}
+
 	FQuat LocalRotation = (RelativeRotation.Inverse() * Rotation).GetShortestArcWith(FQuat::Identity);
 	const FVector RotationVector = LocalRotation.ToRotationVector();
 
@@ -2670,13 +3036,14 @@ FLearningAgentsActionObjectElement ULearningAgentsActions::MakeRotationActionFro
 		(float)FMath::RadiansToDegrees(RotationVector.X),
 		(float)FMath::RadiansToDegrees(RotationVector.Y),
 		(float)FMath::RadiansToDegrees(RotationVector.Z),
-		}, Tag);
+		}, bActive, Tag);
 }
 
 FLearningAgentsActionObjectElement ULearningAgentsActions::MakeScaleAction(
 	ULearningAgentsActionObject* Object,
 	const FVector Scale,
 	const FVector RelativeScale,
+    const bool bActive,	
 	const FName Tag,
 	const bool bVisualLoggerEnabled,
 	ULearningAgentsManagerListener* VisualLoggerListener,
@@ -2684,6 +3051,11 @@ FLearningAgentsActionObjectElement ULearningAgentsActions::MakeScaleAction(
 	const FVector VisualLoggerLocation,
 	const FLinearColor VisualLoggerColor)
 {
+	if (!bActive)
+	{
+		return MakeNullAction(Object, Tag);
+	}
+
 	const FVector LocalLogScale =
 		UE::Learning::Agents::Action::Private::VectorLogSafe(Scale) -
 		UE::Learning::Agents::Action::Private::VectorLogSafe(RelativeScale);
@@ -2708,13 +3080,14 @@ FLearningAgentsActionObjectElement ULearningAgentsActions::MakeScaleAction(
 		(float)LocalLogScale.X,
 		(float)LocalLogScale.Y,
 		(float)LocalLogScale.Z,
-		}, Tag);
+		}, bActive, Tag);
 }
 
 FLearningAgentsActionObjectElement ULearningAgentsActions::MakeTransformAction(
 	ULearningAgentsActionObject* Object,
 	const FTransform Transform,
 	const FTransform RelativeTransform,
+	const bool bActive,	
 	const FName Tag,
 	const bool bVisualLoggerEnabled,
 	ULearningAgentsManagerListener* VisualLoggerListener,
@@ -2722,6 +3095,11 @@ FLearningAgentsActionObjectElement ULearningAgentsActions::MakeTransformAction(
 	const FVector VisualLoggerLocation,
 	const FLinearColor VisualLoggerColor)
 {
+	if (!bActive)
+	{
+		return MakeNullAction(Object, Tag);
+	}
+
 	FVector Location = Transform.GetLocation();
 	FQuat Rotation = Transform.GetRotation();
 	FVector Scale = Transform.GetScale3D();
@@ -2760,13 +3138,14 @@ FLearningAgentsActionObjectElement ULearningAgentsActions::MakeTransformAction(
 			MakeRotationActionFromQuat(Object, Rotation, RelativeTransform.GetRotation()),
 			MakeScaleAction(Object, Scale, RelativeTransform.GetScale3D())
 		},
-		Tag);
+		bActive, Tag);
 }
 
 FLearningAgentsActionObjectElement ULearningAgentsActions::MakeAngleAction(
 	ULearningAgentsActionObject* Object,
 	const float Angle,
 	const float RelativeAngle,
+	const bool bActive,
 	const FName Tag,
 	const bool bVisualLoggerEnabled,
 	ULearningAgentsManagerListener* VisualLoggerListener,
@@ -2775,6 +3154,11 @@ FLearningAgentsActionObjectElement ULearningAgentsActions::MakeAngleAction(
 	const FVector VisualLoggerLocation,
 	const FLinearColor VisualLoggerColor)
 {
+	if (!bActive)
+	{
+		return MakeNullAction(Object, Tag);
+	}
+
 	float LocalAngle = FMath::FindDeltaAngleDegrees(RelativeAngle, Angle);
 	
 #if UE_LEARNING_AGENTS_ENABLE_VISUAL_LOG
@@ -2801,13 +3185,14 @@ FLearningAgentsActionObjectElement ULearningAgentsActions::MakeAngleAction(
 	}
 #endif
 
-	return MakeContinuousActionFromArrayView(Object, { LocalAngle }, Tag);
+	return MakeContinuousActionFromArrayView(Object, { LocalAngle }, bActive, Tag);
 }
 
 FLearningAgentsActionObjectElement ULearningAgentsActions::MakeAngleActionRadians(
 	ULearningAgentsActionObject* Object,
 	const float Angle,
 	const float RelativeAngle,
+    const bool bActive,	
 	const FName Tag,
 	const bool bVisualLoggerEnabled,
 	ULearningAgentsManagerListener* VisualLoggerListener,
@@ -2816,7 +3201,12 @@ FLearningAgentsActionObjectElement ULearningAgentsActions::MakeAngleActionRadian
 	const FVector VisualLoggerLocation,
 	const FLinearColor VisualLoggerColor)
 {
-	return MakeAngleAction(Object, FMath::RadiansToDegrees(Angle), FMath::RadiansToDegrees(RelativeAngle), Tag,
+	if (!bActive)
+	{
+		return MakeNullAction(Object, Tag);
+	}
+
+	return MakeAngleAction(Object, FMath::RadiansToDegrees(Angle), FMath::RadiansToDegrees(RelativeAngle), bActive, Tag,
 		bVisualLoggerEnabled, VisualLoggerListener, VisualLoggerAgentId, VisualLoggerAngleLocation, VisualLoggerLocation, VisualLoggerColor);
 }
 
@@ -2824,6 +3214,7 @@ FLearningAgentsActionObjectElement ULearningAgentsActions::MakeVelocityAction(
 	ULearningAgentsActionObject* Object,
 	const FVector Velocity,
 	const FTransform RelativeTransform,
+	const bool bActive,	
 	const FName Tag,
 	const bool bVisualLoggerEnabled,
 	ULearningAgentsManagerListener* VisualLoggerListener,
@@ -2832,6 +3223,11 @@ FLearningAgentsActionObjectElement ULearningAgentsActions::MakeVelocityAction(
 	const FVector VisualLoggerLocation,
 	const FLinearColor VisualLoggerColor)
 {
+	if (!bActive)
+	{
+		return MakeNullAction(Object, Tag);
+	}
+
 	const FVector LocalVelocity = RelativeTransform.InverseTransformVectorNoScale(Velocity);
 
 #if UE_LEARNING_AGENTS_ENABLE_VISUAL_LOG
@@ -2866,13 +3262,14 @@ FLearningAgentsActionObjectElement ULearningAgentsActions::MakeVelocityAction(
 		(float)LocalVelocity.X,
 		(float)LocalVelocity.Y,
 		(float)LocalVelocity.Z,
-		}, Tag);
+		}, bActive, Tag);
 }
 
 FLearningAgentsActionObjectElement ULearningAgentsActions::MakeDirectionAction(
 	ULearningAgentsActionObject* Object,
 	const FVector Direction,
 	const FTransform RelativeTransform,
+	const bool bActive,	
 	const FName Tag,
 	const bool bVisualLoggerEnabled,
 	ULearningAgentsManagerListener* VisualLoggerListener,
@@ -2882,6 +3279,11 @@ FLearningAgentsActionObjectElement ULearningAgentsActions::MakeDirectionAction(
 	const float VisualLoggerArrowLength,
 	const FLinearColor VisualLoggerColor)
 {
+	if (!bActive)
+	{
+		return MakeNullAction(Object, Tag);
+	}
+
 	const FVector LocalDirection = RelativeTransform.InverseTransformVectorNoScale(Direction).GetSafeNormal(UE_SMALL_NUMBER, FVector::ForwardVector);
 
 #if UE_LEARNING_AGENTS_ENABLE_VISUAL_LOG
@@ -2916,7 +3318,7 @@ FLearningAgentsActionObjectElement ULearningAgentsActions::MakeDirectionAction(
 		(float)LocalDirection.X,
 		(float)LocalDirection.Y,
 		(float)LocalDirection.Z,
-		}, Tag);
+		}, bActive, Tag);
 }
 
 FLearningAgentsActionModifierElement ULearningAgentsActions::MakeNullActionModifier(ULearningAgentsActionModifier* Modifier, const FName Tag)
@@ -3421,7 +3823,7 @@ FLearningAgentsActionModifierElement ULearningAgentsActions::MakeScaleActionModi
 		}, {
 		(float)LocalLogScale.X, 
 		(float)LocalLogScale.Y, 
-		(float)LocalLogScale.Z } , Tag);
+		(float)LocalLogScale.Z }, Tag);
 }
 
 FLearningAgentsActionModifierElement ULearningAgentsActions::MakeAngleActionModifier(ULearningAgentsActionModifier* Modifier, const float MaskedAngle, const bool bMask, const float RelativeAngle, const FName Tag)
@@ -3493,8 +3895,13 @@ bool ULearningAgentsActions::GetNullAction(const ULearningAgentsActionObject* Ob
 	return true;
 }
 
-bool ULearningAgentsActions::GetContinuousActionNum(int32& OutNum, const ULearningAgentsActionObject* Object, const FLearningAgentsActionObjectElement Element, const FName Tag)
+bool ULearningAgentsActions::GetContinuousActionNum(int32& OutNum, const ULearningAgentsActionObject* Object, const FLearningAgentsActionObjectElement Element, const bool bActive, const FName Tag)
 {
+	if (!bActive)
+	{
+		return GetNullAction(Object, Element, Tag);
+	}
+
 	if (!Object)
 	{
 		UE_LOG(LogLearning, Error, TEXT("GetContinuousActionNum: Object is nullptr."));
@@ -3532,6 +3939,7 @@ bool ULearningAgentsActions::GetContinuousAction(
 	TArray<float>& OutValues, 
 	const ULearningAgentsActionObject* Object, 
 	const FLearningAgentsActionObjectElement Element, 
+	const bool bActive,
 	const FName Tag,
 	const bool bVisualLoggerEnabled,
 	ULearningAgentsManagerListener* VisualLoggerListener,
@@ -3539,8 +3947,13 @@ bool ULearningAgentsActions::GetContinuousAction(
 	const FVector VisualLoggerLocation,
 	const FLinearColor VisualLoggerColor)
 {
+	if (!bActive)
+	{
+		return GetNullAction(Object, Element, Tag);
+	}
+
 	int32 OutValueNum = 0;
-	if (!GetContinuousActionNum(OutValueNum, Object, Element, Tag))
+	if (!GetContinuousActionNum(OutValueNum, Object, Element, bActive, Tag))
 	{
 		OutValues.Empty();
 		return false;
@@ -3548,7 +3961,7 @@ bool ULearningAgentsActions::GetContinuousAction(
 
 	OutValues.SetNumUninitialized(OutValueNum);
 
-	if (!GetContinuousActionToArrayView(OutValues, Object, Element, Tag, bVisualLoggerEnabled, VisualLoggerListener, VisualLoggerAgentId, VisualLoggerLocation, VisualLoggerColor))
+	if (!GetContinuousActionToArrayView(OutValues, Object, Element, bActive, Tag, bVisualLoggerEnabled, VisualLoggerListener, VisualLoggerAgentId, VisualLoggerLocation, VisualLoggerColor))
 	{
 		OutValues.Empty();
 		return false;
@@ -3561,6 +3974,7 @@ bool ULearningAgentsActions::GetContinuousActionToArrayView(
 	TArrayView<float> OutValues, 
 	const ULearningAgentsActionObject* Object, 
 	const FLearningAgentsActionObjectElement Element, 
+	const bool bActive,
 	const FName Tag,
 	const bool bVisualLoggerEnabled,
 	ULearningAgentsManagerListener* VisualLoggerListener,
@@ -3568,6 +3982,11 @@ bool ULearningAgentsActions::GetContinuousActionToArrayView(
 	const FVector VisualLoggerLocation,
 	const FLinearColor VisualLoggerColor)
 {
+	if (!bActive)
+	{
+		return GetNullAction(Object, Element, Tag);
+	}
+
 	if (!Object)
 	{
 		UE_LOG(LogLearning, Error, TEXT("GetContinuousActionToArrayView: Object is nullptr."));
@@ -3633,6 +4052,7 @@ bool ULearningAgentsActions::GetExclusiveDiscreteAction(
 	int32& OutIndex, 
 	const ULearningAgentsActionObject* Object, 
 	const FLearningAgentsActionObjectElement Element, 
+	const bool bActive,	
 	const FName Tag,
 	const bool bVisualLoggerEnabled,
 	ULearningAgentsManagerListener* VisualLoggerListener,
@@ -3640,6 +4060,11 @@ bool ULearningAgentsActions::GetExclusiveDiscreteAction(
 	const FVector VisualLoggerLocation,
 	const FLinearColor VisualLoggerColor)
 {
+	if (!bActive)
+	{
+		return GetNullAction(Object, Element, Tag);
+	}
+
 	if (!Object)
 	{
 		UE_LOG(LogLearning, Error, TEXT("GetExclusiveDiscreteAction: Object is nullptr."));
@@ -3693,6 +4118,7 @@ bool ULearningAgentsActions::GetNamedExclusiveDiscreteAction(
 	FName& OutName,
 	const ULearningAgentsActionObject* Object,
 	const FLearningAgentsActionObjectElement Element,
+	const bool bActive,
 	const FName Tag,
 	const bool bVisualLoggerEnabled,
 	ULearningAgentsManagerListener* VisualLoggerListener,
@@ -3700,6 +4126,11 @@ bool ULearningAgentsActions::GetNamedExclusiveDiscreteAction(
 	const FVector VisualLoggerLocation,
 	const FLinearColor VisualLoggerColor)
 {
+	if (!bActive)
+	{
+		return GetNullAction(Object, Element, Tag);
+	}
+
 	if (!Object)
 	{
 		UE_LOG(LogLearning, Error, TEXT("GetNamedExclusiveDiscreteAction: Object is nullptr."));
@@ -3749,8 +4180,13 @@ bool ULearningAgentsActions::GetNamedExclusiveDiscreteAction(
 	return true;
 }
 
-bool ULearningAgentsActions::GetInclusiveDiscreteActionNum(int32& OutNum, const ULearningAgentsActionObject* Object, const FLearningAgentsActionObjectElement Element, const FName Tag)
+bool ULearningAgentsActions::GetInclusiveDiscreteActionNum(int32& OutNum, const ULearningAgentsActionObject* Object, const FLearningAgentsActionObjectElement Element, const bool bActive, const FName Tag)
 {
+	if (!bActive)
+	{
+		return GetNullAction(Object, Element, Tag);
+	}
+
 	if (!Object)
 	{
 		UE_LOG(LogLearning, Error, TEXT("GetInclusiveDiscreteActionNum: Object is nullptr."));
@@ -3788,6 +4224,7 @@ bool ULearningAgentsActions::GetInclusiveDiscreteAction(
 	TArray<int32>& OutIndices, 
 	const ULearningAgentsActionObject* Object, 
 	const FLearningAgentsActionObjectElement Element, 
+	const bool bActive,
 	const FName Tag,
 	const bool bVisualLoggerEnabled,
 	ULearningAgentsManagerListener* VisualLoggerListener,
@@ -3795,8 +4232,13 @@ bool ULearningAgentsActions::GetInclusiveDiscreteAction(
 	const FVector VisualLoggerLocation,
 	const FLinearColor VisualLoggerColor)
 {
+	if (!bActive)
+	{
+		return GetNullAction(Object, Element, Tag);
+	}
+
 	int32 OutIndexNum = 0;
-	if (!GetInclusiveDiscreteActionNum(OutIndexNum, Object, Element, Tag))
+	if (!GetInclusiveDiscreteActionNum(OutIndexNum, Object, Element, bActive, Tag))
 	{
 		OutIndices.Empty();
 		return false;
@@ -3804,7 +4246,7 @@ bool ULearningAgentsActions::GetInclusiveDiscreteAction(
 
 	OutIndices.SetNumUninitialized(OutIndexNum);
 
-	if (!GetInclusiveDiscreteActionToArrayView(OutIndices, Object, Element, Tag, bVisualLoggerEnabled, VisualLoggerListener, VisualLoggerAgentId, VisualLoggerLocation, VisualLoggerColor))
+	if (!GetInclusiveDiscreteActionToArrayView(OutIndices, Object, Element, bActive, Tag, bVisualLoggerEnabled, VisualLoggerListener, VisualLoggerAgentId, VisualLoggerLocation, VisualLoggerColor))
 	{
 		OutIndices.Empty();
 		return false;
@@ -3817,6 +4259,7 @@ bool ULearningAgentsActions::GetInclusiveDiscreteActionToArrayView(
 	TArrayView<int32> OutIndices,
 	const ULearningAgentsActionObject* Object, 
 	const FLearningAgentsActionObjectElement Element, 
+	const bool bActive,
 	const FName Tag,
 	const bool bVisualLoggerEnabled,
 	ULearningAgentsManagerListener* VisualLoggerListener,
@@ -3824,6 +4267,11 @@ bool ULearningAgentsActions::GetInclusiveDiscreteActionToArrayView(
 	const FVector VisualLoggerLocation,
 	const FLinearColor VisualLoggerColor)
 {
+	if (!bActive)
+	{
+		return GetNullAction(Object, Element, Tag);
+	}
+
 	if (!Object)
 	{
 		UE_LOG(LogLearning, Error, TEXT("GetInclusiveDiscreteActionToArrayView: Object is nullptr."));
@@ -3884,8 +4332,13 @@ bool ULearningAgentsActions::GetInclusiveDiscreteActionToArrayView(
 	return true;
 }
 
-bool ULearningAgentsActions::GetNamedInclusiveDiscreteActionNum(int32& OutNum, const ULearningAgentsActionObject* Object, const FLearningAgentsActionObjectElement Element, const FName Tag)
+bool ULearningAgentsActions::GetNamedInclusiveDiscreteActionNum(int32& OutNum, const ULearningAgentsActionObject* Object, const FLearningAgentsActionObjectElement Element, const bool bActive, const FName Tag)
 {
+	if (!bActive)
+	{
+		return GetNullAction(Object, Element, Tag);
+	}
+
 	if (!Object)
 	{
 		UE_LOG(LogLearning, Error, TEXT("GetNamedInclusiveDiscreteActionNum: Object is nullptr."));
@@ -3923,6 +4376,7 @@ bool ULearningAgentsActions::GetNamedInclusiveDiscreteAction(
 	TArray<FName>& OutNames,
 	const ULearningAgentsActionObject* Object,
 	const FLearningAgentsActionObjectElement Element,
+	const bool bActive,
 	const FName Tag,
 	const bool bVisualLoggerEnabled,
 	ULearningAgentsManagerListener* VisualLoggerListener,
@@ -3930,8 +4384,13 @@ bool ULearningAgentsActions::GetNamedInclusiveDiscreteAction(
 	const FVector VisualLoggerLocation,
 	const FLinearColor VisualLoggerColor)
 {
+	if (!bActive)
+	{
+		return GetNullAction(Object, Element, Tag);
+	}
+
 	int32 OutNameNum = 0;
-	if (!GetNamedInclusiveDiscreteActionNum(OutNameNum, Object, Element, Tag))
+	if (!GetNamedInclusiveDiscreteActionNum(OutNameNum, Object, Element, bActive, Tag))
 	{
 		OutNames.Empty();
 		return false;
@@ -3939,7 +4398,7 @@ bool ULearningAgentsActions::GetNamedInclusiveDiscreteAction(
 
 	OutNames.SetNumUninitialized(OutNameNum);
 
-	if (!GetNamedInclusiveDiscreteActionToArrayView(OutNames, Object, Element, Tag, bVisualLoggerEnabled, VisualLoggerListener, VisualLoggerAgentId, VisualLoggerLocation, VisualLoggerColor))
+	if (!GetNamedInclusiveDiscreteActionToArrayView(OutNames, Object, Element, bActive, Tag, bVisualLoggerEnabled, VisualLoggerListener, VisualLoggerAgentId, VisualLoggerLocation, VisualLoggerColor))
 	{
 		OutNames.Empty();
 		return false;
@@ -3952,6 +4411,7 @@ bool ULearningAgentsActions::GetNamedInclusiveDiscreteActionToArrayView(
 	TArrayView<FName> OutNames,
 	const ULearningAgentsActionObject* Object,
 	const FLearningAgentsActionObjectElement Element,
+	const bool bActive,	
 	const FName Tag,
 	const bool bVisualLoggerEnabled,
 	ULearningAgentsManagerListener* VisualLoggerListener,
@@ -3959,6 +4419,11 @@ bool ULearningAgentsActions::GetNamedInclusiveDiscreteActionToArrayView(
 	const FVector VisualLoggerLocation,
 	const FLinearColor VisualLoggerColor)
 {
+	if (!bActive)
+	{
+		return GetNullAction(Object, Element, Tag);
+	}
+
 	if (!Object)
 	{
 		UE_LOG(LogLearning, Error, TEXT("GetNamedInclusiveDiscreteActionToArrayView: Object is nullptr."));
@@ -4020,8 +4485,13 @@ bool ULearningAgentsActions::GetNamedInclusiveDiscreteActionToArrayView(
 }
 
 
-bool ULearningAgentsActions::GetStructActionNum(int32& OutNum, const ULearningAgentsActionObject* Object, const FLearningAgentsActionObjectElement Element, const FName Tag)
+bool ULearningAgentsActions::GetStructActionNum(int32& OutNum, const ULearningAgentsActionObject* Object, const FLearningAgentsActionObjectElement Element, const bool bActive, const FName Tag)
 {
+	if (!bActive)
+	{
+		return GetNullAction(Object, Element, Tag);
+	}
+
 	if (!Object)
 	{
 		UE_LOG(LogLearning, Error, TEXT("GetStructActionNum: Object is nullptr."));
@@ -4057,10 +4527,15 @@ bool ULearningAgentsActions::GetStructActionNum(int32& OutNum, const ULearningAg
 	return true;
 }
 
-bool ULearningAgentsActions::GetStructAction(TMap<FName, FLearningAgentsActionObjectElement>& OutElements, const ULearningAgentsActionObject* Object, const FLearningAgentsActionObjectElement Element, const FName Tag)
+bool ULearningAgentsActions::GetStructAction(TMap<FName, FLearningAgentsActionObjectElement>& OutElements, const ULearningAgentsActionObject* Object, const FLearningAgentsActionObjectElement Element, const bool bActive, const FName Tag)
 {
+	if (!bActive)
+	{
+		return GetNullAction(Object, Element, Tag);
+	}
+
 	int32 OutElementNum = 0;
-	if (!GetStructActionNum(OutElementNum, Object, Element, Tag))
+	if (!GetStructActionNum(OutElementNum, Object, Element, bActive, Tag))
 	{
 		OutElements.Empty();
 		return false;
@@ -4071,7 +4546,7 @@ bool ULearningAgentsActions::GetStructAction(TMap<FName, FLearningAgentsActionOb
 	SubElementNames.SetNumUninitialized(OutElementNum);
 	SubElements.SetNumUninitialized(OutElementNum);
 
-	if (!GetStructActionToArrayViews(SubElementNames, SubElements, Object, Element, Tag))
+	if (!GetStructActionToArrayViews(SubElementNames, SubElements, Object, Element, bActive, Tag))
 	{
 		OutElements.Empty();
 		return false;
@@ -4086,8 +4561,13 @@ bool ULearningAgentsActions::GetStructAction(TMap<FName, FLearningAgentsActionOb
 	return true;
 }
 
-bool ULearningAgentsActions::GetStructActionElement(FLearningAgentsActionObjectElement& OutElement, const ULearningAgentsActionObject* Object, const FLearningAgentsActionObjectElement Element, const FName ElementName, const FName Tag)
+bool ULearningAgentsActions::GetStructActionElement(FLearningAgentsActionObjectElement& OutElement, const ULearningAgentsActionObject* Object, const FLearningAgentsActionObjectElement Element, const FName ElementName, const bool bActive, const FName Tag)
 {
+	if (!bActive)
+	{
+		return GetNullAction(Object, Element, Tag);
+	}
+
 	if (!Object)
 	{
 		UE_LOG(LogLearning, Error, TEXT("GetStructActionElement: Object is nullptr."));
@@ -4132,10 +4612,15 @@ bool ULearningAgentsActions::GetStructActionElement(FLearningAgentsActionObjectE
 	return true;
 }
 
-bool ULearningAgentsActions::GetStructActionToArrays(TArray<FName>& OutElementNames, TArray<FLearningAgentsActionObjectElement>& OutElements, const ULearningAgentsActionObject* Object, const FLearningAgentsActionObjectElement Element, const FName Tag)
+bool ULearningAgentsActions::GetStructActionToArrays(TArray<FName>& OutElementNames, TArray<FLearningAgentsActionObjectElement>& OutElements, const ULearningAgentsActionObject* Object, const FLearningAgentsActionObjectElement Element, const bool bActive, const FName Tag)
 {
+	if (!bActive)
+	{
+		return GetNullAction(Object, Element, Tag);
+	}
+
 	int32 OutElementNum = 0;
-	if (!GetStructActionNum(OutElementNum, Object, Element, Tag))
+	if (!GetStructActionNum(OutElementNum, Object, Element, bActive, Tag))
 	{
 		OutElementNames.Empty();
 		OutElements.Empty();
@@ -4145,7 +4630,7 @@ bool ULearningAgentsActions::GetStructActionToArrays(TArray<FName>& OutElementNa
 	OutElementNames.SetNumUninitialized(OutElementNum);
 	OutElements.SetNumUninitialized(OutElementNum);
 
-	if (!GetStructActionToArrayViews(OutElementNames, OutElements, Object, Element, Tag))
+	if (!GetStructActionToArrayViews(OutElementNames, OutElements, Object, Element, bActive, Tag))
 	{
 		OutElementNames.Empty();
 		OutElements.Empty();
@@ -4155,8 +4640,13 @@ bool ULearningAgentsActions::GetStructActionToArrays(TArray<FName>& OutElementNa
 	return true;
 }
 
-bool ULearningAgentsActions::GetStructActionToArrayViews(TArrayView<FName> OutElementNames, TArrayView<FLearningAgentsActionObjectElement> OutElements, const ULearningAgentsActionObject* Object, const FLearningAgentsActionObjectElement Element, const FName Tag)
+bool ULearningAgentsActions::GetStructActionToArrayViews(TArrayView<FName> OutElementNames, TArrayView<FLearningAgentsActionObjectElement> OutElements, const ULearningAgentsActionObject* Object, const FLearningAgentsActionObjectElement Element, const bool bActive, const FName Tag)
 {
+	if (!bActive)
+	{
+		return GetNullAction(Object, Element, Tag);
+	}
+
 	if (!Object)
 	{
 		UE_LOG(LogLearning, Error, TEXT("GetStructActionToArrayViews: Object is nullptr."));
@@ -4223,8 +4713,13 @@ bool ULearningAgentsActions::GetStructActionToArrayViews(TArrayView<FName> OutEl
 	return true;
 }
 
-bool ULearningAgentsActions::GetExclusiveUnionAction(FName& OutElementName, FLearningAgentsActionObjectElement& OutElement, const ULearningAgentsActionObject* Object, const FLearningAgentsActionObjectElement Element, const FName Tag)
+bool ULearningAgentsActions::GetExclusiveUnionAction(FName& OutElementName, FLearningAgentsActionObjectElement& OutElement, const ULearningAgentsActionObject* Object, const FLearningAgentsActionObjectElement Element, const bool bActive, const FName Tag)
 {
+	if (!bActive)
+	{
+		return GetNullAction(Object, Element, Tag);
+	}
+
 	if (!Object)
 	{
 		UE_LOG(LogLearning, Error, TEXT("GetExclusiveUnionAction: Object is nullptr."));
@@ -4264,8 +4759,13 @@ bool ULearningAgentsActions::GetExclusiveUnionAction(FName& OutElementName, FLea
 	return true;
 }
 
-bool ULearningAgentsActions::GetInclusiveUnionActionNum(int32& OutNum, const ULearningAgentsActionObject* Object, const FLearningAgentsActionObjectElement Element, const FName Tag)
+bool ULearningAgentsActions::GetInclusiveUnionActionNum(int32& OutNum, const ULearningAgentsActionObject* Object, const FLearningAgentsActionObjectElement Element, const bool bActive, const FName Tag)
 {
+	if (!bActive)
+	{
+		return GetNullAction(Object, Element, Tag);
+	}
+
 	if (!Object)
 	{
 		UE_LOG(LogLearning, Error, TEXT("GetInclusiveUnionActionNum: Object is nullptr."));
@@ -4301,10 +4801,15 @@ bool ULearningAgentsActions::GetInclusiveUnionActionNum(int32& OutNum, const ULe
 	return true;
 }
 
-bool ULearningAgentsActions::GetInclusiveUnionAction(TMap<FName, FLearningAgentsActionObjectElement>& OutElements, const ULearningAgentsActionObject* Object, const FLearningAgentsActionObjectElement Element, const FName Tag)
+bool ULearningAgentsActions::GetInclusiveUnionAction(TMap<FName, FLearningAgentsActionObjectElement>& OutElements, const ULearningAgentsActionObject* Object, const FLearningAgentsActionObjectElement Element, const bool bActive, const FName Tag)
 {
+	if (!bActive)
+	{
+		return GetNullAction(Object, Element, Tag);
+	}
+
 	int32 OutElementNum = 0;
-	if (!GetInclusiveUnionActionNum(OutElementNum, Object, Element, Tag))
+	if (!GetInclusiveUnionActionNum(OutElementNum, Object, Element, bActive, Tag))
 	{
 		OutElements.Empty();
 		return false;
@@ -4315,7 +4820,7 @@ bool ULearningAgentsActions::GetInclusiveUnionAction(TMap<FName, FLearningAgents
 	SubElementNames.SetNumUninitialized(OutElementNum);
 	SubElements.SetNumUninitialized(OutElementNum);
 
-	if (!GetInclusiveUnionActionToArrayViews(SubElementNames, SubElements, Object, Element, Tag))
+	if (!GetInclusiveUnionActionToArrayViews(SubElementNames, SubElements, Object, Element, bActive, Tag))
 	{
 		OutElements.Empty();
 		return false;
@@ -4330,10 +4835,15 @@ bool ULearningAgentsActions::GetInclusiveUnionAction(TMap<FName, FLearningAgents
 	return true;
 }
 
-bool ULearningAgentsActions::GetInclusiveUnionActionToArrays(TArray<FName>& OutElementNames, TArray<FLearningAgentsActionObjectElement>& OutElements, const ULearningAgentsActionObject* Object, const FLearningAgentsActionObjectElement Element, const FName Tag)
+bool ULearningAgentsActions::GetInclusiveUnionActionToArrays(TArray<FName>& OutElementNames, TArray<FLearningAgentsActionObjectElement>& OutElements, const ULearningAgentsActionObject* Object, const FLearningAgentsActionObjectElement Element, const bool bActive, const FName Tag)
 {
+	if (!bActive)
+	{
+		return GetNullAction(Object, Element, Tag);
+	}
+
 	int32 OutElementNum = 0;
-	if (!GetInclusiveUnionActionNum(OutElementNum, Object, Element, Tag))
+	if (!GetInclusiveUnionActionNum(OutElementNum, Object, Element, bActive, Tag))
 	{
 		OutElementNames.Empty();
 		OutElements.Empty();
@@ -4343,7 +4853,7 @@ bool ULearningAgentsActions::GetInclusiveUnionActionToArrays(TArray<FName>& OutE
 	OutElementNames.SetNumUninitialized(OutElementNum);
 	OutElements.SetNumUninitialized(OutElementNum);
 
-	if (!GetInclusiveUnionActionToArrayViews(OutElementNames, OutElements, Object, Element, Tag))
+	if (!GetInclusiveUnionActionToArrayViews(OutElementNames, OutElements, Object, Element, bActive, Tag))
 	{
 		OutElementNames.Empty();
 		OutElements.Empty();
@@ -4353,8 +4863,13 @@ bool ULearningAgentsActions::GetInclusiveUnionActionToArrays(TArray<FName>& OutE
 	return true;
 }
 
-bool ULearningAgentsActions::GetInclusiveUnionActionToArrayViews(TArrayView<FName> OutElementNames, TArrayView<FLearningAgentsActionObjectElement> OutElements, const ULearningAgentsActionObject* Object, const FLearningAgentsActionObjectElement Element, const FName Tag)
+bool ULearningAgentsActions::GetInclusiveUnionActionToArrayViews(TArrayView<FName> OutElementNames, TArrayView<FLearningAgentsActionObjectElement> OutElements, const ULearningAgentsActionObject* Object, const FLearningAgentsActionObjectElement Element, const bool bActive, const FName Tag)
 {
+	if (!bActive)
+	{
+		return GetNullAction(Object, Element, Tag);
+	}
+
 	if (!Object)
 	{
 		UE_LOG(LogLearning, Error, TEXT("GetInclusiveUnionActionToArrayViews: Object is nullptr."));
@@ -4416,8 +4931,13 @@ bool ULearningAgentsActions::GetInclusiveUnionActionToArrayViews(TArrayView<FNam
 	return true;
 }
 
-bool ULearningAgentsActions::GetStaticArrayActionNum(int32& OutNum, const ULearningAgentsActionObject* Object, const FLearningAgentsActionObjectElement Element, const FName Tag)
+bool ULearningAgentsActions::GetStaticArrayActionNum(int32& OutNum, const ULearningAgentsActionObject* Object, const FLearningAgentsActionObjectElement Element, const bool bActive, const FName Tag)
 {
+	if (!bActive)
+	{
+		return GetNullAction(Object, Element, Tag);
+	}
+
 	if (!Object)
 	{
 		UE_LOG(LogLearning, Error, TEXT("GetStaticArrayActionNum: Object is nullptr."));
@@ -4451,10 +4971,15 @@ bool ULearningAgentsActions::GetStaticArrayActionNum(int32& OutNum, const ULearn
 	return true;
 }
 
-bool ULearningAgentsActions::GetStaticArrayAction(TArray<FLearningAgentsActionObjectElement>& OutElements, const ULearningAgentsActionObject* Object, const FLearningAgentsActionObjectElement Element, const FName Tag)
+bool ULearningAgentsActions::GetStaticArrayAction(TArray<FLearningAgentsActionObjectElement>& OutElements, const ULearningAgentsActionObject* Object, const FLearningAgentsActionObjectElement Element, const bool bActive, const FName Tag)
 {
+	if (!bActive)
+	{
+		return GetNullAction(Object, Element, Tag);
+	}
+
 	int32 OutElementNum = 0;
-	if (!GetStaticArrayActionNum(OutElementNum, Object, Element, Tag))
+	if (!GetStaticArrayActionNum(OutElementNum, Object, Element, bActive, Tag))
 	{
 		OutElements.Empty();
 		return false;
@@ -4462,7 +4987,7 @@ bool ULearningAgentsActions::GetStaticArrayAction(TArray<FLearningAgentsActionOb
 
 	OutElements.SetNumUninitialized(OutElementNum);
 
-	if (!GetStaticArrayActionToArrayView(OutElements, Object, Element, Tag))
+	if (!GetStaticArrayActionToArrayView(OutElements, Object, Element, bActive, Tag))
 	{
 		OutElements.Empty();
 		return false;
@@ -4471,8 +4996,13 @@ bool ULearningAgentsActions::GetStaticArrayAction(TArray<FLearningAgentsActionOb
 	return true;
 }
 
-bool ULearningAgentsActions::GetStaticArrayActionToArrayView(TArrayView<FLearningAgentsActionObjectElement> OutElements, const ULearningAgentsActionObject* Object, const FLearningAgentsActionObjectElement Element, const FName Tag)
+bool ULearningAgentsActions::GetStaticArrayActionToArrayView(TArrayView<FLearningAgentsActionObjectElement> OutElements, const ULearningAgentsActionObject* Object, const FLearningAgentsActionObjectElement Element, const bool bActive, const FName Tag)
 {
+	if (!bActive)
+	{
+		return GetNullAction(Object, Element, Tag);
+	}
+
 	if (!Object)
 	{
 		UE_LOG(LogLearning, Error, TEXT("GetStaticArrayActionToArrayView: Object is nullptr."));
@@ -4533,11 +5063,16 @@ bool ULearningAgentsActions::GetStaticArrayActionToArrayView(TArrayView<FLearnin
 	return true;
 }
 
-bool ULearningAgentsActions::GetPairAction(FLearningAgentsActionObjectElement& OutKey, FLearningAgentsActionObjectElement& OutValue, const ULearningAgentsActionObject* Object, const FLearningAgentsActionObjectElement Element, const FName Tag)
+bool ULearningAgentsActions::GetPairAction(FLearningAgentsActionObjectElement& OutKey, FLearningAgentsActionObjectElement& OutValue, const ULearningAgentsActionObject* Object, const FLearningAgentsActionObjectElement Element, const bool bActive, const FName Tag)
 {
+	if (!bActive)
+	{
+		return GetNullAction(Object, Element, Tag);
+	}
+
 	TStaticArray<FName, 2> OutElementNames;
 	TStaticArray<FLearningAgentsActionObjectElement, 2> OutElements;
-	if (!GetStructActionToArrayViews(OutElementNames, OutElements, Object, Element, Tag))
+	if (!GetStructActionToArrayViews(OutElementNames, OutElements, Object, Element, bActive, Tag))
 	{
 		OutKey = FLearningAgentsActionObjectElement();
 		OutValue = FLearningAgentsActionObjectElement();
@@ -4554,6 +5089,7 @@ bool ULearningAgentsActions::GetEnumAction(
 	const ULearningAgentsActionObject* Object, 
 	const FLearningAgentsActionObjectElement Element, 
 	const UEnum* Enum, 
+	const bool bActive,
 	const FName Tag,
 	const bool bVisualLoggerEnabled,
 	ULearningAgentsManagerListener* VisualLoggerListener,
@@ -4561,6 +5097,11 @@ bool ULearningAgentsActions::GetEnumAction(
 	const FVector VisualLoggerLocation,
 	const FLinearColor VisualLoggerColor)
 {
+	if (!bActive)
+	{
+		return GetNullAction(Object, Element, Tag);
+	}
+
 	if (!Enum)
 	{
 		UE_LOG(LogLearning, Error, TEXT("GetEnumAction: Enum is nullptr."));
@@ -4569,7 +5110,7 @@ bool ULearningAgentsActions::GetEnumAction(
 	}
 
 	int32 OutIndex = 0;
-	if (!GetExclusiveDiscreteAction(OutIndex, Object, Element, Tag))
+	if (!GetExclusiveDiscreteAction(OutIndex, Object, Element, bActive, Tag))
 	{
 		OutEnumValue = 0;
 		return false;
@@ -4619,6 +5160,7 @@ bool ULearningAgentsActions::GetBitmaskAction(
 	const ULearningAgentsActionObject* Object, 
 	const FLearningAgentsActionObjectElement Element, 
 	const UEnum* Enum, 
+	const bool bActive,
 	const FName Tag,
 	const bool bVisualLoggerEnabled,
 	ULearningAgentsManagerListener* VisualLoggerListener,
@@ -4626,6 +5168,11 @@ bool ULearningAgentsActions::GetBitmaskAction(
 	const FVector VisualLoggerLocation,
 	const FLinearColor VisualLoggerColor)
 {
+	if (!bActive)
+	{
+		return GetNullAction(Object, Element, Tag);
+	}
+
 	if (!Enum)
 	{
 		UE_LOG(LogLearning, Error, TEXT("GetBitmaskAction: Enum is nullptr."));
@@ -4641,7 +5188,7 @@ bool ULearningAgentsActions::GetBitmaskAction(
 	}
 
 	int32 EnumValueNum;
-	if (!GetInclusiveDiscreteActionNum(EnumValueNum, Object, Element, Tag))
+	if (!GetInclusiveDiscreteActionNum(EnumValueNum, Object, Element, bActive, Tag))
 	{
 		OutBitmaskValue = 0;
 		return false;
@@ -4656,7 +5203,7 @@ bool ULearningAgentsActions::GetBitmaskAction(
 
 	TArray<int32, TInlineAllocator<32>> OutIndices;
 	OutIndices.SetNumUninitialized(EnumValueNum);
-	if (!GetInclusiveDiscreteActionToArrayView(OutIndices, Object, Element, Tag))
+	if (!GetInclusiveDiscreteActionToArrayView(OutIndices, Object, Element, bActive, Tag))
 	{
 		OutBitmaskValue = 0;
 		return false;
@@ -4704,10 +5251,15 @@ bool ULearningAgentsActions::GetBitmaskAction(
 	return true;
 }
 
-bool ULearningAgentsActions::GetOptionalAction(ELearningAgentsOptionalAction& OutOption, FLearningAgentsActionObjectElement& OutElement, const ULearningAgentsActionObject* Object, const FLearningAgentsActionObjectElement Element, const FName Tag)
+bool ULearningAgentsActions::GetOptionalAction(ELearningAgentsOptionalAction& OutOption, FLearningAgentsActionObjectElement& OutElement, const ULearningAgentsActionObject* Object, const FLearningAgentsActionObjectElement Element, const bool bActive, const FName Tag)
 {
+	if (!bActive)
+	{
+		return GetNullAction(Object, Element, Tag);
+	}
+
 	FName OutName = NAME_None;
-	if (!GetExclusiveUnionAction(OutName, OutElement, Object, Element, Tag))
+	if (!GetExclusiveUnionAction(OutName, OutElement, Object, Element, bActive, Tag))
 	{
 		OutOption = ELearningAgentsOptionalAction::Null;
 		return false;
@@ -4717,10 +5269,15 @@ bool ULearningAgentsActions::GetOptionalAction(ELearningAgentsOptionalAction& Ou
 	return true;
 }
 
-bool ULearningAgentsActions::GetEitherAction(ELearningAgentsEitherAction& OutEither, FLearningAgentsActionObjectElement& OutElement, const ULearningAgentsActionObject* Object, const FLearningAgentsActionObjectElement Element, const FName Tag)
+bool ULearningAgentsActions::GetEitherAction(ELearningAgentsEitherAction& OutEither, FLearningAgentsActionObjectElement& OutElement, const ULearningAgentsActionObject* Object, const FLearningAgentsActionObjectElement Element, const bool bActive, const FName Tag)
 {
+	if (!bActive)
+	{
+		return GetNullAction(Object, Element, Tag);
+	}
+
 	FName OutName = NAME_None;
-	if (!GetExclusiveUnionAction(OutName, OutElement, Object, Element, Tag))
+	if (!GetExclusiveUnionAction(OutName, OutElement, Object, Element, bActive, Tag))
 	{
 		OutEither = ELearningAgentsEitherAction::A;
 		return false;
@@ -4730,8 +5287,13 @@ bool ULearningAgentsActions::GetEitherAction(ELearningAgentsEitherAction& OutEit
 	return true;
 }
 
-bool ULearningAgentsActions::GetEncodingAction(FLearningAgentsActionObjectElement& OutElement, const ULearningAgentsActionObject* Object, const FLearningAgentsActionObjectElement Element, const FName Tag)
+bool ULearningAgentsActions::GetEncodingAction(FLearningAgentsActionObjectElement& OutElement, const ULearningAgentsActionObject* Object, const FLearningAgentsActionObjectElement Element, const bool bActive, const FName Tag)
 {
+	if (!bActive)
+	{
+		return GetNullAction(Object, Element, Tag);
+	}
+
 	if (!Object)
 	{
 		UE_LOG(LogLearning, Error, TEXT("GetEncodingAction: Object is nullptr."));
@@ -4770,6 +5332,7 @@ bool ULearningAgentsActions::GetBoolAction(
 	bool& bOutValue, 
 	const ULearningAgentsActionObject* Object, 
 	const FLearningAgentsActionObjectElement Element, 
+	const bool bActive,	
 	const FName Tag,
 	const bool bVisualLoggerEnabled,
 	ULearningAgentsManagerListener* VisualLoggerListener,
@@ -4777,8 +5340,13 @@ bool ULearningAgentsActions::GetBoolAction(
 	const FVector VisualLoggerLocation,
 	const FLinearColor VisualLoggerColor)
 {
+	if (!bActive)
+	{
+		return GetNullAction(Object, Element, Tag);
+	}
+
 	int32 OutIndex = 0;
-	if (!GetExclusiveDiscreteAction(OutIndex, Object, Element, Tag))
+	if (!GetExclusiveDiscreteAction(OutIndex, Object, Element, bActive, Tag))
 	{
 		bOutValue = false;
 		return false;
@@ -4808,6 +5376,7 @@ bool ULearningAgentsActions::GetFloatAction(
 	float& OutValue, 
 	const ULearningAgentsActionObject* Object, 
 	const FLearningAgentsActionObjectElement Element, 
+	const bool bActive, 
 	const FName Tag,
 	const bool bVisualLoggerEnabled,
 	ULearningAgentsManagerListener* VisualLoggerListener,
@@ -4815,7 +5384,12 @@ bool ULearningAgentsActions::GetFloatAction(
 	const FVector VisualLoggerLocation,
 	const FLinearColor VisualLoggerColor)
 {
-	if (!GetContinuousActionToArrayView(MakeArrayView(&OutValue, 1), Object, Element, Tag))
+	if (!bActive)
+	{
+		return GetNullAction(Object, Element, Tag);
+	}
+
+	if (!GetContinuousActionToArrayView(MakeArrayView(&OutValue, 1), Object, Element, bActive, Tag))
 	{
 		OutValue = 0.0f;
 		return false;
@@ -4844,6 +5418,7 @@ bool ULearningAgentsActions::GetLocationAction(
 	const ULearningAgentsActionObject* Object, 
 	const FLearningAgentsActionObjectElement Element, 
 	const FTransform RelativeTransform, 
+	const bool bActive,
 	const FName Tag,
 	const bool bVisualLoggerEnabled,
 	ULearningAgentsManagerListener* VisualLoggerListener,
@@ -4851,8 +5426,13 @@ bool ULearningAgentsActions::GetLocationAction(
 	const FVector VisualLoggerLocation,
 	const FLinearColor VisualLoggerColor)
 {
+	if (!bActive)
+	{
+		return GetNullAction(Object, Element, Tag);
+	}
+
 	TStaticArray<float, 3> OutValues;
-	if (!GetContinuousActionToArrayView(OutValues, Object, Element, Tag))
+	if (!GetContinuousActionToArrayView(OutValues, Object, Element, bActive, Tag))
 	{
 		OutLocation = FVector::ZeroVector;
 		return false;
@@ -4903,6 +5483,7 @@ bool ULearningAgentsActions::GetRotationAction(
 	const ULearningAgentsActionObject* Object, 
 	const FLearningAgentsActionObjectElement Element, 
 	const FRotator RelativeRotation, 
+	const bool bActive,
 	const FName Tag,
 	const bool bVisualLoggerEnabled,
 	ULearningAgentsManagerListener* VisualLoggerListener,
@@ -4911,8 +5492,13 @@ bool ULearningAgentsActions::GetRotationAction(
 	const FVector VisualLoggerLocation,
 	const FLinearColor VisualLoggerColor)
 {
+	if (!bActive)
+	{
+		return GetNullAction(Object, Element, Tag);
+	}
+
 	FQuat OutRotationQuat;
-	if (!GetRotationActionAsQuat(OutRotationQuat, Object, Element, FQuat::MakeFromRotator(RelativeRotation), Tag, bVisualLoggerEnabled, VisualLoggerListener, VisualLoggerAgentId, VisualLoggerRotationLocation, VisualLoggerLocation, VisualLoggerColor))
+	if (!GetRotationActionAsQuat(OutRotationQuat, Object, Element, FQuat::MakeFromRotator(RelativeRotation), bActive, Tag, bVisualLoggerEnabled, VisualLoggerListener, VisualLoggerAgentId, VisualLoggerRotationLocation, VisualLoggerLocation, VisualLoggerColor))
 	{
 		OutRotation = FRotator::ZeroRotator;
 		return false;
@@ -4927,6 +5513,7 @@ bool ULearningAgentsActions::GetRotationActionAsQuat(
 	const ULearningAgentsActionObject* Object, 
 	const FLearningAgentsActionObjectElement Element, 
 	const FQuat RelativeRotation, 
+	const bool bActive,
 	const FName Tag,
 	const bool bVisualLoggerEnabled,
 	ULearningAgentsManagerListener* VisualLoggerListener,
@@ -4935,8 +5522,13 @@ bool ULearningAgentsActions::GetRotationActionAsQuat(
 	const FVector VisualLoggerLocation,
 	const FLinearColor VisualLoggerColor)
 {
+	if (!bActive)
+	{
+		return GetNullAction(Object, Element, Tag);
+	}
+
 	TStaticArray<float, 3> OutValues;
-	if (!GetContinuousActionToArrayView(OutValues, Object, Element, Tag))
+	if (!GetContinuousActionToArrayView(OutValues, Object, Element, bActive, Tag))
 	{
 		OutRotation = FQuat::Identity;
 		return false;
@@ -4977,6 +5569,7 @@ bool ULearningAgentsActions::GetScaleAction(
 	const ULearningAgentsActionObject* Object, 
 	const FLearningAgentsActionObjectElement Element, 
 	const FVector RelativeScale, 
+	const bool bActive,
 	const FName Tag,
 	const bool bVisualLoggerEnabled,
 	ULearningAgentsManagerListener* VisualLoggerListener,
@@ -4984,8 +5577,13 @@ bool ULearningAgentsActions::GetScaleAction(
 	const FVector VisualLoggerLocation,
 	const FLinearColor VisualLoggerColor)
 {
+	if (!bActive)
+	{
+		return GetNullAction(Object, Element, Tag);
+	}
+
 	TStaticArray<float, 3> OutValues;
-	if (!GetContinuousActionToArrayView(OutValues, Object, Element, Tag))
+	if (!GetContinuousActionToArrayView(OutValues, Object, Element, bActive, Tag))
 	{
 		OutScale = FVector::OneVector;
 		return false;
@@ -5018,6 +5616,7 @@ bool ULearningAgentsActions::GetTransformAction(
 	const ULearningAgentsActionObject* Object, 
 	const FLearningAgentsActionObjectElement Element, 
 	const FTransform RelativeTransform, 
+	const bool bActive,
 	const FName Tag,
 	const bool bVisualLoggerEnabled,
 	ULearningAgentsManagerListener* VisualLoggerListener,
@@ -5025,9 +5624,14 @@ bool ULearningAgentsActions::GetTransformAction(
 	const FVector VisualLoggerLocation,
 	const FLinearColor VisualLoggerColor)
 {
+	if (!bActive)
+	{
+		return GetNullAction(Object, Element, Tag);
+	}
+
 	TStaticArray<FName, 3> OutElementNames;
 	TStaticArray<FLearningAgentsActionObjectElement, 3> OutElements;
-	if (!GetStructActionToArrayViews(OutElementNames, OutElements, Object, Element, Tag))
+	if (!GetStructActionToArrayViews(OutElementNames, OutElements, Object, Element, bActive, Tag))
 	{
 		OutTransform = FTransform::Identity;
 		return false;
@@ -5087,6 +5691,7 @@ bool ULearningAgentsActions::GetAngleAction(
 	const ULearningAgentsActionObject* Object, 
 	const FLearningAgentsActionObjectElement Element, 
 	const float RelativeAngle, 
+	const bool bActive,
 	const FName Tag,
 	const bool bVisualLoggerEnabled,
 	ULearningAgentsManagerListener* VisualLoggerListener,
@@ -5095,8 +5700,13 @@ bool ULearningAgentsActions::GetAngleAction(
 	const FVector VisualLoggerLocation,
 	const FLinearColor VisualLoggerColor)
 {
+	if (!bActive)
+	{
+		return GetNullAction(Object, Element, Tag);
+	}
+
 	float LocalAngle = 0.0f;
-	if (!GetContinuousActionToArrayView(MakeArrayView(&LocalAngle, 1), Object, Element, Tag))
+	if (!GetContinuousActionToArrayView(MakeArrayView(&LocalAngle, 1), Object, Element, bActive, Tag))
 	{
 		OutAngle = 0.0f;
 		return false;
@@ -5136,6 +5746,7 @@ bool ULearningAgentsActions::GetAngleActionRadians(
 	const ULearningAgentsActionObject* Object, 
 	const FLearningAgentsActionObjectElement Element, 
 	const float RelativeAngle, 
+	const bool bActive,	
 	const FName Tag,
 	const bool bVisualLoggerEnabled,
 	ULearningAgentsManagerListener* VisualLoggerListener,
@@ -5144,7 +5755,12 @@ bool ULearningAgentsActions::GetAngleActionRadians(
 	const FVector VisualLoggerLocation,
 	const FLinearColor VisualLoggerColor)
 {
-	if (!GetAngleAction(OutAngle, Object, Element, FMath::RadiansToDegrees(RelativeAngle), Tag, bVisualLoggerEnabled, VisualLoggerListener, VisualLoggerAgentId, VisualLoggerAngleLocation, VisualLoggerLocation, VisualLoggerColor))
+	if (!bActive)
+	{
+		return GetNullAction(Object, Element, Tag);
+	}
+
+	if (!GetAngleAction(OutAngle, Object, Element, FMath::RadiansToDegrees(RelativeAngle), bActive, Tag, bVisualLoggerEnabled, VisualLoggerListener, VisualLoggerAgentId, VisualLoggerAngleLocation, VisualLoggerLocation, VisualLoggerColor))
 	{
 		OutAngle = 0.0f;
 		return false;
@@ -5159,6 +5775,7 @@ bool ULearningAgentsActions::GetVelocityAction(
 	const ULearningAgentsActionObject* Object, 
 	const FLearningAgentsActionObjectElement Element, 
 	const FTransform RelativeTransform, 
+	const bool bActive,
 	const FName Tag,
 	const bool bVisualLoggerEnabled,
 	ULearningAgentsManagerListener* VisualLoggerListener,
@@ -5167,8 +5784,13 @@ bool ULearningAgentsActions::GetVelocityAction(
 	const FVector VisualLoggerLocation,
 	const FLinearColor VisualLoggerColor)
 {
+	if (!bActive)
+	{
+		return GetNullAction(Object, Element, Tag);
+	}
+
 	TStaticArray<float, 3> OutValues;
-	if (!GetContinuousActionToArrayView(OutValues, Object, Element, Tag))
+	if (!GetContinuousActionToArrayView(OutValues, Object, Element, bActive, Tag))
 	{
 		OutVelocity = FVector::OneVector;
 		return false;
@@ -5213,6 +5835,7 @@ bool ULearningAgentsActions::GetDirectionAction(
 	const ULearningAgentsActionObject* Object, 
 	const FLearningAgentsActionObjectElement Element, 
 	const FTransform RelativeTransform, 
+	const bool bActive,	
 	const FName Tag,
 	const bool bVisualLoggerEnabled,
 	ULearningAgentsManagerListener* VisualLoggerListener,
@@ -5222,8 +5845,13 @@ bool ULearningAgentsActions::GetDirectionAction(
 	const float VisualLoggerArrowLength,
 	const FLinearColor VisualLoggerColor)
 {
+	if (!bActive)
+	{
+		return GetNullAction(Object, Element, Tag);
+	}
+
 	TStaticArray<float, 3> OutValues;
-	if (!GetContinuousActionToArrayView(OutValues, Object, Element, Tag))
+	if (!GetContinuousActionToArrayView(OutValues, Object, Element, bActive, Tag))
 	{
 		OutDirection = FVector::ForwardVector;
 		return false;
