READY TO SUBMIT

diff --git a/Engine/Plugins/Experimental/LearningAgents/Source/LearningAgents/Private/LearningAgentsObservations.cpp b/Engine/Plugins/Experimental/LearningAgents/Source/LearningAgents/Private/LearningAgentsObservations.cpp
--- a/Engine/Plugins/Experimental/LearningAgents/Source/LearningAgents/Private/LearningAgentsObservations.cpp
+++ b/Engine/Plugins/Experimental/LearningAgents/Source/LearningAgents/Private/LearningAgentsObservations.cpp
@@ -675,8 +675,13 @@ FLearningAgentsObservationSchemaElement ULearningAgentsObservations::SpecifyNull
 	return { Schema->ObservationSchema.CreateNull(Tag)};
 }
 
-FLearningAgentsObservationSchemaElement ULearningAgentsObservations::SpecifyContinuousObservation(ULearningAgentsObservationSchema* Schema, const int32 Size, const float Scale, const FName Tag)
+FLearningAgentsObservationSchemaElement ULearningAgentsObservations::SpecifyContinuousObservation(ULearningAgentsObservationSchema* Schema, const int32 Size, const float Scale, const bool bActive, const FName Tag)
 {
+	if (!bActive)
+	{
+		return SpecifyNullObservation(Schema, Tag);
+	}
+
 	if (!Schema)
 	{
 		UE_LOG(LogLearning, Error, TEXT("SpecifyContinuousObservation: Schema is nullptr."));
@@ -703,8 +708,13 @@ FLearningAgentsObservationSchemaElement ULearningAgentsObservations::SpecifyCont
 	return { Schema->ObservationSchema.CreateContinuous({ Size, Scale }, Tag) };
 }
 
-FLearningAgentsObservationSchemaElement ULearningAgentsObservations::SpecifyExclusiveDiscreteObservation(ULearningAgentsObservationSchema* Schema, const int32 Size, const FName Tag)
+FLearningAgentsObservationSchemaElement ULearningAgentsObservations::SpecifyExclusiveDiscreteObservation(ULearningAgentsObservationSchema* Schema, const int32 Size, const bool bActive, const FName Tag)
 {
+	if (!bActive)
+	{
+		return SpecifyNullObservation(Schema, Tag);
+	}
+
 	if (!Schema)
 	{
 		UE_LOG(LogLearning, Error, TEXT("SpecifyExclusiveDiscreteObservation: Schema is nullptr."));
@@ -720,8 +730,13 @@ FLearningAgentsObservationSchemaElement ULearningAgentsObservations::SpecifyExcl
 	return { Schema->ObservationSchema.CreateDiscreteExclusive({ Size }, Tag) };
 }
 
-FLearningAgentsObservationSchemaElement ULearningAgentsObservations::SpecifyInclusiveDiscreteObservation(ULearningAgentsObservationSchema* Schema, const int32 Size, const FName Tag)
+FLearningAgentsObservationSchemaElement ULearningAgentsObservations::SpecifyInclusiveDiscreteObservation(ULearningAgentsObservationSchema* Schema, const int32 Size, const bool bActive, const FName Tag)
 {
+	if (!bActive)
+	{
+		return SpecifyNullObservation(Schema, Tag);
+	}
+
 	if (!Schema)
 	{
 		UE_LOG(LogLearning, Error, TEXT("SpecifyInclusiveDiscreteObservation: Schema is nullptr."));
@@ -742,13 +757,23 @@ FLearningAgentsObservationSchemaElement ULearningAgentsObservations::SpecifyIncl
 	return { Schema->ObservationSchema.CreateDiscreteInclusive({ Size }, Tag) };
 }
 
-FLearningAgentsObservationSchemaElement ULearningAgentsObservations::SpecifyNamedExclusiveDiscreteObservation(ULearningAgentsObservationSchema* Schema, const TArray<FName>& ElementNames, const FName Tag)
+FLearningAgentsObservationSchemaElement ULearningAgentsObservations::SpecifyNamedExclusiveDiscreteObservation(ULearningAgentsObservationSchema* Schema, const TArray<FName>& ElementNames, const bool bActive, const FName Tag)
 {
-	return SpecifyNamedExclusiveDiscreteObservationFromArrayView(Schema, ElementNames, Tag);
+	if (!bActive)
+	{
+		return SpecifyNullObservation(Schema, Tag);
+	}
+
+	return SpecifyNamedExclusiveDiscreteObservationFromArrayView(Schema, ElementNames, bActive, Tag);
 }
 
-FLearningAgentsObservationSchemaElement ULearningAgentsObservations::SpecifyNamedExclusiveDiscreteObservationFromArrayView(ULearningAgentsObservationSchema* Schema, const TArrayView<const FName> ElementNames, const FName Tag)
+FLearningAgentsObservationSchemaElement ULearningAgentsObservations::SpecifyNamedExclusiveDiscreteObservationFromArrayView(ULearningAgentsObservationSchema* Schema, const TArrayView<const FName> ElementNames, const bool bActive, const FName Tag)
 {
+	if (!bActive)
+	{
+		return SpecifyNullObservation(Schema, Tag);
+	}
+
 	if (!Schema)
 	{
 		UE_LOG(LogLearning, Error, TEXT("SpecifyNamedExclusiveDiscreteObservationFromArrayView: Schema is nullptr."));
@@ -770,13 +795,23 @@ FLearningAgentsObservationSchemaElement ULearningAgentsObservations::SpecifyName
 	return { Schema->ObservationSchema.CreateNamedDiscreteExclusive({ ElementNames }, Tag) };
 }
 
-FLearningAgentsObservationSchemaElement ULearningAgentsObservations::SpecifyNamedInclusiveDiscreteObservation(ULearningAgentsObservationSchema* Schema, const TArray<FName>& ElementNames, const FName Tag)
+FLearningAgentsObservationSchemaElement ULearningAgentsObservations::SpecifyNamedInclusiveDiscreteObservation(ULearningAgentsObservationSchema* Schema, const TArray<FName>& ElementNames, const bool bActive, const FName Tag)
 {
-	return SpecifyNamedInclusiveDiscreteObservationFromArrayView(Schema, ElementNames, Tag);
+	if (!bActive)
+	{
+		return SpecifyNullObservation(Schema, Tag);
+	}
+
+	return SpecifyNamedInclusiveDiscreteObservationFromArrayView(Schema, ElementNames, bActive, Tag);
 }
 
-FLearningAgentsObservationSchemaElement ULearningAgentsObservations::SpecifyNamedInclusiveDiscreteObservationFromArrayView(ULearningAgentsObservationSchema* Schema, const TArrayView<const FName> ElementNames, const FName Tag)
+FLearningAgentsObservationSchemaElement ULearningAgentsObservations::SpecifyNamedInclusiveDiscreteObservationFromArrayView(ULearningAgentsObservationSchema* Schema, const TArrayView<const FName> ElementNames, const bool bActive, const FName Tag)
 {
+	if (!bActive)
+	{
+		return SpecifyNullObservation(Schema, Tag);
+	}
+
 	if (!Schema)
 	{
 		UE_LOG(LogLearning, Error, TEXT("SpecifyNamedInclusiveDiscreteObservationFromArrayView: Schema is nullptr."));
@@ -797,13 +832,23 @@ FLearningAgentsObservationSchemaElement ULearningAgentsObservations::SpecifyName
 	return { Schema->ObservationSchema.CreateNamedDiscreteInclusive({ ElementNames }, Tag) };
 }
 
-FLearningAgentsObservationSchemaElement ULearningAgentsObservations::SpecifyCountObservation(ULearningAgentsObservationSchema* Schema, const FName Tag)
+FLearningAgentsObservationSchemaElement ULearningAgentsObservations::SpecifyCountObservation(ULearningAgentsObservationSchema* Schema, const bool bActive, const FName Tag)
 {
-	return SpecifyContinuousObservation(Schema, 1, 1.0f, Tag);
+	if (!bActive)
+	{
+		return SpecifyNullObservation(Schema, Tag);
+	}
+
+	return SpecifyContinuousObservation(Schema, 1, 1.0f, bActive, Tag);
 }
 
-FLearningAgentsObservationSchemaElement ULearningAgentsObservations::SpecifyStructObservation(ULearningAgentsObservationSchema* Schema, const TMap<FName, FLearningAgentsObservationSchemaElement>& Elements, const FName Tag)
+FLearningAgentsObservationSchemaElement ULearningAgentsObservations::SpecifyStructObservation(ULearningAgentsObservationSchema* Schema, const TMap<FName, FLearningAgentsObservationSchemaElement>& Elements, const bool bActive, const FName Tag)
 {
+	if (!bActive)
+	{
+		return SpecifyNullObservation(Schema, Tag);
+	}
+
 	if (Elements.Num() == 0)
 	{
 		UE_LOG(LogLearning, Warning, TEXT("SpecifyStructObservation: Specifying zero-sized Struct Observation."));
@@ -844,16 +889,26 @@ FLearningAgentsObservationSchemaElement ULearningAgentsObservations::SpecifyStru
 		SortedSubElements[Idx] = SubElements[SubElementIndices[Idx]];
 	}
 
-	return SpecifyStructObservationFromArrayViews(Schema, SortedSubElementNames, SortedSubElements, Tag);
+	return SpecifyStructObservationFromArrayViews(Schema, SortedSubElementNames, SortedSubElements, bActive, Tag);
 }
 
-FLearningAgentsObservationSchemaElement ULearningAgentsObservations::SpecifyStructObservationFromArrays(ULearningAgentsObservationSchema* Schema, const TArray<FName>& ElementNames, const TArray<FLearningAgentsObservationSchemaElement>& Elements, const FName Tag)
+FLearningAgentsObservationSchemaElement ULearningAgentsObservations::SpecifyStructObservationFromArrays(ULearningAgentsObservationSchema* Schema, const TArray<FName>& ElementNames, const TArray<FLearningAgentsObservationSchemaElement>& Elements, const bool bActive, const FName Tag)
 {
-	return SpecifyStructObservationFromArrayViews(Schema, ElementNames, Elements, Tag);
+	if (!bActive)
+	{
+		return SpecifyNullObservation(Schema, Tag);
+	}
+
+	return SpecifyStructObservationFromArrayViews(Schema, ElementNames, Elements, bActive, Tag);
 }
 
-FLearningAgentsObservationSchemaElement ULearningAgentsObservations::SpecifyStructObservationFromArrayViews(ULearningAgentsObservationSchema* Schema, const TArrayView<const FName> ElementNames, const TArrayView<const FLearningAgentsObservationSchemaElement> Elements, const FName Tag)
+FLearningAgentsObservationSchemaElement ULearningAgentsObservations::SpecifyStructObservationFromArrayViews(ULearningAgentsObservationSchema* Schema, const TArrayView<const FName> ElementNames, const TArrayView<const FLearningAgentsObservationSchemaElement> Elements, const bool bActive, const FName Tag)
 {
+	if (!bActive)
+	{
+		return SpecifyNullObservation(Schema, Tag);
+	}
+
 	if (!Schema)
 	{
 		UE_LOG(LogLearning, Error, TEXT("SpecifyStructObservationFromArrayViews: Schema is nullptr."));
@@ -894,8 +949,13 @@ FLearningAgentsObservationSchemaElement ULearningAgentsObservations::SpecifyStru
 	return { Schema->ObservationSchema.CreateAnd({ ElementNames, SubElements }, Tag) };
 }
 
-FLearningAgentsObservationSchemaElement ULearningAgentsObservations::SpecifyExclusiveUnionObservation(ULearningAgentsObservationSchema* Schema, const TMap<FName, FLearningAgentsObservationSchemaElement>& Elements, const int32 EncodingSize, const FName Tag)
+FLearningAgentsObservationSchemaElement ULearningAgentsObservations::SpecifyExclusiveUnionObservation(ULearningAgentsObservationSchema* Schema, const TMap<FName, FLearningAgentsObservationSchemaElement>& Elements, const int32 EncodingSize, const bool bActive, const FName Tag)
 {
+	if (!bActive)
+	{
+		return SpecifyNullObservation(Schema, Tag);
+	}
+
 	if (EncodingSize < 1)
 	{
 		UE_LOG(LogLearning, Error, TEXT("SpecifyExclusiveUnionObservation: Invalid Observation EncodingSize '%i' - must be greater than zero."), EncodingSize);
@@ -942,16 +1002,26 @@ FLearningAgentsObservationSchemaElement ULearningAgentsObservations::SpecifyExcl
 		SortedSubElements[Idx] = SubElements[SubElementIndices[Idx]];
 	}
 
-	return SpecifyExclusiveUnionObservationFromArrayViews(Schema, SortedSubElementNames, SortedSubElements, EncodingSize, Tag);
+	return SpecifyExclusiveUnionObservationFromArrayViews(Schema, SortedSubElementNames, SortedSubElements, EncodingSize, bActive, Tag);
 }
 
-FLearningAgentsObservationSchemaElement ULearningAgentsObservations::SpecifyExclusiveUnionObservationFromArrays(ULearningAgentsObservationSchema* Schema, const TArray<FName>& ElementNames, const TArray<FLearningAgentsObservationSchemaElement>& Elements, const int32 EncodingSize, const FName Tag)
+FLearningAgentsObservationSchemaElement ULearningAgentsObservations::SpecifyExclusiveUnionObservationFromArrays(ULearningAgentsObservationSchema* Schema, const TArray<FName>& ElementNames, const TArray<FLearningAgentsObservationSchemaElement>& Elements, const int32 EncodingSize, const bool bActive, const FName Tag)
 {
-	return SpecifyExclusiveUnionObservationFromArrayViews(Schema, ElementNames, Elements, EncodingSize, Tag);
+	if (!bActive)
+	{
+		return SpecifyNullObservation(Schema, Tag);
+	}
+
+	return SpecifyExclusiveUnionObservationFromArrayViews(Schema, ElementNames, Elements, EncodingSize, bActive, Tag);
 }
 
-FLearningAgentsObservationSchemaElement ULearningAgentsObservations::SpecifyExclusiveUnionObservationFromArrayViews(ULearningAgentsObservationSchema* Schema, const TArrayView<const FName> ElementNames, const TArrayView<const FLearningAgentsObservationSchemaElement> Elements, const int32 EncodingSize, const FName Tag)
+FLearningAgentsObservationSchemaElement ULearningAgentsObservations::SpecifyExclusiveUnionObservationFromArrayViews(ULearningAgentsObservationSchema* Schema, const TArrayView<const FName> ElementNames, const TArrayView<const FLearningAgentsObservationSchemaElement> Elements, const int32 EncodingSize, const bool bActive, const FName Tag)
 {
+	if (!bActive)
+	{
+		return SpecifyNullObservation(Schema, Tag);
+	}
+
 	if (!Schema)
 	{
 		UE_LOG(LogLearning, Error, TEXT("SpecifyExclusiveUnionObservationFromArrayViews: Schema is nullptr."));
@@ -998,8 +1068,13 @@ FLearningAgentsObservationSchemaElement ULearningAgentsObservations::SpecifyExcl
 	return { Schema->ObservationSchema.CreateOrExclusive({ ElementNames, SubElements, EncodingSize }, Tag) };
 }
 
-FLearningAgentsObservationSchemaElement ULearningAgentsObservations::SpecifyInclusiveUnionObservation(ULearningAgentsObservationSchema* Schema, const TMap<FName, FLearningAgentsObservationSchemaElement>& Elements, const int32 AttentionEncodingSize, const int32 AttentionHeadNum, const int32 ValueEncodingSize, const FName Tag)
+FLearningAgentsObservationSchemaElement ULearningAgentsObservations::SpecifyInclusiveUnionObservation(ULearningAgentsObservationSchema* Schema, const TMap<FName, FLearningAgentsObservationSchemaElement>& Elements, const int32 AttentionEncodingSize, const int32 AttentionHeadNum, const int32 ValueEncodingSize, const bool bActive, const FName Tag)
 {
+	if (!bActive)
+	{
+		return SpecifyNullObservation(Schema, Tag);
+	}
+
 	if (AttentionEncodingSize < 1 || AttentionHeadNum < 1 || ValueEncodingSize < 1)
 	{
 		UE_LOG(LogLearning, Error, TEXT("SpecifyInclusiveUnionObservation: Invalid Observation Parameters: AttentionEncodingSize: %i, AttentionHeadNum: %i, ValueEncodingSize: %i - must be greater than zero."), AttentionEncodingSize, AttentionHeadNum, ValueEncodingSize);
@@ -1046,16 +1121,26 @@ FLearningAgentsObservationSchemaElement ULearningAgentsObservations::SpecifyIncl
 		SortedSubElements[Idx] = SubElements[SubElementIndices[Idx]];
 	}
 
-	return SpecifyInclusiveUnionObservationFromArrayViews(Schema, SortedSubElementNames, SortedSubElements, AttentionEncodingSize, AttentionHeadNum, ValueEncodingSize, Tag);
+	return SpecifyInclusiveUnionObservationFromArrayViews(Schema, SortedSubElementNames, SortedSubElements, AttentionEncodingSize, AttentionHeadNum, ValueEncodingSize, bActive, Tag);
 }
 
-FLearningAgentsObservationSchemaElement ULearningAgentsObservations::SpecifyInclusiveUnionObservationFromArrays(ULearningAgentsObservationSchema* Schema, const TArray<FName>& ElementNames, const TArray<FLearningAgentsObservationSchemaElement>& Elements, const int32 AttentionEncodingSize, const int32 AttentionHeadNum, const int32 ValueEncodingSize, const FName Tag)
+FLearningAgentsObservationSchemaElement ULearningAgentsObservations::SpecifyInclusiveUnionObservationFromArrays(ULearningAgentsObservationSchema* Schema, const TArray<FName>& ElementNames, const TArray<FLearningAgentsObservationSchemaElement>& Elements, const int32 AttentionEncodingSize, const int32 AttentionHeadNum, const int32 ValueEncodingSize, const bool bActive, const FName Tag)
 {
-	return SpecifyInclusiveUnionObservationFromArrayViews(Schema, ElementNames, Elements, AttentionEncodingSize, AttentionHeadNum, ValueEncodingSize, Tag);
+	if (!bActive)
+	{
+		return SpecifyNullObservation(Schema, Tag);
+	}
+
+	return SpecifyInclusiveUnionObservationFromArrayViews(Schema, ElementNames, Elements, AttentionEncodingSize, AttentionHeadNum, ValueEncodingSize, bActive, Tag);
 }
 
-FLearningAgentsObservationSchemaElement ULearningAgentsObservations::SpecifyInclusiveUnionObservationFromArrayViews(ULearningAgentsObservationSchema* Schema, const TArrayView<const FName> ElementNames, const TArrayView<const FLearningAgentsObservationSchemaElement> Elements, const int32 AttentionEncodingSize, const int32 AttentionHeadNum, const int32 ValueEncodingSize, const FName Tag)
+FLearningAgentsObservationSchemaElement ULearningAgentsObservations::SpecifyInclusiveUnionObservationFromArrayViews(ULearningAgentsObservationSchema* Schema, const TArrayView<const FName> ElementNames, const TArrayView<const FLearningAgentsObservationSchemaElement> Elements, const int32 AttentionEncodingSize, const int32 AttentionHeadNum, const int32 ValueEncodingSize, const bool bActive, const FName Tag)
 {
+	if (!bActive)
+	{
+		return SpecifyNullObservation(Schema, Tag);
+	}
+
 	if (!Schema)
 	{
 		UE_LOG(LogLearning, Error, TEXT("SpecifyInclusiveUnionObservationFromArrayViews: Schema is nullptr."));
@@ -1102,8 +1187,13 @@ FLearningAgentsObservationSchemaElement ULearningAgentsObservations::SpecifyIncl
 	return { Schema->ObservationSchema.CreateOrInclusive({ ElementNames, SubElements, AttentionEncodingSize, AttentionHeadNum, ValueEncodingSize }, Tag) };
 }
 
-FLearningAgentsObservationSchemaElement ULearningAgentsObservations::SpecifyStaticArrayObservation(ULearningAgentsObservationSchema* Schema, const FLearningAgentsObservationSchemaElement Element, const int32 Num, const FName Tag)
+FLearningAgentsObservationSchemaElement ULearningAgentsObservations::SpecifyStaticArrayObservation(ULearningAgentsObservationSchema* Schema, const FLearningAgentsObservationSchemaElement Element, const int32 Num, const bool bActive, const FName Tag)
 {
+	if (!bActive)
+	{
+		return SpecifyNullObservation(Schema, Tag);
+	}
+
 	if (!Schema)
 	{
 		UE_LOG(LogLearning, Error, TEXT("SpecifyStaticArrayObservation: Schema is nullptr."));
@@ -1130,8 +1220,13 @@ FLearningAgentsObservationSchemaElement ULearningAgentsObservations::SpecifyStat
 	return { Schema->ObservationSchema.CreateArray({ Element.SchemaElement, Num }, Tag) };
 }
 
-FLearningAgentsObservationSchemaElement ULearningAgentsObservations::SpecifySetObservation(ULearningAgentsObservationSchema* Schema, const FLearningAgentsObservationSchemaElement Element, const int32 MaxNum, const int32 AttentionEncodingSize, const int32 AttentionHeadNum, const int32 ValueEncodingSize, const FName Tag)
+FLearningAgentsObservationSchemaElement ULearningAgentsObservations::SpecifySetObservation(ULearningAgentsObservationSchema* Schema, const FLearningAgentsObservationSchemaElement Element, const int32 MaxNum, const int32 AttentionEncodingSize, const int32 AttentionHeadNum, const int32 ValueEncodingSize, const bool bActive, const FName Tag)
 {
+	if (!bActive)
+	{
+		return SpecifyNullObservation(Schema, Tag);
+	}
+
 	if (!Schema)
 	{
 		UE_LOG(LogLearning, Error, TEXT("SpecifySetObservation: Schema is nullptr."));
@@ -1164,34 +1259,59 @@ FLearningAgentsObservationSchemaElement ULearningAgentsObservations::SpecifySetO
 	return { Schema->ObservationSchema.CreateSet({ Element.SchemaElement, MaxNum, AttentionEncodingSize, AttentionHeadNum, ValueEncodingSize }, Tag) };
 }
 
-FLearningAgentsObservationSchemaElement ULearningAgentsObservations::SpecifyPairObservation(ULearningAgentsObservationSchema* Schema, const FLearningAgentsObservationSchemaElement Element0, const FLearningAgentsObservationSchemaElement Element1, const FName Tag)
+FLearningAgentsObservationSchemaElement ULearningAgentsObservations::SpecifyPairObservation(ULearningAgentsObservationSchema* Schema, const FLearningAgentsObservationSchemaElement Element0, const FLearningAgentsObservationSchemaElement Element1, const bool bActive, const FName Tag)
 {
-	return SpecifyStructObservationFromArrayViews(Schema, { TEXT("Key"), TEXT("Value") }, { Element0, Element1 }, Tag);
+	if (!bActive)
+	{
+		return SpecifyNullObservation(Schema, Tag);
+	}
+
+	return SpecifyStructObservationFromArrayViews(Schema, { TEXT("Key"), TEXT("Value") }, { Element0, Element1 }, bActive, Tag);
 }
 
-FLearningAgentsObservationSchemaElement ULearningAgentsObservations::SpecifyArrayObservation(ULearningAgentsObservationSchema* Schema, const FLearningAgentsObservationSchemaElement Element, const int32 MaxNum, const int32 AttentionEncodingSize, const int32 AttentionHeadNum, const int32 ValueEncodingSize, const FName Tag)
+FLearningAgentsObservationSchemaElement ULearningAgentsObservations::SpecifyArrayObservation(ULearningAgentsObservationSchema* Schema, const FLearningAgentsObservationSchemaElement Element, const int32 MaxNum, const int32 AttentionEncodingSize, const int32 AttentionHeadNum, const int32 ValueEncodingSize, const bool bActive, const FName Tag)
 {
-	return SpecifySetObservation(Schema, SpecifyPairObservation(Schema, SpecifyCountObservation(Schema), Element), MaxNum, AttentionEncodingSize, AttentionHeadNum, ValueEncodingSize, Tag);
+	if (!bActive)
+	{
+		return SpecifyNullObservation(Schema, Tag);
+	}
+
+	return SpecifySetObservation(Schema, SpecifyPairObservation(Schema, SpecifyCountObservation(Schema), Element), MaxNum, AttentionEncodingSize, AttentionHeadNum, ValueEncodingSize, bActive, Tag);
 }
 
-FLearningAgentsObservationSchemaElement ULearningAgentsObservations::SpecifyMapObservation(ULearningAgentsObservationSchema* Schema, const FLearningAgentsObservationSchemaElement KeyElement, const FLearningAgentsObservationSchemaElement ValueElement, const int32 MaxNum, const int32 AttentionEncodingSize, const int32 AttentionHeadNum, const int32 ValueEncodingSize, const FName Tag)
+FLearningAgentsObservationSchemaElement ULearningAgentsObservations::SpecifyMapObservation(ULearningAgentsObservationSchema* Schema, const FLearningAgentsObservationSchemaElement KeyElement, const FLearningAgentsObservationSchemaElement ValueElement, const int32 MaxNum, const int32 AttentionEncodingSize, const int32 AttentionHeadNum, const int32 ValueEncodingSize, const bool bActive, const FName Tag)
 {
-	return SpecifySetObservation(Schema, SpecifyPairObservation(Schema, KeyElement, ValueElement), MaxNum, AttentionEncodingSize, AttentionHeadNum, ValueEncodingSize, Tag);
+	if (!bActive)
+	{
+		return SpecifyNullObservation(Schema, Tag);
+	}
+
+	return SpecifySetObservation(Schema, SpecifyPairObservation(Schema, KeyElement, ValueElement), MaxNum, AttentionEncodingSize, AttentionHeadNum, ValueEncodingSize, bActive, Tag);
 }
 
-FLearningAgentsObservationSchemaElement ULearningAgentsObservations::SpecifyEnumObservation(ULearningAgentsObservationSchema* Schema, const UEnum* Enum, const FName Tag)
+FLearningAgentsObservationSchemaElement ULearningAgentsObservations::SpecifyEnumObservation(ULearningAgentsObservationSchema* Schema, const UEnum* Enum, const bool bActive, const FName Tag)
 {
+	if (!bActive)
+	{
+		return SpecifyNullObservation(Schema, Tag);
+	}
+
 	if (!Enum)
 	{
 		UE_LOG(LogLearning, Error, TEXT("SpecifyEnumObservation: Enum is nullptr."));
 		return FLearningAgentsObservationSchemaElement();
 	}
 
-	return SpecifyContinuousObservation(Schema, Enum->NumEnums() - 1, 1.0f, Tag);
+	return SpecifyContinuousObservation(Schema, Enum->NumEnums() - 1, 1.0f, bActive, Tag);
 }
 
-FLearningAgentsObservationSchemaElement ULearningAgentsObservations::SpecifyBitmaskObservation(ULearningAgentsObservationSchema* Schema, const UEnum* Enum, const FName Tag)
+FLearningAgentsObservationSchemaElement ULearningAgentsObservations::SpecifyBitmaskObservation(ULearningAgentsObservationSchema* Schema, const UEnum* Enum, const bool bActive, const FName Tag)
 {
+	if (!bActive)
+	{
+		return SpecifyNullObservation(Schema, Tag);
+	}
+
 	if (!Enum)
 	{
 		UE_LOG(LogLearning, Error, TEXT("SpecifyBitmaskObservation: Enum is nullptr."));
@@ -1204,21 +1324,36 @@ FLearningAgentsObservationSchemaElement ULearningAgentsObservations::SpecifyBitm
 		return FLearningAgentsObservationSchemaElement();
 	}
 
-	return SpecifyContinuousObservation(Schema, Enum->NumEnums() - 1, 1.0f, Tag);
+	return SpecifyContinuousObservation(Schema, Enum->NumEnums() - 1, 1.0f, bActive, Tag);
 }
 
-FLearningAgentsObservationSchemaElement ULearningAgentsObservations::SpecifyOptionalObservation(ULearningAgentsObservationSchema* Schema, const FLearningAgentsObservationSchemaElement Element, const int32 EncodingSize, const FName Tag)
+FLearningAgentsObservationSchemaElement ULearningAgentsObservations::SpecifyOptionalObservation(ULearningAgentsObservationSchema* Schema, const FLearningAgentsObservationSchemaElement Element, const int32 EncodingSize, const bool bActive, const FName Tag)
 {
-	return SpecifyExclusiveUnionObservationFromArrayViews(Schema, { TEXT("Null"), TEXT("Valid") }, { SpecifyNullObservation(Schema), Element }, EncodingSize, Tag);
+	if (!bActive)
+	{
+		return SpecifyNullObservation(Schema, Tag);
+	}
+
+	return SpecifyExclusiveUnionObservationFromArrayViews(Schema, { TEXT("Null"), TEXT("Valid") }, { SpecifyNullObservation(Schema), Element }, EncodingSize, bActive, Tag);
 }
 
-FLearningAgentsObservationSchemaElement ULearningAgentsObservations::SpecifyEitherObservation(ULearningAgentsObservationSchema* Schema, const FLearningAgentsObservationSchemaElement A, const FLearningAgentsObservationSchemaElement B, const int32 EncodingSize, const FName Tag)
+FLearningAgentsObservationSchemaElement ULearningAgentsObservations::SpecifyEitherObservation(ULearningAgentsObservationSchema* Schema, const FLearningAgentsObservationSchemaElement A, const FLearningAgentsObservationSchemaElement B, const int32 EncodingSize, const bool bActive, const FName Tag)
 {
-	return SpecifyExclusiveUnionObservationFromArrayViews(Schema, { TEXT("A"), TEXT("B") }, { A, B }, EncodingSize, Tag);
+	if (!bActive)
+	{
+		return SpecifyNullObservation(Schema, Tag);
+	}
+
+	return SpecifyExclusiveUnionObservationFromArrayViews(Schema, { TEXT("A"), TEXT("B") }, { A, B }, EncodingSize, bActive, Tag);
 }
 
-FLearningAgentsObservationSchemaElement ULearningAgentsObservations::SpecifyEncodingObservation(ULearningAgentsObservationSchema* Schema, const FLearningAgentsObservationSchemaElement Element, const int32 EncodingSize, const int32 HiddenLayerNum, const ELearningAgentsActivationFunction ActivationFunction, const FName Tag)
+FLearningAgentsObservationSchemaElement ULearningAgentsObservations::SpecifyEncodingObservation(ULearningAgentsObservationSchema* Schema, const FLearningAgentsObservationSchemaElement Element, const int32 EncodingSize, const int32 HiddenLayerNum, const ELearningAgentsActivationFunction ActivationFunction, const bool bActive, const FName Tag)
 {
+	if (!bActive)
+	{
+		return SpecifyNullObservation(Schema, Tag);
+	}
+
 	if (!Schema)
 	{
 		UE_LOG(LogLearning, Error, TEXT("SpecifyEncodingObservation: Schema is nullptr."));
@@ -1246,13 +1381,23 @@ FLearningAgentsObservationSchemaElement ULearningAgentsObservations::SpecifyEnco
 	return { Schema->ObservationSchema.CreateEncoding({ Element.SchemaElement, EncodingSize, HiddenLayerNum, UE::Learning::Agents::Observation::Private::GetEncodingActivationFunction(ActivationFunction) }, Tag)};
 }
 
-FLearningAgentsObservationSchemaElement ULearningAgentsObservations::SpecifyBoolObservation(ULearningAgentsObservationSchema* Schema, const FName Tag)
+FLearningAgentsObservationSchemaElement ULearningAgentsObservations::SpecifyBoolObservation(ULearningAgentsObservationSchema* Schema, const bool bActive, const FName Tag)
 {
-	return SpecifyContinuousObservation(Schema, 1, 1.0f, Tag);
+	if (!bActive)
+	{
+		return SpecifyNullObservation(Schema, Tag);
+	}
+
+	return SpecifyContinuousObservation(Schema, 1, 1.0f, bActive, Tag);
 }
 
-FLearningAgentsObservationSchemaElement ULearningAgentsObservations::SpecifyConv1dObservation(ULearningAgentsObservationSchema* Schema, const FLearningAgentsObservationSchemaElement Element, const FConv1dObservationParams& Conv1dParams, const FName Tag)
+FLearningAgentsObservationSchemaElement ULearningAgentsObservations::SpecifyConv1dObservation(ULearningAgentsObservationSchema* Schema, const FLearningAgentsObservationSchemaElement Element, const FConv1dObservationParams& Conv1dParams, const bool bActive, const FName Tag)
 {
+	if (!bActive)
+	{
+		return SpecifyNullObservation(Schema, Tag);
+	}
+
 	if (!Schema)
 	{
 		UE_LOG(LogLearning, Error, TEXT("SpecifyConv1dObservation: Schema is nullptr."));
@@ -1308,8 +1453,13 @@ FLearningAgentsObservationSchemaElement ULearningAgentsObservations::SpecifyConv
 	Tag) };
 }
 
-FLearningAgentsObservationSchemaElement ULearningAgentsObservations::SpecifyConv2dObservation(ULearningAgentsObservationSchema* Schema, const FLearningAgentsObservationSchemaElement Element, const FConv2dObservationParams& Conv2dParams, const FName Tag)
+FLearningAgentsObservationSchemaElement ULearningAgentsObservations::SpecifyConv2dObservation(ULearningAgentsObservationSchema* Schema, const FLearningAgentsObservationSchemaElement Element, const FConv2dObservationParams& Conv2dParams, const bool bActive, const FName Tag)
 {
+	if (!bActive)
+	{
+		return SpecifyNullObservation(Schema, Tag);
+	}
+
 	if (!Schema)
 	{
 		UE_LOG(LogLearning, Error, TEXT("SpecifyConv2dObservation: Schema is nullptr."));
@@ -1372,28 +1522,53 @@ FLearningAgentsObservationSchemaElement ULearningAgentsObservations::SpecifyConv
 	Tag) };
 }
 
-FLearningAgentsObservationSchemaElement ULearningAgentsObservations::SpecifyFloatObservation(ULearningAgentsObservationSchema* Schema, const float FloatScale, const FName Tag)
+FLearningAgentsObservationSchemaElement ULearningAgentsObservations::SpecifyFloatObservation(ULearningAgentsObservationSchema* Schema, const float FloatScale, const bool bActive, const FName Tag)
 {
-	return SpecifyContinuousObservation(Schema, 1, FloatScale, Tag);
+	if (!bActive)
+	{
+		return SpecifyNullObservation(Schema, Tag);
+	}
+
+	return SpecifyContinuousObservation(Schema, 1, FloatScale, bActive, Tag);
 }
 
-FLearningAgentsObservationSchemaElement ULearningAgentsObservations::SpecifyLocationObservation(ULearningAgentsObservationSchema* Schema, const float LocationScale, const FName Tag)
+FLearningAgentsObservationSchemaElement ULearningAgentsObservations::SpecifyLocationObservation(ULearningAgentsObservationSchema* Schema, const float LocationScale, const bool bActive, const FName Tag)
 {
-	return SpecifyContinuousObservation(Schema, 3, LocationScale, Tag);
+	if (!bActive)
+	{
+		return SpecifyNullObservation(Schema, Tag);
+	}
+
+	return SpecifyContinuousObservation(Schema, 3, LocationScale, bActive, Tag);
 }
 
-FLearningAgentsObservationSchemaElement ULearningAgentsObservations::SpecifyRotationObservation(ULearningAgentsObservationSchema* Schema, const FName Tag)
+FLearningAgentsObservationSchemaElement ULearningAgentsObservations::SpecifyRotationObservation(ULearningAgentsObservationSchema* Schema, const bool bActive, const FName Tag)
 {
-	return SpecifyContinuousObservation(Schema, 6, 1.0f, Tag);
+	if (!bActive)
+	{
+		return SpecifyNullObservation(Schema, Tag);
+	}
+
+	return SpecifyContinuousObservation(Schema, 6, 1.0f, bActive, Tag);
 }
 
-FLearningAgentsObservationSchemaElement ULearningAgentsObservations::SpecifyScaleObservation(ULearningAgentsObservationSchema* Schema, const FName Tag)
+FLearningAgentsObservationSchemaElement ULearningAgentsObservations::SpecifyScaleObservation(ULearningAgentsObservationSchema* Schema, const bool bActive, const FName Tag)
 {
-	return SpecifyContinuousObservation(Schema, 3, 1.0f, Tag);
+	if (!bActive)
+	{
+		return SpecifyNullObservation(Schema, Tag);
+	}
+
+	return SpecifyContinuousObservation(Schema, 3, 1.0f, bActive, Tag);
 }
 
-FLearningAgentsObservationSchemaElement ULearningAgentsObservations::SpecifyTransformObservation(ULearningAgentsObservationSchema* Schema, const float LocationScale, const FName Tag)
+FLearningAgentsObservationSchemaElement ULearningAgentsObservations::SpecifyTransformObservation(ULearningAgentsObservationSchema* Schema, const float LocationScale, const bool bActive, const FName Tag)
 {
+	if (!bActive)
+	{
+		return SpecifyNullObservation(Schema, Tag);
+	}
+
 	return SpecifyStructObservationFromArrayViews(Schema,
 		{
 			TEXT("Location"),
@@ -1404,32 +1579,57 @@ FLearningAgentsObservationSchemaElement ULearningAgentsObservations::SpecifyTran
 			SpecifyLocationObservation(Schema, LocationScale),
 			SpecifyRotationObservation(Schema),
 			SpecifyScaleObservation(Schema)
-		}, 
-		Tag);
+		},
+		bActive, Tag);
 }
 
-FLearningAgentsObservationSchemaElement ULearningAgentsObservations::SpecifyAngleObservation(ULearningAgentsObservationSchema* Schema, const FName Tag)
+FLearningAgentsObservationSchemaElement ULearningAgentsObservations::SpecifyAngleObservation(ULearningAgentsObservationSchema* Schema, const bool bActive, const FName Tag)
 {
-	return SpecifyContinuousObservation(Schema, 2, 1.0f, Tag);
+	if (!bActive)
+	{
+		return SpecifyNullObservation(Schema, Tag);
+	}
+
+	return SpecifyContinuousObservation(Schema, 2, 1.0f, bActive, Tag);
 }
 
-FLearningAgentsObservationSchemaElement ULearningAgentsObservations::SpecifyVelocityObservation(ULearningAgentsObservationSchema* Schema, const float VelocityScale, const FName Tag)
+FLearningAgentsObservationSchemaElement ULearningAgentsObservations::SpecifyVelocityObservation(ULearningAgentsObservationSchema* Schema, const float VelocityScale, const bool bActive, const FName Tag)
 {
-	return SpecifyContinuousObservation(Schema, 3, VelocityScale, Tag);
+	if (!bActive)
+	{
+		return SpecifyNullObservation(Schema, Tag);
+	}
+
+	return SpecifyContinuousObservation(Schema, 3, VelocityScale, bActive, Tag);
 }
 
-FLearningAgentsObservationSchemaElement ULearningAgentsObservations::SpecifyDirectionObservation(ULearningAgentsObservationSchema* Schema, const FName Tag)
+FLearningAgentsObservationSchemaElement ULearningAgentsObservations::SpecifyDirectionObservation(ULearningAgentsObservationSchema* Schema, const bool bActive, const FName Tag)
 {
-	return SpecifyContinuousObservation(Schema, 3, 1.0f, Tag);
+	if (!bActive)
+	{
+		return SpecifyNullObservation(Schema, Tag);
+	}
+
+	return SpecifyContinuousObservation(Schema, 3, 1.0f, bActive, Tag);
 }
 
-FLearningAgentsObservationSchemaElement ULearningAgentsObservations::SpecifyLocationAlongSplineObservation(ULearningAgentsObservationSchema* Schema, const float LocationScale, const FName Tag)
+FLearningAgentsObservationSchemaElement ULearningAgentsObservations::SpecifyLocationAlongSplineObservation(ULearningAgentsObservationSchema* Schema, const float LocationScale, const bool bActive, const FName Tag)
 {
-	return SpecifyLocationObservation(Schema, LocationScale, Tag);
+	if (!bActive)
+	{
+		return SpecifyNullObservation(Schema, Tag);
+	}
+
+	return SpecifyLocationObservation(Schema, LocationScale, bActive, Tag);
 }
 
-FLearningAgentsObservationSchemaElement ULearningAgentsObservations::SpecifyProportionAlongSplineObservation(ULearningAgentsObservationSchema* Schema, const FName Tag)
+FLearningAgentsObservationSchemaElement ULearningAgentsObservations::SpecifyProportionAlongSplineObservation(ULearningAgentsObservationSchema* Schema, const bool bActive, const FName Tag)
 {
+	if (!bActive)
+	{
+		return SpecifyNullObservation(Schema, Tag);
+	}
+
 	return SpecifyExclusiveUnionObservationFromArrayViews(Schema,
 		{
 			TEXT("Angle"),
@@ -1440,22 +1640,37 @@ FLearningAgentsObservationSchemaElement ULearningAgentsObservations::SpecifyProp
 			SpecifyFloatObservation(Schema)
 		},
 		8,
-		Tag);
+		bActive, Tag);
 }
 
-FLearningAgentsObservationSchemaElement ULearningAgentsObservations::SpecifyDirectionAlongSplineObservation(ULearningAgentsObservationSchema* Schema, const FName Tag)
+FLearningAgentsObservationSchemaElement ULearningAgentsObservations::SpecifyDirectionAlongSplineObservation(ULearningAgentsObservationSchema* Schema, const bool bActive, const FName Tag)
 {
-	return SpecifyDirectionObservation(Schema, Tag);
+	if (!bActive)
+	{
+		return SpecifyNullObservation(Schema, Tag);
+	}
+
+	return SpecifyDirectionObservation(Schema, bActive, Tag);
 }
 
-FLearningAgentsObservationSchemaElement ULearningAgentsObservations::SpecifyProportionAlongRayObservation(ULearningAgentsObservationSchema* Schema, const FName Tag)
+FLearningAgentsObservationSchemaElement ULearningAgentsObservations::SpecifyProportionAlongRayObservation(ULearningAgentsObservationSchema* Schema, const bool bActive, const FName Tag)
 {
-	return SpecifyFloatObservation(Schema, 1.0f, Tag);
+	if (!bActive)
+	{
+		return SpecifyNullObservation(Schema, Tag);
+	}
+
+	return SpecifyFloatObservation(Schema, 1.0f, bActive, Tag);
 }
 
-FLearningAgentsObservationSchemaElement ULearningAgentsObservations::SpecifyLidarObservation(ULearningAgentsObservationSchema* Schema, const FName Tag)
+FLearningAgentsObservationSchemaElement ULearningAgentsObservations::SpecifyLidarObservation(ULearningAgentsObservationSchema* Schema, const bool bActive, const FName Tag)
 {
-	return SpecifyFloatObservation(Schema, 1.0f, Tag);
+	if (!bActive)
+	{
+		return SpecifyNullObservation(Schema, Tag);
+	}
+
+	return SpecifyFloatObservation(Schema, 1.0f, bActive, Tag);
 }
 
 void ULearningAgentsObservations::LogObservation(const ULearningAgentsObservationObject* Object, const FLearningAgentsObservationObjectElement Element)
@@ -1482,7 +1697,8 @@ FLearningAgentsObservationObjectElement ULearningAgentsObservations::MakeNullObs
 
 FLearningAgentsObservationObjectElement ULearningAgentsObservations::MakeContinuousObservation(
 	ULearningAgentsObservationObject* Object, 
-	const TArray<float>& Values, 
+	const TArray<float>& Values,
+	const bool bActive,
 	const FName Tag,
 	const bool bVisualLoggerEnabled,
 	ULearningAgentsManagerListener* VisualLoggerListener,
@@ -1490,12 +1706,18 @@ FLearningAgentsObservationObjectElement ULearningAgentsObservations::MakeContinu
 	const FVector VisualLoggerLocation,
 	const FLinearColor VisualLoggerColor)
 {
-	return MakeContinuousObservationFromArrayView(Object, Values, Tag, bVisualLoggerEnabled, VisualLoggerListener, VisualLoggerAgentId, VisualLoggerLocation, VisualLoggerColor);
+	if (!bActive)
+	{
+		return MakeNullObservation(Object, Tag);
+	}
+
+	return MakeContinuousObservationFromArrayView(Object, Values, bActive, Tag, bVisualLoggerEnabled, VisualLoggerListener, VisualLoggerAgentId, VisualLoggerLocation, VisualLoggerColor);
 }
 
 FLearningAgentsObservationObjectElement ULearningAgentsObservations::MakeContinuousObservationFromArrayView(
 	ULearningAgentsObservationObject* Object, 
 	const TArrayView<const float> Values, 
+	const bool bActive,
 	const FName Tag,
 	const bool bVisualLoggerEnabled,
 	ULearningAgentsManagerListener* VisualLoggerListener,
@@ -1503,6 +1725,11 @@ FLearningAgentsObservationObjectElement ULearningAgentsObservations::MakeContinu
 	const FVector VisualLoggerLocation,
 	const FLinearColor VisualLoggerColor)
 {
+	if (!bActive)
+	{
+		return MakeNullObservation(Object, Tag);
+	}
+
 	if (!Object)
 	{
 		UE_LOG(LogLearning, Error, TEXT("MakeContinuousObservationFromArrayView: Object is nullptr."));
@@ -1534,7 +1761,8 @@ FLearningAgentsObservationObjectElement ULearningAgentsObservations::MakeContinu
 
 FLearningAgentsObservationObjectElement ULearningAgentsObservations::MakeExclusiveDiscreteObservation(
 	ULearningAgentsObservationObject* Object, 
-	const int32 DiscreteIndex, 
+	const int32 DiscreteIndex,
+	const bool bActive,
 	const FName Tag,
 	const bool bVisualLoggerEnabled,
 	ULearningAgentsManagerListener* VisualLoggerListener,
@@ -1542,6 +1770,11 @@ FLearningAgentsObservationObjectElement ULearningAgentsObservations::MakeExclusi
 	const FVector VisualLoggerLocation,
 	const FLinearColor VisualLoggerColor)
 {
+	if (!bActive)
+	{
+		return MakeNullObservation(Object, Tag);
+	}
+
 	if (!Object)
 	{
 		UE_LOG(LogLearning, Error, TEXT("MakeExclusiveDiscreteObservation: Object is nullptr."));
@@ -1569,6 +1802,7 @@ FLearningAgentsObservationObjectElement ULearningAgentsObservations::MakeExclusi
 FLearningAgentsObservationObjectElement ULearningAgentsObservations::MakeInclusiveDiscreteObservation(
 	ULearningAgentsObservationObject* Object, 
 	const TArray<int32>& DiscreteIndices, 
+	const bool bActive,
 	const FName Tag,
 	const bool bVisualLoggerEnabled,
 	ULearningAgentsManagerListener* VisualLoggerListener,
@@ -1576,12 +1810,18 @@ FLearningAgentsObservationObjectElement ULearningAgentsObservations::MakeInclusi
 	const FVector VisualLoggerLocation,
 	const FLinearColor VisualLoggerColor)
 {
-	return MakeInclusiveDiscreteObservationFromArrayView(Object, DiscreteIndices, Tag, bVisualLoggerEnabled, VisualLoggerListener, VisualLoggerAgentId, VisualLoggerLocation, VisualLoggerColor);
+	if (!bActive)
+	{
+		return MakeNullObservation(Object, Tag);
+	}
+
+	return MakeInclusiveDiscreteObservationFromArrayView(Object, DiscreteIndices, bActive, Tag, bVisualLoggerEnabled, VisualLoggerListener, VisualLoggerAgentId, VisualLoggerLocation, VisualLoggerColor);
 }
 
 FLearningAgentsObservationObjectElement ULearningAgentsObservations::MakeInclusiveDiscreteObservationFromArrayView(
 	ULearningAgentsObservationObject* Object, 
 	const TArrayView<const int32> DiscreteIndices, 
+	const bool bActive,
 	const FName Tag,
 	const bool bVisualLoggerEnabled,
 	ULearningAgentsManagerListener* VisualLoggerListener,
@@ -1589,6 +1829,11 @@ FLearningAgentsObservationObjectElement ULearningAgentsObservations::MakeInclusi
 	const FVector VisualLoggerLocation,
 	const FLinearColor VisualLoggerColor)
 {
+	if (!bActive)
+	{
+		return MakeNullObservation(Object, Tag);
+	}
+
 	if (!Object)
 	{
 		UE_LOG(LogLearning, Error, TEXT("MakeExclusiveDiscreteObservation: Object is nullptr."));
@@ -1622,6 +1867,7 @@ FLearningAgentsObservationObjectElement ULearningAgentsObservations::MakeInclusi
 FLearningAgentsObservationObjectElement ULearningAgentsObservations::MakeNamedExclusiveDiscreteObservation(
 	ULearningAgentsObservationObject* Object,
 	const FName ElementName,
+	const bool bActive,
 	const FName Tag,
 	const bool bVisualLoggerEnabled,
 	ULearningAgentsManagerListener* VisualLoggerListener,
@@ -1629,6 +1875,11 @@ FLearningAgentsObservationObjectElement ULearningAgentsObservations::MakeNamedEx
 	const FVector VisualLoggerLocation,
 	const FLinearColor VisualLoggerColor)
 {
+	if (!bActive)
+	{
+		return MakeNullObservation(Object, Tag);
+	}
+
 	if (!Object)
 	{
 		UE_LOG(LogLearning, Error, TEXT("MakeNamedExclusiveDiscreteObservation: Object is nullptr."));
@@ -1656,6 +1907,7 @@ FLearningAgentsObservationObjectElement ULearningAgentsObservations::MakeNamedEx
 FLearningAgentsObservationObjectElement ULearningAgentsObservations::MakeNamedInclusiveDiscreteObservation(
 	ULearningAgentsObservationObject* Object,
 	const TArray<FName>& ElementNames,
+	const bool bActive,
 	const FName Tag,
 	const bool bVisualLoggerEnabled,
 	ULearningAgentsManagerListener* VisualLoggerListener,
@@ -1663,12 +1915,18 @@ FLearningAgentsObservationObjectElement ULearningAgentsObservations::MakeNamedIn
 	const FVector VisualLoggerLocation,
 	const FLinearColor VisualLoggerColor)
 {
-	return MakeNamedInclusiveDiscreteObservationFromArrayView(Object, ElementNames, Tag, bVisualLoggerEnabled, VisualLoggerListener, VisualLoggerAgentId, VisualLoggerLocation, VisualLoggerColor);
+	if (!bActive)
+	{
+		return MakeNullObservation(Object, Tag);
+	}
+
+	return MakeNamedInclusiveDiscreteObservationFromArrayView(Object, ElementNames, bActive, Tag, bVisualLoggerEnabled, VisualLoggerListener, VisualLoggerAgentId, VisualLoggerLocation, VisualLoggerColor);
 }
 
 FLearningAgentsObservationObjectElement ULearningAgentsObservations::MakeNamedInclusiveDiscreteObservationFromArrayView(
 	ULearningAgentsObservationObject* Object,
 	const TArrayView<const FName> ElementNames,
+	const bool bActive,
 	const FName Tag,
 	const bool bVisualLoggerEnabled,
 	ULearningAgentsManagerListener* VisualLoggerListener,
@@ -1676,6 +1934,11 @@ FLearningAgentsObservationObjectElement ULearningAgentsObservations::MakeNamedIn
 	const FVector VisualLoggerLocation,
 	const FLinearColor VisualLoggerColor)
 {
+	if (!bActive)
+	{
+		return MakeNullObservation(Object, Tag);
+	}
+
 	if (!Object)
 	{
 		UE_LOG(LogLearning, Error, TEXT("MakeNamedInclusiveDiscreteObservationFromArrayView: Object is nullptr."));
@@ -1710,6 +1973,7 @@ FLearningAgentsObservationObjectElement ULearningAgentsObservations::MakeCountOb
 	ULearningAgentsObservationObject* Object, 
 	const int32 Num, 
 	const int32 MaxNum, 
+	const bool bActive,
 	const FName Tag,
 	const bool bVisualLoggerEnabled,
 	ULearningAgentsManagerListener* VisualLoggerListener,
@@ -1717,6 +1981,11 @@ FLearningAgentsObservationObjectElement ULearningAgentsObservations::MakeCountOb
 	const FVector VisualLoggerLocation,
 	const FLinearColor VisualLoggerColor)
 {
+	if (!bActive)
+	{
+		return MakeNullObservation(Object, Tag);
+	}
+
 	if (MaxNum == 0)
 	{
 		UE_LOG(LogLearning, Error, TEXT("MakeCountObservation: MaxNum must not be zero for Count Observation."));
@@ -1742,11 +2011,16 @@ FLearningAgentsObservationObjectElement ULearningAgentsObservations::MakeCountOb
 	}
 #endif
 
-	return MakeContinuousObservationFromArrayView(Object, { Encoded }, Tag);
+	return MakeContinuousObservationFromArrayView(Object, { Encoded }, bActive, Tag);
 }
 
-FLearningAgentsObservationObjectElement ULearningAgentsObservations::MakeStructObservation(ULearningAgentsObservationObject* Object, const TMap<FName, FLearningAgentsObservationObjectElement>& Elements, const FName Tag)
+FLearningAgentsObservationObjectElement ULearningAgentsObservations::MakeStructObservation(ULearningAgentsObservationObject* Object, const TMap<FName, FLearningAgentsObservationObjectElement>& Elements, const bool bActive, const FName Tag)
 {
+	if (!bActive)
+	{
+		return MakeNullObservation(Object, Tag);
+	}
+
 	if (Elements.Num() == 0)
 	{
 		UE_LOG(LogLearning, Warning, TEXT("MakeStructObservation: Creating zero-sized Struct Observation."));
@@ -1765,16 +2039,26 @@ FLearningAgentsObservationObjectElement ULearningAgentsObservations::MakeStructO
 		SubElements.Add(Element.Value);
 	}
 
-	return MakeStructObservationFromArrayViews(Object, SubElementNames, SubElements, Tag);
+	return MakeStructObservationFromArrayViews(Object, SubElementNames, SubElements, bActive, Tag);
 }
 
-FLearningAgentsObservationObjectElement ULearningAgentsObservations::MakeStructObservationFromArrays(ULearningAgentsObservationObject* Object, const TArray<FName>& ElementNames, const TArray<FLearningAgentsObservationObjectElement>& Elements, const FName Tag)
+FLearningAgentsObservationObjectElement ULearningAgentsObservations::MakeStructObservationFromArrays(ULearningAgentsObservationObject* Object, const TArray<FName>& ElementNames, const TArray<FLearningAgentsObservationObjectElement>& Elements, const bool bActive, const FName Tag)
 {
-	return MakeStructObservationFromArrayViews(Object, ElementNames, Elements, Tag);
+	if (!bActive)
+	{
+		return MakeNullObservation(Object, Tag);
+	}
+
+	return MakeStructObservationFromArrayViews(Object, ElementNames, Elements, bActive, Tag);
 }
 
-FLearningAgentsObservationObjectElement ULearningAgentsObservations::MakeStructObservationFromArrayViews(ULearningAgentsObservationObject* Object, const TArrayView<const FName> ElementNames, const TArrayView<const FLearningAgentsObservationObjectElement> Elements, const FName Tag)
+FLearningAgentsObservationObjectElement ULearningAgentsObservations::MakeStructObservationFromArrayViews(ULearningAgentsObservationObject* Object, const TArrayView<const FName> ElementNames, const TArrayView<const FLearningAgentsObservationObjectElement> Elements, const bool bActive, const FName Tag)
 {
+	if (!bActive)
+	{
+		return MakeNullObservation(Object, Tag);
+	}
+
 	if (!Object)
 	{
 		UE_LOG(LogLearning, Error, TEXT("MakeStructObservationFromArrayViews: Object is nullptr."));
@@ -1815,8 +2099,13 @@ FLearningAgentsObservationObjectElement ULearningAgentsObservations::MakeStructO
 	return { Object->ObservationObject.CreateAnd({ ElementNames, SubElements }, Tag)};
 }
 
-FLearningAgentsObservationObjectElement ULearningAgentsObservations::MakeExclusiveUnionObservation(ULearningAgentsObservationObject* Object, const FName ElementName, const FLearningAgentsObservationObjectElement Element, const FName Tag)
+FLearningAgentsObservationObjectElement ULearningAgentsObservations::MakeExclusiveUnionObservation(ULearningAgentsObservationObject* Object, const FName ElementName, const FLearningAgentsObservationObjectElement Element, const bool bActive, const FName Tag)
 {
+	if (!bActive)
+	{
+		return MakeNullObservation(Object, Tag);
+	}
+
 	if (!Object)
 	{
 		UE_LOG(LogLearning, Error, TEXT("MakeExclusiveUnionObservation: Object is nullptr."));
@@ -1832,8 +2121,13 @@ FLearningAgentsObservationObjectElement ULearningAgentsObservations::MakeExclusi
 	return { Object->ObservationObject.CreateOrExclusive({ ElementName, Element.ObjectElement }, Tag) };
 }
 
-FLearningAgentsObservationObjectElement ULearningAgentsObservations::MakeInclusiveUnionObservation(ULearningAgentsObservationObject* Object, const TMap<FName, FLearningAgentsObservationObjectElement>& Elements, const FName Tag)
+FLearningAgentsObservationObjectElement ULearningAgentsObservations::MakeInclusiveUnionObservation(ULearningAgentsObservationObject* Object, const TMap<FName, FLearningAgentsObservationObjectElement>& Elements, const bool bActive, const FName Tag)
 {
+	if (!bActive)
+	{
+		return MakeNullObservation(Object, Tag);
+	}
+
 	const int32 SubElementNum = Elements.Num();
 
 	TArray<FName, TInlineAllocator<16>> SubElementNames;
@@ -1847,16 +2141,26 @@ FLearningAgentsObservationObjectElement ULearningAgentsObservations::MakeInclusi
 		SubElements.Add(Element.Value);
 	}
 
-	return MakeInclusiveUnionObservationFromArrayViews(Object, SubElementNames, SubElements, Tag);
+	return MakeInclusiveUnionObservationFromArrayViews(Object, SubElementNames, SubElements, bActive, Tag);
 }
 
-FLearningAgentsObservationObjectElement ULearningAgentsObservations::MakeInclusiveUnionObservationFromArrays(ULearningAgentsObservationObject* Object, const TArray<FName>& ElementNames, const TArray<FLearningAgentsObservationObjectElement>& Elements, const FName Tag)
+FLearningAgentsObservationObjectElement ULearningAgentsObservations::MakeInclusiveUnionObservationFromArrays(ULearningAgentsObservationObject* Object, const TArray<FName>& ElementNames, const TArray<FLearningAgentsObservationObjectElement>& Elements, const bool bActive, const FName Tag)
 {
-	return MakeInclusiveUnionObservationFromArrayViews(Object, ElementNames, Elements, Tag);
+	if (!bActive)
+	{
+		return MakeNullObservation(Object, Tag);
+	}
+
+	return MakeInclusiveUnionObservationFromArrayViews(Object, ElementNames, Elements, bActive, Tag);
 }
 
-FLearningAgentsObservationObjectElement ULearningAgentsObservations::MakeInclusiveUnionObservationFromArrayViews(ULearningAgentsObservationObject* Object, const TArrayView<const FName> ElementNames, const TArrayView<const FLearningAgentsObservationObjectElement> Elements, const FName Tag)
+FLearningAgentsObservationObjectElement ULearningAgentsObservations::MakeInclusiveUnionObservationFromArrayViews(ULearningAgentsObservationObject* Object, const TArrayView<const FName> ElementNames, const TArrayView<const FLearningAgentsObservationObjectElement> Elements, const bool bActive, const FName Tag)
 {
+	if (!bActive)
+	{
+		return MakeNullObservation(Object, Tag);
+	}
+
 	if (!Object)
 	{
 		UE_LOG(LogLearning, Error, TEXT("MakeInclusiveUnionObservationFromArrayViews: Object is nullptr."));
@@ -1892,13 +2196,23 @@ FLearningAgentsObservationObjectElement ULearningAgentsObservations::MakeInclusi
 	return { Object->ObservationObject.CreateOrInclusive({ ElementNames, SubElements }, Tag) };
 }
 
-FLearningAgentsObservationObjectElement ULearningAgentsObservations::MakeStaticArrayObservation(ULearningAgentsObservationObject* Object, const TArray<FLearningAgentsObservationObjectElement>& Elements, const FName Tag)
+FLearningAgentsObservationObjectElement ULearningAgentsObservations::MakeStaticArrayObservation(ULearningAgentsObservationObject* Object, const TArray<FLearningAgentsObservationObjectElement>& Elements, const bool bActive, const FName Tag)
 {
-	return MakeStaticArrayObservationFromArrayView(Object, Elements, Tag);
+	if (!bActive)
+	{
+		return MakeNullObservation(Object, Tag);
+	}
+
+	return MakeStaticArrayObservationFromArrayView(Object, Elements, bActive, Tag);
 }
 
-FLearningAgentsObservationObjectElement ULearningAgentsObservations::MakeStaticArrayObservationFromArrayView(ULearningAgentsObservationObject* Object, const TArrayView<const FLearningAgentsObservationObjectElement> Elements, const FName Tag)
+FLearningAgentsObservationObjectElement ULearningAgentsObservations::MakeStaticArrayObservationFromArrayView(ULearningAgentsObservationObject* Object, const TArrayView<const FLearningAgentsObservationObjectElement> Elements, const bool bActive, const FName Tag)
 {
+	if (!bActive)
+	{
+		return MakeNullObservation(Object, Tag);
+	}
+
 	if (!Object)
 	{
 		UE_LOG(LogLearning, Error, TEXT("MakeStaticArrayObservationFromArrayView: Object is nullptr."));
@@ -1927,8 +2241,13 @@ FLearningAgentsObservationObjectElement ULearningAgentsObservations::MakeStaticA
 	return { Object->ObservationObject.CreateArray({ SubElements }, Tag) };
 }
 
-FLearningAgentsObservationObjectElement ULearningAgentsObservations::MakeSetObservation(ULearningAgentsObservationObject* Object, const TSet<FLearningAgentsObservationObjectElement>& Elements, const FName Tag)
+FLearningAgentsObservationObjectElement ULearningAgentsObservations::MakeSetObservation(ULearningAgentsObservationObject* Object, const TSet<FLearningAgentsObservationObjectElement>& Elements, const bool bActive, const FName Tag)
 {
+	if (!bActive)
+	{
+		return MakeNullObservation(Object, Tag);
+	}
+
 	if (!Object)
 	{
 		UE_LOG(LogLearning, Error, TEXT("MakeSetObservation: Object is nullptr."));
@@ -1952,13 +2271,23 @@ FLearningAgentsObservationObjectElement ULearningAgentsObservations::MakeSetObse
 	return { Object->ObservationObject.CreateSet({ SubElements }, Tag) };
 }
 
-FLearningAgentsObservationObjectElement ULearningAgentsObservations::MakeSetObservationFromArray(ULearningAgentsObservationObject* Object, const TArray<FLearningAgentsObservationObjectElement>& Elements, const FName Tag)
+FLearningAgentsObservationObjectElement ULearningAgentsObservations::MakeSetObservationFromArray(ULearningAgentsObservationObject* Object, const TArray<FLearningAgentsObservationObjectElement>& Elements, const bool bActive, const FName Tag)
 {
-	return MakeSetObservationFromArrayView(Object, Elements, Tag);
+	if (!bActive)
+	{
+		return MakeNullObservation(Object, Tag);
+	}
+
+	return MakeSetObservationFromArrayView(Object, Elements, bActive, Tag);
 }
 
-FLearningAgentsObservationObjectElement ULearningAgentsObservations::MakeSetObservationFromArrayView(ULearningAgentsObservationObject* Object, const TArrayView<const FLearningAgentsObservationObjectElement> Elements, const FName Tag)
+FLearningAgentsObservationObjectElement ULearningAgentsObservations::MakeSetObservationFromArrayView(ULearningAgentsObservationObject* Object, const TArrayView<const FLearningAgentsObservationObjectElement> Elements, const bool bActive, const FName Tag)
 {
+	if (!bActive)
+	{
+		return MakeNullObservation(Object, Tag);
+	}
+
 	if (!Object)
 	{
 		UE_LOG(LogLearning, Error, TEXT("MakeSetObservationFromArrayView: Object is nullptr."));
@@ -1982,18 +2311,33 @@ FLearningAgentsObservationObjectElement ULearningAgentsObservations::MakeSetObse
 	return { Object->ObservationObject.CreateSet({ SubElements }, Tag) };
 }
 
-FLearningAgentsObservationObjectElement ULearningAgentsObservations::MakePairObservation(ULearningAgentsObservationObject* Object, const FLearningAgentsObservationObjectElement Key, const FLearningAgentsObservationObjectElement Value, const FName Tag)
+FLearningAgentsObservationObjectElement ULearningAgentsObservations::MakePairObservation(ULearningAgentsObservationObject* Object, const FLearningAgentsObservationObjectElement Key, const FLearningAgentsObservationObjectElement Value, const bool bActive, const FName Tag)
 {
-	return MakeStructObservationFromArrayViews(Object, { TEXT("Key"), TEXT("Value") }, { Key, Value }, Tag);
+	if (!bActive)
+	{
+		return MakeNullObservation(Object, Tag);
+	}
+
+	return MakeStructObservationFromArrayViews(Object, { TEXT("Key"), TEXT("Value") }, { Key, Value }, bActive, Tag);
 }
 
-FLearningAgentsObservationObjectElement ULearningAgentsObservations::MakeArrayObservation(ULearningAgentsObservationObject* Object, const TArray<FLearningAgentsObservationObjectElement>& Elements, const int32 MaxNum, const FName Tag)
+FLearningAgentsObservationObjectElement ULearningAgentsObservations::MakeArrayObservation(ULearningAgentsObservationObject* Object, const TArray<FLearningAgentsObservationObjectElement>& Elements, const int32 MaxNum, const bool bActive, const FName Tag)
 {
-	return MakeArrayObservationFromArrayView(Object, Elements, MaxNum, Tag);
+	if (!bActive)
+	{
+		return MakeNullObservation(Object, Tag);
+	}
+
+	return MakeArrayObservationFromArrayView(Object, Elements, MaxNum, bActive, Tag);
 }
 
-FLearningAgentsObservationObjectElement ULearningAgentsObservations::MakeArrayObservationFromArrayView(ULearningAgentsObservationObject* Object, const TArrayView<const FLearningAgentsObservationObjectElement> Elements, const int32 MaxNum, const FName Tag)
+FLearningAgentsObservationObjectElement ULearningAgentsObservations::MakeArrayObservationFromArrayView(ULearningAgentsObservationObject* Object, const TArrayView<const FLearningAgentsObservationObjectElement> Elements, const int32 MaxNum, const bool bActive, const FName Tag)
 {
+	if (!bActive)
+	{
+		return MakeNullObservation(Object, Tag);
+	}
+
 	if (!Object)
 	{
 		UE_LOG(LogLearning, Error, TEXT("MakeArrayObservationFromArrayView: Object is nullptr."));
@@ -2016,11 +2360,16 @@ FLearningAgentsObservationObjectElement ULearningAgentsObservations::MakeArrayOb
 		SubElements.Add(MakePairObservation(Object, MakeCountObservation(Object, ElementIdx, MaxNum), Element));
 	}
 
-	return MakeSetObservationFromArrayView(Object, SubElements, Tag);
+	return MakeSetObservationFromArrayView(Object, SubElements, bActive, Tag);
 }
 
-FLearningAgentsObservationObjectElement ULearningAgentsObservations::MakeMapObservation(ULearningAgentsObservationObject* Object, const TMap<FLearningAgentsObservationObjectElement, FLearningAgentsObservationObjectElement>& Map, const FName Tag)
+FLearningAgentsObservationObjectElement ULearningAgentsObservations::MakeMapObservation(ULearningAgentsObservationObject* Object, const TMap<FLearningAgentsObservationObjectElement, FLearningAgentsObservationObjectElement>& Map, const bool bActive, const FName Tag)
 {
+	if (!bActive)
+	{
+		return MakeNullObservation(Object, Tag);
+	}
+
 	if (!Object)
 	{
 		UE_LOG(LogLearning, Error, TEXT("MakeMapObservation: Object is nullptr."));
@@ -2041,16 +2390,26 @@ FLearningAgentsObservationObjectElement ULearningAgentsObservations::MakeMapObse
 		SubElements.Add(MakePairObservation(Object, Item.Key, Item.Value));
 	}
 
-	return MakeSetObservationFromArrayView(Object, SubElements, Tag);
+	return MakeSetObservationFromArrayView(Object, SubElements, bActive, Tag);
 }
 
-FLearningAgentsObservationObjectElement ULearningAgentsObservations::MakeMapObservationFromArrays(ULearningAgentsObservationObject* Object, const TArray<FLearningAgentsObservationObjectElement>& Keys, const TArray<FLearningAgentsObservationObjectElement>& Values, const FName Tag)
+FLearningAgentsObservationObjectElement ULearningAgentsObservations::MakeMapObservationFromArrays(ULearningAgentsObservationObject* Object, const TArray<FLearningAgentsObservationObjectElement>& Keys, const TArray<FLearningAgentsObservationObjectElement>& Values, const bool bActive, const FName Tag)
 {
-	return MakeMapObservationFromArrayViews(Object, Keys, Values, Tag);
+	if (!bActive)
+	{
+		return MakeNullObservation(Object, Tag);
+	}
+
+	return MakeMapObservationFromArrayViews(Object, Keys, Values, bActive, Tag);
 }
 
-FLearningAgentsObservationObjectElement ULearningAgentsObservations::MakeMapObservationFromArrayViews(ULearningAgentsObservationObject* Object, const TArrayView<const FLearningAgentsObservationObjectElement> Keys, const TArrayView<const FLearningAgentsObservationObjectElement> Values, const FName Tag)
+FLearningAgentsObservationObjectElement ULearningAgentsObservations::MakeMapObservationFromArrayViews(ULearningAgentsObservationObject* Object, const TArrayView<const FLearningAgentsObservationObjectElement> Keys, const TArrayView<const FLearningAgentsObservationObjectElement> Values, const bool bActive, const FName Tag)
 {
+	if (!bActive)
+	{
+		return MakeNullObservation(Object, Tag);
+	}
+
 	if (!Object)
 	{
 		UE_LOG(LogLearning, Error, TEXT("MakeMapObservationFromArrayViews: Object is nullptr."));
@@ -2078,13 +2437,14 @@ FLearningAgentsObservationObjectElement ULearningAgentsObservations::MakeMapObse
 		SubElements.Add(MakePairObservation(Object, Keys[ElementIdx], Values[ElementIdx]));
 	}
 
-	return MakeSetObservationFromArrayView(Object, SubElements, Tag);
+	return MakeSetObservationFromArrayView(Object, SubElements, bActive, Tag);
 }
 
 FLearningAgentsObservationObjectElement ULearningAgentsObservations::MakeEnumObservation(
 	ULearningAgentsObservationObject* Object, 
 	const UEnum* Enum, 
 	const uint8 EnumValue, 
+	const bool bActive,
 	const FName Tag,
 	const bool bVisualLoggerEnabled,
 	ULearningAgentsManagerListener* VisualLoggerListener,
@@ -2092,6 +2452,11 @@ FLearningAgentsObservationObjectElement ULearningAgentsObservations::MakeEnumObs
 	const FVector VisualLoggerLocation,
 	const FLinearColor VisualLoggerColor)
 {
+	if (!bActive)
+	{
+		return MakeNullObservation(Object, Tag);
+	}
+
 	if (!Enum)
 	{
 		UE_LOG(LogLearning, Error, TEXT("MakeEnumObservation: Enum is nullptr."));
@@ -2129,13 +2494,14 @@ FLearningAgentsObservationObjectElement ULearningAgentsObservations::MakeEnumObs
 	}
 #endif
 
-	return MakeContinuousObservationFromArrayView(Object, OneHot, Tag);
+	return MakeContinuousObservationFromArrayView(Object, OneHot, bActive, Tag);
 }
 
 FLearningAgentsObservationObjectElement ULearningAgentsObservations::MakeBitmaskObservation(
 	ULearningAgentsObservationObject* Object, 
 	const UEnum* Enum, 
 	const int32 BitmaskValue, 
+	const bool bActive,
 	const FName Tag,
 	const bool bVisualLoggerEnabled,
 	ULearningAgentsManagerListener* VisualLoggerListener,
@@ -2143,6 +2509,11 @@ FLearningAgentsObservationObjectElement ULearningAgentsObservations::MakeBitmask
 	const FVector VisualLoggerLocation,
 	const FLinearColor VisualLoggerColor)
 {
+	if (!bActive)
+	{
+		return MakeNullObservation(Object, Tag);
+	}
+
 	if (!Enum)
 	{
 		UE_LOG(LogLearning, Error, TEXT("MakeBitmaskObservation: Enum is nullptr."));
@@ -2200,45 +2571,80 @@ FLearningAgentsObservationObjectElement ULearningAgentsObservations::MakeBitmask
 	}
 #endif
 
-	return MakeContinuousObservationFromArrayView(Object, OneHot, Tag);
+	return MakeContinuousObservationFromArrayView(Object, OneHot, bActive, Tag);
 }
 
-FLearningAgentsObservationObjectElement ULearningAgentsObservations::MakeOptionalObservation(ULearningAgentsObservationObject* Object, const FLearningAgentsObservationObjectElement Element, const ELearningAgentsOptionalObservation Option, const FName Tag)
+FLearningAgentsObservationObjectElement ULearningAgentsObservations::MakeOptionalObservation(ULearningAgentsObservationObject* Object, const FLearningAgentsObservationObjectElement Element, const ELearningAgentsOptionalObservation Option, const bool bActive, const FName Tag)
 {
+	if (!bActive)
+	{
+		return MakeNullObservation(Object, Tag);
+	}
+
 	return MakeExclusiveUnionObservation(
 		Object,
 		Option == ELearningAgentsOptionalObservation::Null ? TEXT("Null") : TEXT("Valid"),
 		Option == ELearningAgentsOptionalObservation::Null ? MakeNullObservation(Object) : Element,
-		Tag);
+		bActive, Tag);
 }
 
-FLearningAgentsObservationObjectElement ULearningAgentsObservations::MakeOptionalNullObservation(ULearningAgentsObservationObject* Object, const FName Tag)
+FLearningAgentsObservationObjectElement ULearningAgentsObservations::MakeOptionalNullObservation(ULearningAgentsObservationObject* Object, const bool bActive, const FName Tag)
 {
-	return MakeExclusiveUnionObservation(Object, TEXT("Null"), MakeNullObservation(Object), Tag);
+	if (!bActive)
+	{
+		return MakeNullObservation(Object, Tag);
+	}
+
+	return MakeExclusiveUnionObservation(Object, TEXT("Null"), MakeNullObservation(Object), bActive, Tag);
 }
 
-FLearningAgentsObservationObjectElement ULearningAgentsObservations::MakeOptionalValidObservation(ULearningAgentsObservationObject* Object, const FLearningAgentsObservationObjectElement Element, const FName Tag)
+FLearningAgentsObservationObjectElement ULearningAgentsObservations::MakeOptionalValidObservation(ULearningAgentsObservationObject* Object, const FLearningAgentsObservationObjectElement Element, const bool bActive, const FName Tag)
 {
-	return MakeExclusiveUnionObservation(Object, TEXT("Valid"), Element, Tag);
+	if (!bActive)
+	{
+		return MakeNullObservation(Object, Tag);
+	}
+
+	return MakeExclusiveUnionObservation(Object, TEXT("Valid"), Element, bActive, Tag);
 }
 
-FLearningAgentsObservationObjectElement ULearningAgentsObservations::MakeEitherObservation(ULearningAgentsObservationObject* Object, const FLearningAgentsObservationObjectElement Element, const ELearningAgentsEitherObservation Either, const FName Tag)
+FLearningAgentsObservationObjectElement ULearningAgentsObservations::MakeEitherObservation(ULearningAgentsObservationObject* Object, const FLearningAgentsObservationObjectElement Element, const ELearningAgentsEitherObservation Either, const bool bActive, const FName Tag)
 {
-	return MakeExclusiveUnionObservation(Object, Either == ELearningAgentsEitherObservation::A ? TEXT("A") : TEXT("B"), Element, Tag);
+	if (!bActive)
+	{
+		return MakeNullObservation(Object, Tag);
+	}
+
+	return MakeExclusiveUnionObservation(Object, Either == ELearningAgentsEitherObservation::A ? TEXT("A") : TEXT("B"), Element, bActive, Tag);
 }
 
-FLearningAgentsObservationObjectElement ULearningAgentsObservations::MakeEitherAObservation(ULearningAgentsObservationObject* Object, const FLearningAgentsObservationObjectElement A, const FName Tag)
+FLearningAgentsObservationObjectElement ULearningAgentsObservations::MakeEitherAObservation(ULearningAgentsObservationObject* Object, const FLearningAgentsObservationObjectElement A, const bool bActive, const FName Tag)
 {
-	return MakeExclusiveUnionObservation(Object, TEXT("A"), A, Tag);
+	if (!bActive)
+	{
+		return MakeNullObservation(Object, Tag);
+	}
+
+	return MakeExclusiveUnionObservation(Object, TEXT("A"), A, bActive, Tag);
 }
 
-FLearningAgentsObservationObjectElement ULearningAgentsObservations::MakeEitherBObservation(ULearningAgentsObservationObject* Object, const FLearningAgentsObservationObjectElement B, const FName Tag)
+FLearningAgentsObservationObjectElement ULearningAgentsObservations::MakeEitherBObservation(ULearningAgentsObservationObject* Object, const FLearningAgentsObservationObjectElement B, const bool bActive, const FName Tag)
 {
-	return MakeExclusiveUnionObservation(Object, TEXT("B"), B, Tag);
+	if (!bActive)
+	{
+		return MakeNullObservation(Object, Tag);
+	}
+
+	return MakeExclusiveUnionObservation(Object, TEXT("B"), B, bActive, Tag);
 }
 
-FLearningAgentsObservationObjectElement ULearningAgentsObservations::MakeEncodingObservation(ULearningAgentsObservationObject* Object, const FLearningAgentsObservationObjectElement Element, const FName Tag)
+FLearningAgentsObservationObjectElement ULearningAgentsObservations::MakeEncodingObservation(ULearningAgentsObservationObject* Object, const FLearningAgentsObservationObjectElement Element, const bool bActive, const FName Tag)
 {
+	if (!bActive)
+	{
+		return MakeNullObservation(Object, Tag);
+	}
+
 	if (!Object)
 	{
 		UE_LOG(LogLearning, Error, TEXT("MakeEncodingObservation: Object is nullptr."));
@@ -2257,6 +2663,7 @@ FLearningAgentsObservationObjectElement ULearningAgentsObservations::MakeEncodin
 FLearningAgentsObservationObjectElement ULearningAgentsObservations::MakeBoolObservation(
 	ULearningAgentsObservationObject* Object,
 	const bool bValue,
+	const bool bActive,
 	const FName Tag,
 	const bool bVisualLoggerEnabled,
 	ULearningAgentsManagerListener* VisualLoggerListener,
@@ -2264,6 +2671,11 @@ FLearningAgentsObservationObjectElement ULearningAgentsObservations::MakeBoolObs
 	const FVector VisualLoggerLocation,
 	const FLinearColor VisualLoggerColor)
 {
+	if (!bActive)
+	{
+		return MakeNullObservation(Object, Tag);
+	}
+
 	const float EncodedBool = bValue ? 1.0f : -1.0f;
 
 #if UE_LEARNING_AGENTS_ENABLE_VISUAL_LOG
@@ -2282,11 +2694,16 @@ FLearningAgentsObservationObjectElement ULearningAgentsObservations::MakeBoolObs
 	}
 #endif
 
-	return MakeContinuousObservationFromArrayView(Object, { EncodedBool }, Tag);
+	return MakeContinuousObservationFromArrayView(Object, { EncodedBool }, bActive, Tag);
 }
 
-FLearningAgentsObservationObjectElement ULearningAgentsObservations::MakeConv1dObservation(ULearningAgentsObservationObject* Object, const FLearningAgentsObservationObjectElement Element, const FName Tag)
+FLearningAgentsObservationObjectElement ULearningAgentsObservations::MakeConv1dObservation(ULearningAgentsObservationObject* Object, const FLearningAgentsObservationObjectElement Element, const bool bActive, const FName Tag)
 {
+	if (!bActive)
+	{
+		return MakeNullObservation(Object, Tag);
+	}
+
 	if (!Object)
 	{
 		UE_LOG(LogLearning, Error, TEXT("MakeConv1dObservation: Object is nullptr."));
@@ -2302,8 +2719,13 @@ FLearningAgentsObservationObjectElement ULearningAgentsObservations::MakeConv1dO
 	return { Object->ObservationObject.CreateConv1d({ Element.ObjectElement }, Tag) };
 }
 
-FLearningAgentsObservationObjectElement ULearningAgentsObservations::MakeConv2dObservation(ULearningAgentsObservationObject* Object, const FLearningAgentsObservationObjectElement Element, const FName Tag)
+FLearningAgentsObservationObjectElement ULearningAgentsObservations::MakeConv2dObservation(ULearningAgentsObservationObject* Object, const FLearningAgentsObservationObjectElement Element, const bool bActive, const FName Tag)
 {
+	if (!bActive)
+	{
+		return MakeNullObservation(Object, Tag);
+	}
+
 	if (!Object)
 	{
 		UE_LOG(LogLearning, Error, TEXT("MakeConv2dObservation: Object is nullptr."));
@@ -2322,6 +2744,7 @@ FLearningAgentsObservationObjectElement ULearningAgentsObservations::MakeConv2dO
 FLearningAgentsObservationObjectElement ULearningAgentsObservations::MakeFloatObservation(
 	ULearningAgentsObservationObject* Object,
 	const float Value,
+	const bool bActive,
 	const FName Tag,
 	const bool bVisualLoggerEnabled,
 	ULearningAgentsManagerListener* VisualLoggerListener,
@@ -2329,6 +2752,11 @@ FLearningAgentsObservationObjectElement ULearningAgentsObservations::MakeFloatOb
 	const FVector VisualLoggerLocation,
 	const FLinearColor VisualLoggerColor)
 {
+	if (!bActive)
+	{
+		return MakeNullObservation(Object, Tag);
+	}
+
 #if UE_LEARNING_AGENTS_ENABLE_VISUAL_LOG
 	if (bVisualLoggerEnabled && VisualLoggerListener)
 	{
@@ -2344,13 +2772,14 @@ FLearningAgentsObservationObjectElement ULearningAgentsObservations::MakeFloatOb
 	}
 #endif
 
-	return MakeContinuousObservationFromArrayView(Object, { Value }, Tag);
+	return MakeContinuousObservationFromArrayView(Object, { Value }, bActive, Tag);
 }
 
 FLearningAgentsObservationObjectElement ULearningAgentsObservations::MakeLocationObservation(
 	ULearningAgentsObservationObject* Object,
 	const FVector Location,
 	const FTransform RelativeTransform,
+	const bool bActive,
 	const FName Tag,
 	const bool bVisualLoggerEnabled,
 	ULearningAgentsManagerListener* VisualLoggerListener,
@@ -2358,6 +2787,11 @@ FLearningAgentsObservationObjectElement ULearningAgentsObservations::MakeLocatio
 	const FVector VisualLoggerLocation,
 	const FLinearColor VisualLoggerColor)
 {
+	if (!bActive)
+	{
+		return MakeNullObservation(Object, Tag);
+	}
+
 	const FVector LocalLocation = RelativeTransform.InverseTransformPosition(Location);
 
 #if UE_LEARNING_AGENTS_ENABLE_VISUAL_LOG
@@ -2398,13 +2832,14 @@ FLearningAgentsObservationObjectElement ULearningAgentsObservations::MakeLocatio
 		(float)LocalLocation.X,
 		(float)LocalLocation.Y,
 		(float)LocalLocation.Z,
-		}, Tag);
+		}, bActive, Tag);
 }
 
 FLearningAgentsObservationObjectElement ULearningAgentsObservations::MakeRotationObservation(
 	ULearningAgentsObservationObject* Object,
 	const FRotator Rotation,
 	const FRotator RelativeRotation,
+	const bool bActive,
 	const FName Tag,
 	const bool bVisualLoggerEnabled,
 	ULearningAgentsManagerListener* VisualLoggerListener,
@@ -2413,11 +2848,16 @@ FLearningAgentsObservationObjectElement ULearningAgentsObservations::MakeRotatio
 	const FVector VisualLoggerLocation,
 	const FLinearColor VisualLoggerColor)
 {
+	if (!bActive)
+	{
+		return MakeNullObservation(Object, Tag);
+	}
+
 	// Visual logging is handled in the MakeRotationObservationFromQuat
 
 	return MakeRotationObservationFromQuat(
 		Object,
-		FQuat::MakeFromRotator(Rotation), FQuat::MakeFromRotator(RelativeRotation), Tag,
+		FQuat::MakeFromRotator(Rotation), FQuat::MakeFromRotator(RelativeRotation), bActive, Tag,
 		bVisualLoggerEnabled, VisualLoggerListener, VisualLoggerAgentId, VisualLoggerRotationLocation, VisualLoggerLocation, VisualLoggerColor);
 }
 
@@ -2425,6 +2865,7 @@ FLearningAgentsObservationObjectElement ULearningAgentsObservations::MakeRotatio
 	ULearningAgentsObservationObject* Object,
 	const FQuat Rotation,
 	const FQuat RelativeRotation,
+	const bool bActive,
 	const FName Tag,
 	const bool bVisualLoggerEnabled,
 	ULearningAgentsManagerListener* VisualLoggerListener,
@@ -2433,6 +2874,11 @@ FLearningAgentsObservationObjectElement ULearningAgentsObservations::MakeRotatio
 	const FVector VisualLoggerLocation,
 	const FLinearColor VisualLoggerColor)
 {
+	if (!bActive)
+	{
+		return MakeNullObservation(Object, Tag);
+	}
+
 	const FQuat LocalRotation = RelativeRotation.Inverse() * Rotation;
 	const FVector LocalAxisForward = LocalRotation.GetForwardVector();
 	const FVector LocalAxisRight = LocalRotation.GetRightVector();
@@ -2468,13 +2914,14 @@ FLearningAgentsObservationObjectElement ULearningAgentsObservations::MakeRotatio
 		(float)LocalAxisRight.X,
 		(float)LocalAxisRight.Y,
 		(float)LocalAxisRight.Z,
-		}, Tag);
+		}, bActive, Tag);
 }
 
 FLearningAgentsObservationObjectElement ULearningAgentsObservations::MakeScaleObservation(
 	ULearningAgentsObservationObject* Object,
 	const FVector Scale,
 	const FVector RelativeScale,
+	const bool bActive,
 	const FName Tag,
 	const bool bVisualLoggerEnabled,
 	ULearningAgentsManagerListener* VisualLoggerListener,
@@ -2483,6 +2930,11 @@ FLearningAgentsObservationObjectElement ULearningAgentsObservations::MakeScaleOb
 	const FVector VisualLoggerLocation,
 	const FLinearColor VisualLoggerColor)
 {
+	if (!bActive)
+	{
+		return MakeNullObservation(Object, Tag);
+	}
+
 	const FVector LocalLogScale = 
 		UE::Learning::Agents::Observation::Private::VectorLogSafe(Scale) - 
 		UE::Learning::Agents::Observation::Private::VectorLogSafe(RelativeScale);
@@ -2507,13 +2959,14 @@ FLearningAgentsObservationObjectElement ULearningAgentsObservations::MakeScaleOb
 		(float)LocalLogScale.X,
 		(float)LocalLogScale.Y,
 		(float)LocalLogScale.Z,
-		}, Tag);
+		}, bActive, Tag);
 }
 
 FLearningAgentsObservationObjectElement ULearningAgentsObservations::MakeTransformObservation(
 	ULearningAgentsObservationObject* Object,
 	const FTransform Transform,
 	const FTransform RelativeTransform,
+	const bool bActive,
 	const FName Tag,
 	const bool bVisualLoggerEnabled,
 	ULearningAgentsManagerListener* VisualLoggerListener,
@@ -2521,6 +2974,11 @@ FLearningAgentsObservationObjectElement ULearningAgentsObservations::MakeTransfo
 	const FVector VisualLoggerLocation,
 	const FLinearColor VisualLoggerColor)
 {
+	if (!bActive)
+	{
+		return MakeNullObservation(Object, Tag);
+	}
+
 	const FTransform LocalTransform = Transform * RelativeTransform.Inverse();
 	
 #if UE_LEARNING_AGENTS_ENABLE_VISUAL_LOG
@@ -2568,13 +3026,14 @@ FLearningAgentsObservationObjectElement ULearningAgentsObservations::MakeTransfo
 			MakeRotationObservationFromQuat(Object, LocalTransform.GetRotation(), FQuat::Identity),
 			MakeScaleObservation(Object, LocalTransform.GetScale3D(), FVector::OneVector)
 		},
-		Tag);
+		bActive, Tag);
 }
 
 FLearningAgentsObservationObjectElement ULearningAgentsObservations::MakeAngleObservation(
 	ULearningAgentsObservationObject* Object, 
 	const float Angle, 
 	const float RelativeAngle, 
+	const bool bActive,
 	const FName Tag,
 	const bool bVisualLoggerEnabled,
 	ULearningAgentsManagerListener* VisualLoggerListener,
@@ -2582,6 +3041,11 @@ FLearningAgentsObservationObjectElement ULearningAgentsObservations::MakeAngleOb
 	const FVector VisualLoggerLocation,
 	const FLinearColor VisualLoggerColor)
 {
+	if (!bActive)
+	{
+		return MakeNullObservation(Object, Tag);
+	}
+
 	const float LocalAngle = FMath::FindDeltaAngleDegrees(RelativeAngle, Angle);
 	const float EncodedX = FMath::Sin(FMath::DegreesToRadians(LocalAngle));
 	const float EncodedY = FMath::Cos(FMath::DegreesToRadians(LocalAngle));
@@ -2612,13 +3076,14 @@ FLearningAgentsObservationObjectElement ULearningAgentsObservations::MakeAngleOb
 	}
 #endif
 
-	return MakeContinuousObservationFromArrayView(Object, { EncodedX, EncodedY }, Tag);
+	return MakeContinuousObservationFromArrayView(Object, { EncodedX, EncodedY }, bActive, Tag);
 }
 
 FLearningAgentsObservationObjectElement ULearningAgentsObservations::MakeAngleObservationRadians(
 	ULearningAgentsObservationObject* Object, 
 	const float Angle, 
 	const float RelativeAngle, 
+	const bool bActive,
 	const FName Tag,
 	const bool bVisualLoggerEnabled,
 	ULearningAgentsManagerListener* VisualLoggerListener,
@@ -2626,10 +3091,16 @@ FLearningAgentsObservationObjectElement ULearningAgentsObservations::MakeAngleOb
 	const FVector VisualLoggerLocation,
 	const FLinearColor VisualLoggerColor)
 {
+	if (!bActive)
+	{
+		return MakeNullObservation(Object, Tag);
+	}
+
 	return MakeAngleObservation(
 		Object, 
 		FMath::RadiansToDegrees(Angle), 
 		FMath::RadiansToDegrees(RelativeAngle), 
+		bActive,
 		Tag, 
 		bVisualLoggerEnabled, 
 		VisualLoggerListener, 
@@ -2642,6 +3113,7 @@ FLearningAgentsObservationObjectElement ULearningAgentsObservations::MakeVelocit
 	ULearningAgentsObservationObject* Object, 
 	const FVector Velocity, 
 	const FTransform RelativeTransform, 
+	const bool bActive,
 	const FName Tag,
 	const bool bVisualLoggerEnabled,
 	ULearningAgentsManagerListener* VisualLoggerListener,
@@ -2650,6 +3122,11 @@ FLearningAgentsObservationObjectElement ULearningAgentsObservations::MakeVelocit
 	const FVector VisualLoggerLocation,
 	const FLinearColor VisualLoggerColor)
 {
+	if (!bActive)
+	{
+		return MakeNullObservation(Object, Tag);
+	}
+
 	const FVector LocalVelocity = RelativeTransform.InverseTransformVectorNoScale(Velocity);
 
 #if UE_LEARNING_AGENTS_ENABLE_VISUAL_LOG
@@ -2680,13 +3157,14 @@ FLearningAgentsObservationObjectElement ULearningAgentsObservations::MakeVelocit
 	}
 #endif
 
-	return MakeContinuousObservationFromArrayView(Object, { (float)LocalVelocity.X, (float)LocalVelocity.Y, (float)LocalVelocity.Z }, Tag);
+	return MakeContinuousObservationFromArrayView(Object, { (float)LocalVelocity.X, (float)LocalVelocity.Y, (float)LocalVelocity.Z }, bActive, Tag);
 }
 
 FLearningAgentsObservationObjectElement ULearningAgentsObservations::MakeDirectionObservation(
 	ULearningAgentsObservationObject* Object,
 	const FVector Direction,
 	const FTransform RelativeTransform,
+	const bool bActive,
 	const FName Tag,
 	const bool bVisualLoggerEnabled,
 	ULearningAgentsManagerListener* VisualLoggerListener,
@@ -2696,6 +3174,11 @@ FLearningAgentsObservationObjectElement ULearningAgentsObservations::MakeDirecti
 	const float VisualLoggerArrowLength,
 	const FLinearColor VisualLoggerColor)
 {
+	if (!bActive)
+	{
+		return MakeNullObservation(Object, Tag);
+	}
+
 	const FVector LocalDirection = RelativeTransform.InverseTransformVectorNoScale(Direction).GetSafeNormal(UE_SMALL_NUMBER, FVector::ForwardVector);
 
 #if UE_LEARNING_AGENTS_ENABLE_VISUAL_LOG
@@ -2730,7 +3213,7 @@ FLearningAgentsObservationObjectElement ULearningAgentsObservations::MakeDirecti
 		(float)LocalDirection.X,
 		(float)LocalDirection.Y,
 		(float)LocalDirection.Z,
-		}, Tag);
+		}, bActive, Tag);
 }
 
 FLearningAgentsObservationObjectElement ULearningAgentsObservations::MakeLocationAlongSplineObservation(
@@ -2738,6 +3221,7 @@ FLearningAgentsObservationObjectElement ULearningAgentsObservations::MakeLocatio
 	const USplineComponent* SplineComponent,
 	const float DistanceAlongSpline,
 	const FTransform RelativeTransform,
+	const bool bActive,
 	const FName Tag,
 	const bool bVisualLoggerEnabled,
 	ULearningAgentsManagerListener* VisualLoggerListener,
@@ -2745,6 +3229,11 @@ FLearningAgentsObservationObjectElement ULearningAgentsObservations::MakeLocatio
 	const FVector VisualLoggerLocation,
 	const FLinearColor VisualLoggerColor)
 {
+	if (!bActive)
+	{
+		return MakeNullObservation(Object, Tag);
+	}
+
 	if (!SplineComponent)
 	{
 		UE_LOG(LogLearning, Error, TEXT("MakeLocationAlongSplineObservation: SplineComponent was nullptr."));
@@ -2758,6 +3247,7 @@ FLearningAgentsObservationObjectElement ULearningAgentsObservations::MakeLocatio
 		Object, 
 		SplineComponent->GetLocationAtDistanceAlongSpline(LoopedDistance, ESplineCoordinateSpace::World), 
 		RelativeTransform, 
+		bActive,
 		Tag, 
 		bVisualLoggerEnabled, 
 		VisualLoggerListener, 
@@ -2770,6 +3260,7 @@ FLearningAgentsObservationObjectElement ULearningAgentsObservations::MakeProport
 	ULearningAgentsObservationObject* Object, 
 	const USplineComponent* SplineComponent, 
 	const float DistanceAlongSpline, 
+	const bool bActive,
 	const FName Tag,
 	const bool bVisualLoggerEnabled,
 	ULearningAgentsManagerListener* VisualLoggerListener,
@@ -2777,6 +3268,11 @@ FLearningAgentsObservationObjectElement ULearningAgentsObservations::MakeProport
 	const FVector VisualLoggerLocation,
 	const FLinearColor VisualLoggerColor)
 {
+	if (!bActive)
+	{
+		return MakeNullObservation(Object, Tag);
+	}
+
 	if (!SplineComponent)
 	{
 		UE_LOG(LogLearning, Error, TEXT("MakeProportionAlongSplineObservation: SplineComponent was nullptr."));
@@ -2794,12 +3290,13 @@ FLearningAgentsObservationObjectElement ULearningAgentsObservations::MakeProport
 			Object, 
 			Angle,
 			0.0f,
+			bActive,
 			TEXT("AngleObservation"),
 			bVisualLoggerEnabled,
 			VisualLoggerListener,
 			VisualLoggerAgentId,
 			VisualLoggerLocation,
-			VisualLoggerColor), Tag);
+			VisualLoggerColor), bActive, Tag);
 	}
 	else
 	{
@@ -2808,12 +3305,13 @@ FLearningAgentsObservationObjectElement ULearningAgentsObservations::MakeProport
 		return MakeExclusiveUnionObservation(Object, TEXT("Proportion"), MakeFloatObservation(
 			Object, 
 			Proportion, 
+			bActive,
 			TEXT("FloatObservation"),
 			bVisualLoggerEnabled, 
 			VisualLoggerListener, 
 			VisualLoggerAgentId, 
 			VisualLoggerLocation, 
-			VisualLoggerColor), Tag);
+			VisualLoggerColor), bActive, Tag);
 	}
 }
 
@@ -2822,6 +3320,7 @@ FLearningAgentsObservationObjectElement ULearningAgentsObservations::MakeDirecti
 	const USplineComponent* SplineComponent,
 	const float DistanceAlongSpline,
 	const FTransform RelativeTransform,
+	const bool bActive,
 	const FName Tag,
 	const bool bVisualLoggerEnabled,
 	ULearningAgentsManagerListener* VisualLoggerListener,
@@ -2830,6 +3329,11 @@ FLearningAgentsObservationObjectElement ULearningAgentsObservations::MakeDirecti
 	const float VisualLoggerArrowLength,
 	const FLinearColor VisualLoggerColor)
 {
+	if (!bActive)
+	{
+		return MakeNullObservation(Object, Tag);
+	}
+
 	if (!SplineComponent)
 	{
 		UE_LOG(LogLearning, Error, TEXT("MakeDirectionAlongSplineObservation: SplineComponent was nullptr."));
@@ -2852,6 +3356,7 @@ FLearningAgentsObservationObjectElement ULearningAgentsObservations::MakeDirecti
 		Object, 
 		SplineComponent->GetDirectionAtDistanceAlongSpline(LoopedDistance, ESplineCoordinateSpace::World), 
 		RelativeTransform, 
+		bActive,
 		Tag, 
 		bVisualLoggerEnabled, 
 		VisualLoggerListener, 
@@ -2867,7 +3372,8 @@ FLearningAgentsObservationObjectElement ULearningAgentsObservations::MakeProport
 	const FVector RayStart, 
 	const FVector RayEnd, 
 	const FTransform RayTransform, 
-	const ECollisionChannel CollisionChannel, 
+	const ECollisionChannel CollisionChannel,
+	const bool bActive,
 	const FName Tag,
 	const bool bVisualLoggerEnabled,
 	ULearningAgentsManagerListener* VisualLoggerListener,
@@ -2875,6 +3381,11 @@ FLearningAgentsObservationObjectElement ULearningAgentsObservations::MakeProport
 	const FVector VisualLoggerLocation,
 	const FLinearColor VisualLoggerColor)
 {
+	if (!bActive)
+	{
+		return MakeNullObservation(Object, Tag);
+	}
+
 	const FVector RayStartWorld = RayTransform.TransformPosition(RayStart);
 	const FVector RayEndWorld = RayTransform.TransformPosition(RayEnd);
 	const float RayDistance = FVector::Distance(RayStartWorld, RayEndWorld);
@@ -2929,7 +3440,7 @@ FLearningAgentsObservationObjectElement ULearningAgentsObservations::MakeProport
 	}
 #endif
 
-	return MakeFloatObservation(Object, Encoded, Tag);
+	return MakeFloatObservation(Object, Encoded, bActive, Tag);
 }
 
 FLearningAgentsObservationObjectElement ULearningAgentsObservations::MakeLidarObservation(
@@ -2939,6 +3450,7 @@ FLearningAgentsObservationObjectElement ULearningAgentsObservations::MakeLidarOb
 	const FTransform RayTransform, 
 	const ECollisionChannel CollisionChannel, 
 	const AActor* ActorToIgnore,
+	const bool bActive,
 	const FName Tag,
 	const bool bVisualLoggerEnabled,
 	ULearningAgentsManagerListener* VisualLoggerListener,
@@ -2946,6 +3458,11 @@ FLearningAgentsObservationObjectElement ULearningAgentsObservations::MakeLidarOb
 	const FVector VisualLoggerLocation,
 	const FLinearColor VisualLoggerColor)
 {
+	if (!bActive)
+	{
+		return MakeNullObservation(Object, Tag);
+	}
+
 	const FVector RayStartWorld = RayTransform.TransformPosition(RayStart);
 	const FVector RayEndWorld = RayTransform.TransformPosition(RayEnd);
 	const float RayRange = FVector::Dist(RayEndWorld, RayStartWorld);
@@ -3009,7 +3526,7 @@ FLearningAgentsObservationObjectElement ULearningAgentsObservations::MakeLidarOb
 	}
 #endif
 
-	return MakeFloatObservation(Object, Encoded, Tag);
+	return MakeFloatObservation(Object, Encoded, bActive, Tag);
 }
 
 bool ULearningAgentsObservations::GetNullObservation(const ULearningAgentsObservationObject* Object, const FLearningAgentsObservationObjectElement Element, const FName Tag)
@@ -3043,8 +3560,13 @@ bool ULearningAgentsObservations::GetNullObservation(const ULearningAgentsObserv
 	return true;
 }
 
-bool ULearningAgentsObservations::GetContinuousObservationNum(int32& OutNum, const ULearningAgentsObservationObject* Object, const FLearningAgentsObservationObjectElement Element, const FName Tag)
+bool ULearningAgentsObservations::GetContinuousObservationNum(int32& OutNum, const ULearningAgentsObservationObject* Object, const FLearningAgentsObservationObjectElement Element, const bool bActive, const FName Tag)
 {
+	if (!bActive)
+	{
+		return GetNullObservation(Object, Element, Tag);
+	}
+
 	if (!Object)
 	{
 		UE_LOG(LogLearning, Error, TEXT("GetContinuousObservationNum: Object is nullptr."));
@@ -3078,10 +3600,15 @@ bool ULearningAgentsObservations::GetContinuousObservationNum(int32& OutNum, con
 	return true;
 }
 
-bool ULearningAgentsObservations::GetContinuousObservation(TArray<float>& OutValues, const ULearningAgentsObservationObject* Object, const FLearningAgentsObservationObjectElement Element, const FName Tag)
+bool ULearningAgentsObservations::GetContinuousObservation(TArray<float>& OutValues, const ULearningAgentsObservationObject* Object, const FLearningAgentsObservationObjectElement Element, const bool bActive, const FName Tag)
 {
+	if (!bActive)
+	{
+		return GetNullObservation(Object, Element, Tag);
+	}
+
 	int32 OutValueNum = 0;
-	if (!GetContinuousObservationNum(OutValueNum, Object, Element, Tag))
+	if (!GetContinuousObservationNum(OutValueNum, Object, Element, bActive, Tag))
 	{
 		OutValues.Empty();
 		return false;
@@ -3089,7 +3616,7 @@ bool ULearningAgentsObservations::GetContinuousObservation(TArray<float>& OutVal
 
 	OutValues.SetNumUninitialized(OutValueNum);
 
-	if (!GetContinuousObservationToArrayView(OutValues, Object, Element, Tag))
+	if (!GetContinuousObservationToArrayView(OutValues, Object, Element, bActive, Tag))
 	{
 		OutValues.Empty();
 		return false;
@@ -3098,8 +3625,13 @@ bool ULearningAgentsObservations::GetContinuousObservation(TArray<float>& OutVal
 	return true;
 }
 
-bool ULearningAgentsObservations::GetContinuousObservationToArrayView(TArrayView<float> OutValues, const ULearningAgentsObservationObject* Object, const FLearningAgentsObservationObjectElement Element, const FName Tag)
+bool ULearningAgentsObservations::GetContinuousObservationToArrayView(TArrayView<float> OutValues, const ULearningAgentsObservationObject* Object, const FLearningAgentsObservationObjectElement Element, const bool bActive, const FName Tag)
 {
+	if (!bActive)
+	{
+		return GetNullObservation(Object, Element, Tag);
+	}
+
 	if (!Object)
 	{
 		UE_LOG(LogLearning, Error, TEXT("GetContinuousObservationToArrayView: Object is nullptr."));
@@ -3144,8 +3676,13 @@ bool ULearningAgentsObservations::GetContinuousObservationToArrayView(TArrayView
 	return true;
 }
 
-bool ULearningAgentsObservations::GetExclusiveDiscreteObservation(int32& OutIndex, const ULearningAgentsObservationObject* Object, const FLearningAgentsObservationObjectElement Element, const FName Tag)
+bool ULearningAgentsObservations::GetExclusiveDiscreteObservation(int32& OutIndex, const ULearningAgentsObservationObject* Object, const FLearningAgentsObservationObjectElement Element, const bool bActive, const FName Tag)
 {
+	if (!bActive)
+	{
+		return GetNullObservation(Object, Element, Tag);
+	}
+
 	if (!Object)
 	{
 		UE_LOG(LogLearning, Error, TEXT("GetExclusiveDiscreteObservation: Object is nullptr."));
@@ -3179,8 +3716,13 @@ bool ULearningAgentsObservations::GetExclusiveDiscreteObservation(int32& OutInde
 	return true;
 }
 
-bool ULearningAgentsObservations::GetInclusiveDiscreteObservationNum(int32& OutNum, const ULearningAgentsObservationObject* Object, const FLearningAgentsObservationObjectElement Element, const FName Tag)
+bool ULearningAgentsObservations::GetInclusiveDiscreteObservationNum(int32& OutNum, const ULearningAgentsObservationObject* Object, const FLearningAgentsObservationObjectElement Element, const bool bActive, const FName Tag)
 {
+	if (!bActive)
+	{
+		return GetNullObservation(Object, Element, Tag);
+	}
+
 	if (!Object)
 	{
 		UE_LOG(LogLearning, Error, TEXT("GetInclusiveDiscreteObservationNum: Object is nullptr."));
@@ -3214,21 +3756,31 @@ bool ULearningAgentsObservations::GetInclusiveDiscreteObservationNum(int32& OutN
 	return true;
 }
 
-bool ULearningAgentsObservations::GetInclusiveDiscreteObservation(TArray<int32>& OutIndices, const ULearningAgentsObservationObject* Object, const FLearningAgentsObservationObjectElement Element, const FName Tag)
+bool ULearningAgentsObservations::GetInclusiveDiscreteObservation(TArray<int32>& OutIndices, const ULearningAgentsObservationObject* Object, const FLearningAgentsObservationObjectElement Element, const bool bActive, const FName Tag)
 {
+	if (!bActive)
+	{
+		return GetNullObservation(Object, Element, Tag);
+	}
+
 	int32 OutIndicesNum;
-	if (!GetInclusiveDiscreteObservationNum(OutIndicesNum, Object, Element, Tag))
+	if (!GetInclusiveDiscreteObservationNum(OutIndicesNum, Object, Element, bActive, Tag))
 	{
 		OutIndices.Empty();
 		return false;
 	}
 
 	OutIndices.SetNumUninitialized(OutIndicesNum);
-	return GetInclusiveDiscreteObservationToArrayView(OutIndices, Object, Element, Tag);
+	return GetInclusiveDiscreteObservationToArrayView(OutIndices, Object, Element, bActive, Tag);
 }
 
-bool ULearningAgentsObservations::GetInclusiveDiscreteObservationToArrayView(TArrayView<int32> OutIndices, const ULearningAgentsObservationObject* Object, const FLearningAgentsObservationObjectElement Element, const FName Tag)
+bool ULearningAgentsObservations::GetInclusiveDiscreteObservationToArrayView(TArrayView<int32> OutIndices, const ULearningAgentsObservationObject* Object, const FLearningAgentsObservationObjectElement Element, const bool bActive, const FName Tag)
 {
+	if (!bActive)
+	{
+		return GetNullObservation(Object, Element, Tag);
+	}
+
 	if (!Object)
 	{
 		UE_LOG(LogLearning, Error, TEXT("GetInclusiveDiscreteObservationToArrayView: Object is nullptr."));
@@ -3273,8 +3825,13 @@ bool ULearningAgentsObservations::GetInclusiveDiscreteObservationToArrayView(TAr
 	return true;
 }
 
-bool ULearningAgentsObservations::GetNamedExclusiveDiscreteObservation(FName& OutName, const ULearningAgentsObservationObject* Object, const FLearningAgentsObservationObjectElement Element, const FName Tag)
+bool ULearningAgentsObservations::GetNamedExclusiveDiscreteObservation(FName& OutName, const ULearningAgentsObservationObject* Object, const FLearningAgentsObservationObjectElement Element, const bool bActive, const FName Tag)
 {
+	if (!bActive)
+	{
+		return GetNullObservation(Object, Element, Tag);
+	}
+
 	if (!Object)
 	{
 		UE_LOG(LogLearning, Error, TEXT("GetNamedExclusiveDiscreteObservation: Object is nullptr."));
@@ -3308,8 +3865,13 @@ bool ULearningAgentsObservations::GetNamedExclusiveDiscreteObservation(FName& Ou
 	return true;
 }
 
-bool ULearningAgentsObservations::GetNamedInclusiveDiscreteObservationNum(int32& OutNum, const ULearningAgentsObservationObject* Object, const FLearningAgentsObservationObjectElement Element, const FName Tag)
+bool ULearningAgentsObservations::GetNamedInclusiveDiscreteObservationNum(int32& OutNum, const ULearningAgentsObservationObject* Object, const FLearningAgentsObservationObjectElement Element, const bool bActive, const FName Tag)
 {
+	if (!bActive)
+	{
+		return GetNullObservation(Object, Element, Tag);
+	}
+
 	if (!Object)
 	{
 		UE_LOG(LogLearning, Error, TEXT("GetNamedInclusiveDiscreteObservationNum: Object is nullptr."));
@@ -3343,21 +3905,31 @@ bool ULearningAgentsObservations::GetNamedInclusiveDiscreteObservationNum(int32&
 	return true;
 }
 
-bool ULearningAgentsObservations::GetNamedInclusiveDiscreteObservation(TArray<FName>& OutNames, const ULearningAgentsObservationObject* Object, const FLearningAgentsObservationObjectElement Element, const FName Tag)
+bool ULearningAgentsObservations::GetNamedInclusiveDiscreteObservation(TArray<FName>& OutNames, const ULearningAgentsObservationObject* Object, const FLearningAgentsObservationObjectElement Element, const bool bActive, const FName Tag)
 {
+	if (!bActive)
+	{
+		return GetNullObservation(Object, Element, Tag);
+	}
+
 	int32 OutNamesNum;
-	if (!GetNamedInclusiveDiscreteObservationNum(OutNamesNum, Object, Element, Tag))
+	if (!GetNamedInclusiveDiscreteObservationNum(OutNamesNum, Object, Element, bActive, Tag))
 	{
 		OutNames.Empty();
 		return false;
 	}
 
 	OutNames.SetNumUninitialized(OutNamesNum);
-	return GetNamedInclusiveDiscreteObservationToArrayView(OutNames, Object, Element, Tag);
+	return GetNamedInclusiveDiscreteObservationToArrayView(OutNames, Object, Element, bActive, Tag);
 }
 
-bool ULearningAgentsObservations::GetNamedInclusiveDiscreteObservationToArrayView(TArrayView<FName> OutNames, const ULearningAgentsObservationObject* Object, const FLearningAgentsObservationObjectElement Element, const FName Tag)
+bool ULearningAgentsObservations::GetNamedInclusiveDiscreteObservationToArrayView(TArrayView<FName> OutNames, const ULearningAgentsObservationObject* Object, const FLearningAgentsObservationObjectElement Element, const bool bActive, const FName Tag)
 {
+	if (!bActive)
+	{
+		return GetNullObservation(Object, Element, Tag);
+	}
+
 	if (!Object)
 	{
 		UE_LOG(LogLearning, Error, TEXT("GetNamedInclusiveDiscreteObservationToArrayView: Object is nullptr."));
@@ -3402,10 +3974,15 @@ bool ULearningAgentsObservations::GetNamedInclusiveDiscreteObservationToArrayVie
 	return true;
 }
 
-bool ULearningAgentsObservations::GetCountObservation(int32& OutNum, const ULearningAgentsObservationObject* Object, const FLearningAgentsObservationObjectElement Element, const int32 MaxNum, const FName Tag)
+bool ULearningAgentsObservations::GetCountObservation(int32& OutNum, const ULearningAgentsObservationObject* Object, const FLearningAgentsObservationObjectElement Element, const int32 MaxNum, const bool bActive, const FName Tag)
 {
+	if (!bActive)
+	{
+		return GetNullObservation(Object, Element, Tag);
+	}
+
 	float FloatNum = 0.0f;
-	if (!GetContinuousObservationToArrayView(MakeArrayView(&FloatNum, 1), Object, Element, Tag))
+	if (!GetContinuousObservationToArrayView(MakeArrayView(&FloatNum, 1), Object, Element, bActive, Tag))
 	{
 		OutNum = -1;
 		return false;
@@ -3415,8 +3992,13 @@ bool ULearningAgentsObservations::GetCountObservation(int32& OutNum, const ULear
 	return true;
 }
 
-bool ULearningAgentsObservations::GetStructObservationNum(int32& OutNum, const ULearningAgentsObservationObject* Object, const FLearningAgentsObservationObjectElement Element, const FName Tag)
+bool ULearningAgentsObservations::GetStructObservationNum(int32& OutNum, const ULearningAgentsObservationObject* Object, const FLearningAgentsObservationObjectElement Element, const bool bActive, const FName Tag)
 {
+	if (!bActive)
+	{
+		return GetNullObservation(Object, Element, Tag);
+	}
+
 	if (!Object)
 	{
 		UE_LOG(LogLearning, Error, TEXT("GetStructObservationNum: Object is nullptr."));
@@ -3452,10 +4034,15 @@ bool ULearningAgentsObservations::GetStructObservationNum(int32& OutNum, const U
 	return true;
 }
 
-bool ULearningAgentsObservations::GetStructObservation(TMap<FName, FLearningAgentsObservationObjectElement>& OutElements, const ULearningAgentsObservationObject* Object, const FLearningAgentsObservationObjectElement Element, const FName Tag)
+bool ULearningAgentsObservations::GetStructObservation(TMap<FName, FLearningAgentsObservationObjectElement>& OutElements, const ULearningAgentsObservationObject* Object, const FLearningAgentsObservationObjectElement Element, const bool bActive, const FName Tag)
 {
+	if (!bActive)
+	{
+		return GetNullObservation(Object, Element, Tag);
+	}
+
 	int32 OutElementNum = 0;
-	if (!GetStructObservationNum(OutElementNum, Object, Element, Tag))
+	if (!GetStructObservationNum(OutElementNum, Object, Element, bActive, Tag))
 	{
 		OutElements.Empty();
 		return false;
@@ -3466,7 +4053,7 @@ bool ULearningAgentsObservations::GetStructObservation(TMap<FName, FLearningAgen
 	SubElementNames.SetNumUninitialized(OutElementNum);
 	SubElements.SetNumUninitialized(OutElementNum);
 
-	if (!GetStructObservationToArrayViews(SubElementNames, SubElements, Object, Element, Tag))
+	if (!GetStructObservationToArrayViews(SubElementNames, SubElements, Object, Element, bActive, Tag))
 	{
 		OutElements.Empty();
 		return false;
@@ -3481,8 +4068,13 @@ bool ULearningAgentsObservations::GetStructObservation(TMap<FName, FLearningAgen
 	return true;
 }
 
-bool ULearningAgentsObservations::GetStructObservationElement(FLearningAgentsObservationObjectElement& OutElement, const ULearningAgentsObservationObject* Object, const FLearningAgentsObservationObjectElement Element, const FName ElementName, const FName Tag)
+bool ULearningAgentsObservations::GetStructObservationElement(FLearningAgentsObservationObjectElement& OutElement, const ULearningAgentsObservationObject* Object, const FLearningAgentsObservationObjectElement Element, const FName ElementName, const bool bActive, const FName Tag)
 {
+	if (!bActive)
+	{
+		return GetNullObservation(Object, Element, Tag);
+	}
+
 	if (!Object)
 	{
 		UE_LOG(LogLearning, Error, TEXT("GetStructObservationElement: Object is nullptr."));
@@ -3527,10 +4119,15 @@ bool ULearningAgentsObservations::GetStructObservationElement(FLearningAgentsObs
 	return true;
 }
 
-bool ULearningAgentsObservations::GetStructObservationToArrays(TArray<FName>& OutElementNames, TArray<FLearningAgentsObservationObjectElement>& OutElements, const ULearningAgentsObservationObject* Object, const FLearningAgentsObservationObjectElement Element, const FName Tag)
+bool ULearningAgentsObservations::GetStructObservationToArrays(TArray<FName>& OutElementNames, TArray<FLearningAgentsObservationObjectElement>& OutElements, const ULearningAgentsObservationObject* Object, const FLearningAgentsObservationObjectElement Element, const bool bActive, const FName Tag)
 {
+	if (!bActive)
+	{
+		return GetNullObservation(Object, Element, Tag);
+	}
+
 	int32 OutElementNum = 0;
-	if (!GetStructObservationNum(OutElementNum, Object, Element, Tag))
+	if (!GetStructObservationNum(OutElementNum, Object, Element, bActive, Tag))
 	{
 		OutElementNames.Empty();
 		OutElements.Empty();
@@ -3540,7 +4137,7 @@ bool ULearningAgentsObservations::GetStructObservationToArrays(TArray<FName>& Ou
 	OutElementNames.SetNumUninitialized(OutElementNum);
 	OutElements.SetNumUninitialized(OutElementNum);
 
-	if (!GetStructObservationToArrayViews(OutElementNames, OutElements, Object, Element, Tag))
+	if (!GetStructObservationToArrayViews(OutElementNames, OutElements, Object, Element, bActive, Tag))
 	{
 		OutElementNames.Empty();
 		OutElements.Empty();
@@ -3550,8 +4147,13 @@ bool ULearningAgentsObservations::GetStructObservationToArrays(TArray<FName>& Ou
 	return true;
 }
 
-bool ULearningAgentsObservations::GetStructObservationToArrayViews(TArrayView<FName> OutElementNames, TArrayView<FLearningAgentsObservationObjectElement> OutElements, const ULearningAgentsObservationObject* Object, const FLearningAgentsObservationObjectElement Element, const FName Tag)
+bool ULearningAgentsObservations::GetStructObservationToArrayViews(TArrayView<FName> OutElementNames, TArrayView<FLearningAgentsObservationObjectElement> OutElements, const ULearningAgentsObservationObject* Object, const FLearningAgentsObservationObjectElement Element, const bool bActive, const FName Tag)
 {
+	if (!bActive)
+	{
+		return GetNullObservation(Object, Element, Tag);
+	}
+
 	if (!Object)
 	{
 		UE_LOG(LogLearning, Error, TEXT("GetStructObservationToArrayViews: Object is nullptr."));
@@ -3618,8 +4220,13 @@ bool ULearningAgentsObservations::GetStructObservationToArrayViews(TArrayView<FN
 	return true;
 }
 
-bool ULearningAgentsObservations::GetExclusiveUnionObservation(FName& OutElementName, FLearningAgentsObservationObjectElement& OutElement, const ULearningAgentsObservationObject* Object, const FLearningAgentsObservationObjectElement Element, const FName Tag)
+bool ULearningAgentsObservations::GetExclusiveUnionObservation(FName& OutElementName, FLearningAgentsObservationObjectElement& OutElement, const ULearningAgentsObservationObject* Object, const FLearningAgentsObservationObjectElement Element, const bool bActive, const FName Tag)
 {
+	if (!bActive)
+	{
+		return GetNullObservation(Object, Element, Tag);
+	}
+
 	if (!Object)
 	{
 		UE_LOG(LogLearning, Error, TEXT("GetStructObservationToArrayViews: Object is nullptr."));
@@ -3658,8 +4265,13 @@ bool ULearningAgentsObservations::GetExclusiveUnionObservation(FName& OutElement
 	return true;
 }
 
-bool ULearningAgentsObservations::GetInclusiveUnionObservationNum(int32& OutNum, const ULearningAgentsObservationObject* Object, const FLearningAgentsObservationObjectElement Element, const FName Tag)
+bool ULearningAgentsObservations::GetInclusiveUnionObservationNum(int32& OutNum, const ULearningAgentsObservationObject* Object, const FLearningAgentsObservationObjectElement Element, const bool bActive, const FName Tag)
 {
+	if (!bActive)
+	{
+		return GetNullObservation(Object, Element, Tag);
+	}
+
 	if (!Object)
 	{
 		UE_LOG(LogLearning, Error, TEXT("GetInclusiveUnionObservationNum: Object is nullptr."));
@@ -3695,10 +4307,15 @@ bool ULearningAgentsObservations::GetInclusiveUnionObservationNum(int32& OutNum,
 	return true;
 }
 
-bool ULearningAgentsObservations::GetInclusiveUnionObservation(TMap<FName, FLearningAgentsObservationObjectElement>& OutElements, const ULearningAgentsObservationObject* Object, const FLearningAgentsObservationObjectElement Element, const FName Tag)
+bool ULearningAgentsObservations::GetInclusiveUnionObservation(TMap<FName, FLearningAgentsObservationObjectElement>& OutElements, const ULearningAgentsObservationObject* Object, const FLearningAgentsObservationObjectElement Element, const bool bActive, const FName Tag)
 {
+	if (!bActive)
+	{
+		return GetNullObservation(Object, Element, Tag);
+	}
+
 	int32 OutElementNum = 0;
-	if (!GetInclusiveUnionObservationNum(OutElementNum, Object, Element, Tag))
+	if (!GetInclusiveUnionObservationNum(OutElementNum, Object, Element, bActive, Tag))
 	{
 		OutElements.Empty();
 		return false;
@@ -3709,7 +4326,7 @@ bool ULearningAgentsObservations::GetInclusiveUnionObservation(TMap<FName, FLear
 	SubElementNames.SetNumUninitialized(OutElementNum);
 	SubElements.SetNumUninitialized(OutElementNum);
 
-	if (!GetInclusiveUnionObservationToArrayViews(SubElementNames, SubElements, Object, Element, Tag))
+	if (!GetInclusiveUnionObservationToArrayViews(SubElementNames, SubElements, Object, Element, bActive, Tag))
 	{
 		OutElements.Empty();
 		return false;
@@ -3724,10 +4341,15 @@ bool ULearningAgentsObservations::GetInclusiveUnionObservation(TMap<FName, FLear
 	return true;
 }
 
-bool ULearningAgentsObservations::GetInclusiveUnionObservationToArrays(TArray<FName>& OutElementNames, TArray<FLearningAgentsObservationObjectElement>& OutElements, const ULearningAgentsObservationObject* Object, const FLearningAgentsObservationObjectElement Element, const FName Tag)
+bool ULearningAgentsObservations::GetInclusiveUnionObservationToArrays(TArray<FName>& OutElementNames, TArray<FLearningAgentsObservationObjectElement>& OutElements, const ULearningAgentsObservationObject* Object, const FLearningAgentsObservationObjectElement Element, const bool bActive, const FName Tag)
 {
+	if (!bActive)
+	{
+		return GetNullObservation(Object, Element, Tag);
+	}
+
 	int32 OutElementNum = 0;
-	if (!GetInclusiveUnionObservationNum(OutElementNum, Object, Element, Tag))
+	if (!GetInclusiveUnionObservationNum(OutElementNum, Object, Element, bActive, Tag))
 	{
 		OutElementNames.Empty();
 		OutElements.Empty();
@@ -3737,7 +4359,7 @@ bool ULearningAgentsObservations::GetInclusiveUnionObservationToArrays(TArray<FN
 	OutElementNames.SetNumUninitialized(OutElementNum);
 	OutElements.SetNumUninitialized(OutElementNum);
 
-	if (!GetInclusiveUnionObservationToArrayViews(OutElementNames, OutElements, Object, Element, Tag))
+	if (!GetInclusiveUnionObservationToArrayViews(OutElementNames, OutElements, Object, Element, bActive, Tag))
 	{
 		OutElementNames.Empty();
 		OutElements.Empty();
@@ -3747,8 +4369,13 @@ bool ULearningAgentsObservations::GetInclusiveUnionObservationToArrays(TArray<FN
 	return true;
 }
 
-bool ULearningAgentsObservations::GetInclusiveUnionObservationToArrayViews(TArrayView<FName> OutElementNames, TArrayView<FLearningAgentsObservationObjectElement> OutElements, const ULearningAgentsObservationObject* Object, const FLearningAgentsObservationObjectElement Element, const FName Tag)
+bool ULearningAgentsObservations::GetInclusiveUnionObservationToArrayViews(TArrayView<FName> OutElementNames, TArrayView<FLearningAgentsObservationObjectElement> OutElements, const ULearningAgentsObservationObject* Object, const FLearningAgentsObservationObjectElement Element, const bool bActive, const FName Tag)
 {
+	if (!bActive)
+	{
+		return GetNullObservation(Object, Element, Tag);
+	}
+
 	if (!Object)
 	{
 		UE_LOG(LogLearning, Error, TEXT("GetInclusiveUnionObservationToArrayViews: Object is nullptr."));
@@ -3810,8 +4437,13 @@ bool ULearningAgentsObservations::GetInclusiveUnionObservationToArrayViews(TArra
 	return true;
 }
 
-bool ULearningAgentsObservations::GetStaticArrayObservationNum(int32& OutNum, const ULearningAgentsObservationObject* Object, const FLearningAgentsObservationObjectElement Element, const FName Tag)
+bool ULearningAgentsObservations::GetStaticArrayObservationNum(int32& OutNum, const ULearningAgentsObservationObject* Object, const FLearningAgentsObservationObjectElement Element, const bool bActive, const FName Tag)
 {
+	if (!bActive)
+	{
+		return GetNullObservation(Object, Element, Tag);
+	}
+
 	if (!Object)
 	{
 		UE_LOG(LogLearning, Error, TEXT("GetStaticArrayObservationNum: Object is nullptr."));
@@ -3845,10 +4477,15 @@ bool ULearningAgentsObservations::GetStaticArrayObservationNum(int32& OutNum, co
 	return true;
 }
 
-bool ULearningAgentsObservations::GetStaticArrayObservation(TArray<FLearningAgentsObservationObjectElement>& OutElements, const ULearningAgentsObservationObject* Object, const FLearningAgentsObservationObjectElement Element, const FName Tag)
+bool ULearningAgentsObservations::GetStaticArrayObservation(TArray<FLearningAgentsObservationObjectElement>& OutElements, const ULearningAgentsObservationObject* Object, const FLearningAgentsObservationObjectElement Element, const bool bActive, const FName Tag)
 {
+	if (!bActive)
+	{
+		return GetNullObservation(Object, Element, Tag);
+	}
+
 	int32 OutElementNum = 0;
-	if (!GetStaticArrayObservationNum(OutElementNum, Object, Element, Tag))
+	if (!GetStaticArrayObservationNum(OutElementNum, Object, Element, bActive, Tag))
 	{
 		OutElements.Empty();
 		return false;
@@ -3856,7 +4493,7 @@ bool ULearningAgentsObservations::GetStaticArrayObservation(TArray<FLearningAgen
 
 	OutElements.SetNumUninitialized(OutElementNum);
 
-	if (!GetStaticArrayObservationToArrayView(OutElements, Object, Element, Tag))
+	if (!GetStaticArrayObservationToArrayView(OutElements, Object, Element, bActive, Tag))
 	{
 		OutElements.Empty();
 		return false;
@@ -3865,8 +4502,13 @@ bool ULearningAgentsObservations::GetStaticArrayObservation(TArray<FLearningAgen
 	return true;
 }
 
-bool ULearningAgentsObservations::GetStaticArrayObservationToArrayView(TArrayView<FLearningAgentsObservationObjectElement> OutElements, const ULearningAgentsObservationObject* Object, const FLearningAgentsObservationObjectElement Element, const FName Tag)
+bool ULearningAgentsObservations::GetStaticArrayObservationToArrayView(TArrayView<FLearningAgentsObservationObjectElement> OutElements, const ULearningAgentsObservationObject* Object, const FLearningAgentsObservationObjectElement Element, const bool bActive, const FName Tag)
 {
+	if (!bActive)
+	{
+		return GetNullObservation(Object, Element, Tag);
+	}
+
 	if (!Object)
 	{
 		UE_LOG(LogLearning, Error, TEXT("GetStaticArrayObservationNum: Object is nullptr."));
@@ -3927,8 +4569,13 @@ bool ULearningAgentsObservations::GetStaticArrayObservationToArrayView(TArrayVie
 	return true;
 }
 
-bool ULearningAgentsObservations::GetSetObservationNum(int32& OutNum, const ULearningAgentsObservationObject* Object, const FLearningAgentsObservationObjectElement Element, const FName Tag)
+bool ULearningAgentsObservations::GetSetObservationNum(int32& OutNum, const ULearningAgentsObservationObject* Object, const FLearningAgentsObservationObjectElement Element, const bool bActive, const FName Tag)
 {
+	if (!bActive)
+	{
+		return GetNullObservation(Object, Element, Tag);
+	}
+
 	if (!Object)
 	{
 		UE_LOG(LogLearning, Error, TEXT("GetSetObservationNum: Object is nullptr."));
@@ -3962,10 +4609,15 @@ bool ULearningAgentsObservations::GetSetObservationNum(int32& OutNum, const ULea
 	return true;
 }
 
-bool ULearningAgentsObservations::GetSetObservation(TSet<FLearningAgentsObservationObjectElement>& OutElements, const ULearningAgentsObservationObject* Object, const FLearningAgentsObservationObjectElement Element, const FName Tag)
+bool ULearningAgentsObservations::GetSetObservation(TSet<FLearningAgentsObservationObjectElement>& OutElements, const ULearningAgentsObservationObject* Object, const FLearningAgentsObservationObjectElement Element, const bool bActive, const FName Tag)
 {
+	if (!bActive)
+	{
+		return GetNullObservation(Object, Element, Tag);
+	}
+
 	int32 OutElementNum = 0;
-	if (!GetSetObservationNum(OutElementNum, Object, Element, Tag))
+	if (!GetSetObservationNum(OutElementNum, Object, Element, bActive, Tag))
 	{
 		OutElements.Empty();
 		return false;
@@ -4018,10 +4670,15 @@ bool ULearningAgentsObservations::GetSetObservation(TSet<FLearningAgentsObservat
 	return true;
 }
 
-bool ULearningAgentsObservations::GetSetObservationToArray(TArray<FLearningAgentsObservationObjectElement>& OutElements, const ULearningAgentsObservationObject* Object, const FLearningAgentsObservationObjectElement Element, const FName Tag)
+bool ULearningAgentsObservations::GetSetObservationToArray(TArray<FLearningAgentsObservationObjectElement>& OutElements, const ULearningAgentsObservationObject* Object, const FLearningAgentsObservationObjectElement Element, const bool bActive, const FName Tag)
 {
+	if (!bActive)
+	{
+		return GetNullObservation(Object, Element, Tag);
+	}
+
 	int32 OutElementNum = 0;
-	if (!GetSetObservationNum(OutElementNum, Object, Element, Tag))
+	if (!GetSetObservationNum(OutElementNum, Object, Element, bActive, Tag))
 	{
 		OutElements.Empty();
 		return false;
@@ -4029,7 +4686,7 @@ bool ULearningAgentsObservations::GetSetObservationToArray(TArray<FLearningAgent
 
 	OutElements.SetNumUninitialized(OutElementNum);
 
-	if (!GetSetObservationToArrayView(OutElements, Object, Element, Tag))
+	if (!GetSetObservationToArrayView(OutElements, Object, Element, bActive, Tag))
 	{
 		OutElements.Empty();
 		return false;
@@ -4038,8 +4695,13 @@ bool ULearningAgentsObservations::GetSetObservationToArray(TArray<FLearningAgent
 	return true;
 }
 
-bool ULearningAgentsObservations::GetSetObservationToArrayView(TArrayView<FLearningAgentsObservationObjectElement> OutElements, const ULearningAgentsObservationObject* Object, const FLearningAgentsObservationObjectElement Element, const FName Tag)
+bool ULearningAgentsObservations::GetSetObservationToArrayView(TArrayView<FLearningAgentsObservationObjectElement> OutElements, const ULearningAgentsObservationObject* Object, const FLearningAgentsObservationObjectElement Element, const bool bActive, const FName Tag)
 {
+	if (!bActive)
+	{
+		return GetNullObservation(Object, Element, Tag);
+	}
+
 	if (!Object)
 	{
 		UE_LOG(LogLearning, Error, TEXT("GetSetObservationToArrayView: Object is nullptr."));
@@ -4095,11 +4757,16 @@ bool ULearningAgentsObservations::GetSetObservationToArrayView(TArrayView<FLearn
 	return true;
 }
 
-bool ULearningAgentsObservations::GetPairObservation(FLearningAgentsObservationObjectElement& OutKey, FLearningAgentsObservationObjectElement& OutValue, const ULearningAgentsObservationObject* Object, const FLearningAgentsObservationObjectElement Element, const FName Tag)
+bool ULearningAgentsObservations::GetPairObservation(FLearningAgentsObservationObjectElement& OutKey, FLearningAgentsObservationObjectElement& OutValue, const ULearningAgentsObservationObject* Object, const FLearningAgentsObservationObjectElement Element, const bool bActive, const FName Tag)
 {
+	if (!bActive)
+	{
+		return GetNullObservation(Object, Element, Tag);
+	}
+
 	TStaticArray<FName, 2> OutElementNames;
 	TStaticArray<FLearningAgentsObservationObjectElement, 2> OutElements;
-	if (!GetStructObservationToArrayViews(OutElementNames, OutElements, Object, Element, Tag))
+	if (!GetStructObservationToArrayViews(OutElementNames, OutElements, Object, Element, bActive, Tag))
 	{
 		OutKey = FLearningAgentsObservationObjectElement();
 		OutValue = FLearningAgentsObservationObjectElement();
@@ -4111,15 +4778,25 @@ bool ULearningAgentsObservations::GetPairObservation(FLearningAgentsObservationO
 	return true;
 }
 
-bool ULearningAgentsObservations::GetArrayObservationNum(int32& OutNum, const ULearningAgentsObservationObject* Object, const FLearningAgentsObservationObjectElement Element, const FName Tag)
+bool ULearningAgentsObservations::GetArrayObservationNum(int32& OutNum, const ULearningAgentsObservationObject* Object, const FLearningAgentsObservationObjectElement Element, const bool bActive, const FName Tag)
 {
-	return GetSetObservationNum(OutNum, Object, Element, Tag);
+	if (!bActive)
+	{
+		return GetNullObservation(Object, Element, Tag);
+	}
+
+	return GetSetObservationNum(OutNum, Object, Element, bActive, Tag);
 }
 
-bool ULearningAgentsObservations::GetArrayObservation(TArray<FLearningAgentsObservationObjectElement>& OutElements, const ULearningAgentsObservationObject* Object, const FLearningAgentsObservationObjectElement Element, const int32 MaxNum, const FName Tag)
+bool ULearningAgentsObservations::GetArrayObservation(TArray<FLearningAgentsObservationObjectElement>& OutElements, const ULearningAgentsObservationObject* Object, const FLearningAgentsObservationObjectElement Element, const int32 MaxNum, const bool bActive, const FName Tag)
 {
+	if (!bActive)
+	{
+		return GetNullObservation(Object, Element, Tag);
+	}
+
 	int32 OutElementNum = 0;
-	if (!GetArrayObservationNum(OutElementNum, Object, Element, Tag))
+	if (!GetArrayObservationNum(OutElementNum, Object, Element, bActive, Tag))
 	{
 		OutElements.Empty();
 		return false;
@@ -4127,7 +4804,7 @@ bool ULearningAgentsObservations::GetArrayObservation(TArray<FLearningAgentsObse
 
 	OutElements.SetNumUninitialized(OutElementNum);
 
-	if (!GetArrayObservationToArrayView(OutElements, Object, Element, MaxNum, Tag))
+	if (!GetArrayObservationToArrayView(OutElements, Object, Element, MaxNum, bActive, Tag))
 	{
 		OutElements.Empty();
 		return false;
@@ -4136,11 +4813,16 @@ bool ULearningAgentsObservations::GetArrayObservation(TArray<FLearningAgentsObse
 	return true;
 }
 
-bool ULearningAgentsObservations::GetArrayObservationToArrayView(TArrayView<FLearningAgentsObservationObjectElement> OutElements, const ULearningAgentsObservationObject* Object, const FLearningAgentsObservationObjectElement Element, const int32 MaxNum, const FName Tag)
+bool ULearningAgentsObservations::GetArrayObservationToArrayView(TArrayView<FLearningAgentsObservationObjectElement> OutElements, const ULearningAgentsObservationObject* Object, const FLearningAgentsObservationObjectElement Element, const int32 MaxNum, const bool bActive, const FName Tag)
 {
+	if (!bActive)
+	{
+		return GetNullObservation(Object, Element, Tag);
+	}
+
 	TArray<FLearningAgentsObservationObjectElement, TInlineAllocator<16>> Pairs;
 	Pairs.SetNumUninitialized(OutElements.Num());
-	if (!GetSetObservationToArrayView(Pairs, Object, Element, Tag))
+	if (!GetSetObservationToArrayView(Pairs, Object, Element, bActive, Tag))
 	{
 		UE::Learning::Array::Set<1, FLearningAgentsObservationObjectElement>(OutElements, FLearningAgentsObservationObjectElement());
 		return false;
@@ -4169,15 +4851,25 @@ bool ULearningAgentsObservations::GetArrayObservationToArrayView(TArrayView<FLea
 	return true;
 }
 
-bool ULearningAgentsObservations::GetMapObservationNum(int32& OutNum, const ULearningAgentsObservationObject* Object, const FLearningAgentsObservationObjectElement Element, const FName Tag)
+bool ULearningAgentsObservations::GetMapObservationNum(int32& OutNum, const ULearningAgentsObservationObject* Object, const FLearningAgentsObservationObjectElement Element, const bool bActive, const FName Tag)
 {
-	return GetSetObservationNum(OutNum, Object, Element, Tag);
+	if (!bActive)
+	{
+		return GetNullObservation(Object, Element, Tag);
+	}
+
+	return GetSetObservationNum(OutNum, Object, Element, bActive, Tag);
 }
 
-bool ULearningAgentsObservations::GetMapObservation(TMap<FLearningAgentsObservationObjectElement, FLearningAgentsObservationObjectElement>& OutElements, const ULearningAgentsObservationObject* Object, const FLearningAgentsObservationObjectElement Element, const FName Tag)
+bool ULearningAgentsObservations::GetMapObservation(TMap<FLearningAgentsObservationObjectElement, FLearningAgentsObservationObjectElement>& OutElements, const ULearningAgentsObservationObject* Object, const FLearningAgentsObservationObjectElement Element, const bool bActive, const FName Tag)
 {
+	if (!bActive)
+	{
+		return GetNullObservation(Object, Element, Tag);
+	}
+
 	int32 OutElementNum = 0;
-	if (!GetMapObservationNum(OutElementNum, Object, Element, Tag))
+	if (!GetMapObservationNum(OutElementNum, Object, Element, bActive, Tag))
 	{
 		OutElements.Empty();
 		return false;
@@ -4185,7 +4877,7 @@ bool ULearningAgentsObservations::GetMapObservation(TMap<FLearningAgentsObservat
 
 	TArray<FLearningAgentsObservationObjectElement, TInlineAllocator<16>> Pairs;
 	Pairs.SetNumUninitialized(OutElementNum);
-	if (!GetSetObservationToArrayView(Pairs, Object, Element, Tag))
+	if (!GetSetObservationToArrayView(Pairs, Object, Element, bActive, Tag))
 	{
 		OutElements.Empty();
 		return false;
@@ -4207,10 +4899,15 @@ bool ULearningAgentsObservations::GetMapObservation(TMap<FLearningAgentsObservat
 	return true;
 }
 
-bool ULearningAgentsObservations::GetMapObservationToArrays(TArray<FLearningAgentsObservationObjectElement>& OutKeys, TArray<FLearningAgentsObservationObjectElement>& OutValues, const ULearningAgentsObservationObject* Object, const FLearningAgentsObservationObjectElement Element, const FName Tag)
+bool ULearningAgentsObservations::GetMapObservationToArrays(TArray<FLearningAgentsObservationObjectElement>& OutKeys, TArray<FLearningAgentsObservationObjectElement>& OutValues, const ULearningAgentsObservationObject* Object, const FLearningAgentsObservationObjectElement Element, const bool bActive, const FName Tag)
 {
+	if (!bActive)
+	{
+		return GetNullObservation(Object, Element, Tag);
+	}
+
 	int32 OutElementNum = 0;
-	if (!GetMapObservationNum(OutElementNum, Object, Element, Tag))
+	if (!GetMapObservationNum(OutElementNum, Object, Element, bActive, Tag))
 	{
 		OutKeys.Empty();
 		OutValues.Empty();
@@ -4219,7 +4916,7 @@ bool ULearningAgentsObservations::GetMapObservationToArrays(TArray<FLearningAgen
 
 	OutKeys.SetNumUninitialized(OutElementNum);
 	OutValues.SetNumUninitialized(OutElementNum);
-	if (!GetMapObservationToArrayViews(OutKeys, OutValues, Object, Element, Tag))
+	if (!GetMapObservationToArrayViews(OutKeys, OutValues, Object, Element, bActive, Tag))
 	{
 		OutKeys.Empty();
 		OutValues.Empty();
@@ -4229,11 +4926,16 @@ bool ULearningAgentsObservations::GetMapObservationToArrays(TArray<FLearningAgen
 	return true;
 }
 
-bool ULearningAgentsObservations::GetMapObservationToArrayViews(TArrayView<FLearningAgentsObservationObjectElement> OutKeys, TArrayView<FLearningAgentsObservationObjectElement> OutValues, const ULearningAgentsObservationObject* Object, const FLearningAgentsObservationObjectElement Element, const FName Tag)
+bool ULearningAgentsObservations::GetMapObservationToArrayViews(TArrayView<FLearningAgentsObservationObjectElement> OutKeys, TArrayView<FLearningAgentsObservationObjectElement> OutValues, const ULearningAgentsObservationObject* Object, const FLearningAgentsObservationObjectElement Element, const bool bActive, const FName Tag)
 {
+	if (!bActive)
+	{
+		return GetNullObservation(Object, Element, Tag);
+	}
+
 	TArray<FLearningAgentsObservationObjectElement, TInlineAllocator<16>> Pairs;
 	Pairs.SetNumUninitialized(OutKeys.Num());
-	if (!GetSetObservationToArrayView(Pairs, Object, Element, Tag))
+	if (!GetSetObservationToArrayView(Pairs, Object, Element, bActive, Tag))
 	{
 		UE::Learning::Array::Set<1, FLearningAgentsObservationObjectElement>(OutKeys, FLearningAgentsObservationObjectElement());
 		UE::Learning::Array::Set<1, FLearningAgentsObservationObjectElement>(OutValues, FLearningAgentsObservationObjectElement());
@@ -4257,8 +4959,13 @@ bool ULearningAgentsObservations::GetMapObservationToArrayViews(TArrayView<FLear
 	return true;
 }
 
-bool ULearningAgentsObservations::GetEnumObservation(uint8& OutEnumValue, const ULearningAgentsObservationObject* Object, const FLearningAgentsObservationObjectElement Element, const UEnum* Enum, const FName Tag)
+bool ULearningAgentsObservations::GetEnumObservation(uint8& OutEnumValue, const ULearningAgentsObservationObject* Object, const FLearningAgentsObservationObjectElement Element, const UEnum* Enum, const bool bActive, const FName Tag)
 {
+	if (!bActive)
+	{
+		return GetNullObservation(Object, Element, Tag);
+	}
+
 	if (!Enum)
 	{
 		UE_LOG(LogLearning, Error, TEXT("GetEnumObservation: Enum is nullptr."));
@@ -4267,7 +4974,7 @@ bool ULearningAgentsObservations::GetEnumObservation(uint8& OutEnumValue, const
 	}
 
 	int32 EnumValueNum;
-	if (!GetContinuousObservationNum(EnumValueNum, Object, Element, Tag))
+	if (!GetContinuousObservationNum(EnumValueNum, Object, Element, bActive, Tag))
 	{
 		OutEnumValue = 0;
 		return false;
@@ -4282,7 +4989,7 @@ bool ULearningAgentsObservations::GetEnumObservation(uint8& OutEnumValue, const
 
 	TArray<float, TInlineAllocator<32>> OneHot;
 	OneHot.SetNumUninitialized(EnumValueNum);
-	if (!GetContinuousObservationToArrayView(OneHot, Object, Element, Tag))
+	if (!GetContinuousObservationToArrayView(OneHot, Object, Element, bActive, Tag))
 	{
 		OutEnumValue = 0;
 		return false;
@@ -4318,8 +5025,13 @@ bool ULearningAgentsObservations::GetEnumObservation(uint8& OutEnumValue, const
 	return true;
 }
 
-bool ULearningAgentsObservations::GetBitmaskObservation(int32& OutBitmaskValue, const ULearningAgentsObservationObject* Object, const FLearningAgentsObservationObjectElement Element, const UEnum* Enum, const FName Tag)
+bool ULearningAgentsObservations::GetBitmaskObservation(int32& OutBitmaskValue, const ULearningAgentsObservationObject* Object, const FLearningAgentsObservationObjectElement Element, const UEnum* Enum, const bool bActive, const FName Tag)
 {
+	if (!bActive)
+	{
+		return GetNullObservation(Object, Element, Tag);
+	}
+
 	if (!Enum)
 	{
 		UE_LOG(LogLearning, Error, TEXT("GetBitmaskObservation: Enum is nullptr."));
@@ -4335,7 +5047,7 @@ bool ULearningAgentsObservations::GetBitmaskObservation(int32& OutBitmaskValue,
 	}
 
 	int32 EnumValueNum;
-	if (!GetContinuousObservationNum(EnumValueNum, Object, Element, Tag))
+	if (!GetContinuousObservationNum(EnumValueNum, Object, Element, bActive, Tag))
 	{
 		OutBitmaskValue = 0;
 		return false;
@@ -4350,7 +5062,7 @@ bool ULearningAgentsObservations::GetBitmaskObservation(int32& OutBitmaskValue,
 
 	TArray<float, TInlineAllocator<32>> OneHot;
 	OneHot.Init(0.0f, EnumValueNum);
-	if (!GetContinuousObservationToArrayView(OneHot, Object, Element, Tag))
+	if (!GetContinuousObservationToArrayView(OneHot, Object, Element, bActive, Tag))
 	{
 		OutBitmaskValue = 0;
 		return false;
@@ -4368,10 +5080,15 @@ bool ULearningAgentsObservations::GetBitmaskObservation(int32& OutBitmaskValue,
 }
 
 
-bool ULearningAgentsObservations::GetOptionalObservation(ELearningAgentsOptionalObservation& OutOption, FLearningAgentsObservationObjectElement& OutElement, const ULearningAgentsObservationObject* Object, const FLearningAgentsObservationObjectElement Element, const FName Tag)
+bool ULearningAgentsObservations::GetOptionalObservation(ELearningAgentsOptionalObservation& OutOption, FLearningAgentsObservationObjectElement& OutElement, const ULearningAgentsObservationObject* Object, const FLearningAgentsObservationObjectElement Element, const bool bActive, const FName Tag)
 {
+	if (!bActive)
+	{
+		return GetNullObservation(Object, Element, Tag);
+	}
+
 	FName OutName = NAME_None;
-	if (!GetExclusiveUnionObservation(OutName, OutElement, Object, Element, Tag))
+	if (!GetExclusiveUnionObservation(OutName, OutElement, Object, Element, bActive, Tag))
 	{
 		OutOption = ELearningAgentsOptionalObservation::Null;
 		return false;
@@ -4381,10 +5098,15 @@ bool ULearningAgentsObservations::GetOptionalObservation(ELearningAgentsOptional
 	return true;
 }
 
-bool ULearningAgentsObservations::GetEitherObservation(ELearningAgentsEitherObservation& OutEither, FLearningAgentsObservationObjectElement& OutElement, const ULearningAgentsObservationObject* Object, const FLearningAgentsObservationObjectElement Element, const FName Tag)
+bool ULearningAgentsObservations::GetEitherObservation(ELearningAgentsEitherObservation& OutEither, FLearningAgentsObservationObjectElement& OutElement, const ULearningAgentsObservationObject* Object, const FLearningAgentsObservationObjectElement Element, const bool bActive, const FName Tag)
 {
+	if (!bActive)
+	{
+		return GetNullObservation(Object, Element, Tag);
+	}
+
 	FName OutName = NAME_None;
-	if (!GetExclusiveUnionObservation(OutName, OutElement, Object, Element, Tag))
+	if (!GetExclusiveUnionObservation(OutName, OutElement, Object, Element, bActive, Tag))
 	{
 		OutEither = ELearningAgentsEitherObservation::A;
 		return false;
@@ -4394,8 +5116,13 @@ bool ULearningAgentsObservations::GetEitherObservation(ELearningAgentsEitherObse
 	return true;
 }
 
-bool ULearningAgentsObservations::GetEncodingObservation(FLearningAgentsObservationObjectElement& OutElement, const ULearningAgentsObservationObject* Object, const FLearningAgentsObservationObjectElement Element, const FName Tag)
+bool ULearningAgentsObservations::GetEncodingObservation(FLearningAgentsObservationObjectElement& OutElement, const ULearningAgentsObservationObject* Object, const FLearningAgentsObservationObjectElement Element, const bool bActive, const FName Tag)
 {
+	if (!bActive)
+	{
+		return GetNullObservation(Object, Element, Tag);
+	}
+
 	if (!Object)
 	{
 		UE_LOG(LogLearning, Error, TEXT("GetEncodingObservation: Object is nullptr."));
@@ -4430,10 +5157,15 @@ bool ULearningAgentsObservations::GetEncodingObservation(FLearningAgentsObservat
 	return true;
 }
 
-bool ULearningAgentsObservations::GetBoolObservation(bool& bOutValue, const ULearningAgentsObservationObject* Object, const FLearningAgentsObservationObjectElement Element, const FName Tag)
+bool ULearningAgentsObservations::GetBoolObservation(bool& bOutValue, const ULearningAgentsObservationObject* Object, const FLearningAgentsObservationObjectElement Element, const bool bActive, const FName Tag)
 {
+	if (!bActive)
+	{
+		return GetNullObservation(Object, Element, Tag);
+	}
+
 	float OutValue = 0.0f;
-	if (!GetFloatObservation(OutValue, Object, Element, Tag))
+	if (!GetFloatObservation(OutValue, Object, Element, bActive, Tag))
 	{
 		bOutValue = false;
 		return false;
@@ -4443,9 +5175,14 @@ bool ULearningAgentsObservations::GetBoolObservation(bool& bOutValue, const ULea
 	return true;
 }
 
-bool ULearningAgentsObservations::GetFloatObservation(float& OutValue, const ULearningAgentsObservationObject* Object, const FLearningAgentsObservationObjectElement Element, const FName Tag)
+bool ULearningAgentsObservations::GetFloatObservation(float& OutValue, const ULearningAgentsObservationObject* Object, const FLearningAgentsObservationObjectElement Element, const bool bActive, const FName Tag)
 {
-	if (!GetContinuousObservationToArrayView(MakeArrayView(&OutValue, 1), Object, Element, Tag))
+	if (!bActive)
+	{
+		return GetNullObservation(Object, Element, Tag);
+	}
+
+	if (!GetContinuousObservationToArrayView(MakeArrayView(&OutValue, 1), Object, Element, bActive, Tag))
 	{
 		OutValue = 0.0f;
 		return false;
@@ -4454,10 +5191,15 @@ bool ULearningAgentsObservations::GetFloatObservation(float& OutValue, const ULe
 	return true;
 }
 
-bool ULearningAgentsObservations::GetLocationObservation(FVector& OutLocation, const ULearningAgentsObservationObject* Object, const FLearningAgentsObservationObjectElement Element, const FTransform RelativeTransform, const FName Tag)
+bool ULearningAgentsObservations::GetLocationObservation(FVector& OutLocation, const ULearningAgentsObservationObject* Object, const FLearningAgentsObservationObjectElement Element, const FTransform RelativeTransform, const bool bActive, const FName Tag)
 {
+	if (!bActive)
+	{
+		return GetNullObservation(Object, Element, Tag);
+	}
+
 	TStaticArray<float, 3> OutValues;
-	if (!GetContinuousObservationToArrayView(OutValues, Object, Element, Tag))
+	if (!GetContinuousObservationToArrayView(OutValues, Object, Element, bActive, Tag))
 	{
 		OutLocation = FVector::ZeroVector;
 		return false;
@@ -4467,10 +5209,15 @@ bool ULearningAgentsObservations::GetLocationObservation(FVector& OutLocation, c
 	return true;
 }
 
-bool ULearningAgentsObservations::GetRotationObservation(FRotator& OutRotation, const ULearningAgentsObservationObject* Object, const FLearningAgentsObservationObjectElement Element, const FRotator RelativeRotation, const FName Tag)
+bool ULearningAgentsObservations::GetRotationObservation(FRotator& OutRotation, const ULearningAgentsObservationObject* Object, const FLearningAgentsObservationObjectElement Element, const FRotator RelativeRotation, const bool bActive, const FName Tag)
 {
+	if (!bActive)
+	{
+		return GetNullObservation(Object, Element, Tag);
+	}
+
 	FQuat OutRotationQuat;
-	if (!GetRotationObservationAsQuat(OutRotationQuat, Object, Element, FQuat::MakeFromRotator(RelativeRotation), Tag))
+	if (!GetRotationObservationAsQuat(OutRotationQuat, Object, Element, FQuat::MakeFromRotator(RelativeRotation), bActive, Tag))
 	{
 		OutRotation = FRotator::ZeroRotator;
 		return false;
@@ -4480,10 +5227,15 @@ bool ULearningAgentsObservations::GetRotationObservation(FRotator& OutRotation,
 	return true;
 }
 
-bool ULearningAgentsObservations::GetRotationObservationAsQuat(FQuat& OutRotation, const ULearningAgentsObservationObject* Object, const FLearningAgentsObservationObjectElement Element, const FQuat RelativeRotation, const FName Tag)
+bool ULearningAgentsObservations::GetRotationObservationAsQuat(FQuat& OutRotation, const ULearningAgentsObservationObject* Object, const FLearningAgentsObservationObjectElement Element, const FQuat RelativeRotation, const bool bActive, const FName Tag)
 {
+	if (!bActive)
+	{
+		return GetNullObservation(Object, Element, Tag);
+	}
+
 	TStaticArray<float, 6> OutValues;
-	if (!GetContinuousObservationToArrayView(OutValues, Object, Element, Tag))
+	if (!GetContinuousObservationToArrayView(OutValues, Object, Element, bActive, Tag))
 	{
 		OutRotation = FQuat::Identity;
 		return false;
@@ -4504,10 +5256,15 @@ bool ULearningAgentsObservations::GetRotationObservationAsQuat(FQuat& OutRotatio
 	return true;
 }
 
-bool ULearningAgentsObservations::GetScaleObservation(FVector& OutScale, const ULearningAgentsObservationObject* Object, const FLearningAgentsObservationObjectElement Element, const FVector RelativeScale, const FName Tag)
+bool ULearningAgentsObservations::GetScaleObservation(FVector& OutScale, const ULearningAgentsObservationObject* Object, const FLearningAgentsObservationObjectElement Element, const FVector RelativeScale, const bool bActive, const FName Tag)
 {
+	if (!bActive)
+	{
+		return GetNullObservation(Object, Element, Tag);
+	}
+
 	TStaticArray<float, 3> OutValues;
-	if (!GetContinuousObservationToArrayView(OutValues, Object, Element, Tag))
+	if (!GetContinuousObservationToArrayView(OutValues, Object, Element, bActive, Tag))
 	{
 		OutScale = FVector::OneVector;
 		return false;
@@ -4517,11 +5274,16 @@ bool ULearningAgentsObservations::GetScaleObservation(FVector& OutScale, const U
 	return true;
 }
 
-bool ULearningAgentsObservations::GetTransformObservation(FTransform& OutTransform, const ULearningAgentsObservationObject* Object, const FLearningAgentsObservationObjectElement Element, const FTransform RelativeTransform, const FName Tag)
+bool ULearningAgentsObservations::GetTransformObservation(FTransform& OutTransform, const ULearningAgentsObservationObject* Object, const FLearningAgentsObservationObjectElement Element, const FTransform RelativeTransform, const bool bActive, const FName Tag)
 {
+	if (!bActive)
+	{
+		return GetNullObservation(Object, Element, Tag);
+	}
+
 	TStaticArray<FName, 3> OutElementNames;
 	TStaticArray<FLearningAgentsObservationObjectElement, 3> OutElements;
-	if (!GetStructObservationToArrayViews(OutElementNames, OutElements, Object, Element, Tag))
+	if (!GetStructObservationToArrayViews(OutElementNames, OutElements, Object, Element, bActive, Tag))
 	{
 		OutTransform = FTransform::Identity;
 		return false;
@@ -4555,10 +5317,15 @@ bool ULearningAgentsObservations::GetTransformObservation(FTransform& OutTransfo
 	return true;
 }
 
-bool ULearningAgentsObservations::GetAngleObservationRadians(float& OutAngle, const ULearningAgentsObservationObject* Object, const FLearningAgentsObservationObjectElement Element, const float RelativeAngle, const FName Tag)
+bool ULearningAgentsObservations::GetAngleObservationRadians(float& OutAngle, const ULearningAgentsObservationObject* Object, const FLearningAgentsObservationObjectElement Element, const float RelativeAngle, const bool bActive, const FName Tag)
 {
+	if (!bActive)
+	{
+		return GetNullObservation(Object, Element, Tag);
+	}
+
 	TStaticArray<float, 2> OutValues;
-	if (!GetContinuousObservationToArrayView(OutValues, Object, Element, Tag))
+	if (!GetContinuousObservationToArrayView(OutValues, Object, Element, bActive, Tag))
 	{
 		OutAngle = 0.0f;
 		return false;
@@ -4569,9 +5336,14 @@ bool ULearningAgentsObservations::GetAngleObservationRadians(float& OutAngle, co
 }
 
 
-bool ULearningAgentsObservations::GetAngleObservation(float& OutAngle, const ULearningAgentsObservationObject* Object, const FLearningAgentsObservationObjectElement Element, const float RelativeAngle, const FName Tag)
+bool ULearningAgentsObservations::GetAngleObservation(float& OutAngle, const ULearningAgentsObservationObject* Object, const FLearningAgentsObservationObjectElement Element, const float RelativeAngle, const bool bActive, const FName Tag)
 {
-	if (!GetAngleObservationRadians(OutAngle, Object, Element, FMath::DegreesToRadians(RelativeAngle), Tag))
+	if (!bActive)
+	{
+		return GetNullObservation(Object, Element, Tag);
+	}
+
+	if (!GetAngleObservationRadians(OutAngle, Object, Element, FMath::DegreesToRadians(RelativeAngle), bActive, Tag))
 	{
 		return false;
 	}
@@ -4580,10 +5352,15 @@ bool ULearningAgentsObservations::GetAngleObservation(float& OutAngle, const ULe
 	return true;
 }
 
-bool ULearningAgentsObservations::GetVelocityObservation(FVector& OutVelocity, const ULearningAgentsObservationObject* Object, const FLearningAgentsObservationObjectElement Element, const FTransform RelativeTransform, const FName Tag)
+bool ULearningAgentsObservations::GetVelocityObservation(FVector& OutVelocity, const ULearningAgentsObservationObject* Object, const FLearningAgentsObservationObjectElement Element, const FTransform RelativeTransform, const bool bActive, const FName Tag)
 {
+	if (!bActive)
+	{
+		return GetNullObservation(Object, Element, Tag);
+	}
+
 	TStaticArray<float, 3> OutValues;
-	if (!GetContinuousObservationToArrayView(OutValues, Object, Element, Tag))
+	if (!GetContinuousObservationToArrayView(OutValues, Object, Element, bActive, Tag))
 	{
 		OutVelocity = FVector::ZeroVector;
 		return false;
@@ -4593,10 +5370,15 @@ bool ULearningAgentsObservations::GetVelocityObservation(FVector& OutVelocity, c
 	return true;
 }
 
-bool ULearningAgentsObservations::GetDirectionObservation(FVector& OutDirection, const ULearningAgentsObservationObject* Object, const FLearningAgentsObservationObjectElement Element, const FTransform RelativeTransform, const FName Tag)
+bool ULearningAgentsObservations::GetDirectionObservation(FVector& OutDirection, const ULearningAgentsObservationObject* Object, const FLearningAgentsObservationObjectElement Element, const FTransform RelativeTransform, const bool bActive, const FName Tag)
 {
+	if (!bActive)
+	{
+		return GetNullObservation(Object, Element, Tag);
+	}
+
 	TStaticArray<float, 3> OutValues;
-	if (!GetContinuousObservationToArrayView(OutValues, Object, Element, Tag))
+	if (!GetContinuousObservationToArrayView(OutValues, Object, Element, bActive, Tag))
 	{
 		OutDirection = FVector::ForwardVector;
 		return false;
@@ -4606,16 +5388,26 @@ bool ULearningAgentsObservations::GetDirectionObservation(FVector& OutDirection,
 	return true;
 }
 
-bool ULearningAgentsObservations::GetLocationAlongSplineObservation(FVector& OutLocation, const ULearningAgentsObservationObject* Object, const FLearningAgentsObservationObjectElement Element, const FTransform RelativeTransform, const FName Tag)
+bool ULearningAgentsObservations::GetLocationAlongSplineObservation(FVector& OutLocation, const ULearningAgentsObservationObject* Object, const FLearningAgentsObservationObjectElement Element, const FTransform RelativeTransform, const bool bActive, const FName Tag)
 {
-	return GetLocationObservation(OutLocation, Object, Element, RelativeTransform, Tag);
+	if (!bActive)
+	{
+		return GetNullObservation(Object, Element, Tag);
+	}
+
+	return GetLocationObservation(OutLocation, Object, Element, RelativeTransform, bActive, Tag);
 }
 
-bool ULearningAgentsObservations::GetProportionAlongSplineObservation(bool& bOutIsClosedLoop, float& OutAngle, float& OutPropotion, const ULearningAgentsObservationObject* Object, const FLearningAgentsObservationObjectElement Element, const FName Tag)
+bool ULearningAgentsObservations::GetProportionAlongSplineObservation(bool& bOutIsClosedLoop, float& OutAngle, float& OutPropotion, const ULearningAgentsObservationObject* Object, const FLearningAgentsObservationObjectElement Element, const bool bActive, const FName Tag)
 {
+	if (!bActive)
+	{
+		return GetNullObservation(Object, Element, Tag);
+	}
+
 	FName SubName;
 	FLearningAgentsObservationObjectElement SubElement;
-	if (!GetExclusiveUnionObservation(SubName, SubElement, Object, Element, Tag))
+	if (!GetExclusiveUnionObservation(SubName, SubElement, Object, Element, bActive, Tag))
 	{
 		bOutIsClosedLoop = false;
 		OutAngle = 0.0f;
@@ -4637,14 +5429,24 @@ bool ULearningAgentsObservations::GetProportionAlongSplineObservation(bool& bOut
 	}
 }
 
-bool ULearningAgentsObservations::GetDirectionAlongSplineObservation(FVector& OutDirection, const ULearningAgentsObservationObject* Object, const FLearningAgentsObservationObjectElement Element, const FTransform RelativeTransform, const FName Tag)
+bool ULearningAgentsObservations::GetDirectionAlongSplineObservation(FVector& OutDirection, const ULearningAgentsObservationObject* Object, const FLearningAgentsObservationObjectElement Element, const FTransform RelativeTransform, const bool bActive, const FName Tag)
 {
-	return GetDirectionObservation(OutDirection, Object, Element, RelativeTransform, Tag);
+	if (!bActive)
+	{
+		return GetNullObservation(Object, Element, Tag);
+	}
+
+	return GetDirectionObservation(OutDirection, Object, Element, RelativeTransform, bActive, Tag);
 }
 
-bool ULearningAgentsObservations::GetProportionAlongRayObservation(float& OutProportion, const ULearningAgentsObservationObject* Object, const FLearningAgentsObservationObjectElement Element, const FName Tag)
+bool ULearningAgentsObservations::GetProportionAlongRayObservation(float& OutProportion, const ULearningAgentsObservationObject* Object, const FLearningAgentsObservationObjectElement Element, const bool bActive, const FName Tag)
 {
-	if (!GetFloatObservation(OutProportion, Object, Element, Tag))
+	if (!bActive)
+	{
+		return GetNullObservation(Object, Element, Tag);
+	}
+
+	if (!GetFloatObservation(OutProportion, Object, Element, bActive, Tag))
 	{
 		OutProportion = 0.0f;
 		return false;
@@ -4654,8 +5456,13 @@ bool ULearningAgentsObservations::GetProportionAlongRayObservation(float& OutPro
 	return true;
 }
 
-bool ULearningAgentsObservations::GetConv1dObservation(FLearningAgentsObservationObjectElement& OutElement, const ULearningAgentsObservationObject* Object, const FLearningAgentsObservationObjectElement Element, const FName Tag)
+bool ULearningAgentsObservations::GetConv1dObservation(FLearningAgentsObservationObjectElement& OutElement, const ULearningAgentsObservationObject* Object, const FLearningAgentsObservationObjectElement Element, const bool bActive, const FName Tag)
 {
+	if (!bActive)
+	{
+		return GetNullObservation(Object, Element, Tag);
+	}
+
 	if (!Object)
 	{
 		UE_LOG(LogLearning, Error, TEXT("GetConv1dObservation: Object is nullptr."));
@@ -4690,8 +5497,13 @@ bool ULearningAgentsObservations::GetConv1dObservation(FLearningAgentsObservatio
 	return true;
 }
 
-bool ULearningAgentsObservations::GetConv2dObservation(FLearningAgentsObservationObjectElement& OutElement, const ULearningAgentsObservationObject* Object, const FLearningAgentsObservationObjectElement Element, const FName Tag)
+bool ULearningAgentsObservations::GetConv2dObservation(FLearningAgentsObservationObjectElement& OutElement, const ULearningAgentsObservationObject* Object, const FLearningAgentsObservationObjectElement Element, const bool bActive, const FName Tag)
 {
+	if (!bActive)
+	{
+		return GetNullObservation(Object, Element, Tag);
+	}
+
 	if (!Object)
 	{
 		UE_LOG(LogLearning, Error, TEXT("GetConv2dObservation: Object is nullptr."));
