READY TO SUBMIT

diff --git a/Engine/Plugins/Experimental/LearningAgents/Source/LearningAgents/Public/LearningAgentsActions.h b/Engine/Plugins/Experimental/LearningAgents/Source/LearningAgents/Public/LearningAgentsActions.h
--- a/Engine/Plugins/Experimental/LearningAgents/Source/LearningAgents/Public/LearningAgentsActions.h
+++ b/Engine/Plugins/Experimental/LearningAgents/Source/LearningAgents/Public/LearningAgentsActions.h
@@ -214,10 +214,11 @@ public:
 	 * @param Size The number of float values in the action.
 	 * @param Scale The scale used for this action.
 	 * @param Tag The tag of this new action. Used during action object validation and debugging.
+	 * @param bActive Whether to enable the specified action as a valid network output. If set to inactive, this action will be considered an empty action. This flag allows easy toggles of specified actions.
 	 * @return The newly created action schema element.
 	 */
 	UFUNCTION(BlueprintPure, Category = "LearningAgents", meta = (AdvancedDisplay = 3))
-	static UE_API FLearningAgentsActionSchemaElement SpecifyContinuousAction(ULearningAgentsActionSchema* Schema, const int32 Size, const float Scale = 1.0f, const FName Tag = TEXT("ContinuousAction"));
+	static UE_API FLearningAgentsActionSchemaElement SpecifyContinuousAction(ULearningAgentsActionSchema* Schema, const int32 Size, const float Scale = 1.0f, const bool bActive = true, const FName Tag = TEXT("ContinuousAction"));
 
 	/**
 	 * Specifies a new exclusive discrete action. This represents an action which is an exclusive choice from a number of discrete options, sampled 
@@ -227,10 +228,11 @@ public:
 	 * @param Size The number of discrete options in the action.
 	 * @param PriorProbabilities The prior probabilities of each option. Can be left empty to use a uniform distribution over options. Should sum to one.
 	 * @param Tag The tag of this new action. Used during action object validation and debugging.
+	 * @param bActive Whether to enable the specified action as a valid network output. If set to inactive, this action will be considered an empty action. This flag allows easy toggles of specified actions.
 	 * @return The newly created action schema element.
 	 */
 	UFUNCTION(BlueprintPure, Category = "LearningAgents", meta = (AdvancedDisplay = 2, AutoCreateRefTerm = "PriorProbabilities"))
-	static UE_API FLearningAgentsActionSchemaElement SpecifyExclusiveDiscreteAction(ULearningAgentsActionSchema* Schema, const int32 Size, const TArray<float>& PriorProbabilities, const FName Tag = TEXT("DiscreteExclusiveAction"));
+	static UE_API FLearningAgentsActionSchemaElement SpecifyExclusiveDiscreteAction(ULearningAgentsActionSchema* Schema, const int32 Size, const TArray<float>& PriorProbabilities, const bool bActive = true, const FName Tag = TEXT("DiscreteExclusiveAction"));
 	
 	/**
 	 * Specifies a new exclusive discrete action. This represents an action which is an exclusive choice from a number of discrete options, sampled
@@ -239,10 +241,11 @@ public:
 	 * @param Schema The Action Schema
 	 * @param Size The number of discrete options in the action.
 	 * @param PriorProbabilities The prior probabilities of each option. Can be left empty to use a uniform distribution over options. Should sum to one.
+	 * @param bActive Whether to enable the specified action as a valid network output. If set to inactive, this action will be considered an empty action. This flag allows easy toggles of specified actions.
 	 * @param Tag The tag of this new action. Used during action object validation and debugging.
 	 * @return The newly created action schema element.
 	 */
-	static UE_API FLearningAgentsActionSchemaElement SpecifyExclusiveDiscreteActionFromArrayView(ULearningAgentsActionSchema* Schema, const int32 Size, const TArrayView<const float> PriorProbabilities = {}, const FName Tag = TEXT("DiscreteExclusiveAction"));
+	static UE_API FLearningAgentsActionSchemaElement SpecifyExclusiveDiscreteActionFromArrayView(ULearningAgentsActionSchema* Schema, const int32 Size, const TArrayView<const float> PriorProbabilities = {}, const bool bActive = true, const FName Tag = TEXT("DiscreteExclusiveAction"));
 
 	/**
 	 * Specifies a new named exclusive discrete action. This represents an action which is an exclusive choice from a number of discrete options, sampled
@@ -251,11 +254,12 @@ public:
 	 * @param Schema The Action Schema
 	 * @param Names The names of the discrete options in the action.
 	 * @param PriorProbabilities The prior probabilities of each option. Can be left empty to use a uniform distribution over options. Should sum to one.
+	 * @param bActive Whether to enable the specified action as a valid network output. If set to inactive, this action will be considered an empty action. This flag allows easy toggles of specified actions.
 	 * @param Tag The tag of this new action. Used during action object validation and debugging.
 	 * @return The newly created action schema element.
 	 */
 	UFUNCTION(BlueprintPure, Category = "LearningAgents", meta = (AdvancedDisplay = 2, AutoCreateRefTerm = "PriorProbabilities"))
-	static UE_API FLearningAgentsActionSchemaElement SpecifyNamedExclusiveDiscreteAction(ULearningAgentsActionSchema* Schema, const TArray<FName>& Names, const TMap<FName, float>& PriorProbabilities, const FName Tag = TEXT("NamedDiscreteExclusiveAction"));
+	static UE_API FLearningAgentsActionSchemaElement SpecifyNamedExclusiveDiscreteAction(ULearningAgentsActionSchema* Schema, const TArray<FName>& Names, const TMap<FName, float>& PriorProbabilities, const bool bActive = true, const FName Tag = TEXT("NamedDiscreteExclusiveAction"));
 
 	/**
 	 * Specifies a new named exclusive discrete action. This represents an action which is an exclusive choice from a number of discrete options, sampled
@@ -265,10 +269,11 @@ public:
 	 * @param Names The names of the discrete options in the action.
 	 * @param PriorProbabilityNames The names of the prior probabilities. Can be left empty to use a uniform distribution over options.
 	 * @param PriorProbabilityValues The prior probabilities of each option. Can be left empty to use a uniform distribution over options. Should sum to one.
+	 * @param bActive Whether to enable the specified action as a valid network output. If set to inactive, this action will be considered an empty action. This flag allows easy toggles of specified actions.
 	 * @param Tag The tag of this new action. Used during action object validation and debugging.
 	 * @return The newly created action schema element.
 	 */
-	static UE_API FLearningAgentsActionSchemaElement SpecifyNamedExclusiveDiscreteActionFromArrayViews(ULearningAgentsActionSchema* Schema, const TArrayView<const FName> Names, const TArrayView<const float> PriorProbabilities, const FName Tag = TEXT("NamedDiscreteExclusiveAction"));
+	static UE_API FLearningAgentsActionSchemaElement SpecifyNamedExclusiveDiscreteActionFromArrayViews(ULearningAgentsActionSchema* Schema, const TArrayView<const FName> Names, const TArrayView<const float> PriorProbabilities, const bool bActive = true, const FName Tag = TEXT("NamedDiscreteExclusiveAction"));
 
 	/**
 	 * Specifies a new inclusive discrete action. This represents an action which is an inclusive choice from a number of discrete options, sampled
@@ -277,11 +282,12 @@ public:
 	 * @param Schema The Action Schema
 	 * @param Size The number of discrete options in the action.
 	 * @param PriorProbabilities The prior probabilities of each option. Can be left empty to use a probability of 0.5 for each option.
+	 * @param bActive Whether to enable the specified action as a valid network output. If set to inactive, this action will be considered an empty action. This flag allows easy toggles of specified actions.
 	 * @param Tag The tag of this new action. Used during action object validation and debugging.
 	 * @return The newly created action schema element.
 	 */
 	UFUNCTION(BlueprintPure, Category = "LearningAgents", meta = (AdvancedDisplay = 2, AutoCreateRefTerm = "PriorProbabilities"))
-	static UE_API FLearningAgentsActionSchemaElement SpecifyInclusiveDiscreteAction(ULearningAgentsActionSchema* Schema, const int32 Size, const TArray<float>& PriorProbabilities, const FName Tag = TEXT("DiscreteInclusiveAction"));
+	static UE_API FLearningAgentsActionSchemaElement SpecifyInclusiveDiscreteAction(ULearningAgentsActionSchema* Schema, const int32 Size, const TArray<float>& PriorProbabilities, const bool bActive = true, const FName Tag = TEXT("DiscreteInclusiveAction"));
 	
 	/**
 	 * Specifies a new inclusive discrete action. This represents an action which is an inclusive choice from a number of discrete options, sampled
@@ -290,10 +296,11 @@ public:
 	 * @param Schema The Action Schema
 	 * @param Size The number of discrete options in the action.
 	 * @param PriorProbabilities The prior probabilities of each option. Can be left empty to use a probability of 0.5 for each option.
+	 * @param bActive Whether to enable the specified action as a valid network output. If set to inactive, this action will be considered an empty action. This flag allows easy toggles of specified actions.
 	 * @param Tag The tag of this new action. Used during action object validation and debugging.
 	 * @return The newly created action schema element.
 	 */
-	static UE_API FLearningAgentsActionSchemaElement SpecifyInclusiveDiscreteActionFromArrayView(ULearningAgentsActionSchema* Schema, const int32 Size, const TArrayView<const float> PriorProbabilities = {}, const FName Tag = TEXT("DiscreteInclusiveAction"));
+	static UE_API FLearningAgentsActionSchemaElement SpecifyInclusiveDiscreteActionFromArrayView(ULearningAgentsActionSchema* Schema, const int32 Size, const TArrayView<const float> PriorProbabilities = {}, const bool bActive = true, const FName Tag = TEXT("DiscreteInclusiveAction"));
 
 	/**
 	 * Specifies a new named inclusive discrete action. This represents an action which is an inclusive choice from a number of discrete options, sampled
@@ -302,11 +309,12 @@ public:
 	 * @param Schema The Action Schema
 	 * @param Names The names of the discrete options in the action.
 	 * @param PriorProbabilities The prior probabilities of each option. Can be left empty to use a probability of 0.5 for each option.
+	 * @param bActive Whether to enable the specified action as a valid network output. If set to inactive, this action will be considered an empty action. This flag allows easy toggles of specified actions.
 	 * @param Tag The tag of this new action. Used during action object validation and debugging.
 	 * @return The newly created action schema element.
 	 */
 	UFUNCTION(BlueprintPure, Category = "LearningAgents", meta = (AdvancedDisplay = 2, AutoCreateRefTerm = "PriorProbabilities"))
-	static UE_API FLearningAgentsActionSchemaElement SpecifyNamedInclusiveDiscreteAction(ULearningAgentsActionSchema* Schema, const TArray<FName> Names, const TMap<FName, float>& PriorProbabilities, const FName Tag = TEXT("NamedDiscreteInclusiveAction"));
+	static UE_API FLearningAgentsActionSchemaElement SpecifyNamedInclusiveDiscreteAction(ULearningAgentsActionSchema* Schema, const TArray<FName> Names, const TMap<FName, float>& PriorProbabilities, const bool bActive = true, const FName Tag = TEXT("NamedDiscreteInclusiveAction"));
 
 	/**
 	 * Specifies a new named inclusive discrete action. This represents an action which is an inclusive choice from a number of discrete options, sampled
@@ -315,21 +323,23 @@ public:
 	 * @param Schema The Action Schema
 	 * @param Names The names of the discrete options in the action.
 	 * @param PriorProbabilities The prior probabilities of each option. Can be left empty to use a probability of 0.5 for each option.
+	 * @param bActive Whether to enable the specified action as a valid network output. If set to inactive, this action will be considered an empty action. This flag allows easy toggles of specified actions.
 	 * @param Tag The tag of this new action. Used during action object validation and debugging.
 	 * @return The newly created action schema element.
 	 */
-	static UE_API FLearningAgentsActionSchemaElement SpecifyNamedInclusiveDiscreteActionFromArrayViews(ULearningAgentsActionSchema* Schema, const TArrayView<const FName> Names, const TArrayView<const float> PriorProbabilities = {}, const FName Tag = TEXT("NamedDiscreteInclusiveAction"));
+	static UE_API FLearningAgentsActionSchemaElement SpecifyNamedInclusiveDiscreteActionFromArrayViews(ULearningAgentsActionSchema* Schema, const TArrayView<const FName> Names, const TArrayView<const float> PriorProbabilities = {}, const bool bActive = true, const FName Tag = TEXT("NamedDiscreteInclusiveAction"));
 
 	/**
 	 * Specifies a new struct action. This represents an action which is made up of a number of named sub-actions.
 	 *
 	 * @param Schema The Action Schema
 	 * @param Elements The sub-actions.
+	 * @param bActive Whether to enable the specified action as a valid network output. If set to inactive, this action will be considered an empty action. This flag allows easy toggles of specified actions.
 	 * @param Tag The tag of this new action. Used during action object validation and debugging.
 	 * @return The newly created action schema element.
 	 */
 	UFUNCTION(BlueprintPure, Category = "LearningAgents", meta = (AdvancedDisplay = 2))
-	static UE_API FLearningAgentsActionSchemaElement SpecifyStructAction(ULearningAgentsActionSchema* Schema, const TMap<FName, FLearningAgentsActionSchemaElement>& Elements, const FName Tag = TEXT("StructAction"));
+	static UE_API FLearningAgentsActionSchemaElement SpecifyStructAction(ULearningAgentsActionSchema* Schema, const TMap<FName, FLearningAgentsActionSchemaElement>& Elements, const bool bActive = true, const FName Tag = TEXT("StructAction"));
 
 	/**
 	 * Specifies a new struct action. This represents an action which is made up of a number of named sub-actions.
@@ -337,11 +347,12 @@ public:
 	 * @param Schema The Action Schema
 	 * @param ElementNames The names of the sub-actions.
 	 * @param Elements The corresponding sub-actions. Must be the same size as ElementNames.
+	 * @param bActive Whether to enable the specified action as a valid network output. If set to inactive, this action will be considered an empty action. This flag allows easy toggles of specified actions.
 	 * @param Tag The tag of this new action. Used during action object validation and debugging.
 	 * @return The newly created action schema element.
 	 */
 	UFUNCTION(BlueprintPure, Category = "LearningAgents", meta = (AdvancedDisplay = 3))
-	static UE_API FLearningAgentsActionSchemaElement SpecifyStructActionFromArrays(ULearningAgentsActionSchema* Schema, const TArray<FName>& ElementNames, const TArray<FLearningAgentsActionSchemaElement>& Elements, const FName Tag = TEXT("StructAction"));
+	static UE_API FLearningAgentsActionSchemaElement SpecifyStructActionFromArrays(ULearningAgentsActionSchema* Schema, const TArray<FName>& ElementNames, const TArray<FLearningAgentsActionSchemaElement>& Elements, const bool bActive = true, const FName Tag = TEXT("StructAction"));
 	
 	/**
 	 * Specifies a new struct action. This represents an action which is made up of a number of named sub-actions.
@@ -349,10 +360,11 @@ public:
 	 * @param Schema The Action Schema
 	 * @param ElementNames The names of the sub-actions.
 	 * @param Elements The corresponding sub-actions. Must be the same size as ElementNames.
+	 * @param bActive Whether to enable the specified action as a valid network output. If set to inactive, this action will be considered an empty action. This flag allows easy toggles of specified actions.
 	 * @param Tag The tag of this new action. Used during action object validation and debugging.
 	 * @return The newly created action schema element.
 	 */
-	static UE_API FLearningAgentsActionSchemaElement SpecifyStructActionFromArrayViews(ULearningAgentsActionSchema* Schema, const TArrayView<const FName> ElementNames, const TArrayView<const FLearningAgentsActionSchemaElement> Elements, const FName Tag = TEXT("StructAction"));
+	static UE_API FLearningAgentsActionSchemaElement SpecifyStructActionFromArrayViews(ULearningAgentsActionSchema* Schema, const TArrayView<const FName> ElementNames, const TArrayView<const FLearningAgentsActionSchemaElement> Elements, const bool bActive = true, const FName Tag = TEXT("StructAction"));
 
 	/**
 	 * Specifies a new exclusive union action. This represents an action which is an exclusive choice from a number of named sub-actions, sampled
@@ -361,11 +373,12 @@ public:
 	 * @param Schema The Action Schema
 	 * @param Elements The sub-actions.
 	 * @param PriorProbabilities The prior probabilities of each option. Can be left empty to use a uniform distribution over options. Should sum to one.
+	 * @param bActive Whether to enable the specified action as a valid network output. If set to inactive, this action will be considered an empty action. This flag allows easy toggles of specified actions.
 	 * @param Tag The tag of this new action. Used during action object validation and debugging.
 	 * @return The newly created action schema element.
 	 */
 	UFUNCTION(BlueprintPure, Category = "LearningAgents", meta = (AdvancedDisplay = 2, AutoCreateRefTerm = "PriorProbabilities"))
-	static UE_API FLearningAgentsActionSchemaElement SpecifyExclusiveUnionAction(ULearningAgentsActionSchema* Schema, const TMap<FName, FLearningAgentsActionSchemaElement>& Elements, const TMap<FName, float>& PriorProbabilities, const FName Tag = TEXT("ExclusiveUnionAction"));
+	static UE_API FLearningAgentsActionSchemaElement SpecifyExclusiveUnionAction(ULearningAgentsActionSchema* Schema, const TMap<FName, FLearningAgentsActionSchemaElement>& Elements, const TMap<FName, float>& PriorProbabilities, const bool bActive = true, const FName Tag = TEXT("ExclusiveUnionAction"));
 
 	/**
 	 * Specifies a new exclusive union action. This represents an action which is an exclusive choice from a number of named sub-actions, sampled
@@ -375,11 +388,12 @@ public:
 	 * @param ElementNames The names of the sub-actions.
 	 * @param Elements The corresponding sub-actions. Must be the same size as ElementNames.
 	 * @param PriorProbabilities The prior probabilities of each option. Can be left empty to use a uniform distribution over options. Should sum to one.
+	 * @param bActive Whether to enable the specified action as a valid network output. If set to inactive, this action will be considered an empty action. This flag allows easy toggles of specified actions.
 	 * @param Tag The tag of this new action. Used during action object validation and debugging.
 	 * @return The newly created action schema element.
 	 */
 	UFUNCTION(BlueprintPure, Category = "LearningAgents", meta = (AdvancedDisplay = 3, AutoCreateRefTerm = "PriorProbabilities"))
-	static UE_API FLearningAgentsActionSchemaElement SpecifyExclusiveUnionActionFromArrays(ULearningAgentsActionSchema* Schema, const TArray<FName>& ElementNames, const TArray<FLearningAgentsActionSchemaElement>& Elements, const TArray<float>& PriorProbabilities, const FName Tag = TEXT("ExclusiveUnionAction"));
+	static UE_API FLearningAgentsActionSchemaElement SpecifyExclusiveUnionActionFromArrays(ULearningAgentsActionSchema* Schema, const TArray<FName>& ElementNames, const TArray<FLearningAgentsActionSchemaElement>& Elements, const TArray<float>& PriorProbabilities, const bool bActive = true, const FName Tag = TEXT("ExclusiveUnionAction"));
 	
 	/**
 	 * Specifies a new exclusive union action. This represents an action which is an exclusive choice from a number of named sub-actions, sampled
@@ -389,10 +403,11 @@ public:
 	 * @param ElementNames The names of the sub-actions.
 	 * @param Elements The corresponding sub-actions. Must be the same size as ElementNames.
 	 * @param PriorProbabilities The prior probabilities of each option. Can be left empty to use a uniform distribution over options. Should sum to one.
+	 * @param bActive Whether to enable the specified action as a valid network output. If set to inactive, this action will be considered an empty action. This flag allows easy toggles of specified actions.
 	 * @param Tag The tag of this new action. Used during action object validation and debugging.
 	 * @return The newly created action schema element.
 	 */
-	static UE_API FLearningAgentsActionSchemaElement SpecifyExclusiveUnionActionFromArrayViews(ULearningAgentsActionSchema* Schema, const TArrayView<const FName> ElementNames, const TArrayView<const FLearningAgentsActionSchemaElement> Elements, const TArrayView<const float> PriorProbabilities = {}, const FName Tag = TEXT("ExclusiveUnionAction"));
+	static UE_API FLearningAgentsActionSchemaElement SpecifyExclusiveUnionActionFromArrayViews(ULearningAgentsActionSchema* Schema, const TArrayView<const FName> ElementNames, const TArrayView<const FLearningAgentsActionSchemaElement> Elements, const TArrayView<const float> PriorProbabilities = {}, const bool bActive = true, const FName Tag = TEXT("ExclusiveUnionAction"));
 
 	/**
 	 * Specifies a new inclusive union action. This represents an action which is an inclusive choice from a number of named sub-actions, sampled
@@ -401,11 +416,12 @@ public:
 	 * @param Schema The Action Schema
 	 * @param Elements The sub-actions.
 	 * @param PriorProbabilities The prior probabilities of each option. Can be left empty to use a probability of 0.5 for each option.
+	 * @param bActive Whether to enable the specified action as a valid network output. If set to inactive, this action will be considered an empty action. This flag allows easy toggles of specified actions.
 	 * @param Tag The tag of this new action. Used during action object validation and debugging.
 	 * @return The newly created action schema element.
 	 */
 	UFUNCTION(BlueprintPure, Category = "LearningAgents", meta = (AdvancedDisplay = 2, AutoCreateRefTerm = "PriorProbabilities"))
-	static UE_API FLearningAgentsActionSchemaElement SpecifyInclusiveUnionAction(ULearningAgentsActionSchema* Schema, const TMap<FName, FLearningAgentsActionSchemaElement>& Elements, const TMap<FName, float>& PriorProbabilities, const FName Tag = TEXT("InclusiveUnionAction"));
+	static UE_API FLearningAgentsActionSchemaElement SpecifyInclusiveUnionAction(ULearningAgentsActionSchema* Schema, const TMap<FName, FLearningAgentsActionSchemaElement>& Elements, const TMap<FName, float>& PriorProbabilities, const bool bActive = true, const FName Tag = TEXT("InclusiveUnionAction"));
 
 	/**
 	 * Specifies a new inclusive union action. This represents an action which is an inclusive choice from a number of named sub-actions, sampled
@@ -415,11 +431,12 @@ public:
 	 * @param ElementNames The names of the sub-actions.
 	 * @param Elements The corresponding sub-actions. Must be the same size as ElementNames.
 	 * @param PriorProbabilities The prior probabilities of each option. Can be left empty to use a probability of 0.5 for each option.
+	 * @param bActive Whether to enable the specified action as a valid network output. If set to inactive, this action will be considered an empty action. This flag allows easy toggles of specified actions.
 	 * @param Tag The tag of this new action. Used during action object validation and debugging.
 	 * @return The newly created action schema element.
 	 */
 	UFUNCTION(BlueprintPure, Category = "LearningAgents", meta = (AdvancedDisplay = 3, AutoCreateRefTerm = "PriorProbabilities"))
-	static UE_API FLearningAgentsActionSchemaElement SpecifyInclusiveUnionActionFromArrays(ULearningAgentsActionSchema* Schema, const TArray<FName> ElementNames, const TArray<FLearningAgentsActionSchemaElement>& Elements, const TArray<float>& PriorProbabilities, const FName Tag = TEXT("InclusiveUnionAction"));
+	static UE_API FLearningAgentsActionSchemaElement SpecifyInclusiveUnionActionFromArrays(ULearningAgentsActionSchema* Schema, const TArray<FName> ElementNames, const TArray<FLearningAgentsActionSchemaElement>& Elements, const TArray<float>& PriorProbabilities, const bool bActive = true, const FName Tag = TEXT("InclusiveUnionAction"));
 	
 	/**
 	 * Specifies a new inclusive union action. This represents an action which is an inclusive choice from a number of named sub-actions, sampled
@@ -429,10 +446,11 @@ public:
 	 * @param ElementNames The names of the sub-actions.
 	 * @param Elements The corresponding sub-actions. Must be the same size as ElementNames.
 	 * @param PriorProbabilities The prior probabilities of each option. Can be left empty to use a probability of 0.5 for each option.
+	 * @param bActive Whether to enable the specified action as a valid network output. If set to inactive, this action will be considered an empty action. This flag allows easy toggles of specified actions.
 	 * @param Tag The tag of this new action. Used during action object validation and debugging.
 	 * @return The newly created action schema element.
 	 */
-	static UE_API FLearningAgentsActionSchemaElement SpecifyInclusiveUnionActionFromArrayViews(ULearningAgentsActionSchema* Schema, const TArrayView<const FName> ElementNames, const TArrayView<const FLearningAgentsActionSchemaElement> Elements, const TArrayView<const float> PriorProbabilities = {}, const FName Tag = TEXT("InclusiveUnionAction"));
+	static UE_API FLearningAgentsActionSchemaElement SpecifyInclusiveUnionActionFromArrayViews(ULearningAgentsActionSchema* Schema, const TArrayView<const FName> ElementNames, const TArrayView<const FLearningAgentsActionSchemaElement> Elements, const TArrayView<const float> PriorProbabilities = {}, const bool bActive = true, const FName Tag = TEXT("InclusiveUnionAction"));
 
 	/**
 	 * Specifies a new static array action. This represents an action which is a fixed sized array of some sub-action.
@@ -440,11 +458,12 @@ public:
 	 * @param Schema The Action Schema
 	 * @param Element The sub-action.
 	 * @param Num The number of elements in the array.
+	 * @param bActive Whether to enable the specified action as a valid network output. If set to inactive, this action will be considered an empty action. This flag allows easy toggles of specified actions.
 	 * @param Tag The tag of this new action. Used during action object validation and debugging.
 	 * @return The newly created action schema element.
 	 */
 	UFUNCTION(BlueprintPure, Category = "LearningAgents", meta = (AdvancedDisplay = 3))
-	static UE_API FLearningAgentsActionSchemaElement SpecifyStaticArrayAction(ULearningAgentsActionSchema* Schema, const FLearningAgentsActionSchemaElement Element, const int32 Num, const FName Tag = TEXT("StaticArrayAction"));
+	static UE_API FLearningAgentsActionSchemaElement SpecifyStaticArrayAction(ULearningAgentsActionSchema* Schema, const FLearningAgentsActionSchemaElement Element, const int32 Num, const bool bActive = true, const FName Tag = TEXT("StaticArrayAction"));
 
 	/**
 	 * Specifies a new pair action. This represents an action which is made up of a key and value sub-actions.
@@ -452,11 +471,12 @@ public:
 	 * @param Schema The Action Schema
 	 * @param Key The key sub-action.
 	 * @param Value The value sub-action.
+	 * @param bActive Whether to enable the specified action as a valid network output. If set to inactive, this action will be considered an empty action. This flag allows easy toggles of specified actions.
 	 * @param Tag The tag of this new action. Used during action object validation and debugging.
 	 * @return The newly created action schema element.
 	 */
 	UFUNCTION(BlueprintPure, Category = "LearningAgents", meta = (AdvancedDisplay = 3))
-	static UE_API FLearningAgentsActionSchemaElement SpecifyPairAction(ULearningAgentsActionSchema* Schema, const FLearningAgentsActionSchemaElement Key, const FLearningAgentsActionSchemaElement Value, const FName Tag = TEXT("PairAction"));
+	static UE_API FLearningAgentsActionSchemaElement SpecifyPairAction(ULearningAgentsActionSchema* Schema, const FLearningAgentsActionSchemaElement Key, const FLearningAgentsActionSchemaElement Value, const bool bActive = true, const FName Tag = TEXT("PairAction"));
 
 	/**
 	 * Specifies a new enum action. This represents an action which is an exclusive choice from entries of an Enum, sampled from a Categorical distribution.
@@ -464,11 +484,12 @@ public:
 	 * @param Schema The Action Schema
 	 * @param Enum The Enum type.
 	 * @param PriorProbabilities The prior probabilities of each enum element. Can be left empty to use a uniform distribution over elements. Should sum to one.
+	 * @param bActive Whether to enable the specified action as a valid network output. If set to inactive, this action will be considered an empty action. This flag allows easy toggles of specified actions.
 	 * @param Tag The tag of this new action. Used during action object validation and debugging.
 	 * @return The newly created action schema element.
 	 */
 	UFUNCTION(BlueprintPure, Category = "LearningAgents", meta = (AdvancedDisplay = 2, AutoCreateRefTerm = "PriorProbabilities"))
-	static UE_API FLearningAgentsActionSchemaElement SpecifyEnumAction(ULearningAgentsActionSchema* Schema, const UEnum* Enum, const TMap<uint8, float>& PriorProbabilities, const FName Tag = TEXT("EnumAction"));
+	static UE_API FLearningAgentsActionSchemaElement SpecifyEnumAction(ULearningAgentsActionSchema* Schema, const UEnum* Enum, const TMap<uint8, float>& PriorProbabilities, const bool bActive = true, const FName Tag = TEXT("EnumAction"));
 
 	/**
 	 * Specifies a new enum action. This represents an action which is an exclusive choice from entries of an Enum, sampled from a Categorical distribution.
@@ -476,11 +497,12 @@ public:
 	 * @param Schema The Action Schema
 	 * @param Enum The Enum type.
 	 * @param PriorProbabilities The prior probabilities of each enum element. Can be left empty to use a uniform distribution over elements. Should sum to one.
+	 * @param bActive Whether to enable the specified action as a valid network output. If set to inactive, this action will be considered an empty action. This flag allows easy toggles of specified actions.
 	 * @param Tag The tag of this new action. Used during action object validation and debugging.
 	 * @return The newly created action schema element.
 	 */
 	UFUNCTION(BlueprintPure, Category = "LearningAgents", meta = (AdvancedDisplay = 2, AutoCreateRefTerm = "PriorProbabilities"))
-	static UE_API FLearningAgentsActionSchemaElement SpecifyEnumActionFromArray(ULearningAgentsActionSchema* Schema, const UEnum* Enum, const TArray<float>& PriorProbabilities, const FName Tag = TEXT("EnumAction"));
+	static UE_API FLearningAgentsActionSchemaElement SpecifyEnumActionFromArray(ULearningAgentsActionSchema* Schema, const UEnum* Enum, const TArray<float>& PriorProbabilities, const bool bActive = true, const FName Tag = TEXT("EnumAction"));
 	
 	/**
 	 * Specifies a new enum action. This represents an action which is an exclusive choice from entries of an Enum, sampled from a Categorical distribution.
@@ -488,10 +510,11 @@ public:
 	 * @param Schema The Action Schema
 	 * @param Enum The Enum type.
 	 * @param PriorProbabilities The prior probabilities of each enum element. Can be left empty to use a uniform distribution over elements. Should sum to one.
+	 * @param bActive Whether to enable the specified action as a valid network output. If set to inactive, this action will be considered an empty action. This flag allows easy toggles of specified actions.
 	 * @param Tag The tag of this new action. Used during action object validation and debugging.
 	 * @return The newly created action schema element.
 	 */
-	static UE_API FLearningAgentsActionSchemaElement SpecifyEnumActionFromArrayView(ULearningAgentsActionSchema* Schema, const UEnum* Enum, const TArrayView<const float> PriorProbabilities = {}, const FName Tag = TEXT("EnumAction"));
+	static UE_API FLearningAgentsActionSchemaElement SpecifyEnumActionFromArrayView(ULearningAgentsActionSchema* Schema, const UEnum* Enum, const TArrayView<const float> PriorProbabilities = {}, const bool bActive = true, const FName Tag = TEXT("EnumAction"));
 
 	/**
 	 * Specifies a new bitmask action. This represents an action which is an inclusive choice from entries of an Enum, sampled from a Bernoulli 
@@ -500,11 +523,12 @@ public:
 	 * @param Schema The Action Schema
 	 * @param Enum The Enum type.
 	 * @param PriorProbabilities The prior probabilities of each enum element. Can be left empty to use a probability of 0.5 for each element.
+	 * @param bActive Whether to enable the specified action as a valid network output. If set to inactive, this action will be considered an empty action. This flag allows easy toggles of specified actions.
 	 * @param Tag The tag of this new action. Used during action object validation and debugging.
 	 * @return The newly created action schema element.
 	 */
 	UFUNCTION(BlueprintPure, Category = "LearningAgents", meta = (AdvancedDisplay = 2, AutoCreateRefTerm = "PriorProbabilities"))
-	static UE_API FLearningAgentsActionSchemaElement SpecifyBitmaskAction(ULearningAgentsActionSchema* Schema, const UEnum* Enum, const TMap<uint8, float>& PriorProbabilities, const FName Tag = TEXT("BitmaskAction"));
+	static UE_API FLearningAgentsActionSchemaElement SpecifyBitmaskAction(ULearningAgentsActionSchema* Schema, const UEnum* Enum, const TMap<uint8, float>& PriorProbabilities, const bool bActive = true, const FName Tag = TEXT("BitmaskAction"));
 
 	/**
 	 * Specifies a new bitmask action. This represents an action which is an inclusive choice from entries of an Enum, sampled from a Bernoulli
@@ -513,11 +537,12 @@ public:
 	 * @param Schema The Action Schema
 	 * @param Enum The Enum type.
 	 * @param PriorProbabilities The prior probabilities of each enum element. Can be left empty to use a probability of 0.5 for each element.
+	 * @param bActive Whether to enable the specified action as a valid network output. If set to inactive, this action will be considered an empty action. This flag allows easy toggles of specified actions.
 	 * @param Tag The tag of this new action. Used during action object validation and debugging.
 	 * @return The newly created action schema element.
 	 */
 	UFUNCTION(BlueprintPure, Category = "LearningAgents", meta = (AdvancedDisplay = 2, AutoCreateRefTerm = "PriorProbabilities"))
-	static UE_API FLearningAgentsActionSchemaElement SpecifyBitmaskActionFromArray(ULearningAgentsActionSchema* Schema, const UEnum* Enum, const TArray<float>& PriorProbabilities, const FName Tag = TEXT("BitmaskAction"));
+	static UE_API FLearningAgentsActionSchemaElement SpecifyBitmaskActionFromArray(ULearningAgentsActionSchema* Schema, const UEnum* Enum, const TArray<float>& PriorProbabilities, const bool bActive = true, const FName Tag = TEXT("BitmaskAction"));
 	
 	/**
 	 * Specifies a new bitmask action. This represents an action which is an inclusive choice from entries of an Enum, sampled from a Bernoulli
@@ -526,10 +551,11 @@ public:
 	 * @param Schema The Action Schema
 	 * @param Enum The Enum type.
 	 * @param PriorProbabilities The prior probabilities of each enum element. Can be left empty to use a probability of 0.5 for each element.
+	 * @param bActive Whether to enable the specified action as a valid network output. If set to inactive, this action will be considered an empty action. This flag allows easy toggles of specified actions.
 	 * @param Tag The tag of this new action. Used during action object validation and debugging.
 	 * @return The newly created action schema element.
 	 */
-	static UE_API FLearningAgentsActionSchemaElement SpecifyBitmaskActionFromArrayView(ULearningAgentsActionSchema* Schema, const UEnum* Enum, const TArrayView<const float> PriorProbabilities = {}, const FName Tag = TEXT("BitmaskAction"));
+	static UE_API FLearningAgentsActionSchemaElement SpecifyBitmaskActionFromArrayView(ULearningAgentsActionSchema* Schema, const UEnum* Enum, const TArrayView<const float> PriorProbabilities = {}, const bool bActive = true, const FName Tag = TEXT("BitmaskAction"));
 
 	/**
 	 * Specifies a new optional action. This represents an action which may or may not be generated.
@@ -537,11 +563,12 @@ public:
 	 * @param Schema The Action Schema
 	 * @param Element The sub-action.
 	 * @param PriorProbabilities The prior probability of sampling this action.
+	 * @param bActive Whether to enable the specified action as a valid network output. If set to inactive, this action will be considered an empty action. This flag allows easy toggles of specified actions.
 	 * @param Tag The tag of this new action. Used during action object validation and debugging.
 	 * @return The newly created action schema element.
 	 */
 	UFUNCTION(BlueprintPure, Category = "LearningAgents", meta = (AdvancedDisplay = 2))
-	static UE_API FLearningAgentsActionSchemaElement SpecifyOptionalAction(ULearningAgentsActionSchema* Schema, const FLearningAgentsActionSchemaElement Element, const float PriorProbability = 0.5f, const FName Tag = TEXT("OptionalAction"));
+	static UE_API FLearningAgentsActionSchemaElement SpecifyOptionalAction(ULearningAgentsActionSchema* Schema, const FLearningAgentsActionSchemaElement Element, const float PriorProbability = 0.5f, const bool bActive = true, const FName Tag = TEXT("OptionalAction"));
 
 	/**
 	 * Specifies a new either action. This represents an action which is either action A or action B.
@@ -550,11 +577,12 @@ public:
 	 * @param A The sub-action A.
 	 * @param B The sub-action B.
 	 * @param PriorProbabilityOfA The prior probability of sampling action A over action B.
+	 * @param bActive Whether to enable the specified action as a valid network output. If set to inactive, this action will be considered an empty action. This flag allows easy toggles of specified actions.
 	 * @param Tag The tag of this new action. Used during action object validation and debugging.
 	 * @return The newly created action schema element.
 	 */
 	UFUNCTION(BlueprintPure, Category = "LearningAgents", meta = (AdvancedDisplay = 3))
-	static UE_API FLearningAgentsActionSchemaElement SpecifyEitherAction(ULearningAgentsActionSchema* Schema, const FLearningAgentsActionSchemaElement A, const FLearningAgentsActionSchemaElement B, const float PriorProbabilityOfA = 0.5f, const FName Tag = TEXT("EitherAction"));
+	static UE_API FLearningAgentsActionSchemaElement SpecifyEitherAction(ULearningAgentsActionSchema* Schema, const FLearningAgentsActionSchemaElement A, const FLearningAgentsActionSchemaElement B, const float PriorProbabilityOfA = 0.5f, const bool bActive = true, const FName Tag = TEXT("EitherAction"));
 
 	/**
 	 * Specifies a new encoding action. This represents an action which will be a decoding of another sub-action using a small neural network.
@@ -564,22 +592,24 @@ public:
 	 * @param EncodingSize The encoding size used to decode this sub-action.
 	 * @param HiddenLayerNum The number of hidden layers used to decode this sub-action.
 	 * @param ActivationFunction The activation function used to decode this sub-action.
+	 * @param bActive Whether to enable the specified action as a valid network output. If set to inactive, this action will be considered an empty action. This flag allows easy toggles of specified actions.
 	 * @param Tag The tag of this new action. Used during action object validation and debugging.
 	 * @return The newly created action schema element.
 	 */
 	UFUNCTION(BlueprintPure, Category = "LearningAgents", meta = (AdvancedDisplay = 2))
-	static UE_API FLearningAgentsActionSchemaElement SpecifyEncodingAction(ULearningAgentsActionSchema* Schema, const FLearningAgentsActionSchemaElement Element, const int32 EncodingSize = 128, const int32 HiddenLayerNum = 1, const ELearningAgentsActivationFunction ActivationFunction = ELearningAgentsActivationFunction::ELU, const FName Tag = TEXT("EncodingAction"));
+	static UE_API FLearningAgentsActionSchemaElement SpecifyEncodingAction(ULearningAgentsActionSchema* Schema, const FLearningAgentsActionSchemaElement Element, const int32 EncodingSize = 128, const int32 HiddenLayerNum = 1, const ELearningAgentsActivationFunction ActivationFunction = ELearningAgentsActivationFunction::ELU, const bool bActive = true, const FName Tag = TEXT("EncodingAction"));
 
 	/**
 	 * Specifies a new bool action. This represents an action which is either true or false.
 	 *
 	 * @param Schema The Action Schema
 	 * @param PriorProbability The prior probability of this action being true.
+	 * @param bActive Whether to enable the specified action as a valid network output. If set to inactive, this action will be considered an empty action. This flag allows easy toggles of specified actions.
 	 * @param Tag The tag of this new action. Used during action object validation and debugging.
 	 * @return The newly created action schema element.
 	 */
 	UFUNCTION(BlueprintPure, Category = "LearningAgents", meta = (AdvancedDisplay = 1))
-	static UE_API FLearningAgentsActionSchemaElement SpecifyBoolAction(ULearningAgentsActionSchema* Schema, const float PriorProbability = 0.5f, const FName Tag = TEXT("BoolAction"));
+	static UE_API FLearningAgentsActionSchemaElement SpecifyBoolAction(ULearningAgentsActionSchema* Schema, const float PriorProbability = 0.5f, const bool bActive = true, const FName Tag = TEXT("BoolAction"));
 
 	/**
 	 * Specifies a new float action. This represents an action which is a single float sampled from a Gaussian distribution. It can be used as a 
@@ -587,44 +617,48 @@ public:
 	 *
 	 * @param Schema The Action Schema
 	 * @param FloatScale The scale used for this action.
+	 * @param bActive Whether to enable the specified action as a valid network output. If set to inactive, this action will be considered an empty action. This flag allows easy toggles of specified actions.
 	 * @param Tag The tag of this new action. Used during action object validation and debugging.
 	 * @return The newly created action schema element.
 	 */
 	UFUNCTION(BlueprintPure, Category = "LearningAgents", meta = (AdvancedDisplay = 2))
-	static UE_API FLearningAgentsActionSchemaElement SpecifyFloatAction(ULearningAgentsActionSchema* Schema, const float FloatScale = 1.0f, const FName Tag = TEXT("FloatAction"));
+	static UE_API FLearningAgentsActionSchemaElement SpecifyFloatAction(ULearningAgentsActionSchema* Schema, const float FloatScale = 1.0f, const bool bActive = true, const FName Tag = TEXT("FloatAction"));
 
 	/**
 	 * Specifies a new location action. This represents an action which is a location sampled from a Gaussian distribution.
 	 *
 	 * @param Schema The Action Schema
 	 * @param LocationScale The scale used for this action in cm.
+	 * @param bActive Whether to enable the specified action as a valid network output. If set to inactive, this action will be considered an empty action. This flag allows easy toggles of specified actions.
 	 * @param Tag The tag of this new action. Used during action object validation and debugging.
 	 * @return The newly created action schema element.
 	 */
 	UFUNCTION(BlueprintPure, Category = "LearningAgents", meta = (AdvancedDisplay = 2))
-	static UE_API FLearningAgentsActionSchemaElement SpecifyLocationAction(ULearningAgentsActionSchema* Schema, const float LocationScale = 100.0f, const FName Tag = TEXT("LocationAction"));
+	static UE_API FLearningAgentsActionSchemaElement SpecifyLocationAction(ULearningAgentsActionSchema* Schema, const float LocationScale = 100.0f, const bool bActive = true, const FName Tag = TEXT("LocationAction"));
 
 	/**
 	 * Specifies a new rotation action. This represents an action which is a rotation sampled from a Gaussian distribution in the angle-axis space.
 	 *
 	 * @param Schema The Action Schema
 	 * @param RotationScale The scale used for this action in degrees.
+	 * @param bActive Whether to enable the specified action as a valid network output. If set to inactive, this action will be considered an empty action. This flag allows easy toggles of specified actions.
 	 * @param Tag The tag of this new action. Used during action object validation and debugging.
 	 * @return The newly created action schema element.
 	 */
 	UFUNCTION(BlueprintPure, Category = "LearningAgents", meta = (AdvancedDisplay = 2))
-	static UE_API FLearningAgentsActionSchemaElement SpecifyRotationAction(ULearningAgentsActionSchema* Schema, const float RotationScale = 90.0f, const FName Tag = TEXT("RotationAction"));
+	static UE_API FLearningAgentsActionSchemaElement SpecifyRotationAction(ULearningAgentsActionSchema* Schema, const float RotationScale = 90.0f, const bool bActive = true, const FName Tag = TEXT("RotationAction"));
 
 	/**
 	 * Specifies a new scale action. This represents an action which is a scale sampled from a Gaussian distribution in the log space.
 	 *
 	 * @param Schema The Action Schema
 	 * @param ScaleScale The scale used for this action.
+	 * @param bActive Whether to enable the specified action as a valid network output. If set to inactive, this action will be considered an empty action. This flag allows easy toggles of specified actions.
 	 * @param Tag The tag of this new action. Used during action object validation and debugging.
 	 * @return The newly created action schema element.
 	 */
 	UFUNCTION(BlueprintPure, Category = "LearningAgents", meta = (AdvancedDisplay = 2))
-	static UE_API FLearningAgentsActionSchemaElement SpecifyScaleAction(ULearningAgentsActionSchema* Schema, const float ScaleScale = 1.0f, const FName Tag = TEXT("ScaleAction"));
+	static UE_API FLearningAgentsActionSchemaElement SpecifyScaleAction(ULearningAgentsActionSchema* Schema, const float ScaleScale = 1.0f, const bool bActive = true, const FName Tag = TEXT("ScaleAction"));
 
 	/**
 	 * Specifies a new transform action.
@@ -633,6 +667,7 @@ public:
 	 * @param LocationScale The scale used for the Location part of the transform in this action in cm.
 	 * @param RotationScale The scale used for the Rotation part of the transform in this action in degrees.
 	 * @param ScaleScale The scale used for the Scale part of the transform in this action.
+	 * @param bActive Whether to enable the specified action as a valid network output. If set to inactive, this action will be considered an empty action. This flag allows easy toggles of specified actions.
 	 * @param Tag The tag of this new action. Used during action object validation and debugging.
 	 * @return The newly created action schema element.
 	 */
@@ -642,6 +677,7 @@ public:
 		const float LocationScale = 100.0f, 
 		const float RotationScale = 90.0f, 
 		const float ScaleScale = 1.0f, 
+		const bool bActive = true, 
 		const FName Tag = TEXT("TransformAction"));
 
 	/**
@@ -649,32 +685,35 @@ public:
 	 *
 	 * @param Schema The Action Schema
 	 * @param AngleScale The scale used for this action in degrees.
+	 * @param bActive Whether to enable the specified action as a valid network output. If set to inactive, this action will be considered an empty action. This flag allows easy toggles of specified actions.
 	 * @param Tag The tag of this new action. Used during action object validation and debugging.
 	 * @return The newly created action schema element.
 	 */
 	UFUNCTION(BlueprintPure, Category = "LearningAgents", meta = (AdvancedDisplay = 2))
-	static UE_API FLearningAgentsActionSchemaElement SpecifyAngleAction(ULearningAgentsActionSchema* Schema, const float AngleScale = 90.0f, const FName Tag = TEXT("AngleAction"));
+	static UE_API FLearningAgentsActionSchemaElement SpecifyAngleAction(ULearningAgentsActionSchema* Schema, const float AngleScale = 90.0f, const bool bActive = true, const FName Tag = TEXT("AngleAction"));
 
 	/**
 	 * Specifies a new velocity action. This represents an action which is a velocity sampled from a Gaussian distribution.
 	 *
 	 * @param Schema The Action Schema
 	 * @param VelocityScale The scale used for this action in cm/s.
+	 * @param bActive Whether to enable the specified action as a valid network output. If set to inactive, this action will be considered an empty action. This flag allows easy toggles of specified actions.
 	 * @param Tag The tag of this new action. Used during action object validation and debugging.
 	 * @return The newly created action schema element.
 	 */
 	UFUNCTION(BlueprintPure, Category = "LearningAgents", meta = (AdvancedDisplay = 2))
-	static UE_API FLearningAgentsActionSchemaElement SpecifyVelocityAction(ULearningAgentsActionSchema* Schema, const float VelocityScale = 200.0f, const FName Tag = TEXT("VelocityAction"));
+	static UE_API FLearningAgentsActionSchemaElement SpecifyVelocityAction(ULearningAgentsActionSchema* Schema, const float VelocityScale = 200.0f, const bool bActive = true, const FName Tag = TEXT("VelocityAction"));
 
 	/**
 	 * Specifies a new direction action. This represents an action which is a direction sampled from a Gaussian distribution and normalized.
 	 *
 	 * @param Schema The Action Schema
+	 * @param bActive Whether to enable the specified action as a valid network output. If set to inactive, this action will be considered an empty action. This flag allows easy toggles of specified actions.
 	 * @param Tag The tag of this new action. Used during action object validation and debugging.
 	 * @return The newly created action schema element.
 	 */
 	UFUNCTION(BlueprintPure, Category = "LearningAgents", meta = (AdvancedDisplay = 1))
-	static UE_API FLearningAgentsActionSchemaElement SpecifyDirectionAction(ULearningAgentsActionSchema* Schema, const FName Tag = TEXT("DirectionAction"));
+	static UE_API FLearningAgentsActionSchemaElement SpecifyDirectionAction(ULearningAgentsActionSchema* Schema, const bool bActive = true, const FName Tag = TEXT("DirectionAction"));
 
 public:
 
@@ -685,6 +724,7 @@ public:
 	static UE_API FLearningAgentsActionObjectElement MakeContinuousAction(
 		ULearningAgentsActionObject* Object,
 		const TArray<float>& Values,
+		const bool bActive = true, 
 		const FName Tag = TEXT("ContinuousAction"),
 		const bool bVisualLoggerEnabled = false,
 		ULearningAgentsManagerListener* VisualLoggerListener = nullptr,
@@ -695,6 +735,7 @@ public:
 	static UE_API FLearningAgentsActionObjectElement MakeContinuousActionFromArrayView(
 		ULearningAgentsActionObject* Object,
 		const TArrayView<const float> Values,
+		const bool bActive = true, 
 		const FName Tag = TEXT("ContinuousAction"),
 		const bool bVisualLoggerEnabled = false,
 		ULearningAgentsManagerListener* VisualLoggerListener = nullptr,
@@ -706,6 +747,7 @@ public:
 	static UE_API FLearningAgentsActionObjectElement MakeExclusiveDiscreteAction(
 		ULearningAgentsActionObject* Object,
 		const int32 Index,
+		const bool bActive = true, 
 		const FName Tag = TEXT("DiscreteExclusiveAction"),
 		const bool bVisualLoggerEnabled = false,
 		ULearningAgentsManagerListener* VisualLoggerListener = nullptr,
@@ -717,6 +759,7 @@ public:
 	static UE_API FLearningAgentsActionObjectElement MakeNamedExclusiveDiscreteAction(
 		ULearningAgentsActionObject* Object,
 		const FName Name,
+		const bool bActive = true, 
 		const FName Tag = TEXT("NamedDiscreteExclusiveAction"),
 		const bool bVisualLoggerEnabled = false,
 		ULearningAgentsManagerListener* VisualLoggerListener = nullptr,
@@ -728,6 +771,7 @@ public:
 	static UE_API FLearningAgentsActionObjectElement MakeInclusiveDiscreteAction(
 		ULearningAgentsActionObject* Object,
 		const TArray<int32>& Indices,
+		const bool bActive = true, 
 		const FName Tag = TEXT("DiscreteInclusiveAction"),
 		const bool bVisualLoggerEnabled = false,
 		ULearningAgentsManagerListener* VisualLoggerListener = nullptr,
@@ -738,6 +782,7 @@ public:
 	static UE_API FLearningAgentsActionObjectElement MakeInclusiveDiscreteActionFromArrayView(
 		ULearningAgentsActionObject* Object,
 		const TArrayView<const int32> Indices,
+		const bool bActive = true,
 		const FName Tag = TEXT("DiscreteInclusiveAction"),
 		const bool bVisualLoggerEnabled = false,
 		ULearningAgentsManagerListener* VisualLoggerListener = nullptr,
@@ -749,6 +794,7 @@ public:
 	static UE_API FLearningAgentsActionObjectElement MakeNamedInclusiveDiscreteAction(
 		ULearningAgentsActionObject* Object,
 		const TArray<FName>& Names,
+		const bool bActive = true, 
 		const FName Tag = TEXT("NamedDiscreteInclusiveAction"),
 		const bool bVisualLoggerEnabled = false,
 		ULearningAgentsManagerListener* VisualLoggerListener = nullptr,
@@ -759,6 +805,7 @@ public:
 	static UE_API FLearningAgentsActionObjectElement MakeNamedInclusiveDiscreteActionFromArrayView(
 		ULearningAgentsActionObject* Object,
 		const TArrayView<const FName> Names,
+		const bool bActive = true, 
 		const FName Tag = TEXT("NamedDiscreteInclusiveAction"),
 		const bool bVisualLoggerEnabled = false,
 		ULearningAgentsManagerListener* VisualLoggerListener = nullptr,
@@ -767,34 +814,35 @@ public:
 		const FLinearColor VisualLoggerColor = FLinearColor::Blue);
 
 	UFUNCTION(BlueprintPure, Category = "LearningAgents", meta = (AdvancedDisplay = 2))
-	static UE_API FLearningAgentsActionObjectElement MakeStructAction(ULearningAgentsActionObject* Object, const TMap<FName, FLearningAgentsActionObjectElement>& Elements, const FName Tag = TEXT("StructAction"));
+	static UE_API FLearningAgentsActionObjectElement MakeStructAction(ULearningAgentsActionObject* Object, const TMap<FName, FLearningAgentsActionObjectElement>& Elements, const bool bActive = true, const FName Tag = TEXT("StructAction"));
 
 	UFUNCTION(BlueprintPure, Category = "LearningAgents", meta = (AdvancedDisplay = 3))
-	static UE_API FLearningAgentsActionObjectElement MakeStructActionFromArrays(ULearningAgentsActionObject* Object, const TArray<FName>& ElementNames, const TArray<FLearningAgentsActionObjectElement>& Elements, const FName Tag = TEXT("StructAction"));
-	static UE_API FLearningAgentsActionObjectElement MakeStructActionFromArrayViews(ULearningAgentsActionObject* Object, const TArrayView<const FName> ElementNames, const TArrayView<const FLearningAgentsActionObjectElement> Elements, const FName Tag = TEXT("StructAction"));
+	static UE_API FLearningAgentsActionObjectElement MakeStructActionFromArrays(ULearningAgentsActionObject* Object, const TArray<FName>& ElementNames, const TArray<FLearningAgentsActionObjectElement>& Elements, const bool bActive = true, const FName Tag = TEXT("StructAction"));
+	static UE_API FLearningAgentsActionObjectElement MakeStructActionFromArrayViews(ULearningAgentsActionObject* Object, const TArrayView<const FName> ElementNames, const TArrayView<const FLearningAgentsActionObjectElement> Elements, const bool bActive = true, const FName Tag = TEXT("StructAction"));
 
 	UFUNCTION(BlueprintPure, Category = "LearningAgents", meta = (AdvancedDisplay = 3))
-	static UE_API FLearningAgentsActionObjectElement MakeExclusiveUnionAction(ULearningAgentsActionObject* Object, const FName ElementName, const FLearningAgentsActionObjectElement Element, const FName Tag = TEXT("ExclusiveUnionAction"));
+	static UE_API FLearningAgentsActionObjectElement MakeExclusiveUnionAction(ULearningAgentsActionObject* Object, const FName ElementName, const FLearningAgentsActionObjectElement Element, const bool bActive = true, const FName Tag = TEXT("ExclusiveUnionAction"));
 
 	UFUNCTION(BlueprintPure, Category = "LearningAgents", meta = (AdvancedDisplay = 2))
-	static UE_API FLearningAgentsActionObjectElement MakeInclusiveUnionAction(ULearningAgentsActionObject* Object, const TMap<FName, FLearningAgentsActionObjectElement>& Elements, const FName Tag = TEXT("InclusiveUnionAction"));
+	static UE_API FLearningAgentsActionObjectElement MakeInclusiveUnionAction(ULearningAgentsActionObject* Object, const TMap<FName, FLearningAgentsActionObjectElement>& Elements, const bool bActive = true, const FName Tag = TEXT("InclusiveUnionAction"));
 
 	UFUNCTION(BlueprintPure, Category = "LearningAgents", meta = (AdvancedDisplay = 3))
-	static UE_API FLearningAgentsActionObjectElement MakeInclusiveUnionActionFromArrays(ULearningAgentsActionObject* Object, const TArray<FName>& ElementNames, const TArray<FLearningAgentsActionObjectElement>& Elements, const FName Tag = TEXT("InclusiveUnionAction"));
-	static UE_API FLearningAgentsActionObjectElement MakeInclusiveUnionActionFromArrayViews(ULearningAgentsActionObject* Object, const TArrayView<const FName> ElementNames, const TArrayView<const FLearningAgentsActionObjectElement> Elements, const FName Tag = TEXT("InclusiveUnionAction"));
+	static UE_API FLearningAgentsActionObjectElement MakeInclusiveUnionActionFromArrays(ULearningAgentsActionObject* Object, const TArray<FName>& ElementNames, const TArray<FLearningAgentsActionObjectElement>& Elements, const bool bActive = true, const FName Tag = TEXT("InclusiveUnionAction"));
+	static UE_API FLearningAgentsActionObjectElement MakeInclusiveUnionActionFromArrayViews(ULearningAgentsActionObject* Object, const TArrayView<const FName> ElementNames, const TArrayView<const FLearningAgentsActionObjectElement> Elements, const bool bActive = true, const FName Tag = TEXT("InclusiveUnionAction"));
 
 	UFUNCTION(BlueprintPure, Category = "LearningAgents", meta = (AdvancedDisplay = 2))
-	static UE_API FLearningAgentsActionObjectElement MakeStaticArrayAction(ULearningAgentsActionObject* Object, const TArray<FLearningAgentsActionObjectElement>& Elements, const FName Tag = TEXT("StaticArrayAction"));
-	static UE_API FLearningAgentsActionObjectElement MakeStaticArrayActionFromArrayView(ULearningAgentsActionObject* Object, const TArrayView<const FLearningAgentsActionObjectElement> Elements, const FName Tag = TEXT("StaticArrayAction"));
+	static UE_API FLearningAgentsActionObjectElement MakeStaticArrayAction(ULearningAgentsActionObject* Object, const TArray<FLearningAgentsActionObjectElement>& Elements, const bool bActive = true, const FName Tag = TEXT("StaticArrayAction"));
+	static UE_API FLearningAgentsActionObjectElement MakeStaticArrayActionFromArrayView(ULearningAgentsActionObject* Object, const TArrayView<const FLearningAgentsActionObjectElement> Elements, const bool bActive = true, const FName Tag = TEXT("StaticArrayAction"));
 
 	UFUNCTION(BlueprintPure, Category = "LearningAgents", meta = (AdvancedDisplay = 3))
-	static UE_API FLearningAgentsActionObjectElement MakePairAction(ULearningAgentsActionObject* Object, const FLearningAgentsActionObjectElement Key, const FLearningAgentsActionObjectElement Value, const FName Tag = TEXT("PairAction"));
+	static UE_API FLearningAgentsActionObjectElement MakePairAction(ULearningAgentsActionObject* Object, const FLearningAgentsActionObjectElement Key, const FLearningAgentsActionObjectElement Value, const bool bActive = true, const FName Tag = TEXT("PairAction"));
 
 	UFUNCTION(BlueprintPure, Category = "LearningAgents", meta = (AdvancedDisplay = 2, DefaultToSelf = "VisualLoggerListener"))
 	static UE_API FLearningAgentsActionObjectElement MakeEnumAction(
 		ULearningAgentsActionObject* Object,
 		const UEnum* Enum,
 		const uint8 EnumValue,
+		const bool bActive = true, 
 		const FName Tag = TEXT("EnumAction"),
 		const bool bVisualLoggerEnabled = false,
 		ULearningAgentsManagerListener* VisualLoggerListener = nullptr,
@@ -807,6 +855,7 @@ public:
 		ULearningAgentsActionObject* Object,
 		const UEnum* Enum,
 		const int32 BitmaskValue,
+		const bool bActive = true, 
 		const FName Tag = TEXT("BitmaskAction"),
 		const bool bVisualLoggerEnabled = false,
 		ULearningAgentsManagerListener* VisualLoggerListener = nullptr,
@@ -815,30 +864,31 @@ public:
 		const FLinearColor VisualLoggerColor = FLinearColor::Blue);
 
 	UFUNCTION(BlueprintPure, Category = "LearningAgents", meta = (AdvancedDisplay = 3))
-	static UE_API FLearningAgentsActionObjectElement MakeOptionalAction(ULearningAgentsActionObject* Object, const FLearningAgentsActionObjectElement Element, const ELearningAgentsOptionalAction Option, const FName Tag = TEXT("OptionalAction"));
+	static UE_API FLearningAgentsActionObjectElement MakeOptionalAction(ULearningAgentsActionObject* Object, const FLearningAgentsActionObjectElement Element, const ELearningAgentsOptionalAction Option, const bool bActive = true, const FName Tag = TEXT("OptionalAction"));
 
 	UFUNCTION(BlueprintPure, Category = "LearningAgents", meta = (AdvancedDisplay = 1))
-	static UE_API FLearningAgentsActionObjectElement MakeOptionalNullAction(ULearningAgentsActionObject* Object, const FName Tag = TEXT("OptionalAction"));
+	static UE_API FLearningAgentsActionObjectElement MakeOptionalNullAction(ULearningAgentsActionObject* Object, const bool bActive = true, const FName Tag = TEXT("OptionalAction"));
 
 	UFUNCTION(BlueprintPure, Category = "LearningAgents", meta = (AdvancedDisplay = 2))
-	static UE_API FLearningAgentsActionObjectElement MakeOptionalValidAction(ULearningAgentsActionObject* Object, const FLearningAgentsActionObjectElement Element, const FName Tag = TEXT("OptionalAction"));
+	static UE_API FLearningAgentsActionObjectElement MakeOptionalValidAction(ULearningAgentsActionObject* Object, const FLearningAgentsActionObjectElement Element, const bool bActive = true, const FName Tag = TEXT("OptionalAction"));
 
 	UFUNCTION(BlueprintPure, Category = "LearningAgents", meta = (AdvancedDisplay = 3))
-	static UE_API FLearningAgentsActionObjectElement MakeEitherAction(ULearningAgentsActionObject* Object, const FLearningAgentsActionObjectElement Element, const ELearningAgentsEitherAction Either, const FName Tag = TEXT("EitherAction"));
+	static UE_API FLearningAgentsActionObjectElement MakeEitherAction(ULearningAgentsActionObject* Object, const FLearningAgentsActionObjectElement Element, const ELearningAgentsEitherAction Either, const bool bActive = true, const FName Tag = TEXT("EitherAction"));
 
 	UFUNCTION(BlueprintPure, Category = "LearningAgents", meta = (AdvancedDisplay = 2, DisplayName = "Make Either A Action"))
-	static UE_API FLearningAgentsActionObjectElement MakeEitherAAction(ULearningAgentsActionObject* Object, const FLearningAgentsActionObjectElement A, const FName Tag = TEXT("EitherAction"));
+	static UE_API FLearningAgentsActionObjectElement MakeEitherAAction(ULearningAgentsActionObject* Object, const FLearningAgentsActionObjectElement A, const bool bActive = true, const FName Tag = TEXT("EitherAction"));
 
 	UFUNCTION(BlueprintPure, Category = "LearningAgents", meta = (AdvancedDisplay = 2, DisplayName = "Make Either B Action"))
-	static UE_API FLearningAgentsActionObjectElement MakeEitherBAction(ULearningAgentsActionObject* Object, const FLearningAgentsActionObjectElement B, const FName Tag = TEXT("EitherAction"));
+	static UE_API FLearningAgentsActionObjectElement MakeEitherBAction(ULearningAgentsActionObject* Object, const FLearningAgentsActionObjectElement B, const bool bActive = true, const FName Tag = TEXT("EitherAction"));
 
 	UFUNCTION(BlueprintPure, Category = "LearningAgents", meta = (AdvancedDisplay = 2))
-	static UE_API FLearningAgentsActionObjectElement MakeEncodingAction(ULearningAgentsActionObject* Object, const FLearningAgentsActionObjectElement Element, const FName Tag = TEXT("EncodingAction"));
+	static UE_API FLearningAgentsActionObjectElement MakeEncodingAction(ULearningAgentsActionObject* Object, const FLearningAgentsActionObjectElement Element, const bool bActive = true, const FName Tag = TEXT("EncodingAction"));
 
 	UFUNCTION(BlueprintPure, Category = "LearningAgents", meta = (AdvancedDisplay = 2, DefaultToSelf = "VisualLoggerListener"))
 	static UE_API FLearningAgentsActionObjectElement MakeBoolAction(
 		ULearningAgentsActionObject* Object,
 		const bool bValue,
+		const bool bActive = true, 
 		const FName Tag = TEXT("BoolAction"),
 		const bool bVisualLoggerEnabled = false,
 		ULearningAgentsManagerListener * VisualLoggerListener = nullptr,
@@ -850,6 +900,7 @@ public:
 	static UE_API FLearningAgentsActionObjectElement MakeFloatAction(
 		ULearningAgentsActionObject* Object,
 		const float Value,
+		const bool bActive = true, 
 		const FName Tag = TEXT("FloatAction"),
 		const bool bVisualLoggerEnabled = false,
 		ULearningAgentsManagerListener* VisualLoggerListener = nullptr,
@@ -862,6 +913,7 @@ public:
 		ULearningAgentsActionObject* Object,
 		const FVector Location,
 		const FTransform RelativeTransform = FTransform(),
+		const bool bActive = true, 
 		const FName Tag = TEXT("LocationAction"),
 		const bool bVisualLoggerEnabled = false,
 		ULearningAgentsManagerListener* VisualLoggerListener = nullptr,
@@ -874,6 +926,7 @@ public:
 		ULearningAgentsActionObject* Object,
 		const FRotator Rotation,
 		const FRotator RelativeRotation = FRotator::ZeroRotator,
+		const bool bActive = true, 
 		const FName Tag = TEXT("RotationAction"),
 		const bool bVisualLoggerEnabled = false,
 		ULearningAgentsManagerListener* VisualLoggerListener = nullptr,
@@ -887,6 +940,7 @@ public:
 		ULearningAgentsActionObject* Object,
 		const FQuat Rotation,
 		const FQuat RelativeRotation,
+		const bool bActive = true, 
 		const FName Tag = TEXT("RotationAction"),
 		const bool bVisualLoggerEnabled = false,
 		ULearningAgentsManagerListener* VisualLoggerListener = nullptr,
@@ -900,6 +954,7 @@ public:
 		ULearningAgentsActionObject* Object,
 		const FVector Scale,
 		const FVector RelativeScale = FVector(1, 1, 1),
+		const bool bActive = true, 
 		const FName Tag = TEXT("ScaleAction"),
 		const bool bVisualLoggerEnabled = false,
 		ULearningAgentsManagerListener* VisualLoggerListener = nullptr,
@@ -912,6 +967,7 @@ public:
 		ULearningAgentsActionObject* Object,
 		const FTransform Transform,
 		const FTransform RelativeTransform = FTransform(),
+		const bool bActive = true, 
 		const FName Tag = TEXT("TransformAction"),
 		const bool bVisualLoggerEnabled = false,
 		ULearningAgentsManagerListener* VisualLoggerListener = nullptr,
@@ -924,6 +980,7 @@ public:
 		ULearningAgentsActionObject* Object,
 		const float Angle,
 		const float RelativeAngle = 0.0f,
+		const bool bActive = true, 
 		const FName Tag = TEXT("AngleAction"),
 		const bool bVisualLoggerEnabled = false,
 		ULearningAgentsManagerListener* VisualLoggerListener = nullptr,
@@ -937,6 +994,7 @@ public:
 		ULearningAgentsActionObject* Object,
 		const float Angle,
 		const float RelativeAngle = 0.0f,
+		const bool bActive = true, 
 		const FName Tag = TEXT("AngleAction"),
 		const bool bVisualLoggerEnabled = false,
 		ULearningAgentsManagerListener* VisualLoggerListener = nullptr,
@@ -950,6 +1008,7 @@ public:
 		ULearningAgentsActionObject* Object,
 		const FVector Velocity,
 		const FTransform RelativeTransform = FTransform(),
+		const bool bActive = true, 
 		const FName Tag = TEXT("VelocityAction"),
 		const bool bVisualLoggerEnabled = false,
 		ULearningAgentsManagerListener* VisualLoggerListener = nullptr,
@@ -963,6 +1022,7 @@ public:
 		ULearningAgentsActionObject* Object,
 		const FVector Direction,
 		const FTransform RelativeTransform = FTransform(),
+		const bool bActive = true, 
 		const FName Tag = TEXT("DirectionAction"),
 		const bool bVisualLoggerEnabled = false,
 		ULearningAgentsManagerListener* VisualLoggerListener = nullptr,
@@ -1084,11 +1144,12 @@ public:
 	 * @param OutNum The output number of values.
 	 * @param Object The Action Object
 	 * @param Element The Action Object Element
+	 * @param bActive Whether the action should be considered active. Must match the value given during Specify.
 	 * @param Tag The tag of the corresponding action. Must match the tag given during Specify.
 	 * @return true if the provided Element is the correct type, otherwise false.
 	 */
 	UFUNCTION(BlueprintPure = false, Category = "LearningAgents", meta = (AdvancedDisplay = 3, ReturnDisplayName = "Success"))
-	static UE_API bool GetContinuousActionNum(int32& OutNum, const ULearningAgentsActionObject* Object, const FLearningAgentsActionObjectElement Element, const FName Tag = TEXT("ContinuousAction"));
+	static UE_API bool GetContinuousActionNum(int32& OutNum, const ULearningAgentsActionObject* Object, const FLearningAgentsActionObjectElement Element, const bool bActive = true, const FName Tag = TEXT("ContinuousAction"));
 
 	/**
 	 * Get the values of a continuous action.
@@ -1096,6 +1157,7 @@ public:
 	 * @param OutValues The output values.
 	 * @param Object The Action Object
 	 * @param Element The Action Object Element
+	 * @param bActive Whether the action should be considered active. Must match the value given during Specify.
 	 * @param Tag The tag of the corresponding action. Must match the tag given during Specify.
 	 * @param bVisualLoggerEnabled When true, debug data will be sent to the visual logger.
 	 * @param VisualLoggerListener The listener object which is making this action. This must be set to use logging.
@@ -1109,6 +1171,7 @@ public:
 		TArray<float>& OutValues, 
 		const ULearningAgentsActionObject* Object, 
 		const FLearningAgentsActionObjectElement Element, 
+		const bool bActive = true, 
 		const FName Tag = TEXT("ContinuousAction"),
 		const bool bVisualLoggerEnabled = false,
 		ULearningAgentsManagerListener* VisualLoggerListener = nullptr,
@@ -1122,6 +1185,7 @@ public:
 	 * @param OutValues The output values.
 	 * @param Object The Action Object
 	 * @param Element The Action Object Element
+	 * @param bActive Whether the action should be considered active. Must match the value given during Specify.
 	 * @param Tag The tag of the corresponding action. Must match the tag given during Specify.
 	 * @param bVisualLoggerEnabled When true, debug data will be sent to the visual logger.
 	 * @param VisualLoggerListener The listener object which is making this action. This must be set to use logging.
@@ -1134,6 +1198,7 @@ public:
 		TArrayView<float> OutValues, 
 		const ULearningAgentsActionObject* Object, 
 		const FLearningAgentsActionObjectElement Element, 
+		const bool bActive = true, 
 		const FName Tag = TEXT("ContinuousAction"),
 		const bool bVisualLoggerEnabled = false,
 		ULearningAgentsManagerListener* VisualLoggerListener = nullptr,
@@ -1147,6 +1212,7 @@ public:
 	 * @param OutIndex The output index.
 	 * @param Object The Action Object
 	 * @param Element The Action Object Element
+	 * @param bActive Whether the action should be considered active. Must match the value given during Specify.
 	 * @param Tag The tag of the corresponding action. Must match the tag given during Specify.
 	 * @param bVisualLoggerEnabled When true, debug data will be sent to the visual logger.
 	 * @param VisualLoggerListener The listener object which is making this action. This must be set to use logging.
@@ -1160,6 +1226,7 @@ public:
 		int32& OutIndex, 
 		const ULearningAgentsActionObject* Object, 
 		const FLearningAgentsActionObjectElement Element, 
+		const bool bActive = true, 
 		const FName Tag = TEXT("DiscreteExclusiveAction"),
 		const bool bVisualLoggerEnabled = false,
 		ULearningAgentsManagerListener* VisualLoggerListener = nullptr,
@@ -1173,6 +1240,7 @@ public:
 	 * @param OutName The output name.
 	 * @param Object The Action Object
 	 * @param Element The Action Object Element
+	 * @param bActive Whether the action should be considered active. Must match the value given during Specify.
 	 * @param Tag The tag of the corresponding action. Must match the tag given during Specify.
 	 * @param bVisualLoggerEnabled When true, debug data will be sent to the visual logger.
 	 * @param VisualLoggerListener The listener object which is making this action. This must be set to use logging.
@@ -1186,6 +1254,7 @@ public:
 		FName& OutName,
 		const ULearningAgentsActionObject* Object,
 		const FLearningAgentsActionObjectElement Element,
+		const bool bActive = true, 
 		const FName Tag = TEXT("NamedDiscreteExclusiveAction"),
 		const bool bVisualLoggerEnabled = false,
 		ULearningAgentsManagerListener* VisualLoggerListener = nullptr,
@@ -1199,11 +1268,12 @@ public:
 	 * @param OutNum The output number of indices.
 	 * @param Object The Action Object
 	 * @param Element The Action Object Element
+	 * @param bActive Whether the action should be considered active. Must match the value given during Specify.
 	 * @param Tag The tag of the corresponding action. Must match the tag given during Specify.
 	 * @return true if the provided Element is the correct type, otherwise false.
 	 */
 	UFUNCTION(BlueprintPure = false, Category = "LearningAgents", meta = (AdvancedDisplay = 3, ReturnDisplayName = "Success"))
-	static UE_API bool GetInclusiveDiscreteActionNum(int32& OutNum, const ULearningAgentsActionObject* Object, const FLearningAgentsActionObjectElement Element, const FName Tag = TEXT("DiscreteInclusiveAction"));
+	static UE_API bool GetInclusiveDiscreteActionNum(int32& OutNum, const ULearningAgentsActionObject* Object, const FLearningAgentsActionObjectElement Element, const bool bActive = true, const FName Tag = TEXT("DiscreteInclusiveAction"));
 
 	/**
 	 * Get the indices for an inclusive discrete action.
@@ -1211,6 +1281,7 @@ public:
 	 * @param OutIndices The output indices.
 	 * @param Object The Action Object
 	 * @param Element The Action Object Element
+	 * @param bActive Whether the action should be considered active. Must match the value given during Specify.
 	 * @param Tag The tag of the corresponding action. Must match the tag given during Specify.
 	 * @param bVisualLoggerEnabled When true, debug data will be sent to the visual logger.
 	 * @param VisualLoggerListener The listener object which is making this action. This must be set to use logging.
@@ -1224,6 +1295,7 @@ public:
 		TArray<int32>& OutIndices, 
 		const ULearningAgentsActionObject* Object, 
 		const FLearningAgentsActionObjectElement Element, 
+		const bool bActive = true, 
 		const FName Tag = TEXT("DiscreteInclusiveAction"),
 		const bool bVisualLoggerEnabled = false,
 		ULearningAgentsManagerListener* VisualLoggerListener = nullptr,
@@ -1237,6 +1309,7 @@ public:
 	 * @param OutIndices The output indices.
 	 * @param Object The Action Object
 	 * @param Element The Action Object Element
+	 * @param bActive Whether the action should be considered active. Must match the value given during Specify.
 	 * @param Tag The tag of the corresponding action. Must match the tag given during Specify.
 	 * @param bVisualLoggerEnabled When true, debug data will be sent to the visual logger.
 	 * @param VisualLoggerListener The listener object which is making this action. This must be set to use logging.
@@ -1249,6 +1322,7 @@ public:
 		TArrayView<int32> OutIndices, 
 		const ULearningAgentsActionObject* Object, 
 		const FLearningAgentsActionObjectElement Element, 
+		const bool bActive = true, 
 		const FName Tag = TEXT("DiscreteInclusiveAction"),
 		const bool bVisualLoggerEnabled = false,
 		ULearningAgentsManagerListener* VisualLoggerListener = nullptr,
@@ -1262,11 +1336,12 @@ public:
 	 * @param OutNum The output number of names.
 	 * @param Object The Action Object
 	 * @param Element The Action Object Element
+	 * @param bActive Whether the action should be considered active. Must match the value given during Specify.
 	 * @param Tag The tag of the corresponding action. Must match the tag given during Specify.
 	 * @return true if the provided Element is the correct type, otherwise false.
 	 */
 	UFUNCTION(BlueprintPure = false, Category = "LearningAgents", meta = (AdvancedDisplay = 3, ReturnDisplayName = "Success"))
-	static UE_API bool GetNamedInclusiveDiscreteActionNum(int32& OutNum, const ULearningAgentsActionObject* Object, const FLearningAgentsActionObjectElement Element, const FName Tag = TEXT("NamedDiscreteInclusiveAction"));
+	static UE_API bool GetNamedInclusiveDiscreteActionNum(int32& OutNum, const ULearningAgentsActionObject* Object, const FLearningAgentsActionObjectElement Element, const bool bActive = true, const FName Tag = TEXT("NamedDiscreteInclusiveAction"));
 
 	/**
 	 * Get the names for a named inclusive discrete action.
@@ -1274,6 +1349,7 @@ public:
 	 * @param OutNames The output names.
 	 * @param Object The Action Object
 	 * @param Element The Action Object Element
+	 * @param bActive Whether the action should be considered active. Must match the value given during Specify.
 	 * @param Tag The tag of the corresponding action. Must match the tag given during Specify.
 	 * @param bVisualLoggerEnabled When true, debug data will be sent to the visual logger.
 	 * @param VisualLoggerListener The listener object which is making this action. This must be set to use logging.
@@ -1287,6 +1363,7 @@ public:
 		TArray<FName>& OutNames,
 		const ULearningAgentsActionObject* Object,
 		const FLearningAgentsActionObjectElement Element,
+		const bool bActive = true, 
 		const FName Tag = TEXT("NamedDiscreteInclusiveAction"),
 		const bool bVisualLoggerEnabled = false,
 		ULearningAgentsManagerListener* VisualLoggerListener = nullptr,
@@ -1300,6 +1377,7 @@ public:
 	 * @param OutNames The output names.
 	 * @param Object The Action Object
 	 * @param Element The Action Object Element
+	 * @param bActive Whether the action should be considered active. Must match the value given during Specify.
 	 * @param Tag The tag of the corresponding action. Must match the tag given during Specify.
 	 * @param bVisualLoggerEnabled When true, debug data will be sent to the visual logger.
 	 * @param VisualLoggerListener The listener object which is making this action. This must be set to use logging.
@@ -1312,6 +1390,7 @@ public:
 		TArrayView<FName> OutNames,
 		const ULearningAgentsActionObject* Object,
 		const FLearningAgentsActionObjectElement Element,
+		const bool bActive = true, 
 		const FName Tag = TEXT("NamedDiscreteInclusiveAction"),
 		const bool bVisualLoggerEnabled = false,
 		ULearningAgentsManagerListener* VisualLoggerListener = nullptr,
@@ -1325,11 +1404,12 @@ public:
 	 * @param OutNum The output number of sub-actions.
 	 * @param Object The Action Object
 	 * @param Element The Action Object Element
+	 * @param bActive Whether the action should be considered active. Must match the value given during Specify.
 	 * @param Tag The tag of the corresponding action. Must match the tag given during Specify.
 	 * @return true if the provided Element is the correct type, otherwise false.
 	 */
 	UFUNCTION(BlueprintPure = false, Category = "LearningAgents", meta = (AdvancedDisplay = 3, ReturnDisplayName = "Success"))
-	static UE_API bool GetStructActionNum(int32& OutNum, const ULearningAgentsActionObject* Object, const FLearningAgentsActionObjectElement Element, const FName Tag = TEXT("StructAction"));
+	static UE_API bool GetStructActionNum(int32& OutNum, const ULearningAgentsActionObject* Object, const FLearningAgentsActionObjectElement Element, const bool bActive = true, const FName Tag = TEXT("StructAction"));
 
 	/**
 	 * Get the sub-actions for a struct action.
@@ -1337,11 +1417,12 @@ public:
 	 * @param OutElements The output sub-actions.
 	 * @param Object The Action Object
 	 * @param Element The Action Object Element
+	 * @param bActive Whether the action should be considered active. Must match the value given during Specify.
 	 * @param Tag The tag of the corresponding action. Must match the tag given during Specify.
 	 * @return true if the provided Element is the correct type, otherwise false.
 	 */
 	UFUNCTION(BlueprintPure = false, Category = "LearningAgents", meta = (AdvancedDisplay = 3, ReturnDisplayName = "Success"))
-	static UE_API bool GetStructAction(TMap<FName, FLearningAgentsActionObjectElement>& OutElements, const ULearningAgentsActionObject* Object, const FLearningAgentsActionObjectElement Element, const FName Tag = TEXT("StructAction"));
+	static UE_API bool GetStructAction(TMap<FName, FLearningAgentsActionObjectElement>& OutElements, const ULearningAgentsActionObject* Object, const FLearningAgentsActionObjectElement Element, const bool bActive = true, const FName Tag = TEXT("StructAction"));
 
 	/**
 	 * Get the sub-action given its name for a struct action.
@@ -1350,11 +1431,12 @@ public:
 	 * @param Object The Action Object
 	 * @param Element The Action Object Element
 	 * @param ElementName The Sub-action Name.
+	 * @param bActive Whether the action should be considered active. Must match the value given during Specify.
 	 * @param Tag The tag of the corresponding action. Must match the tag given during Specify.
 	 * @return true if the provided Element has the sub-action, otherwise false.
 	 */
 	UFUNCTION(BlueprintPure = false, Category = "LearningAgents", meta = (AdvancedDisplay = 4, ReturnDisplayName = "Success"))
-	static UE_API bool GetStructActionElement(FLearningAgentsActionObjectElement& OutElement, const ULearningAgentsActionObject* Object, const FLearningAgentsActionObjectElement Element, const FName ElementName, const FName Tag = TEXT("StructAction"));
+	static UE_API bool GetStructActionElement(FLearningAgentsActionObjectElement& OutElement, const ULearningAgentsActionObject* Object, const FLearningAgentsActionObjectElement Element, const FName ElementName, const bool bActive = true, const FName Tag = TEXT("StructAction"));
 
 	/**
 	 * Get the sub-actions for a struct action.
@@ -1363,11 +1445,12 @@ public:
 	 * @param OutElements The output sub-actions.
 	 * @param Object The Action Object
 	 * @param Element The Action Object Element
+	 * @param bActive Whether the action should be considered active. Must match the value given during Specify.
 	 * @param Tag The tag of the corresponding action. Must match the tag given during Specify.
 	 * @return true if the provided Element is the correct type, otherwise false.
 	 */
 	UFUNCTION(BlueprintPure = false, Category = "LearningAgents", meta = (AdvancedDisplay = 4, ReturnDisplayName = "Success"))
-	static UE_API bool GetStructActionToArrays(TArray<FName>& OutElementNames, TArray<FLearningAgentsActionObjectElement>& OutElements, const ULearningAgentsActionObject* Object, const FLearningAgentsActionObjectElement Element, const FName Tag = TEXT("StructAction"));
+	static UE_API bool GetStructActionToArrays(TArray<FName>& OutElementNames, TArray<FLearningAgentsActionObjectElement>& OutElements, const ULearningAgentsActionObject* Object, const FLearningAgentsActionObjectElement Element, const bool bActive = true, const FName Tag = TEXT("StructAction"));
 	
 	/**
 	 * Get the sub-actions for a struct action. The OutElementNames and OutElements ArrayViews must be the correct size.
@@ -1376,10 +1459,11 @@ public:
 	 * @param OutElements The output sub-actions.
 	 * @param Object The Action Object
 	 * @param Element The Action Object Element
+	 * @param bActive Whether the action should be considered active. Must match the value given during Specify.
 	 * @param Tag The tag of the corresponding action. Must match the tag given during Specify.
 	 * @return true if the provided Element is the correct type, otherwise false.
 	 */
-	static UE_API bool GetStructActionToArrayViews(TArrayView<FName> OutElementNames, TArrayView<FLearningAgentsActionObjectElement> OutElements, const ULearningAgentsActionObject* Object, const FLearningAgentsActionObjectElement Element, const FName Tag = TEXT("StructAction"));
+	static UE_API bool GetStructActionToArrayViews(TArrayView<FName> OutElementNames, TArrayView<FLearningAgentsActionObjectElement> OutElements, const ULearningAgentsActionObject* Object, const FLearningAgentsActionObjectElement Element, const bool bActive = true, const FName Tag = TEXT("StructAction"));
 
 	/**
 	 * Get the chosen sub-action for an exclusive union action.
@@ -1387,11 +1471,12 @@ public:
 	 * @param OutElement The output sub-action.
 	 * @param Object The Action Object
 	 * @param Element The Action Object Element
+	 * @param bActive Whether the action should be considered active. Must match the value given during Specify.
 	 * @param Tag The tag of the corresponding action. Must match the tag given during Specify.
 	 * @return true if the provided Element is the correct type, otherwise false.
 	 */
 	UFUNCTION(BlueprintPure = false, Category = "LearningAgents", meta = (AdvancedDisplay = 4, ReturnDisplayName = "Success"))
-	static UE_API bool GetExclusiveUnionAction(FName& OutElementName, FLearningAgentsActionObjectElement& OutElement, const ULearningAgentsActionObject* Object, const FLearningAgentsActionObjectElement Element, const FName Tag = TEXT("ExclusiveUnionAction"));
+	static UE_API bool GetExclusiveUnionAction(FName& OutElementName, FLearningAgentsActionObjectElement& OutElement, const ULearningAgentsActionObject* Object, const FLearningAgentsActionObjectElement Element, const bool bActive = true, const FName Tag = TEXT("ExclusiveUnionAction"));
 
 	/**
 	 * Get the number of sub-actions for an inclusive union action.
@@ -1399,11 +1484,12 @@ public:
 	 * @param OutNum The output number of sub-actions.
 	 * @param Object The Action Object
 	 * @param Element The Action Object Element
+	 * @param bActive Whether the action should be considered active. Must match the value given during Specify.
 	 * @param Tag The tag of the corresponding action. Must match the tag given during Specify.
 	 * @return true if the provided Element is the correct type, otherwise false.
 	 */
 	UFUNCTION(BlueprintPure = false, Category = "LearningAgents", meta = (AdvancedDisplay = 3, ReturnDisplayName = "Success"))
-	static UE_API bool GetInclusiveUnionActionNum(int32& OutNum, const ULearningAgentsActionObject* Object, const FLearningAgentsActionObjectElement Element, const FName Tag = TEXT("InclusiveUnionAction"));
+	static UE_API bool GetInclusiveUnionActionNum(int32& OutNum, const ULearningAgentsActionObject* Object, const FLearningAgentsActionObjectElement Element, const bool bActive = true, const FName Tag = TEXT("InclusiveUnionAction"));
 
 	/**
 	 * Get the chosen sub-actions for an inclusive union action.
@@ -1411,11 +1497,12 @@ public:
 	 * @param OutElements The output sub-actions.
 	 * @param Object The Action Object
 	 * @param Element The Action Object Element
+	 * @param bActive Whether the action should be considered active. Must match the value given during Specify.
 	 * @param Tag The tag of the corresponding action. Must match the tag given during Specify.
 	 * @return true if the provided Element is the correct type, otherwise false.
 	 */
 	UFUNCTION(BlueprintPure = false, Category = "LearningAgents", meta = (AdvancedDisplay = 3, ReturnDisplayName = "Success"))
-	static UE_API bool GetInclusiveUnionAction(TMap<FName, FLearningAgentsActionObjectElement>& OutElements, const ULearningAgentsActionObject* Object, const FLearningAgentsActionObjectElement Element, const FName Tag = TEXT("InclusiveUnionAction"));
+	static UE_API bool GetInclusiveUnionAction(TMap<FName, FLearningAgentsActionObjectElement>& OutElements, const ULearningAgentsActionObject* Object, const FLearningAgentsActionObjectElement Element, const bool bActive = true, const FName Tag = TEXT("InclusiveUnionAction"));
 
 	/**
 	 * Get the chosen sub-actions for an inclusive union action.
@@ -1424,11 +1511,12 @@ public:
 	 * @param OutElements The output sub-actions.
 	 * @param Object The Action Object
 	 * @param Element The Action Object Element
+	 * @param bActive Whether the action should be considered active. Must match the value given during Specify.
 	 * @param Tag The tag of the corresponding action. Must match the tag given during Specify.
 	 * @return true if the provided Element is the correct type, otherwise false.
 	 */
 	UFUNCTION(BlueprintPure = false, Category = "LearningAgents", meta = (AdvancedDisplay = 4, ReturnDisplayName = "Success"))
-	static UE_API bool GetInclusiveUnionActionToArrays(TArray<FName>& OutElementNames, TArray<FLearningAgentsActionObjectElement>& OutElements, const ULearningAgentsActionObject* Object, const FLearningAgentsActionObjectElement Element, const FName Tag = TEXT("InclusiveUnionAction"));
+	static UE_API bool GetInclusiveUnionActionToArrays(TArray<FName>& OutElementNames, TArray<FLearningAgentsActionObjectElement>& OutElements, const ULearningAgentsActionObject* Object, const FLearningAgentsActionObjectElement Element, const bool bActive = true, const FName Tag = TEXT("InclusiveUnionAction"));
 	
 	/**
 	 * Get the chosen sub-actions for an inclusive union action. The OutElementNames and OutElements ArrayViews must be the correct size.
@@ -1437,10 +1525,11 @@ public:
 	 * @param OutElements The output sub-actions.
 	 * @param Object The Action Object
 	 * @param Element The Action Object Element
+	 * @param bActive Whether the action should be considered active. Must match the value given during Specify.
 	 * @param Tag The tag of the corresponding action. Must match the tag given during Specify.
 	 * @return true if the provided Element is the correct type, otherwise false.
 	 */
-	static UE_API bool GetInclusiveUnionActionToArrayViews(TArrayView<FName> OutElementNames, TArrayView<FLearningAgentsActionObjectElement> OutElements, const ULearningAgentsActionObject* Object, const FLearningAgentsActionObjectElement Element, const FName Tag = TEXT("InclusiveUnionAction"));
+	static UE_API bool GetInclusiveUnionActionToArrayViews(TArrayView<FName> OutElementNames, TArrayView<FLearningAgentsActionObjectElement> OutElements, const ULearningAgentsActionObject* Object, const FLearningAgentsActionObjectElement Element, const bool bActive = true, const FName Tag = TEXT("InclusiveUnionAction"));
 
 	/**
 	 * Get the number of entries in a static array action.
@@ -1448,11 +1537,12 @@ public:
 	 * @param OutNum The output number of entries.
 	 * @param Object The Action Object
 	 * @param Element The Action Object Element
+	 * @param bActive Whether the action should be considered active. Must match the value given during Specify.
 	 * @param Tag The tag of the corresponding action. Must match the tag given during Specify.
 	 * @return true if the provided Element is the correct type, otherwise false.
 	 */
 	UFUNCTION(BlueprintPure = false, Category = "LearningAgents", meta = (AdvancedDisplay = 3, ReturnDisplayName = "Success"))
-	static UE_API bool GetStaticArrayActionNum(int32& OutNum, const ULearningAgentsActionObject* Object, const FLearningAgentsActionObjectElement Element, const FName Tag = TEXT("StaticArrayAction"));
+	static UE_API bool GetStaticArrayActionNum(int32& OutNum, const ULearningAgentsActionObject* Object, const FLearningAgentsActionObjectElement Element, const bool bActive = true, const FName Tag = TEXT("StaticArrayAction"));
 
 	/**
 	 * Get the entries of a static array action.
@@ -1460,11 +1550,12 @@ public:
 	 * @param OutElements The output sub-elements.
 	 * @param Object The Action Object
 	 * @param Element The Action Object Element
+	 * @param bActive Whether the action should be considered active. Must match the value given during Specify.
 	 * @param Tag The tag of the corresponding action. Must match the tag given during Specify.
 	 * @return true if the provided Element is the correct type, otherwise false.
 	 */
 	UFUNCTION(BlueprintPure = false, Category = "LearningAgents", meta = (AdvancedDisplay = 3, ReturnDisplayName = "Success"))
-	static UE_API bool GetStaticArrayAction(TArray<FLearningAgentsActionObjectElement>& OutElements, const ULearningAgentsActionObject* Object, const FLearningAgentsActionObjectElement Element, const FName Tag = TEXT("StaticArrayAction"));
+	static UE_API bool GetStaticArrayAction(TArray<FLearningAgentsActionObjectElement>& OutElements, const ULearningAgentsActionObject* Object, const FLearningAgentsActionObjectElement Element, const bool bActive = true, const FName Tag = TEXT("StaticArrayAction"));
 	
 	/**
 	 * Get the entries of a static array action. The OutElements ArrayView must be the correct size.
@@ -1472,10 +1563,11 @@ public:
 	 * @param OutElements The output sub-elements.
 	 * @param Object The Action Object
 	 * @param Element The Action Object Element
+	 * @param bActive Whether the action should be considered active. Must match the value given during Specify.
 	 * @param Tag The tag of the corresponding action. Must match the tag given during Specify.
 	 * @return true if the provided Element is the correct type, otherwise false.
 	 */
-	static UE_API bool GetStaticArrayActionToArrayView(TArrayView<FLearningAgentsActionObjectElement> OutElements, const ULearningAgentsActionObject* Object, const FLearningAgentsActionObjectElement Element, const FName Tag = TEXT("StaticArrayAction"));
+	static UE_API bool GetStaticArrayActionToArrayView(TArrayView<FLearningAgentsActionObjectElement> OutElements, const ULearningAgentsActionObject* Object, const FLearningAgentsActionObjectElement Element, const bool bActive = true, const FName Tag = TEXT("StaticArrayAction"));
 
 	/**
 	 * Get the sub-actions of a pair action.
@@ -1484,11 +1576,12 @@ public:
 	 * @param OutValue The output value sub-element.
 	 * @param Object The Action Object
 	 * @param Element The Action Object Element
+	 * @param bActive Whether the action should be considered active. Must match the value given during Specify.
 	 * @param Tag The tag of the corresponding action. Must match the tag given during Specify.
 	 * @return true if the provided Element is the correct type, otherwise false.
 	 */
 	UFUNCTION(BlueprintPure = false, Category = "LearningAgents", meta = (AdvancedDisplay = 4, ReturnDisplayName = "Success"))
-	static UE_API bool GetPairAction(FLearningAgentsActionObjectElement& OutKey, FLearningAgentsActionObjectElement& OutValue, const ULearningAgentsActionObject* Object, const FLearningAgentsActionObjectElement Element, const FName Tag = TEXT("PairAction"));
+	static UE_API bool GetPairAction(FLearningAgentsActionObjectElement& OutKey, FLearningAgentsActionObjectElement& OutValue, const ULearningAgentsActionObject* Object, const FLearningAgentsActionObjectElement Element, const bool bActive = true, const FName Tag = TEXT("PairAction"));
 
 	/**
 	 * Get the enum value of an enum action.
@@ -1497,6 +1590,7 @@ public:
 	 * @param Object The Action Object
 	 * @param Element The Action Object Element
 	 * @param Enum The Enum
+	 * @param bActive Whether the action should be considered active. Must match the value given during Specify.
 	 * @param Tag The tag of the corresponding action. Must match the tag given during Specify.
 	 * @param bVisualLoggerEnabled When true, debug data will be sent to the visual logger.
 	 * @param VisualLoggerListener The listener object which is making this action. This must be set to use logging.
@@ -1511,6 +1605,7 @@ public:
 		const ULearningAgentsActionObject* Object, 
 		const FLearningAgentsActionObjectElement Element, 
 		const UEnum* Enum, 
+		const bool bActive = true, 
 		const FName Tag = TEXT("EnumAction"),
 		const bool bVisualLoggerEnabled = false,
 		ULearningAgentsManagerListener* VisualLoggerListener = nullptr,
@@ -1525,6 +1620,7 @@ public:
 	 * @param Object The Action Object
 	 * @param Element The Action Object Element
 	 * @param Enum The Enum
+	 * @param bActive Whether the action should be considered active. Must match the value given during Specify.
 	 * @param Tag The tag of the corresponding action. Must match the tag given during Specify.
 	 * @param bVisualLoggerEnabled When true, debug data will be sent to the visual logger.
 	 * @param VisualLoggerListener The listener object which is making this action. This must be set to use logging.
@@ -1539,6 +1635,7 @@ public:
 		const ULearningAgentsActionObject* Object, 
 		const FLearningAgentsActionObjectElement Element, 
 		const UEnum* Enum, 
+		const bool bActive = true, 
 		const FName Tag = TEXT("BitmaskAction"),
 		const bool bVisualLoggerEnabled = false,
 		ULearningAgentsManagerListener* VisualLoggerListener = nullptr,
@@ -1553,11 +1650,12 @@ public:
 	 * @param OutElement The output sub-action.
 	 * @param Object The Action Object
 	 * @param Element The Action Object Element
+	 * @param bActive Whether the action should be considered active. Must match the value given during Specify.
 	 * @param Tag The tag of the corresponding action. Must match the tag given during Specify.
 	 * @return true if the provided Element is the correct type, otherwise false.
 	 */
 	UFUNCTION(BlueprintPure = false, Category = "LearningAgents", meta = (AdvancedDisplay = 4, ExpandEnumAsExecs = "OutOption", ReturnDisplayName = "Success"))
-	static UE_API bool GetOptionalAction(ELearningAgentsOptionalAction& OutOption, FLearningAgentsActionObjectElement& OutElement, const ULearningAgentsActionObject* Object, const FLearningAgentsActionObjectElement Element, const FName Tag = TEXT("OptionalAction"));
+	static UE_API bool GetOptionalAction(ELearningAgentsOptionalAction& OutOption, FLearningAgentsActionObjectElement& OutElement, const ULearningAgentsActionObject* Object, const FLearningAgentsActionObjectElement Element, const bool bActive = true, const FName Tag = TEXT("OptionalAction"));
 
 	/**
 	 * Get the sub-action of an either action.
@@ -1566,11 +1664,12 @@ public:
 	 * @param OutElement The output sub-action.
 	 * @param Object The Action Object
 	 * @param Element The Action Object Element
+	 * @param bActive Whether the action should be considered active. Must match the value given during Specify.
 	 * @param Tag The tag of the corresponding action. Must match the tag given during Specify.
 	 * @return true if the provided Element is the correct type, otherwise false.
 	 */
 	UFUNCTION(BlueprintPure = false, Category = "LearningAgents", meta = (AdvancedDisplay = 4, ExpandEnumAsExecs = "OutEither", ReturnDisplayName = "Success"))
-	static UE_API bool GetEitherAction(ELearningAgentsEitherAction& OutEither, FLearningAgentsActionObjectElement& OutElement, const ULearningAgentsActionObject* Object, const FLearningAgentsActionObjectElement Element, const FName Tag = TEXT("EitherAction"));
+	static UE_API bool GetEitherAction(ELearningAgentsEitherAction& OutEither, FLearningAgentsActionObjectElement& OutElement, const ULearningAgentsActionObject* Object, const FLearningAgentsActionObjectElement Element, const bool bActive = true, const FName Tag = TEXT("EitherAction"));
 
 	/**
 	 * Get the sub-action of an encoding action.
@@ -1578,11 +1677,12 @@ public:
 	 * @param OutElement The output sub-action.
 	 * @param Object The Action Object
 	 * @param Element The Action Object Element
+	 * @param bActive Whether the action should be considered active. Must match the value given during Specify.
 	 * @param Tag The tag of the corresponding action. Must match the tag given during Specify.
 	 * @return true if the provided Element is the correct type, otherwise false.
 	 */
 	UFUNCTION(BlueprintPure = false, Category = "LearningAgents", meta = (AdvancedDisplay = 3, ReturnDisplayName = "Success"))
-	static UE_API bool GetEncodingAction(FLearningAgentsActionObjectElement& OutElement, const ULearningAgentsActionObject* Object, const FLearningAgentsActionObjectElement Element, const FName Tag = TEXT("EncodingAction"));
+	static UE_API bool GetEncodingAction(FLearningAgentsActionObjectElement& OutElement, const ULearningAgentsActionObject* Object, const FLearningAgentsActionObjectElement Element, const bool bActive = true, const FName Tag = TEXT("EncodingAction"));
 
 	/**
 	 * Get the value for a bool action.
@@ -1590,6 +1690,7 @@ public:
 	 * @param bOutValue The output bool value.
 	 * @param Object The Action Object
 	 * @param Element The Action Object Element
+	 * @param bActive Whether the action should be considered active. Must match the value given during Specify.
 	 * @param Tag The tag of the corresponding action. Must match the tag given during Specify.
 	 * @param bVisualLoggerEnabled When true, debug data will be sent to the visual logger.
 	 * @param VisualLoggerListener The listener object which is making this action. This must be set to use logging.
@@ -1603,6 +1704,7 @@ public:
 		bool& bOutValue, 
 		const ULearningAgentsActionObject* Object, 
 		const FLearningAgentsActionObjectElement Element, 
+		const bool bActive = true, 
 		const FName Tag = TEXT("BoolAction"),
 		const bool bVisualLoggerEnabled = false,
 		ULearningAgentsManagerListener* VisualLoggerListener = nullptr,
@@ -1616,6 +1718,7 @@ public:
 	 * @param OutValue The output float value.
 	 * @param Object The Action Object
 	 * @param Element The Action Object Element
+	 * @param bActive Whether the action should be considered active. Must match the value given during Specify.
 	 * @param Tag The tag of the corresponding action. Must match the tag given during Specify.
 	 * @param bVisualLoggerEnabled When true, debug data will be sent to the visual logger.
 	 * @param VisualLoggerListener The listener object which is making this action. This must be set to use logging.
@@ -1629,6 +1732,7 @@ public:
 		float& OutValue, 
 		const ULearningAgentsActionObject* Object, 
 		const FLearningAgentsActionObjectElement Element, 
+		const bool bActive = true, 
 		const FName Tag = TEXT("FloatAction"),
 		const bool bVisualLoggerEnabled = false,
 		ULearningAgentsManagerListener* VisualLoggerListener = nullptr,
@@ -1643,6 +1747,7 @@ public:
 	 * @param Object The Action Object
 	 * @param Element The Action Object Element
 	 * @param RelativeTransform The relative transform to transform the location by.
+	 * @param bActive Whether the action should be considered active. Must match the value given during Specify.
 	 * @param Tag The tag of the corresponding action. Must match the tag given during Specify.
 	 * @param bVisualLoggerEnabled When true, debug data will be sent to the visual logger.
 	 * @param VisualLoggerListener The listener object which is making this action. This must be set to use logging.
@@ -1657,6 +1762,7 @@ public:
 		const ULearningAgentsActionObject* Object, 
 		const FLearningAgentsActionObjectElement Element, 
 		const FTransform RelativeTransform = FTransform(), 
+		const bool bActive = true, 
 		const FName Tag = TEXT("LocationAction"),
 		const bool bVisualLoggerEnabled = false,
 		ULearningAgentsManagerListener* VisualLoggerListener = nullptr,
@@ -1671,6 +1777,7 @@ public:
 	 * @param Object The Action Object
 	 * @param Element The Action Object Element
 	 * @param RelativeRotation The relative rotation to transform the rotation by.
+	 * @param bActive Whether the action should be considered active. Must match the value given during Specify.
 	 * @param Tag The tag of the corresponding action. Must match the tag given during Specify.
 	 * @param bVisualLoggerEnabled When true, debug data will be sent to the visual logger.
 	 * @param VisualLoggerListener The listener object which is making this action. This must be set to use logging.
@@ -1686,6 +1793,7 @@ public:
 		const ULearningAgentsActionObject* Object, 
 		const FLearningAgentsActionObjectElement Element, 
 		const FRotator RelativeRotation = FRotator::ZeroRotator, 
+		const bool bActive = true, 
 		const FName Tag = TEXT("RotationAction"),
 		const bool bVisualLoggerEnabled = false,
 		ULearningAgentsManagerListener* VisualLoggerListener = nullptr,
@@ -1701,6 +1809,7 @@ public:
 	 * @param Object The Action Object
 	 * @param Element The Action Object Element
 	 * @param RelativeRotation The relative rotation to transform the rotation by.
+	 * @param bActive Whether the action should be considered active. Must match the value given during Specify.
 	 * @param Tag The tag of the corresponding action. Must match the tag given during Specify.
 	 * @param bVisualLoggerEnabled When true, debug data will be sent to the visual logger.
 	 * @param VisualLoggerListener The listener object which is making this action. This must be set to use logging.
@@ -1716,6 +1825,7 @@ public:
 		const ULearningAgentsActionObject* Object, 
 		const FLearningAgentsActionObjectElement Element, 
 		const FQuat RelativeRotation, 
+		const bool bActive = true, 
 		const FName Tag = TEXT("RotationAction"),
 		const bool bVisualLoggerEnabled = false,
 		ULearningAgentsManagerListener* VisualLoggerListener = nullptr,
@@ -1731,6 +1841,7 @@ public:
 	 * @param Object The Action Object
 	 * @param Element The Action Object Element
 	 * @param RelativeScale The relative scale to transform the scale by.
+	 * @param bActive Whether the action should be considered active. Must match the value given during Specify.
 	 * @param Tag The tag of the corresponding action. Must match the tag given during Specify.
 	 * @param bVisualLoggerEnabled When true, debug data will be sent to the visual logger.
 	 * @param VisualLoggerListener The listener object which is making this action. This must be set to use logging.
@@ -1745,6 +1856,7 @@ public:
 		const ULearningAgentsActionObject* Object, 
 		const FLearningAgentsActionObjectElement Element, 
 		const FVector RelativeScale = FVector(1, 1, 1), 
+		const bool bActive = true, 
 		const FName Tag = TEXT("ScaleAction"),
 		const bool bVisualLoggerEnabled = false,
 		ULearningAgentsManagerListener* VisualLoggerListener = nullptr,
@@ -1759,6 +1871,7 @@ public:
 	 * @param Object The Action Object
 	 * @param Element The Action Object Element
 	 * @param RelativeTransform The relative transform.
+	 * @param bActive Whether the action should be considered active. Must match the value given during Specify.
 	 * @param Tag The tag of the corresponding action. Must match the tag given during Specify.
 	 * @param bVisualLoggerEnabled When true, debug data will be sent to the visual logger.
 	 * @param VisualLoggerListener The listener object which is making this action. This must be set to use logging.
@@ -1773,6 +1886,7 @@ public:
 		const ULearningAgentsActionObject* Object, 
 		const FLearningAgentsActionObjectElement Element, 
 		const FTransform RelativeTransform = FTransform(), 
+		const bool bActive = true, 
 		const FName Tag = TEXT("TransformAction"),
 		const bool bVisualLoggerEnabled = false,
 		ULearningAgentsManagerListener* VisualLoggerListener = nullptr,
@@ -1787,6 +1901,7 @@ public:
 	 * @param Object The Action Object
 	 * @param Element The Action Object Element
 	 * @param RelativeAngle The relative angle to transform the angle by.
+	 * @param bActive Whether the action should be considered active. Must match the value given during Specify.
 	 * @param Tag The tag of the corresponding action. Must match the tag given during Specify.
 	 * @param bVisualLoggerEnabled When true, debug data will be sent to the visual logger.
 	 * @param VisualLoggerListener The listener object which is making this action. This must be set to use logging.
@@ -1802,6 +1917,7 @@ public:
 		const ULearningAgentsActionObject* Object, 
 		const FLearningAgentsActionObjectElement Element, 
 		const float RelativeAngle = 0.0f, 
+		const bool bActive = true, 
 		const FName Tag = TEXT("AngleAction"),
 		const bool bVisualLoggerEnabled = false,
 		ULearningAgentsManagerListener* VisualLoggerListener = nullptr,
@@ -1817,6 +1933,7 @@ public:
 	 * @param Object The Action Object
 	 * @param Element The Action Object Element
 	 * @param RelativeAngle The relative angle to transform the angle by.
+	 * @param bActive Whether the action should be considered active. Must match the value given during Specify.
 	 * @param Tag The tag of the corresponding action. Must match the tag given during Specify.
 	 * @param bVisualLoggerEnabled When true, debug data will be sent to the visual logger.
 	 * @param VisualLoggerListener The listener object which is making this action. This must be set to use logging.
@@ -1832,6 +1949,7 @@ public:
 		const ULearningAgentsActionObject* Object,
 		const FLearningAgentsActionObjectElement Element, 
 		const float RelativeAngle = 0.0f,
+		const bool bActive = true, 
 		const FName Tag = TEXT("AngleAction"),
 		const bool bVisualLoggerEnabled = false,
 		ULearningAgentsManagerListener* VisualLoggerListener = nullptr,
@@ -1847,6 +1965,7 @@ public:
 	 * @param Object The Action Object
 	 * @param Element The Action Object Element
 	 * @param RelativeTransform The relative transform to transform the velocity by.
+	 * @param bActive Whether the action should be considered active. Must match the value given during Specify.
 	 * @param Tag The tag of the corresponding action. Must match the tag given during Specify.
 	 * @param bVisualLoggerEnabled When true, debug data will be sent to the visual logger.
 	 * @param VisualLoggerListener The listener object which is making this action. This must be set to use logging.
@@ -1862,6 +1981,7 @@ public:
 		const ULearningAgentsActionObject* Object, 
 		const FLearningAgentsActionObjectElement Element, 
 		const FTransform RelativeTransform = FTransform(), 
+		const bool bActive = true, 
 		const FName Tag = TEXT("VelocityAction"),
 		const bool bVisualLoggerEnabled = false,
 		ULearningAgentsManagerListener* VisualLoggerListener = nullptr,
@@ -1877,6 +1997,7 @@ public:
 	 * @param Object The Action Object
 	 * @param Element The Action Object Element
 	 * @param RelativeTransform The relative transform to transform the direction by.
+	 * @param bActive Whether the action should be considered active. Must match the value given during Specify.
 	 * @param Tag The tag of the corresponding action. Must match the tag given during Specify.
 	 * @param bVisualLoggerEnabled When true, debug data will be sent to the visual logger.
 	 * @param VisualLoggerListener The listener object which is making this action. This must be set to use logging.
@@ -1893,6 +2014,7 @@ public:
 		const ULearningAgentsActionObject* Object, 
 		const FLearningAgentsActionObjectElement Element, 
 		const FTransform RelativeTransform = FTransform(),
+		const bool bActive = true, 
 		const FName Tag = TEXT("DirectionAction"),
 		const bool bVisualLoggerEnabled = false,
 		ULearningAgentsManagerListener* VisualLoggerListener = nullptr,
